MovieScene: Renamed module; reorganized files.

diff --git a/Engine/Config/BaseEngine.ini b/Engine/Config/BaseEngine.ini
--- a/Engine/Config/BaseEngine.ini
+++ b/Engine/Config/BaseEngine.ini
@@ -478,6 +478,39 @@ bUseStreamingPause=false
 +ActiveClassRedirects=(OldClassName="PathFollowingComponent",NewClassName="/Script/AIModule.PathFollowingComponent")
 +ActiveClassRedirects=(OldClassName="PawnSensingComponent",NewClassName="/Script/AIModule.PawnSensingComponent")
 
+; MovieSceneCore moved to MovieScene
++ActiveClassRedirects=(OldClassName="MovieScene",NewClassName="/Script/MovieScene.MovieScene")
++ActiveClassRedirects=(OldClassName="MovieSceneBindings",NewClassName="/Script/MovieScene.MovieSceneBindings")
++ActiveClassRedirects=(OldClassName="MovieScenePropertyTrack",NewClassName="/Script/MovieScene.MovieScenePropertyTrack")
++ActiveClassRedirects=(OldClassName="MovieSceneSection",NewClassName="/Script/MovieScene.MovieSceneSection")
++ActiveClassRedirects=(OldClassName="MovieSceneTrack",NewClassName="/Script/MovieScene.MovieSceneTrack")
++ActiveClassRedirects=(OldClassName="RuntimeMovieScenePlayer",NewClassName="/Script/MovieScene.RuntimeMovieScenePlayer")
+
+; MovieSceneCoreTypes moved to MovieSceneTracks
++ActiveClassRedirects=(OldClassName="MovieScene3DTransformSection",NewClassName="/Script/MovieSceneTracks.MovieScene3DTransformSection")
++ActiveClassRedirects=(OldClassName="MovieSceneAnimationSection",NewClassName="/Script/MovieSceneTracks.MovieSceneAnimationSection")
++ActiveClassRedirects=(OldClassName="MovieSceneAudioSection",NewClassName="/Script/MovieSceneTracks.MovieSceneAudioSection")
++ActiveClassRedirects=(OldClassName="MovieSceneBoolSection",NewClassName="/Script/MovieSceneTracks.MovieSceneBoolSection")
++ActiveClassRedirects=(OldClassName="MovieSceneByteSection",NewClassName="/Script/MovieSceneTracks.MovieSceneByteSection")
++ActiveClassRedirects=(OldClassName="MovieSceneColorSection",NewClassName="/Script/MovieSceneTracks.MovieSceneColorSection")
++ActiveClassRedirects=(OldClassName="MovieSceneFloatSection",NewClassName="/Script/MovieSceneTracks.MovieSceneFloatSection")
++ActiveClassRedirects=(OldClassName="MovieSceneParticleSection",NewClassName="/Script/MovieSceneTracks.MovieSceneParticleSection")
++ActiveClassRedirects=(OldClassName="MovieSceneShotSection",NewClassName="/Script/MovieSceneTracks.MovieSceneShotSection")
++ActiveClassRedirects=(OldClassName="MovieSceneVectorSection",NewClassName="/Script/MovieSceneTracks.MovieSceneVectorSection")
++ActiveClassRedirects=(OldClassName="SubMovieSceneSection",NewClassName="/Script/MovieSceneTracks.SubMovieSceneSection")
+
++ActiveClassRedirects=(OldClassName="MovieScene3DTransformTrack",NewClassName="/Script/MovieSceneTracks.MovieScene3DTransformTrack")
++ActiveClassRedirects=(OldClassName="MovieSceneAnimationTrack",NewClassName="/Script/MovieSceneTracks.MovieSceneAnimationTrack")
++ActiveClassRedirects=(OldClassName="MovieSceneAudioTrack",NewClassName="/Script/MovieSceneTracks.MovieSceneAudioTrack")
++ActiveClassRedirects=(OldClassName="MovieSceneBoolTrack",NewClassName="/Script/MovieSceneTracks.MovieSceneBoolTrack")
++ActiveClassRedirects=(OldClassName="MovieSceneByteTrack",NewClassName="/Script/MovieSceneTracks.MovieSceneByteTrack")
++ActiveClassRedirects=(OldClassName="MovieSceneColorTrack",NewClassName="/Script/MovieSceneTracks.MovieSceneColorTrack")
++ActiveClassRedirects=(OldClassName="MovieSceneDirectorTrack",NewClassName="/Script/MovieSceneTracks.MovieSceneDirectorTrack")
++ActiveClassRedirects=(OldClassName="MovieSceneFloatTrack",NewClassName="/Script/MovieSceneTracks.MovieSceneFloatTrack")
++ActiveClassRedirects=(OldClassName="MovieSceneParticleTrack",NewClassName="/Script/MovieSceneTracks.MovieSceneParticleTrack")
++ActiveClassRedirects=(OldClassName="MovieSceneVectorTrack",NewClassName="/Script/MovieSceneTracks.MovieSceneVectorTrack")
++ActiveClassRedirects=(OldClassName="SubMovieSceneTrack",NewClassName="/Script/MovieSceneTracks.SubMovieSceneTrack")
+
 ; We renamed EditorGameAgnosticSettings to be less imposing
 +ActiveClassRedirects=(OldClassName="EditorGameAgnosticSettings",NewClassName="/Script/UnrealEd.EditorSettings")
 +ActiveClassRedirects=(OldClassName="EditorUserSettings",NewClassName="/Script/UnrealEd.EditorPerProjectUserSettings")
diff --git a/Engine/Source/Editor/KismetCompiler/KismetCompiler.Build.cs b/Engine/Source/Editor/KismetCompiler/KismetCompiler.Build.cs
--- a/Engine/Source/Editor/KismetCompiler/KismetCompiler.Build.cs
+++ b/Engine/Source/Editor/KismetCompiler/KismetCompiler.Build.cs
@@ -13,7 +13,7 @@ public class KismetCompiler : ModuleRules
 				"Engine",
                 "InputCore",
 				"UnrealEd",
-				"MovieSceneCore",
+				"MovieScene",
 				"MovieSceneTools",
 				"FunctionalTesting",
 				"BlueprintGraph",
diff --git a/Engine/Source/Editor/MovieSceneTools/MovieSceneTools.Build.cs b/Engine/Source/Editor/MovieSceneTools/MovieSceneTools.Build.cs
--- a/Engine/Source/Editor/MovieSceneTools/MovieSceneTools.Build.cs
+++ b/Engine/Source/Editor/MovieSceneTools/MovieSceneTools.Build.cs
@@ -27,7 +27,7 @@ public class MovieSceneTools : ModuleRules
 
 		PrivateDependencyModuleNames.AddRange(
 			new string[] {
-				"MovieSceneCore",
+				"MovieScene",
 				"MovieSceneTracks",
 				"BlueprintGraph",
 				"Slate",
diff --git a/Engine/Source/Editor/NiagaraEditor/NiagaraEditor.Build.cs b/Engine/Source/Editor/NiagaraEditor/NiagaraEditor.Build.cs
--- a/Engine/Source/Editor/NiagaraEditor/NiagaraEditor.Build.cs
+++ b/Engine/Source/Editor/NiagaraEditor/NiagaraEditor.Build.cs
@@ -21,7 +21,7 @@ public class NiagaraEditor : ModuleRules
 				"GraphEditor", 
 				"VectorVM",
                 "Niagara",
-                "MovieSceneCore",
+                "MovieScene",
 				"Sequencer",
 				"PropertyEditor",
 			}
diff --git a/Engine/Source/Editor/NiagaraEditor/Private/NiagaraSequencer.h b/Engine/Source/Editor/NiagaraEditor/Private/NiagaraSequencer.h
--- a/Engine/Source/Editor/NiagaraEditor/Private/NiagaraSequencer.h
+++ b/Engine/Source/Editor/NiagaraEditor/Private/NiagaraSequencer.h
@@ -5,10 +5,10 @@
 #include "Editor/Sequencer/Public/ISequencerObjectBindingManager.h"
 #include "Editor/Sequencer/Public/MovieSceneTrackEditor.h"
 #include "Editor/Sequencer/Public/ISequencerSection.h"
-#include "Runtime/MovieSceneCore/Classes/MovieScene.h"
-#include "Runtime/MovieSceneCore/Classes/MovieSceneTrack.h"
-#include "Runtime/MovieSceneCore/Public/IMovieSceneTrackInstance.h"
-#include "Runtime/MovieSceneCore/Classes/MovieSceneSection.h"
+#include "Runtime/MovieScene/Public/MovieScene.h"
+#include "Runtime/MovieScene/Public/MovieSceneTrack.h"
+#include "Runtime/MovieScene/Public/IMovieSceneTrackInstance.h"
+#include "Runtime/MovieScene/Public/MovieSceneSection.h"
 
 #include "NiagaraSequencer.generated.h"
 
diff --git a/Engine/Source/Editor/Sequencer/Public/ISequencer.h b/Engine/Source/Editor/Sequencer/Public/ISequencer.h
--- a/Engine/Source/Editor/Sequencer/Public/ISequencer.h
+++ b/Engine/Source/Editor/Sequencer/Public/ISequencer.h
@@ -3,7 +3,7 @@
 #pragma once
 
 
-#include "Runtime/MovieSceneCore/Public/IMovieScenePlayer.h"
+#include "Runtime/MovieScene/Public/IMovieScenePlayer.h"
 #include "KeyPropertyParams.h"
 
 class UMovieScene;
diff --git a/Engine/Source/Editor/Sequencer/Sequencer.Build.cs b/Engine/Source/Editor/Sequencer/Sequencer.Build.cs
--- a/Engine/Source/Editor/Sequencer/Sequencer.Build.cs
+++ b/Engine/Source/Editor/Sequencer/Sequencer.Build.cs
@@ -18,7 +18,7 @@ public class Sequencer : ModuleRules
 				"SlateCore",
                 "EditorStyle",
 				"UnrealEd", 
-				"MovieSceneCore", 
+				"MovieScene", 
 				"MovieSceneTracks", 
 				"MovieSceneTools", 
 				"EditorWidgets", 
diff --git a/Engine/Source/Editor/UMGEditor/Private/WidgetBlueprint.cpp b/Engine/Source/Editor/UMGEditor/Private/WidgetBlueprint.cpp
--- a/Engine/Source/Editor/UMGEditor/Private/WidgetBlueprint.cpp
+++ b/Engine/Source/Editor/UMGEditor/Private/WidgetBlueprint.cpp
@@ -2,7 +2,7 @@
 
 #include "UMGEditorPrivatePCH.h"
 
-#include "Runtime/MovieSceneCore/Classes/MovieScene.h"
+#include "Runtime/MovieScene/Public/MovieScene.h"
 #include "Editor/UnrealEd/Public/Kismet2/StructureEditorUtils.h"
 
 #include "Blueprint/WidgetTree.h"
diff --git a/Engine/Source/Editor/UMGEditor/UMGEditor.Build.cs b/Engine/Source/Editor/UMGEditor/UMGEditor.Build.cs
--- a/Engine/Source/Editor/UMGEditor/UMGEditor.Build.cs
+++ b/Engine/Source/Editor/UMGEditor/UMGEditor.Build.cs
@@ -47,7 +47,7 @@ public class UMGEditor : ModuleRules
 				"EditorStyle",
 				"Slate",
 				"SlateCore",
-				"MovieSceneCore",
+				"MovieScene",
 				"MovieSceneTools",
 				"Sequencer",
 				"DetailCustomizations",
diff --git a/Engine/Source/Runtime/Engine/Engine.Build.cs b/Engine/Source/Runtime/Engine/Engine.Build.cs
--- a/Engine/Source/Runtime/Engine/Engine.Build.cs
+++ b/Engine/Source/Runtime/Engine/Engine.Build.cs
@@ -85,7 +85,7 @@ public class Engine : ModuleRules
 		DynamicallyLoadedModuleNames.AddRange(
 			new string[]
 			{
-				"MovieSceneCore",
+				"MovieScene",
 				"MovieSceneTracks",
 				"HeadMountedDisplay",
 				"StreamingPauseRendering",
diff --git a/Engine/Source/Runtime/Engine/Private/Level.cpp b/Engine/Source/Runtime/Engine/Private/Level.cpp
--- a/Engine/Source/Runtime/Engine/Private/Level.cpp
+++ b/Engine/Source/Runtime/Engine/Private/Level.cpp
@@ -1714,7 +1714,7 @@ void ULevel::AddMovieSceneBindings( class UMovieSceneBindings* MovieSceneBinding
 	{
 		Modify();
 
-		// @todo sequencer UObjects: Dangerous cast here to work around MovieSceneCore not being a dependency module of engine
+		// @todo sequencer UObjects: Dangerous cast here to work around MovieScene not being a dependency module of engine
 		UObject* ObjectToAdd = (UObject*)MovieSceneBindings;
 		MovieSceneBindingsArray.AddUnique( ObjectToAdd );
 	}
diff --git a/Engine/Source/Runtime/MovieScene/MovieScene.Build.cs b/Engine/Source/Runtime/MovieScene/MovieScene.Build.cs
new file mode 100644
--- /dev/null
+++ b/Engine/Source/Runtime/MovieScene/MovieScene.Build.cs
@@ -0,0 +1,19 @@
+// Copyright 1998-2015 Epic Games, Inc. All Rights Reserved.
+
+using UnrealBuildTool;
+
+public class MovieScene : ModuleRules
+{
+	public MovieScene(TargetInfo Target)
+	{
+		PrivateIncludePaths.Add("Runtime/MovieScene/Private");
+
+		PublicDependencyModuleNames.AddRange(
+            new string[] {
+                "Core",
+                "CoreUObject",
+                "Engine"
+            }
+        );
+	}
+}
diff --git a/Engine/Source/Runtime/MovieScene/Private/MovieScene.cpp b/Engine/Source/Runtime/MovieScene/Private/MovieScene.cpp
new file mode 100644
--- /dev/null
+++ b/Engine/Source/Runtime/MovieScene/Private/MovieScene.cpp
@@ -0,0 +1,387 @@
+// Copyright 1998-2015 Epic Games, Inc. All Rights Reserved.
+
+#include "MovieScenePrivatePCH.h"
+#include "MovieScene.h"
+
+
+FMovieSceneSpawnable::FMovieSceneSpawnable( const FString& InitName, UClass* InitClass, UObject* InitCounterpartGamePreviewObject )
+{
+	Guid = FGuid::NewGuid();
+	Name = InitName;
+	GeneratedClass = InitClass;
+	CounterpartGamePreviewObject = InitCounterpartGamePreviewObject;
+}
+
+
+
+FMovieScenePossessable::FMovieScenePossessable( const FString& InitName, UClass* InitPossessedObjectClass )
+{
+	Guid = FGuid::NewGuid();
+	Name = InitName;
+	PossessedObjectClass = InitPossessedObjectClass;
+}
+
+
+TRange<float> FMovieSceneObjectBinding::GetTimeRange() const
+{
+	TArray< TRange<float> > Bounds;
+	for (int32 TypeIndex = 0; TypeIndex < Tracks.Num(); ++TypeIndex)
+	{
+		Bounds.Add(Tracks[TypeIndex]->GetSectionBoundaries());
+	}
+	return TRange<float>::Hull(Bounds);
+}
+
+void FMovieSceneObjectBinding::AddTrack( UMovieSceneTrack& NewTrack )
+{
+	Tracks.Add( &NewTrack );
+}
+
+bool FMovieSceneObjectBinding::RemoveTrack( UMovieSceneTrack& Track )
+{
+	return Tracks.RemoveSingle( &Track ) != 0;
+}
+
+UMovieScene::UMovieScene( const FObjectInitializer& ObjectInitializer )
+	: Super( ObjectInitializer )
+{
+}
+
+#if WITH_EDITOR
+// @todo sequencer: Some of these methods should only be used by tools, and should probably move out of MovieScene!
+FGuid UMovieScene::AddSpawnable( const FString& Name, UBlueprint* Blueprint, UObject* CounterpartGamePreviewObject )
+{
+	check( (Blueprint != NULL) && (Blueprint->GeneratedClass) );
+
+	Modify();
+
+	FMovieSceneSpawnable NewSpawnable( Name, Blueprint->GeneratedClass, CounterpartGamePreviewObject );
+	Spawnables.Add( NewSpawnable );
+
+	// Add a new binding so that tracks can be added to it
+	new (ObjectBindings) FMovieSceneObjectBinding( NewSpawnable.GetGuid(), NewSpawnable.GetName() );
+
+	return NewSpawnable.GetGuid();
+}
+
+bool UMovieScene::RemoveSpawnable( const FGuid& Guid )
+{
+	bool bAnythingRemoved = false;
+	if( ensure( Guid.IsValid() ) )
+	{
+		for( auto SpawnableIter( Spawnables.CreateIterator() ); SpawnableIter; ++SpawnableIter )
+		{
+			auto& CurSpawnable = *SpawnableIter;
+			if( CurSpawnable.GetGuid() == Guid )
+			{
+				Modify();
+
+				{
+					UClass* GeneratedClass = CurSpawnable.GetClass();
+					UBlueprint* Blueprint = GeneratedClass ? Cast<UBlueprint>(GeneratedClass->ClassGeneratedBy) : NULL;
+					check(NULL != Blueprint);
+					// @todo sequencer: Also remove created Blueprint inner object.  Is this sufficient?  Needs to work with Undo too!
+					Blueprint->ClearFlags( RF_Standalone );	// @todo sequencer: Probably not needed for Blueprint
+					Blueprint->MarkPendingKill();
+				}
+
+				RemoveObjectBinding( Guid );
+
+				// Found it!
+				Spawnables.RemoveAt( SpawnableIter.GetIndex() );
+
+
+				bAnythingRemoved = true;
+				break;
+			}
+		}
+	}
+
+	return bAnythingRemoved;
+}
+#endif //WITH_EDITOR
+
+
+int32 UMovieScene::GetSpawnableCount() const
+{
+	return Spawnables.Num();
+}
+
+
+FMovieSceneSpawnable& UMovieScene::GetSpawnable( const int32 Index )
+{
+	return Spawnables[ Index ];
+}
+
+
+FMovieSceneSpawnable* UMovieScene::FindSpawnable( const FGuid& Guid )
+{
+	for( auto CurSpawnableIt( Spawnables.CreateIterator() ); CurSpawnableIt; ++CurSpawnableIt )
+	{
+		auto& CurSpawnable = *CurSpawnableIt;
+		if( CurSpawnable.GetGuid() == Guid )
+		{
+			return &CurSpawnable;
+		}
+	}
+
+	return NULL;
+}
+
+
+const FMovieSceneSpawnable* UMovieScene::FindSpawnableForCounterpart( UObject* GamePreviewObject ) const
+{
+	check( GamePreviewObject != NULL );
+
+	// Must only be called for objects in game preview worlds
+	const bool bIsGamePreviewObject = !!( GamePreviewObject->GetOutermost()->PackageFlags & PKG_PlayInEditor );
+	check( bIsGamePreviewObject );
+
+	for( auto CurSpawnableIt( Spawnables.CreateConstIterator() ); CurSpawnableIt; ++CurSpawnableIt )
+	{
+		auto& CurSpawnable = *CurSpawnableIt;
+		if( CurSpawnable.GetCounterpartGamePreviewObject() == GamePreviewObject )
+		{
+			return &CurSpawnable;
+		}
+	}
+
+	return NULL;
+}
+
+
+FGuid UMovieScene::AddPossessable( const FString& Name, UClass* Class )
+{
+	Modify();
+
+	FMovieScenePossessable NewPossessable( Name, Class );
+	Possessables.Add( NewPossessable );
+
+	// Add a new binding so that tracks can be added to it
+	new (ObjectBindings) FMovieSceneObjectBinding( NewPossessable.GetGuid(), NewPossessable.GetName() );
+
+	return NewPossessable.GetGuid();
+}
+
+
+bool UMovieScene::RemovePossessable( const FGuid& PossessableGuid )
+{
+	bool bAnythingRemoved = false;
+
+	for( auto PossesableIter( Possessables.CreateIterator() ); PossesableIter; ++PossesableIter )
+	{
+		auto& CurPossesable = *PossesableIter;
+
+		if( CurPossesable.GetGuid() == PossessableGuid )
+		{	
+			Modify();
+
+			// Found it!
+			Possessables.RemoveAt( PossesableIter.GetIndex() );
+
+			RemoveObjectBinding( PossessableGuid );
+
+			bAnythingRemoved = true;
+			break;
+		}
+	}
+
+	return bAnythingRemoved;
+}
+
+
+FMovieScenePossessable* UMovieScene::FindPossessable( const FGuid& Guid )
+{
+	for( auto CurPossessableIt( Possessables.CreateIterator() ); CurPossessableIt; ++CurPossessableIt )
+	{
+		auto& CurPossessable = *CurPossessableIt;
+		if( CurPossessable.GetGuid() == Guid )
+		{
+			return &CurPossessable;
+		}
+	}
+
+	return NULL;
+}
+
+
+int32 UMovieScene::GetPossessableCount() const
+{
+	return Possessables.Num();
+}
+
+
+FMovieScenePossessable& UMovieScene::GetPossessable( const int32 Index )
+{
+	return Possessables[ Index ];
+}
+
+TRange<float> UMovieScene::GetTimeRange() const
+{
+	// Get the range of all sections combined
+	TArray< TRange<float> > Bounds;
+	for (int32 TypeIndex = 0; TypeIndex < MasterTracks.Num(); ++TypeIndex)
+	{
+		Bounds.Add(MasterTracks[TypeIndex]->GetSectionBoundaries());
+	}
+	for (int32 BindingIndex = 0; BindingIndex < ObjectBindings.Num(); ++BindingIndex)
+	{
+		Bounds.Add(ObjectBindings[BindingIndex].GetTimeRange());
+	}
+	return TRange<float>::Hull(Bounds);
+}
+
+TArray<UMovieSceneSection*> UMovieScene::GetAllSections() const
+{
+	TArray<UMovieSceneSection*> OutSections;
+
+	// Add all master type sections 
+	for( int32 TrackIndex = 0; TrackIndex < MasterTracks.Num(); ++TrackIndex )
+	{
+		OutSections.Append( MasterTracks[TrackIndex]->GetAllSections() );
+	}
+
+	// Add all object binding sections
+	for( int32 ObjectBindingIndex = 0; ObjectBindingIndex < ObjectBindings.Num(); ++ObjectBindingIndex )
+	{
+		const FMovieSceneObjectBinding& ObjectBinding = ObjectBindings[ObjectBindingIndex];
+		const TArray<UMovieSceneTrack*>& Tracks = ObjectBinding.GetTracks();
+
+		for( int32 TrackIndex = 0; TrackIndex < Tracks.Num(); ++TrackIndex )
+		{
+			OutSections.Append( Tracks[TrackIndex]->GetAllSections() );
+		}
+	}
+
+	return OutSections;
+}
+
+void UMovieScene::RemoveObjectBinding( const FGuid& Guid )
+{
+	// Update each type
+	for( int32 BindingIndex = 0; BindingIndex < ObjectBindings.Num(); ++BindingIndex )
+	{
+		if( ObjectBindings[BindingIndex].GetObjectGuid() == Guid )
+		{
+			ObjectBindings.RemoveAt( BindingIndex );
+			break;
+		}
+	}
+}
+
+UMovieSceneTrack* UMovieScene::FindTrack( TSubclassOf<UMovieSceneTrack> TrackClass, const FGuid& ObjectGuid, FName UniqueTrackName ) const
+{
+	UMovieSceneTrack* FoundTrack = NULL;
+	
+	check( UniqueTrackName != NAME_None );
+	check( ObjectGuid.IsValid() );
+	
+	for( int32 ObjectBindingIndex = 0; ObjectBindingIndex < ObjectBindings.Num(); ++ObjectBindingIndex )
+	{
+		const FMovieSceneObjectBinding& ObjectBinding = ObjectBindings[ObjectBindingIndex];
+
+		if( ObjectBinding.GetObjectGuid() == ObjectGuid ) 
+		{
+			const TArray<UMovieSceneTrack*>& Tracks = ObjectBinding.GetTracks();
+			for( int32 TrackIndex = 0; TrackIndex < Tracks.Num(); ++TrackIndex )
+			{
+				UMovieSceneTrack* Track = Tracks[TrackIndex];
+				if( Track->GetClass() == TrackClass && Track->GetTrackName() == UniqueTrackName )
+				{
+					FoundTrack = Track;
+					break;
+				}
+			}
+		}
+	}
+	
+	return FoundTrack;
+}
+
+class UMovieSceneTrack* UMovieScene::AddTrack( TSubclassOf<UMovieSceneTrack> TrackClass, const FGuid& ObjectGuid )
+{
+	UMovieSceneTrack* CreatedType = NULL;
+
+	check( ObjectGuid.IsValid() )
+
+	for( int32 ObjectBindingIndex = 0; ObjectBindingIndex < ObjectBindings.Num(); ++ObjectBindingIndex )
+	{
+		FMovieSceneObjectBinding& ObjectBinding = ObjectBindings[ObjectBindingIndex];
+
+		if( ObjectBinding.GetObjectGuid() == ObjectGuid ) 
+		{
+			Modify();
+
+			CreatedType = NewObject<UMovieSceneTrack>(this, TrackClass, NAME_None, RF_Transactional);
+			
+			ObjectBinding.AddTrack( *CreatedType );
+		}
+	}
+
+	return CreatedType;
+}
+
+bool UMovieScene::RemoveTrack( UMovieSceneTrack* Track )
+{
+	Modify();
+	bool bAnythingRemoved = false;
+	for( int32 ObjectBindingIndex = 0; ObjectBindingIndex < ObjectBindings.Num(); ++ObjectBindingIndex )
+	{
+		FMovieSceneObjectBinding& ObjectBinding = ObjectBindings[ObjectBindingIndex];
+
+		if( ObjectBinding.RemoveTrack( *Track ) )
+		{
+			bAnythingRemoved = true;
+			// The track was removed from the current binding, stop searching now as it cannot exist in any other binding
+			break;
+		}
+	}
+
+	return bAnythingRemoved;
+}
+
+UMovieSceneTrack* UMovieScene::FindMasterTrack( TSubclassOf<UMovieSceneTrack> TrackClass ) const
+{
+	UMovieSceneTrack* FoundTrack = NULL;
+	for( int32 TrackIndex = 0; TrackIndex < MasterTracks.Num(); ++TrackIndex )
+	{
+		UMovieSceneTrack* Track = MasterTracks[TrackIndex];
+		if( Track->GetClass() == TrackClass )
+		{
+			FoundTrack = Track;
+			break;
+		}
+	}
+
+	return FoundTrack;
+}
+
+class UMovieSceneTrack* UMovieScene::AddMasterTrack( TSubclassOf<UMovieSceneTrack> TrackClass )
+{
+	Modify();
+
+	UMovieSceneTrack* CreatedType = NewObject<UMovieSceneTrack>(this, TrackClass, NAME_None, RF_Transactional);
+
+	MasterTracks.Add( CreatedType );
+	
+	return CreatedType;
+}
+
+bool UMovieScene::RemoveMasterTrack( UMovieSceneTrack* Track ) 
+{
+	Modify();
+
+	return MasterTracks.RemoveSingle( Track ) != 0;
+}
+
+bool UMovieScene::IsAMasterTrack(const UMovieSceneTrack* Track) const
+{
+	for( int32 TrackIndex = 0; TrackIndex < MasterTracks.Num(); ++TrackIndex )
+	{
+		if (Track == MasterTracks[TrackIndex])
+		{
+			return true;
+		}
+	}
+	return false;
+}
diff --git a/Engine/Source/Runtime/MovieScene/Private/MovieSceneBindings.cpp b/Engine/Source/Runtime/MovieScene/Private/MovieSceneBindings.cpp
new file mode 100644
--- /dev/null
+++ b/Engine/Source/Runtime/MovieScene/Private/MovieSceneBindings.cpp
@@ -0,0 +1,66 @@
+// Copyright 1998-2015 Epic Games, Inc. All Rights Reserved.
+
+#include "MovieScenePrivatePCH.h"
+#include "MovieSceneBindings.h"
+
+FMovieSceneBoundObject::FMovieSceneBoundObject( const FGuid& InitPosessableGuid, const TArray< UObject* >& InitObjects )
+	: PossessableGuid( InitPosessableGuid ),
+	  Objects( InitObjects )
+{
+}
+
+
+
+UMovieSceneBindings::UMovieSceneBindings( const FObjectInitializer& ObjectInitializer )
+	: Super( ObjectInitializer )
+{
+
+	RootMovieScene = NULL;
+}
+
+
+
+UMovieScene* UMovieSceneBindings::GetRootMovieScene()
+{
+	return RootMovieScene;
+}
+
+
+void UMovieSceneBindings::SetRootMovieScene( class UMovieScene* NewMovieScene )
+{
+	if( NewMovieScene != RootMovieScene )
+	{
+		Modify();
+		RootMovieScene = NewMovieScene;
+	}
+}
+
+
+FMovieSceneBoundObject& UMovieSceneBindings::AddBinding( const FGuid& PosessableGuid, const TArray< UObject* >& Objects )
+{
+	Modify();
+	BoundObjects.Add( FMovieSceneBoundObject( PosessableGuid, Objects ) );
+	return BoundObjects.Last();
+}
+
+
+TArray< UObject* > UMovieSceneBindings::FindBoundObjects( const FGuid& Guid ) const
+{
+	for( TArray< FMovieSceneBoundObject >::TConstIterator BoundObjectIter( BoundObjects.CreateConstIterator() ); BoundObjectIter; ++BoundObjectIter )
+	{
+		const FMovieSceneBoundObject& BoundObject = *BoundObjectIter;
+
+		if( BoundObject.GetPossessableGuid() == Guid )
+		{
+			return BoundObject.GetObjects();
+		}
+	}
+
+	return TArray< UObject* >();
+}
+
+
+TArray< FMovieSceneBoundObject >& UMovieSceneBindings::GetBoundObjects()
+{
+	return BoundObjects;
+}
diff --git a/Engine/Source/Runtime/MovieScene/Private/MovieSceneCommonHelpers.cpp b/Engine/Source/Runtime/MovieScene/Private/MovieSceneCommonHelpers.cpp
new file mode 100644
--- /dev/null
+++ b/Engine/Source/Runtime/MovieScene/Private/MovieSceneCommonHelpers.cpp
@@ -0,0 +1,187 @@
+// Copyright 1998-2015 Epic Games, Inc. All Rights Reserved.
+
+#include "MovieScenePrivatePCH.h"
+#include "IMovieScenePlayer.h"
+#include "MovieSceneCommonHelpers.h"
+
+
+TArray<UMovieSceneSection*> MovieSceneHelpers::GetTraversedSections( const TArray<UMovieSceneSection*>& Sections, float CurrentTime, float PreviousTime )
+{
+	TArray<UMovieSceneSection*> TraversedSections;
+
+	bool bPlayingBackwards = CurrentTime - PreviousTime < 0.0f;
+	
+	float MaxTime = bPlayingBackwards ? PreviousTime : CurrentTime;
+	float MinTime = bPlayingBackwards ? CurrentTime : PreviousTime;
+
+	TRange<float> SliderRange(MinTime, MaxTime);
+
+	for( int32 SectionIndex = 0; SectionIndex < Sections.Num(); ++SectionIndex )
+	{
+		UMovieSceneSection* Section = Sections[SectionIndex];
+		if( SliderRange.Overlaps( TRange<float>( Section->GetRange() ) ) )
+		{
+			TraversedSections.Add( Section );
+		}
+	}
+
+	return TraversedSections;
+}
+
+UMovieSceneSection* MovieSceneHelpers::FindSectionAtTime( const TArray<UMovieSceneSection*>& Sections, float Time )
+{
+	for( int32 SectionIndex = 0; SectionIndex < Sections.Num(); ++SectionIndex )
+	{
+		UMovieSceneSection* Section = Sections[SectionIndex];
+
+		//@todo Sequencer - There can be multiple sections overlapping in time. Returning instantly does not account for that.
+		if( Section->IsTimeWithinSection( Time ) )
+		{
+			return Section;
+		}
+	}
+
+	return NULL;
+}
+
+UMovieSceneSection* MovieSceneHelpers::FindNearestSectionAtTime( const TArray<UMovieSceneSection*>& Sections, float Time )
+{
+	// Only update the section if the position is within the time span of the section
+	// Or if there are no sections at the time, the left closest section to the time
+	// Or in the case that Time is before all sections, take the section with the earliest start time
+	UMovieSceneSection* ClosestSection = NULL;
+	float ClosestSectionTime = 0.f;
+	UMovieSceneSection* EarliestSection = NULL;
+	float EarliestSectionTime = 0.f;
+	for( int32 SectionIndex = 0; SectionIndex < Sections.Num(); ++SectionIndex )
+	{
+		UMovieSceneSection* Section = Sections[SectionIndex];
+
+		//@todo Sequencer - There can be multiple sections overlapping in time. Returning instantly does not account for that.
+		if( Section->IsTimeWithinSection( Time ) )
+		{
+			return Section;
+		}
+
+		float EndTime = Section->GetEndTime();
+		if (EndTime < Time)
+		{
+			float ClosestTime = Time - EndTime;
+			if (!ClosestSection || ClosestTime < ClosestSectionTime)
+			{
+				ClosestSection = Section;
+				ClosestSectionTime = ClosestTime;
+			}
+		}
+
+		float StartTime = Section->GetStartTime();
+		if (!EarliestSection || StartTime < EarliestSectionTime)
+		{
+			EarliestSection = Section;
+			EarliestSectionTime = StartTime;
+		}
+	}
+
+	// if we get here, we are off of any section
+	// if ClosestSection, then we take the closest to left of this time
+	// else, we take the EarliestSection
+	// if that's NULL, then there are no sections
+	return ClosestSection ? ClosestSection : EarliestSection;
+}
+
+FTrackInstancePropertyBindings::FTrackInstancePropertyBindings( FName InPropertyName, const FString& InPropertyPath, const FName& InFunctionName )
+    : PropertyPath( InPropertyPath )
+	, PropertyName( InPropertyName )
+{
+	if (InFunctionName != FName())
+	{
+		FunctionName = InFunctionName;
+	}
+	else
+	{
+		static const FString Set(TEXT("Set"));
+
+		const FString FunctionString = Set + PropertyName.ToString();
+
+		FunctionName = FName(*FunctionString);
+	}
+}
+
+void FTrackInstancePropertyBindings::CallFunction( UObject* InRuntimeObject, void* FunctionParams )
+{
+	FPropertyAndFunction PropAndFunction = RuntimeObjectToFunctionMap.FindRef(InRuntimeObject);
+	if(PropAndFunction.Function)
+	{
+		InRuntimeObject->ProcessEvent(PropAndFunction.Function, FunctionParams);
+	}
+}
+
+
+FTrackInstancePropertyBindings::FPropertyAddress FTrackInstancePropertyBindings::FindPropertyRecursive( UObject* Object, void* BasePointer, UStruct* InStruct, TArray<FString>& InPropertyNames, uint32 Index ) const
+{
+	UProperty* Property = FindField<UProperty>(InStruct, *InPropertyNames[Index]);
+	
+	FTrackInstancePropertyBindings::FPropertyAddress NewAddress;
+
+	UStructProperty* StructProp = Cast<UStructProperty>( Property );
+	if( StructProp )
+	{
+		NewAddress.Property = StructProp;
+		NewAddress.Address = BasePointer;
+
+		if( InPropertyNames.IsValidIndex(Index+1) )
+		{
+			void* StructContainer = StructProp->ContainerPtrToValuePtr<void>(BasePointer);
+			return FindPropertyRecursive( Object, StructContainer, StructProp->Struct, InPropertyNames, Index+1 );
+		}
+		else
+		{
+			check( StructProp->GetName() == InPropertyNames[Index] );
+		}
+	}
+	else if( Property )
+	{
+		NewAddress.Property = Property;
+		NewAddress.Address = BasePointer;
+	}
+
+	return NewAddress;
+
+}
+
+FTrackInstancePropertyBindings::FPropertyAddress FTrackInstancePropertyBindings::FindProperty( UObject* InObject, const FString& InPropertyPath ) const
+{
+	TArray<FString> PropertyNames;
+
+	InPropertyPath.ParseIntoArray(PropertyNames, TEXT("."), true);
+
+	if( PropertyNames.Num() > 0 )
+	{
+		return FindPropertyRecursive( InObject, (void*)InObject, InObject->GetClass(), PropertyNames, 0 );
+	}
+	else
+	{
+		return FTrackInstancePropertyBindings::FPropertyAddress();
+	}
+
+}
+
+void FTrackInstancePropertyBindings::UpdateBindings( const TArray<UObject*>& InRuntimeObjects )
+{
+	for(UObject* Object : InRuntimeObjects)
+	{
+		FPropertyAndFunction PropAndFunction;
+
+		PropAndFunction.Function = Object->FindFunction(FunctionName);
+		if(PropAndFunction.Function)
+		{
+			PropAndFunction.PropertyAddress = FindProperty( Object, PropertyPath );
+			RuntimeObjectToFunctionMap.Add(Object, PropAndFunction);
+		}
+		else
+		{
+			// Dont call potentially invalid functions
+			RuntimeObjectToFunctionMap.Remove(Object);
+		}
+	}
+}
diff --git a/Engine/Source/Runtime/MovieScene/Private/MovieSceneCoreModule.cpp b/Engine/Source/Runtime/MovieScene/Private/MovieSceneCoreModule.cpp
new file mode 100644
--- /dev/null
+++ b/Engine/Source/Runtime/MovieScene/Private/MovieSceneCoreModule.cpp
@@ -0,0 +1,39 @@
+// Copyright 1998-2015 Epic Games, Inc. All Rights Reserved.
+
+#include "MovieScenePrivatePCH.h"
+#include "ModuleManager.h"
+
+#include "MovieScene.h"
+#include "MovieSceneBindings.h"
+#include "RuntimeMovieScenePlayer.h"
+
+DEFINE_LOG_CATEGORY(LogSequencerRuntime);
+
+
+/**
+ * MovieScene module implementation (private)
+ */
+class FMovieSceneModule : public IMovieSceneModule
+{
+
+public:
+
+	/** IModuleInterface */
+	virtual void StartupModule() override;
+	virtual void ShutdownModule() override;
+
+};
+
+
+IMPLEMENT_MODULE( FMovieSceneModule, MovieScene );
+
+
+
+void FMovieSceneModule::StartupModule()
+{
+}
+
+
+void FMovieSceneModule::ShutdownModule()
+{
+}
\ No newline at end of file
diff --git a/Engine/Source/Runtime/MovieScene/Private/MovieSceneInstance.cpp b/Engine/Source/Runtime/MovieScene/Private/MovieSceneInstance.cpp
new file mode 100644
--- /dev/null
+++ b/Engine/Source/Runtime/MovieScene/Private/MovieSceneInstance.cpp
@@ -0,0 +1,146 @@
+// Copyright 1998-2015 Epic Games, Inc. All Rights Reserved.
+
+#include "MovieScenePrivatePCH.h"
+#include "MovieSceneInstance.h"
+
+FMovieSceneInstance::FMovieSceneInstance( UMovieScene& InMovieScene )
+	: MovieScene( &InMovieScene )
+{
+}
+
+void FMovieSceneInstance::SaveState()
+{
+	TMap<FGuid, FMovieSceneObjectBindingInstance>::TIterator ObjectIt = ObjectBindingInstances.CreateIterator();
+	for (; ObjectIt; ++ObjectIt)
+	{
+		FMovieSceneObjectBindingInstance& ObjectBindingInstance = ObjectIt.Value();
+
+		for (FMovieSceneInstanceMap::TIterator It = ObjectBindingInstance.TrackInstances.CreateIterator(); It; ++It)
+		{
+			It.Value()->SaveState(ObjectBindingInstance.RuntimeObjects);
+		}
+	}
+}
+
+void FMovieSceneInstance::RestoreState()
+{
+	TMap<FGuid, FMovieSceneObjectBindingInstance>::TIterator ObjectIt = ObjectBindingInstances.CreateIterator();
+	for (; ObjectIt; ++ObjectIt)
+	{
+		FMovieSceneObjectBindingInstance& ObjectBindingInstance = ObjectIt.Value();
+
+		for (FMovieSceneInstanceMap::TIterator It = ObjectBindingInstance.TrackInstances.CreateIterator(); It; ++It)
+		{
+			It.Value()->RestoreState(ObjectBindingInstance.RuntimeObjects);
+		}
+	}
+}
+
+void FMovieSceneInstance::Update( float Position, float LastPosition, class IMovieScenePlayer& Player )
+{
+	// Update each type
+	TArray<UObject*> Objects;
+	for( FMovieSceneInstanceMap::TIterator It( MasterTrackInstances ); It; ++It )
+	{
+		It.Value()->Update( Position, LastPosition, Objects, Player );
+	}
+
+	TMap<FGuid, FMovieSceneObjectBindingInstance>::TIterator ObjectIt = ObjectBindingInstances.CreateIterator();
+	for(; ObjectIt; ++ObjectIt )
+	{
+		FMovieSceneObjectBindingInstance& ObjectBindingInstance = ObjectIt.Value();
+		
+		for( FMovieSceneInstanceMap::TIterator It = ObjectBindingInstance.TrackInstances.CreateIterator(); It; ++It )
+		{
+			It.Value()->Update( Position, LastPosition, ObjectBindingInstance.RuntimeObjects, Player );
+		}
+	}
+}
+
+void FMovieSceneInstance::RefreshInstance( IMovieScenePlayer& Player )
+{
+	// Get all the master tracks and create instances for them if needed
+	const TArray<UMovieSceneTrack*>& MasterTracks = MovieScene->GetMasterTracks();
+	TArray<UObject*> Objects;
+	RefreshInstanceMap( MasterTracks, Objects, MasterTrackInstances, Player );
+
+	TSet< FGuid > FoundObjectBindings;
+	// Get all tracks for each object binding and create instances for them if needed
+	const TArray<FMovieSceneObjectBinding>& ObjectBindings = MovieScene->GetObjectBindings();
+	for( int32 BindingIndex = 0; BindingIndex < ObjectBindings.Num(); ++BindingIndex )
+	{
+		const FMovieSceneObjectBinding& ObjectBinding = ObjectBindings[BindingIndex];
+
+		// Create an instance for this object binding
+		FMovieSceneObjectBindingInstance& BindingInstance = ObjectBindingInstances.FindOrAdd( ObjectBinding.GetObjectGuid() );
+		BindingInstance.ObjectGuid = ObjectBinding.GetObjectGuid();
+
+		FoundObjectBindings.Add( ObjectBinding.GetObjectGuid() );
+
+		// Spawn the runtime objects
+		// @todo Sequencer SubMovieScenes: We need to know which actors were removed and which actors were added so we know which saved actor state to restore/create
+		BindingInstance.RuntimeObjects.Empty();
+		Player.GetRuntimeObjects( SharedThis( this ), BindingInstance.ObjectGuid, BindingInstance.RuntimeObjects );
+		
+		// Refresh the instance's tracks
+		const TArray<UMovieSceneTrack*>& Tracks = ObjectBinding.GetTracks();
+		RefreshInstanceMap( Tracks, BindingInstance.RuntimeObjects, BindingInstance.TrackInstances, Player );
+	}
+
+
+	// Remove object binding instances which are no longer bound
+	TMap<FGuid, FMovieSceneObjectBindingInstance>::TIterator It = ObjectBindingInstances.CreateIterator();
+	for( ; It; ++It )
+	{
+		if( !FoundObjectBindings.Contains( It.Key() ) )
+		{
+			// The instance no longer is bound to an existing guid
+			It.RemoveCurrent();
+		}
+	}
+}
+
+
+void FMovieSceneInstance::RefreshInstanceMap( const TArray<UMovieSceneTrack*>& Tracks, const TArray<UObject*>& RuntimeObjects, FMovieSceneInstanceMap& TrackInstances, IMovieScenePlayer& Player  )
+{
+	// All the tracks we found during this pass
+	TSet< UMovieSceneTrack* > FoundTracks;
+
+	// For every track, check if it has an instance, if not create one, otherwise refresh that instance
+	for( int32 TrackIndex = 0; TrackIndex < Tracks.Num(); ++TrackIndex )
+	{
+		UMovieSceneTrack* Track = Tracks[TrackIndex];
+
+		// A new track has been encountered
+		FoundTracks.Add( Track );
+
+		// See if the track has an instance
+		TSharedPtr<IMovieSceneTrackInstance> Instance = TrackInstances.FindRef( Track );
+		if ( !Instance.IsValid() )
+		{
+			// The track does not have an instance, create one
+			Instance = Track->CreateInstance();
+			Instance->SaveState(RuntimeObjects);
+			Instance->RefreshInstance( RuntimeObjects, Player );
+
+			TrackInstances.Add( Track, Instance );
+		}
+		else
+		{
+			// The track has an instance, refresh it
+			Instance->RefreshInstance( RuntimeObjects, Player );
+		}
+
+	}
+
+	// Remove entries which no longer have a track associated with them
+	FMovieSceneInstanceMap::TIterator It = TrackInstances.CreateIterator();
+	for( ; It; ++It )
+	{
+		if( !FoundTracks.Contains( It.Key().Get() ) )
+		{
+			// This track was not found in the moviescene's track list so it was removed.
+			It.RemoveCurrent();
+		}
+	}
+}
\ No newline at end of file
diff --git a/Engine/Source/Runtime/MovieScene/Private/MovieScenePrivatePCH.h b/Engine/Source/Runtime/MovieScene/Private/MovieScenePrivatePCH.h
new file mode 100644
--- /dev/null
+++ b/Engine/Source/Runtime/MovieScene/Private/MovieScenePrivatePCH.h
@@ -0,0 +1,18 @@
+// Copyright 1998-2015 Epic Games, Inc. All Rights Reserved.
+
+#ifndef MOVIE_SCENE_CORE_PRIVATE
+#define MOVIE_SCENE_CORE_PRIVATE
+
+#include "Core.h"
+#include "Engine.h"
+
+#include "IMovieSceneModule.h"
+#include "IMovieScenePlayer.h"
+#include "MovieScene.h"
+#include "MovieSceneSection.h"
+#include "MovieSceneTrack.h"
+#include "MovieScenePropertyTrack.h"
+#include "MovieSceneInstance.h"
+#include "IMovieSceneTrackInstance.h"
+
+#endif
\ No newline at end of file
diff --git a/Engine/Source/Runtime/MovieScene/Private/MovieScenePropertyTrack.cpp b/Engine/Source/Runtime/MovieScene/Private/MovieScenePropertyTrack.cpp
new file mode 100644
--- /dev/null
+++ b/Engine/Source/Runtime/MovieScene/Private/MovieScenePropertyTrack.cpp
@@ -0,0 +1,128 @@
+// Copyright 1998-2015 Epic Games, Inc. All Rights Reserved.
+
+#include "MovieScenePrivatePCH.h"
+#include "MovieScenePropertyTrack.h"
+
+
+UMovieScenePropertyTrack::UMovieScenePropertyTrack( const FObjectInitializer& ObjectInitializer )
+	: Super( ObjectInitializer )
+{
+	bSectionsAreShowable = false;
+}
+
+
+void UMovieScenePropertyTrack::SetPropertyNameAndPath( FName InPropertyName, const FString& InPropertyPath )
+{
+	check( InPropertyName != NAME_None && !InPropertyPath.IsEmpty() );
+
+	PropertyName = InPropertyName;
+	PropertyPath = InPropertyPath;
+}
+
+TArray<UMovieSceneSection*> UMovieScenePropertyTrack::GetAllSections() const
+{
+	TArray<UMovieSceneSection*> OutSectionPayloads;
+	for( int32 SectionIndex = 0; SectionIndex < Sections.Num(); ++SectionIndex )
+	{
+		OutSectionPayloads.Add( Sections[SectionIndex] );
+	}
+	
+	return OutSectionPayloads;
+}
+
+
+void UMovieScenePropertyTrack::RemoveAllAnimationData()
+{
+	Sections.Empty();
+}
+
+bool UMovieScenePropertyTrack::HasSection( UMovieSceneSection* Section ) const 
+{
+	return Sections.Contains( Section );
+}
+
+void UMovieScenePropertyTrack::RemoveSection( UMovieSceneSection* Section ) 
+{
+	Sections.Remove( Section );
+}
+
+bool UMovieScenePropertyTrack::IsEmpty() const
+{
+	return Sections.Num() == 0;
+}
+
+
+
+UMovieSceneSection* UMovieScenePropertyTrack::FindOrAddSection( float Time )
+{
+	// Find a spot for the section so that they are sorted by start time
+	for( int32 SectionIndex = 0; SectionIndex < Sections.Num(); ++SectionIndex )
+	{
+		UMovieSceneSection* Section = Sections[SectionIndex];
+
+		if( Section->IsTimeWithinSection( Time ) )
+		{
+			return Section;
+		}
+
+		// Check if there are no more sections that would overlap the time 
+		if( !Sections.IsValidIndex( SectionIndex+1 ) || Sections[SectionIndex+1]->GetStartTime() > Time )
+		{
+			// No sections overlap the time
+		
+			if( SectionIndex > 0 )
+			{
+				// Append and grow the previous section
+				UMovieSceneSection* PreviousSection = Sections[ SectionIndex ? SectionIndex-1 : 0 ];
+		
+				PreviousSection->SetEndTime( Time );
+				return PreviousSection;
+			}
+			else if( Sections.IsValidIndex( SectionIndex+1 ) )
+			{
+				// Prepend and grow the next section because there are no sections before this one
+				UMovieSceneSection* NextSection = Sections[SectionIndex+1];
+				NextSection->SetStartTime( Time );
+				return NextSection;
+			}	
+			else
+			{
+				// SectionIndex == 0 
+				UMovieSceneSection* PreviousSection = Sections[0];
+				if( PreviousSection->GetEndTime() < Time )
+				{
+					// Append and grow the section
+					PreviousSection->SetEndTime( Time );
+				}
+				else
+				{
+					// Prepend and grow the section
+					PreviousSection->SetStartTime( Time );
+				}
+				return PreviousSection;
+			}
+		}
+
+	}
+
+	check( Sections.Num() == 0 );
+
+	// Add a new section that starts and ends at the same time
+	UMovieSceneSection* NewSection = CreateNewSection();
+	NewSection->SetStartTime( Time );
+	NewSection->SetEndTime( Time );
+
+	Sections.Add( NewSection );
+
+	return NewSection;
+}
+
+TRange<float> UMovieScenePropertyTrack::GetSectionBoundaries() const
+{
+	TArray< TRange<float> > Bounds;
+	for (int32 SectionIndex = 0; SectionIndex < Sections.Num(); ++SectionIndex)
+	{
+		Bounds.Add(Sections[SectionIndex]->GetRange());
+	}
+	return TRange<float>::Hull(Bounds);
+}
diff --git a/Engine/Source/Runtime/MovieScene/Private/MovieSceneSection.cpp b/Engine/Source/Runtime/MovieScene/Private/MovieSceneSection.cpp
new file mode 100644
--- /dev/null
+++ b/Engine/Source/Runtime/MovieScene/Private/MovieSceneSection.cpp
@@ -0,0 +1,69 @@
+// Copyright 1998-2015 Epic Games, Inc. All Rights Reserved.
+
+#include "MovieScenePrivatePCH.h"
+#include "MovieSceneSection.h"
+
+UMovieSceneSection::UMovieSceneSection( const FObjectInitializer& ObjectInitializer )
+	: Super( ObjectInitializer )
+{
+
+	StartTime = 0.0f;
+	EndTime = 0.0f;
+	RowIndex = 0;
+}
+
+const UMovieSceneSection* UMovieSceneSection::OverlapsWithSections(const TArray<UMovieSceneSection*>& Sections, int32 TrackDelta, float TimeDelta) const
+{
+	int32 NewTrackIndex = RowIndex + TrackDelta;
+	TRange<float> NewSectionRange = TRange<float>(StartTime + TimeDelta, EndTime + TimeDelta);
+	for (int32 SectionIndex = 0; SectionIndex < Sections.Num(); ++SectionIndex)
+	{
+		const UMovieSceneSection* InSection = Sections[SectionIndex];
+		if (this != InSection && InSection->GetRowIndex() == NewTrackIndex)
+		{
+			if (NewSectionRange.Overlaps(InSection->GetRange()))
+			{
+				return InSection;
+			}
+		}
+	}
+	return NULL;
+}
+
+void UMovieSceneSection::InitialPlacement(const TArray<UMovieSceneSection*>& Sections, float InStartTime, float InEndTime, bool bAllowMultipleRows)
+{
+	check(StartTime <= EndTime);
+
+	StartTime = InStartTime;
+	EndTime = InEndTime;
+
+	RowIndex = 0;
+	if (bAllowMultipleRows)
+	{
+		while (OverlapsWithSections(Sections) != NULL)
+		{
+			++RowIndex;
+		}
+	}
+	else
+	{
+		for (;;)
+		{
+			const UMovieSceneSection* OverlappedSection = OverlapsWithSections(Sections);
+
+			if (OverlappedSection == NULL) {break;}
+
+			TSet<FKeyHandle> KeyHandles;
+			MoveSection(OverlappedSection->GetEndTime() - StartTime, KeyHandles);
+		}
+	}
+}
+
+void UMovieSceneSection::AddKeyToCurve( FRichCurve& InCurve, float Time, float Value )
+{
+	if(IsTimeWithinSection(Time))
+	{
+		Modify();
+		InCurve.UpdateOrAddKey(Time, Value);
+	}
+}
\ No newline at end of file
diff --git a/Engine/Source/Runtime/MovieScene/Private/MovieSceneTrack.cpp b/Engine/Source/Runtime/MovieScene/Private/MovieSceneTrack.cpp
new file mode 100644
--- /dev/null
+++ b/Engine/Source/Runtime/MovieScene/Private/MovieSceneTrack.cpp
@@ -0,0 +1,12 @@
+// Copyright 1998-2015 Epic Games, Inc. All Rights Reserved.
+
+#include "MovieScenePrivatePCH.h"
+#include "MovieScene.h"
+#include "MovieSceneTrack.h"
+
+
+
+UMovieSceneTrack::UMovieSceneTrack( const FObjectInitializer& ObjectInitializer )
+	: Super( ObjectInitializer )
+{
+}
diff --git a/Engine/Source/Runtime/MovieScene/Private/RuntimeMovieScenePlayer.cpp b/Engine/Source/Runtime/MovieScene/Private/RuntimeMovieScenePlayer.cpp
new file mode 100644
--- /dev/null
+++ b/Engine/Source/Runtime/MovieScene/Private/RuntimeMovieScenePlayer.cpp
@@ -0,0 +1,219 @@
+// Copyright 1998-2015 Epic Games, Inc. All Rights Reserved.
+
+#include "MovieScenePrivatePCH.h"
+#include "RuntimeMovieScenePlayer.h"
+#include "MovieSceneInstance.h"
+
+URuntimeMovieScenePlayer* URuntimeMovieScenePlayer::CreateRuntimeMovieScenePlayer( ULevel* Level, UMovieSceneBindings* MovieSceneBindings )
+{	
+	URuntimeMovieScenePlayer* NewRuntimeMovieScenePlayer = NewObject<URuntimeMovieScenePlayer>(
+		(UObject*)GetTransientPackage(), 
+		NAME_None, 
+		RF_Transient );
+	check( NewRuntimeMovieScenePlayer != NULL );
+
+	// Associate the player with its world
+	NewRuntimeMovieScenePlayer->World = Level->OwningWorld;
+	check( NewRuntimeMovieScenePlayer->World.IsValid() );
+
+	// Attach bindings to the newly-created RuntimeMovieScenePlayer
+	NewRuntimeMovieScenePlayer->SetMovieSceneBindings( MovieSceneBindings );
+
+	// Spawn actors for each spawnable!
+
+	// Add this to the level's list of active RuntimeMovieScenePlayers.  The level will own this player from now on.
+	check( Level != NULL );
+	Level->AddActiveRuntimeMovieScenePlayer( NewRuntimeMovieScenePlayer );
+
+	return NewRuntimeMovieScenePlayer;
+}
+
+
+URuntimeMovieScenePlayer::URuntimeMovieScenePlayer( const FObjectInitializer& ObjectInitializer )
+	: Super( ObjectInitializer )
+{
+
+	MovieSceneBindings = NULL;
+	TimeCursorPosition = 0.0f;
+	bIsPlaying = false;
+}
+
+
+void URuntimeMovieScenePlayer::SpawnActorsForMovie( TSharedRef<FMovieSceneInstance> MovieSceneInstance  )
+{
+	UWorld* WorldPtr = World.Get();
+	if( WorldPtr != NULL && MovieSceneBindings != NULL )
+	{
+		UMovieScene* MovieScene = MovieSceneInstance->GetMovieScene();
+		if( MovieScene != NULL )
+		{
+			TArray<FSpawnedActorInfo>* FoundSpawnedActors = InstanceToSpawnedActorMap.Find( MovieSceneInstance );
+			if( FoundSpawnedActors )
+			{
+				// Remove existing spawned actors for this movie
+				DestroyActorsForMovie( MovieSceneInstance );
+			}
+
+			TArray<FSpawnedActorInfo>& SpawnedActorList = InstanceToSpawnedActorMap.Add( MovieSceneInstance, TArray<FSpawnedActorInfo>() );
+
+			for( auto SpawnableIndex = 0; SpawnableIndex < MovieScene->GetSpawnableCount(); ++SpawnableIndex )
+			{
+				auto& Spawnable = MovieScene->GetSpawnable( SpawnableIndex );
+
+				UClass* GeneratedClass = Spawnable.GetClass();
+				if ( GeneratedClass != NULL )
+				{
+					const bool bIsActorBlueprint = GeneratedClass->IsChildOf( AActor::StaticClass() );
+					if ( bIsActorBlueprint )
+					{
+						AActor* ActorCDO = CastChecked< AActor >( GeneratedClass->ClassDefaultObject );
+
+						const FVector SpawnLocation = ActorCDO->GetRootComponent()->RelativeLocation;
+						const FRotator SpawnRotation = ActorCDO->GetRootComponent()->RelativeRotation;
+
+						FActorSpawnParameters SpawnInfo;
+						SpawnInfo.ObjectFlags = RF_NoFlags;
+						AActor* NewActor = WorldPtr->SpawnActor( GeneratedClass, &SpawnLocation, &SpawnRotation, SpawnInfo );
+						if( NewActor )
+						{	
+							// Actor was spawned OK!
+							FSpawnedActorInfo NewInfo;
+							NewInfo.RuntimeGuid = Spawnable.GetGuid();
+							NewInfo.SpawnedActor = NewActor;
+						
+							SpawnedActorList.Add( NewInfo );
+						}
+					}
+				}
+			}
+		}
+	}
+}
+
+void URuntimeMovieScenePlayer::DestroyActorsForMovie( TSharedRef<FMovieSceneInstance> MovieSceneInstance )
+{
+	UWorld* WorldPtr = World.Get();
+	if( WorldPtr != NULL && MovieSceneBindings != NULL )
+	{
+		TArray<FSpawnedActorInfo>* SpawnedActors = InstanceToSpawnedActorMap.Find( MovieSceneInstance );
+		if( SpawnedActors )
+		{
+			TArray<FSpawnedActorInfo>& SpawnedActorsRef = *SpawnedActors;
+			for( int32 ActorIndex = 0; ActorIndex < SpawnedActors->Num(); ++ActorIndex )
+			{
+				AActor* Actor = SpawnedActorsRef[ActorIndex].SpawnedActor.Get();
+				if( Actor )
+				{
+					// @todo Sequencer figure this out.  Defaults to false.
+					const bool bNetForce = false;
+
+					// At runtime, level modification is not needed
+					const bool bShouldModifyLevel = false;
+					Actor->GetWorld()->DestroyActor( Actor, bNetForce, bShouldModifyLevel );
+				}
+			}
+			InstanceToSpawnedActorMap.Remove( MovieSceneInstance );
+		}
+	}
+}
+
+bool URuntimeMovieScenePlayer::IsPlaying() const
+{
+	return bIsPlaying;
+}
+
+
+void URuntimeMovieScenePlayer::Play()
+{
+	// Init the root movie scene instance
+	if( MovieSceneBindings != NULL )
+	{
+		UMovieScene* MovieScene = MovieSceneBindings->GetRootMovieScene();
+	
+		// @todo Sequencer playback: Should we recreate the instance every time?
+		RootMovieSceneInstance = MakeShareable( new FMovieSceneInstance( *MovieScene ) );
+
+		// @odo Sequencer Should we spawn actors here?
+		SpawnActorsForMovie( RootMovieSceneInstance.ToSharedRef() );
+
+		RootMovieSceneInstance->RefreshInstance( *this );
+	}
+	
+	bIsPlaying = true;
+}
+
+
+void URuntimeMovieScenePlayer::Pause()
+{
+	bIsPlaying = false;
+}
+
+
+void URuntimeMovieScenePlayer::SetMovieSceneBindings( UMovieSceneBindings* NewMovieSceneBindings )
+{
+	MovieSceneBindings = NewMovieSceneBindings;
+}
+
+
+void URuntimeMovieScenePlayer::Tick( const float DeltaSeconds )
+{
+	float LastTimePosition = TimeCursorPosition;
+
+	if( bIsPlaying )
+	{
+		TimeCursorPosition += DeltaSeconds;
+	}
+
+	// Update the movie scene!
+	if( MovieSceneBindings != NULL && RootMovieSceneInstance.IsValid() )
+	{
+		RootMovieSceneInstance->Update( TimeCursorPosition, LastTimePosition, *this );
+	}
+}
+
+	
+void URuntimeMovieScenePlayer::GetRuntimeObjects( TSharedRef<FMovieSceneInstance> MovieSceneInstance, const FGuid& ObjectHandle, TArray< UObject* >& OutObjects ) const
+{
+	// @todo sequencer runtime: Add support for individually spawning actors on demand when first requested?
+	//    This may be important to reduce the up-front hitch when spawning actors for the entire movie, but
+	//    may introduce smaller hitches during playback.  Needs experimentation.
+
+	// First, check to see if we have a spawned actor for this handle
+	const TArray<FSpawnedActorInfo>* SpawnedActors = InstanceToSpawnedActorMap.Find( MovieSceneInstance );
+	if( SpawnedActors && SpawnedActors->Num() > 0 )
+	{
+		const TArray<FSpawnedActorInfo>& SpawnedActorInfoRef = *SpawnedActors;
+
+		for( int32 SpawnedActorIndex = 0; SpawnedActorIndex < SpawnedActorInfoRef.Num(); ++SpawnedActorIndex )
+		{
+			const FSpawnedActorInfo ActorInfo = SpawnedActorInfoRef[SpawnedActorIndex];
+
+			if( ActorInfo.RuntimeGuid == ObjectHandle && ActorInfo.SpawnedActor.IsValid() )
+			{
+				OutObjects.Add( ActorInfo.SpawnedActor.Get() );
+			}
+		}
+	
+	}
+	// Otherwise, check to see if we have one or more possessed actors that are mapped to this handle
+	else if( MovieSceneBindings != NULL )
+	{
+		OutObjects = MovieSceneBindings->FindBoundObjects( ObjectHandle );
+	}
+}
+
+EMovieScenePlayerStatus::Type URuntimeMovieScenePlayer::GetPlaybackStatus() const
+{
+	return bIsPlaying ? EMovieScenePlayerStatus::Playing : EMovieScenePlayerStatus::Stopped;
+}
+
+void URuntimeMovieScenePlayer::AddMovieSceneInstance( UMovieSceneSection& MovieSceneSection, TSharedRef<FMovieSceneInstance> InstanceToAdd )
+{
+	SpawnActorsForMovie( InstanceToAdd );
+}
+
+void URuntimeMovieScenePlayer::RemoveMovieSceneInstance( UMovieSceneSection& MovieSceneSection, TSharedRef<FMovieSceneInstance> InstanceToRemove )
+{
+	const bool bDestroyAll = true;
+	DestroyActorsForMovie( InstanceToRemove );
+}
\ No newline at end of file
diff --git a/Engine/Source/Runtime/MovieScene/Public/IMovieSceneModule.h b/Engine/Source/Runtime/MovieScene/Public/IMovieSceneModule.h
new file mode 100644
--- /dev/null
+++ b/Engine/Source/Runtime/MovieScene/Public/IMovieSceneModule.h
@@ -0,0 +1,35 @@
+// Copyright 1998-2015 Epic Games, Inc. All Rights Reserved.
+
+#pragma once
+
+#include "ModuleInterface.h"
+#include "ModuleManager.h"		// For inline LoadModuleChecked()
+
+/**
+ * The public interface of the MovieScene module
+ */
+class IMovieSceneModule
+	: public IModuleInterface
+{
+public:
+
+	/**
+	 * Singleton-like access to IMovieScene
+	 *
+	 * @return Returns MovieScene singleton instance, loading the module on demand if needed
+	 */
+	static inline IMovieSceneModule& Get()
+	{
+		return FModuleManager::LoadModuleChecked< IMovieSceneModule >( "MovieScene" );
+	}
+
+	/**
+	 * Checks to see if this module is loaded and ready.  It is only valid to call Get() if IsAvailable() returns true.
+	 *
+	 * @return True if the module is loaded and ready to use
+	 */
+	static inline bool IsAvailable()
+	{
+		return FModuleManager::Get().IsModuleLoaded( "MovieScene" );
+	}
+};
diff --git a/Engine/Source/Runtime/MovieScene/Public/IMovieScenePlayer.h b/Engine/Source/Runtime/MovieScene/Public/IMovieScenePlayer.h
new file mode 100644
--- /dev/null
+++ b/Engine/Source/Runtime/MovieScene/Public/IMovieScenePlayer.h
@@ -0,0 +1,67 @@
+// Copyright 1998-2015 Epic Games, Inc. All Rights Reserved.
+
+#pragma once
+
+namespace EMovieScenePlayerStatus
+{
+	enum Type
+	{
+		Stopped,
+		Playing,
+		Recording,
+		BeginningScrubbing,
+		Scrubbing,
+		MAX
+	};
+}
+
+class FMovieSceneInstance;
+
+/**
+ * Interface for movie scene players
+ * Provides information for playback of a movie scene
+ */
+class IMovieScenePlayer
+{
+public:
+	
+	/**
+	 * Gets runtime objects associated with an object handle for manipulation during playback
+	 *
+	 * @param ObjectHandle The handle to runtime objects
+	 * @Param The list of runtime objects that will be populated
+	 */
+	virtual void GetRuntimeObjects( TSharedRef<FMovieSceneInstance> MovieSceneInstance, const FGuid& ObjectHandle, TArray< UObject* >& OutObjects ) const = 0;
+
+
+	/**
+	 * Updates the perspective viewports with the actor to view through
+	 *
+	 * @param ObjectToViewThrough The object, probably a camera, that the viewports should lock to
+	 */
+	virtual void UpdatePreviewViewports(UObject* ObjectToViewThrough) const = 0;
+
+	/**
+	 * Adds a MovieScene instance to the player.  MovieScene players need to know about each instance for actor spawning
+	 *
+	 * @param MovieSceneSection	The section owning the MovieScene being instanced. 
+	 * @param InstanceToAdd		The instance being added
+	 */
+	virtual void AddMovieSceneInstance( class UMovieSceneSection& MovieSceneSection, TSharedRef<FMovieSceneInstance> InstanceToAdd ) = 0;
+
+	/**
+	 * Removes a MovieScene instance from the player. 
+	 *
+	 * @param MovieSceneSection	The section owning the MovieScene that was instanced. 
+	 * @param InstanceToAdd		The instance being removed
+	 */
+	virtual void RemoveMovieSceneInstance( class UMovieSceneSection& MovieSceneSection, TSharedRef<FMovieSceneInstance> InstanceToRemove ) = 0;
+
+	/**
+	 * @return The root movie scene instance being played
+	 */
+	virtual TSharedRef<FMovieSceneInstance> GetRootMovieSceneInstance() const = 0;
+
+	/** @return whether the player is currently playing, scrubbing, etc. */
+	virtual EMovieScenePlayerStatus::Type GetPlaybackStatus() const = 0;
+};
diff --git a/Engine/Source/Runtime/MovieScene/Public/IMovieSceneTrackInstance.h b/Engine/Source/Runtime/MovieScene/Public/IMovieSceneTrackInstance.h
new file mode 100644
--- /dev/null
+++ b/Engine/Source/Runtime/MovieScene/Public/IMovieSceneTrackInstance.h
@@ -0,0 +1,43 @@
+// Copyright 1998-2015 Epic Games, Inc. All Rights Reserved.
+
+#pragma once
+
+
+class IMovieScenePlayer;
+
+
+/**
+ * Base class for an animated type instance in a Movie Scene
+ */
+class IMovieSceneTrackInstance
+{
+public:
+
+	/** Virtual destructor. */
+	virtual ~IMovieSceneTrackInstance() { }
+
+	/**
+	 * Save state of objects that this instance will be editing.
+	 */
+	virtual void SaveState(const TArray<UObject*>& RuntimeObjects) { }
+
+	/**
+	 * Restore state of objects that this instance edited.
+	 */
+	virtual void RestoreState(const TArray<UObject*>& RuntimeObjects) { }
+
+	/**
+	 * Main update function for track instances.  Called in game and in editor when we update a moviescene.
+	 *
+	 * @param Position			The current position of the moviescene that is playing
+	 * @param LastPosition		The previous playback position
+	 * @param RuntimeObjects	Runtime objects bound to this instance (if any)
+	 * @param Player			The playback interface.  Contains state and some other functionality for runtime playback
+	 */
+	virtual void Update(float Position, float LastPosition, const TArray<UObject*>& RuntimeObjects, IMovieScenePlayer& Player) = 0;
+	
+	/**
+	 * Refreshes the current instance
+	 */
+	virtual void RefreshInstance(const TArray<UObject*>& RuntimeObjects, IMovieScenePlayer& Player) = 0;
+};
diff --git a/Engine/Source/Runtime/MovieScene/Public/MovieScene.h b/Engine/Source/Runtime/MovieScene/Public/MovieScene.h
new file mode 100644
--- /dev/null
+++ b/Engine/Source/Runtime/MovieScene/Public/MovieScene.h
@@ -0,0 +1,438 @@
+// Copyright 1998-2015 Epic Games, Inc. All Rights Reserved.
+
+#pragma once
+
+#include "MovieScene.generated.h"
+
+MOVIESCENE_API DECLARE_LOG_CATEGORY_EXTERN(LogSequencerRuntime, Log, All);
+
+/**
+ * MovieSceneSpawnable describes an object that can be spawned for this MovieScene
+ */
+USTRUCT()
+struct FMovieSceneSpawnable 
+{
+	GENERATED_USTRUCT_BODY()
+
+public:
+
+	/** FMovieSceneSpawnable default constructor */
+	FMovieSceneSpawnable()
+	{
+	}
+
+	/** FMovieSceneSpawnable initialization constructor */
+	FMovieSceneSpawnable( const FString& InitName, UClass* InitClass, UObject* InitCounterpartGamePreviewObject );
+
+	/** @return Returns the guid for this possessable */
+	const FGuid& GetGuid() const
+	{
+		return Guid;
+	}
+
+	/** @return Returns the name of this spawnable */
+	const FString& GetName() const
+	{
+		return Name;
+	}
+
+	/** @return Returns the blueprint associated with this spawnable */
+	UClass* GetClass() { return GeneratedClass; }
+
+	/** @return	Returns the game preview counterpart object for this spawnable, if it has one. */
+	const FWeakObjectPtr& GetCounterpartGamePreviewObject() const
+	{
+		return CounterpartGamePreviewObject;
+	}
+
+private:
+
+	/** Guid */
+	// @todo sequencer: Guids need to be handled carefully when the asset is duplicated (or loaded after being copied on disk).
+	//					Sometimes we'll need to generate fresh Guids.
+	UPROPERTY()
+	FGuid Guid;
+
+	/** Name label */
+	// @todo sequencer: Should be editor-only probably
+	UPROPERTY()
+	FString Name;
+
+	/** Data-only blueprint-generated-class for this object */
+	// @todo sequencer: Could be weak object ptr here, IF blueprints that are inners are housekept properly without references
+	UPROPERTY()
+	UClass* GeneratedClass;
+
+	/** Optional transient weak pointer to the game preview object this spawnable was created to capture data for.  This is
+	    used in the editor when capturing keyframe data from a live simulation */
+	// @todo sequencer data: Should be editor only
+	FWeakObjectPtr CounterpartGamePreviewObject;
+};
+
+/**
+ * MovieScenePossessable is a "typed slot" used to allow the MovieScene to control an already-existing object
+ */
+USTRUCT()
+struct FMovieScenePossessable
+{
+	GENERATED_USTRUCT_BODY( FMovieScenePossessable )
+
+public:
+
+	/** FMovieScenePossessable default constructor */
+	FMovieScenePossessable()
+	{
+	}
+
+	/** FMovieScenePossessable initialization constructor */
+	FMovieScenePossessable( const FString& InitName, UClass* InitPossessedObjectClass );
+
+	/** @return Returns the guid for this possessable */
+	const FGuid& GetGuid() const
+	{
+		return Guid;
+	}
+
+	/** @return Returns the name of this spawnable */
+	const FString& GetName() const
+	{
+		return Name;
+	}
+
+	/** @return Returns the class of the object we'll possess */
+	const UClass* GetPossessedObjectClass() const
+	{
+		return PossessedObjectClass;
+	}
+
+private:
+
+	/** Guid */
+	// @todo sequencer: Guids need to be handled carefully when the asset is duplicated (or loaded after being copied on disk).
+	//					Sometimes we'll need to generate fresh Guids.
+	UPROPERTY()
+	FGuid Guid;
+
+	/** Name label for this slot */
+	// @todo sequencer: Should be editor-only probably
+	UPROPERTY()
+	FString Name;
+
+	/** Type of the object we'll be possessing */
+	// @todo sequencer: Might be able to be editor-only.  We'll see.
+	// @todo sequencer: This isn't used for anything yet.  We could use it to gate which types of objects can be bound to a
+	// possessable "slot" though.  Or we could use it to generate a "preview" spawnable puppet when previewing with no
+	// possessable object available.
+	UPROPERTY()
+	UClass* PossessedObjectClass;
+
+};
+
+/**
+ * Editor only data that needs to be saved between sessions for editing but has no runtime purpose
+ */
+USTRUCT()
+struct FMovieSceneEditorData
+{
+	GENERATED_USTRUCT_BODY()
+
+	/** List of collapsed sequencer nodes.  We store collapsed instead of expanded so that new nodes with no saved state are expanded by default */
+	UPROPERTY()
+	TArray<FString> CollapsedSequencerNodes;
+};
+
+/**
+ * A set of tracks bound to runtime objects
+ */
+USTRUCT()
+struct FMovieSceneObjectBinding
+{
+	GENERATED_USTRUCT_BODY()
+
+	FMovieSceneObjectBinding()
+	{}
+
+	FMovieSceneObjectBinding( const FGuid& InObjectGuid, const FString& InBindingName, const TArray<class UMovieSceneTrack*>& InTracks )
+		: ObjectGuid( InObjectGuid )
+		, BindingName( InBindingName )
+		, Tracks( InTracks )
+	{}
+
+	FMovieSceneObjectBinding( const FGuid& InObjectGuid, const FString& InBindingName )
+		: ObjectGuid( InObjectGuid )
+		, BindingName( InBindingName )
+	{}
+
+	/**
+	 * @return The time range of all tracks in this binding
+	 */
+	TRange<float> GetTimeRange() const;
+
+	/**
+	 * @return The guid of runtime objects in this binding
+	 */
+	const FGuid& GetObjectGuid() const { return ObjectGuid; }
+
+	/**
+	 * @return The display name of the binding
+	 */
+	const FString& GetName() const { return BindingName; } 
+
+	/**
+	 * Adds a new track to this binding
+	 *
+	 * @param NewTrack	The track to add
+	 */
+	void AddTrack( UMovieSceneTrack& NewTrack );
+
+	/**
+	 * Removes a track from this binding
+	 * 
+	 * @param Track	The track to remove
+	 * @return true if the track was successfully removed, false if the track could not be found
+	 */
+	bool RemoveTrack( UMovieSceneTrack& Track );
+
+	/**
+	 * @return All tracks in this binding
+	 */
+	const TArray<UMovieSceneTrack*>& GetTracks() const { return Tracks; }
+private:
+	/** Object binding guid for runtime objects */
+	UPROPERTY()
+	FGuid ObjectGuid;
+	
+	/** Display name */
+	UPROPERTY()
+	FString BindingName;
+
+	/** All tracks in this binding */
+	UPROPERTY()
+	TArray<class UMovieSceneTrack*> Tracks;
+};
+
+
+/**
+ * MovieScene asset
+ */
+UCLASS( MinimalAPI )
+class UMovieScene : public UObject
+{
+	GENERATED_UCLASS_BODY()
+public:
+
+#if WITH_EDITOR
+	/**
+	 * Adds a spawnable to this movie scene's list of owned blueprints.  These objects are stored as "inners"
+	 * of the MovieScene.
+	 *
+	 * @param	Name		Name of the spawnable
+	 * @param	Blueprint	The blueprint to add.
+	 * @param	CounterpartGamePreviewObject	Optional game preview object to map to this spawnable.  Only a weak pointer to this object is stored.
+	 *
+	 * @return	Guid of the newly-added spawnable
+	 */
+	virtual FGuid AddSpawnable( const FString& Name, class UBlueprint* Blueprint, UObject* CounterpartGamePreviewObject );
+
+	/**
+	 * Removes a spawnable from this movie scene.
+	 *
+	 * @param	Guid	The guid of a spawnable to find and remove 
+	 * @return	true if anything was removed
+	 */
+	virtual bool RemoveSpawnable( const FGuid& Guid );
+#endif //WITH_EDITOR
+
+	/**
+	 * Grabs a reference to a specific spawnable by index
+	 *
+	 * @param	Index of spawnable to return
+	 *
+	 * @return Returns the specified spawnable by index
+	 */
+	virtual struct FMovieSceneSpawnable& GetSpawnable( const int32 Index );
+
+	/**
+	 * Tries to locate a spawnable in this MovieScene for the specified spawnable guid
+	 *
+	 * @param	Guid	The spawnable guid to search for
+	 * 
+	 * @return	Spawnable object that was found (or NULL if not found)
+	 */
+	virtual struct FMovieSceneSpawnable* FindSpawnable( const FGuid& Guid );
+
+	/**
+	 * Tries to locate a spawnable for the specified game preview object (e.g. a PIE-world actor)
+	 *
+	 * @param	GamePreviewObject	The object to find a spawnable counterpart for
+	 *
+	 * @return	Spawnable object that was found (or NULL if not found)
+	 */
+	virtual const struct FMovieSceneSpawnable* FindSpawnableForCounterpart( UObject* GamePreviewObject ) const;
+
+	/**
+	 * @return Returns the number of spawnables in this MovieScene
+	 */
+	virtual int32 GetSpawnableCount() const;
+	
+	/**
+	 * Adds a possessable to this movie scene.
+	 *
+	 * @param	Name		Name of the possessable
+	 * @param	Class		The class of object that will be possessed
+	 *
+	 * @return	Guid of the newly-added possessable
+	 */
+	virtual FGuid AddPossessable( const FString& Name, UClass* Class );
+
+	/**
+	 * Removes a possessable from this movie scene.
+	 *
+	 * @param	PossessableGuid		Guid of possessable to remove
+	 */
+	virtual bool RemovePossessable( const FGuid& PossessableGuid );
+	
+	/**
+	 * Tries to locate a possessable in this MovieScene for the specified possessable guid
+	 *
+	 * @param	Guid	The possessable guid to search for
+	 * 
+	 * @return	Possessable object that was found (or NULL if not found)
+	 */
+	virtual struct FMovieScenePossessable* FindPossessable( const FGuid& Guid );
+
+	/**
+	 * Grabs a reference to a specific possessable by index
+	 *
+	 * @param	Index of possessable to return
+	 *
+	 * @return Returns the specified possessable by index
+	 */
+	virtual struct FMovieScenePossessable& GetPossessable( const int32 Index );
+
+	/**
+	 * @return Returns the number of possessables in this MovieScene
+	 */
+	virtual int32 GetPossessableCount() const;
+
+	/**
+	 * Finds a track 
+	 *
+	 * @param TrackClass 	The class of the track to find
+	 * @param ObjectGuid		The runtime object guid that the track is bound to
+	 * @param UniqueTypeName	The unique name of the track to differentiate the one we are searching for from other tracks of the same class
+	 * @return The found track or NULL if one does not exist
+	 */
+	virtual class UMovieSceneTrack* FindTrack( TSubclassOf<UMovieSceneTrack> TrackClass, const FGuid& ObjectGuid, FName UniqueTypeName ) const;
+
+	/**
+	 * Adds a track.  The type should not already exist
+	 * @param TrackClass The class of the track to create
+	 * @param ObjectGuid	The runtime object guid that the type should bind to
+	 *
+	 * @param Type	The newly created type
+	 */
+	virtual class UMovieSceneTrack* AddTrack( TSubclassOf<UMovieSceneTrack> TrackClass, const FGuid& ObjectGuid );
+	
+	/**
+	 * Removes a track
+	 *
+	 * @param Track The track to remove
+	 * @return true if anything was removed
+	 */
+	virtual bool RemoveTrack( UMovieSceneTrack* Track );
+
+	/**
+	 * Finds a master track (one not bound to a runtime objects)
+	 *
+	 * @param TrackClass 	The class of the track to find
+	 * @return The found track or NULL if one does not exist
+	 */
+	virtual class UMovieSceneTrack* FindMasterTrack( TSubclassOf<UMovieSceneTrack> TrackClass ) const;
+
+	/**
+	 * Adds a master track.  The type should not already exist
+	 * @param TrackClass The class of the track to create
+	 *
+	 * @param Type	The newly created type
+	 */
+	virtual class UMovieSceneTrack* AddMasterTrack(  TSubclassOf<UMovieSceneTrack> TrackClass );
+	
+	/**
+	 * Removes a master track
+	 *
+	 * @param Track The track to remove
+	 * @return true if anything was removed
+	 */
+	virtual bool RemoveMasterTrack( UMovieSceneTrack* Track );
+
+	/**
+	 * @return True if the passed in track is a master track
+	 */
+	virtual bool IsAMasterTrack(const UMovieSceneTrack* Track) const;
+
+	/**
+	 * @return All master tracks
+	 */
+	virtual const TArray<class UMovieSceneTrack*>& GetMasterTracks() const { return MasterTracks; }
+
+	/**
+	 * @return All object bindings
+	 */
+	virtual const TArray<FMovieSceneObjectBinding>& GetObjectBindings() const { return ObjectBindings; }
+
+	/**
+	 * @return The time range of the movie scene (defined by all sections in the scene)
+	 */
+	virtual TRange<float> GetTimeRange() const;
+
+	/**
+	 * Returns all sections and their associated binding data
+	 *
+	 * @return A list of sections with object bindings and names
+	 */
+	virtual TArray<class UMovieSceneSection*> GetAllSections() const;
+
+#if WITH_EDITORONLY_DATA
+	/**
+	 * @return The editor only data for use with this movie scene
+	 */
+	FMovieSceneEditorData& GetEditorData() { return EditorData; }
+#endif
+
+private:
+	/**
+	 * Removes animation data bound to a guid
+	 *
+	 * @param Guid	The guid bound to animation data to remove
+	 */
+	void RemoveObjectBinding( const FGuid& Guid );
+
+private:
+	/** Data-only blueprints for all of the objects that we we're able to spawn.  These describe objects and actors
+		that we may instantiate at runtime, or create proxy objects for previewing in the editor */
+	UPROPERTY()
+	TArray< FMovieSceneSpawnable > Spawnables;
+
+	/** Typed slots for already-spawned objects that we are able to control with this MovieScene */
+	UPROPERTY()
+	TArray< FMovieScenePossessable > Possessables;
+
+	/** Tracks bound to possessed or spawned objects */
+	UPROPERTY()
+	TArray< FMovieSceneObjectBinding > ObjectBindings;
+
+	/**
+	 * Master tracks which are not bound to spawned or possessed objects
+	 */
+	UPROPERTY()
+	TArray<class UMovieSceneTrack*> MasterTracks;
+
+#if WITH_EDITORONLY_DATA
+	/** Editor only data that needs to be saved between sessions for editing but has no runtime purpose */
+	UPROPERTY()
+	FMovieSceneEditorData EditorData;
+#endif
+};
+
+
diff --git a/Engine/Source/Runtime/MovieScene/Public/MovieSceneBindings.h b/Engine/Source/Runtime/MovieScene/Public/MovieSceneBindings.h
new file mode 100644
--- /dev/null
+++ b/Engine/Source/Runtime/MovieScene/Public/MovieSceneBindings.h
@@ -0,0 +1,142 @@
+// Copyright 1998-2015 Epic Games, Inc. All Rights Reserved.
+
+#pragma once
+
+#include "MovieSceneBindings.generated.h"
+
+
+/**
+ * MovieSceneBoundObject connects an object to a "possessable" slot in the MovieScene asset
+ */
+USTRUCT()
+struct FMovieSceneBoundObject
+{
+	GENERATED_USTRUCT_BODY( FMovieSceneBoundObject )
+
+public:
+
+	/** FMovieSceneBoundObject default constructor */
+	FMovieSceneBoundObject()
+	{
+	}
+
+	/**
+	 * FMovieSceneBoundObject initialization constructor.  Binds the specified possessable from a MovieScene asset, to the specified object
+	 *
+	 * @param	InitPossessableGuid		The Guid of the possessable from the MovieScene asset to bind to
+	 * @param	InitObjects				The objects to bind
+	 */
+	FMovieSceneBoundObject( const FGuid& InitPosessableGuid, const TArray< UObject* >& InitObjects );
+
+	/** @return Returns the guid of the possessable that we're bound to */
+	const FGuid& GetPossessableGuid() const
+	{
+		return PossessableGuid;
+	}
+
+	/** @return Returns the objects we're bound to */
+	TArray< UObject* >& GetObjects()
+	{
+		return Objects;
+	}
+
+	/** @return Returns the objects we're bound to */
+	const TArray< UObject* >& GetObjects() const
+	{
+		return Objects;
+	}
+
+	/**
+	 * Returns whether or not an object is bound
+	 *
+	 * @param Object	The object to search for
+	 * @return true if the passed in object is bound, false otherwise 
+	 */
+	bool ContainsObject( UObject* Object )
+	{
+		return Objects.Contains( Object );	
+	}
+
+	/**
+	 * Sets the actual objects bound to this slot.  Remember to call Modify() on the owner object if you change this!
+	 * 
+	 * @param	NewObjects	The list of new bound object (can be NULL)
+	 */
+	void SetObjects( TArray< UObject* > NewObjects )
+	{
+		Objects = NewObjects;
+	}
+
+private:
+
+	/** Guid of the possessable in the MovieScene that we're binding to.  Note that we always want this GUID to be preserved
+	    when the object is duplicated (copy should have same relationship with possessed object.) */
+	UPROPERTY()
+	FGuid PossessableGuid;
+
+	/** The actual bound objects (we support binding multiple objects to the same possessable) */
+	UPROPERTY()
+	TArray<	UObject* > Objects;
+
+};
+
+
+/**
+ * MovieSceneBindings.  Contains the relationship between MovieScene possessables and actual objects being possessed.
+ */
+UCLASS( MinimalAPI )
+class UMovieSceneBindings : public UObject
+{
+	GENERATED_UCLASS_BODY()
+
+public:
+
+	/**
+	 * Associates a MovieScene asset with these bindings
+	 *
+	 * @param	NewMovieScene	The MovieScene asset to use (can be NULL to clear it)
+	 */
+	virtual void SetRootMovieScene( class UMovieScene* NewMovieScene );
+
+	/** @return	Returns the root MovieScene associated with these bindings (or NULL, if no MovieScene is associated yet) */
+	virtual class UMovieScene* GetRootMovieScene();
+
+	/**
+	 * Adds a binding between the specified possessable from a MovieScene asset, to the specified objects
+	 *
+	 * @param	PossessableGuid		The Guid of the possessable from the MovieScene asset to bind to
+	 * @param	Objects				The objects to bind
+	 *
+	 * @return	Reference to the wrapper structure for this bound object
+	 */
+	virtual FMovieSceneBoundObject& AddBinding( const FGuid& PosessableGuid, const TArray< UObject* >& Objects );
+
+	/**
+	 * Given a guid for a possessable, attempts to find the object bound to that guid
+	 *
+	 * @param	Guid	The possessable guid to search for
+	 *
+	 * @return	The bound object, if found, otherwise NULL
+	 */
+	virtual TArray< UObject* > FindBoundObjects( const FGuid& Guid ) const;
+
+	/**
+	 * Gets all bound objects
+	 *
+	 * @return	List of all bound objects.
+	 */
+	virtual TArray< FMovieSceneBoundObject >& GetBoundObjects();
+
+
+private:
+
+	/** The Root MovieScene asset to play */
+	UPROPERTY()
+	class UMovieScene* RootMovieScene;
+
+	/** List of all of the objects that we're bound to */
+	UPROPERTY()
+	TArray< FMovieSceneBoundObject > BoundObjects;
+};
+
+
diff --git a/Engine/Source/Runtime/MovieScene/Public/MovieSceneCommonHelpers.h b/Engine/Source/Runtime/MovieScene/Public/MovieSceneCommonHelpers.h
new file mode 100644
--- /dev/null
+++ b/Engine/Source/Runtime/MovieScene/Public/MovieSceneCommonHelpers.h
@@ -0,0 +1,117 @@
+// Copyright 1998-2015 Epic Games, Inc. All Rights Reserved.
+
+#pragma once
+
+
+class MOVIESCENE_API MovieSceneHelpers
+{
+public:
+	/**
+	 * Gets the sections that were traversed over between the current time and the previous time
+	 */
+	static TArray<UMovieSceneSection*> GetTraversedSections( const TArray<UMovieSceneSection*>& Sections, float CurrentTime, float PreviousTime );
+
+	/**
+	* Finds a section that exists at a given time
+	*
+	* @param Time	The time to find a section at
+	* @return The found section or null
+	*/
+	static UMovieSceneSection* FindSectionAtTime( const TArray<UMovieSceneSection*>& Sections, float Time );
+
+	/**
+	* Finds the nearest section to the given time
+	*
+	* @param Time	The time to find a section at
+	* @return The found section or null
+	*/
+	static UMovieSceneSection* FindNearestSectionAtTime( const TArray<UMovieSceneSection*>& Sections, float Time );
+};
+
+/**
+ * Manages bindings to keyed properties for a track instance. 
+ * Calls UFunctions to set the value on runtime objects
+ */
+class MOVIESCENE_API FTrackInstancePropertyBindings
+{
+public:
+	FTrackInstancePropertyBindings( FName InPropertyName, const FString& InPropertyPath, const FName& InFunctionName = FName());
+
+	/**
+	 * Calls the setter function for a specific runtime object
+	 *
+	 * @param InRuntimeObject	The runtime object whose function to call
+	 * @param FunctionParams	Parameters to pass to the function
+	 */
+	void CallFunction( UObject* InRuntimeObject, void* FunctionParams );
+
+	/**
+	 * Rebuilds the property and function mappings for a set of runtime objects
+	 *
+	 * @param InRuntimeObjects	The objects to rebuild mappings for
+	 */
+	void UpdateBindings( const TArray<UObject*>& InRuntimeObjects );
+
+	/**
+	 * Gets the current value of a property on an object
+	 *
+	 * @param Object	The object to get the property from
+	 * @return ValueType	The current value
+	 */
+	template <typename ValueType>
+	ValueType GetCurrentValue(const UObject* Object) const
+	{
+		FPropertyAndFunction PropAndFunction = RuntimeObjectToFunctionMap.FindRef(Object);
+
+		if(PropAndFunction.PropertyAddress.Address)
+		{
+			const ValueType* Val = PropAndFunction.PropertyAddress.Property->ContainerPtrToValuePtr<ValueType>(PropAndFunction.PropertyAddress.Address);
+			if(Val)
+			{
+				return *Val;
+			}
+		}
+
+		return ValueType();
+	}
+private:
+
+	struct FPropertyAddress
+	{
+		UProperty* Property;
+		void* Address;
+
+		FPropertyAddress()
+			: Property(nullptr)
+			, Address(nullptr)
+		{}
+	};
+
+	struct FPropertyAndFunction
+	{
+		FPropertyAddress PropertyAddress;
+		UFunction* Function;
+
+		FPropertyAndFunction()
+			: PropertyAddress()
+			, Function( nullptr )
+		{}
+	};
+
+	FPropertyAddress FindPropertyRecursive(UObject* Object, void* BasePointer, UStruct* InStruct, TArray<FString>& InPropertyNames, uint32 Index) const;
+	FPropertyAddress FindProperty(UObject* Object, const FString& InPropertyPath) const;
+
+private:
+	/** Mapping of objects to bound functions that will be called to update data on the track */
+	TMap< TWeakObjectPtr<UObject>, FPropertyAndFunction > RuntimeObjectToFunctionMap;
+
+	/** Path to the property we are bound to */
+	FString PropertyPath;
+
+	/** Name of the function to call to set values */
+	FName FunctionName;
+
+	/** Actual name of the property we are bound to */
+	FName PropertyName;
+
+};
\ No newline at end of file
diff --git a/Engine/Source/Runtime/MovieScene/Public/MovieSceneInstance.h b/Engine/Source/Runtime/MovieScene/Public/MovieSceneInstance.h
new file mode 100644
--- /dev/null
+++ b/Engine/Source/Runtime/MovieScene/Public/MovieSceneInstance.h
@@ -0,0 +1,61 @@
+// Copyright 1998-2015 Epic Games, Inc. All Rights Reserved.
+
+#pragma once
+
+class UMovieSceneTrack;
+
+typedef TMap< TWeakObjectPtr<UMovieSceneTrack>, TSharedPtr<class IMovieSceneTrackInstance> > FMovieSceneInstanceMap;
+class FMovieSceneInstance : public TSharedFromThis<FMovieSceneInstance>
+{
+public:
+	MOVIESCENE_API FMovieSceneInstance( UMovieScene& InMovieScene );
+
+	/**
+	* Save state of the objects that this movie scene controls.
+	*/
+	MOVIESCENE_API void SaveState ();
+
+	/**
+	* Restore state of the objects that this movie scene controls.
+	*/
+	MOVIESCENE_API void RestoreState ();
+
+	/**
+	 * Updates this movie scene
+	 *
+	 * @param Position The local playback position
+	 * @param Player Movie scene player interface for interaction with runtime data
+	 */
+	MOVIESCENE_API void Update( float Position, float LastPosition, class IMovieScenePlayer& Player );
+
+	/**
+	 * Refreshes the existing instance.  Called when something significant about moviescene data changes (like adding or removing a track)
+	 * Instances all new tracks found and removes instances for tracks that no longer exist
+	 */
+	MOVIESCENE_API void RefreshInstance( class IMovieScenePlayer& Player );
+
+	/** @return The movie scene being instanced */
+	UMovieScene* GetMovieScene() { return MovieScene.Get(); }
+private:
+	void RefreshInstanceMap( const TArray<UMovieSceneTrack*>& Tracks, const TArray<UObject*>& RuntimeObjects, FMovieSceneInstanceMap& TrackInstances, class IMovieScenePlayer& Player );
+private:
+	/** A paring of a runtime object guid to a runtime data and track instances animating the runtime objects */
+	struct FMovieSceneObjectBindingInstance
+	{
+		/** Runtime object guid for looking up runtime UObjects */
+		FGuid ObjectGuid;
+		/** Actual runtime objects */	
+		TArray<UObject*> RuntimeObjects;
+		/** Instances that animate the runtime objects */
+		FMovieSceneInstanceMap TrackInstances;
+	};
+
+	/** MovieScene that is instanced */
+	TWeakObjectPtr<UMovieScene> MovieScene;
+	/** All Master track instances */
+	FMovieSceneInstanceMap MasterTrackInstances;
+	/** All object binding instances */
+	TMap<FGuid, FMovieSceneObjectBindingInstance> ObjectBindingInstances;
+};
+
+
diff --git a/Engine/Source/Runtime/MovieScene/Public/MovieScenePropertyTrack.h b/Engine/Source/Runtime/MovieScene/Public/MovieScenePropertyTrack.h
new file mode 100644
--- /dev/null
+++ b/Engine/Source/Runtime/MovieScene/Public/MovieScenePropertyTrack.h
@@ -0,0 +1,69 @@
+// Copyright 1998-2015 Epic Games, Inc. All Rights Reserved.
+
+#pragma once
+
+#include "MovieSceneTrack.h"
+#include "MovieScenePropertyTrack.generated.h"
+
+/**
+ * Base class for tracks that animate an object property
+ */
+UCLASS( abstract )
+class MOVIESCENE_API UMovieScenePropertyTrack : public UMovieSceneTrack
+{
+	GENERATED_UCLASS_BODY()
+
+public:
+	/** UMovieSceneTrack interface */
+	virtual FName GetTrackName() const override { return PropertyName; }
+	virtual void RemoveAllAnimationData() override;
+	virtual bool HasSection( UMovieSceneSection* Section ) const override;
+	virtual void RemoveSection( UMovieSceneSection* Section ) override;
+	virtual bool IsEmpty() const override;
+	virtual TRange<float> GetSectionBoundaries() const override;
+	virtual TArray<UMovieSceneSection*> GetAllSections() const override;
+	virtual bool HasShowableData() const override {return bSectionsAreShowable;}
+	
+	/**
+	 * Sets the property name for this animatable property
+	 *
+	 * @param InPropertyName	The property being animated
+	 */
+	virtual void SetPropertyNameAndPath( FName InPropertyName, const FString& InPropertyPath );
+
+	/** @return the name of the property being animataed by this track */
+	FName GetPropertyName() const { return PropertyName; }
+	
+	/** @return The property path for this track */
+	const FString& GetPropertyPath() const { return PropertyPath; }
+
+	/** Sets this track as showable permanently */
+	void SetAsShowable() {bSectionsAreShowable = true;}
+
+protected:
+	/**
+	 * Finds a section at the current time.
+	 *
+	 * @param Time	The time relative to the owning movie scene where the section should be
+	 *
+	 * @return The found section, or NULL if it can't be found
+	 */
+	class UMovieSceneSection* FindOrAddSection(  float Time );
+protected:
+
+	/** Name of the property being changed */
+	UPROPERTY()
+	FName PropertyName;
+
+	/** Path to the property from the source object being changed */
+	UPROPERTY()
+	FString PropertyPath;
+
+	/** All the sections in this list */
+	UPROPERTY()
+	TArray<UMovieSceneSection*> Sections;
+	
+	/** True if this should generate a display node in Sequencer */
+	UPROPERTY()
+	bool bSectionsAreShowable;
+};
diff --git a/Engine/Source/Runtime/MovieScene/Public/MovieSceneSection.h b/Engine/Source/Runtime/MovieScene/Public/MovieSceneSection.h
new file mode 100644
--- /dev/null
+++ b/Engine/Source/Runtime/MovieScene/Public/MovieSceneSection.h
@@ -0,0 +1,181 @@
+// Copyright 1998-2015 Epic Games, Inc. All Rights Reserved.
+
+#pragma once
+
+#include "MovieSceneSection.generated.h"
+
+/**
+ * Base class for movie scene sections
+ */
+UCLASS( abstract, MinimalAPI )
+class UMovieSceneSection : public UObject
+{
+	GENERATED_UCLASS_BODY()
+public:
+
+	/**
+	 * @return The start time of the section
+	 */
+	float GetStartTime() const { return StartTime; }
+
+	/**
+	 * @return The end time of the section
+	 */
+	float GetEndTime() const { return EndTime; }
+	
+	/**
+	 * @return The size of the time range of the section
+	 */
+	float GetTimeSize() const {return EndTime - StartTime;}
+
+	/**
+	 * Sets a new end time for this section
+	 * 
+	 * @param InEndTime	The new end time
+	 */
+	void SetStartTime( float NewStartTime ) { StartTime = NewStartTime; }
+
+	/**
+	 * Sets a new end time for this section
+	 * 
+	 * @param InEndTime	The new end time
+	 */
+	void SetEndTime( float NewEndTime ) { EndTime = NewEndTime; }
+	
+	/**
+	 * @return The range of times of the section
+	 */
+	TRange<float> GetRange() const 
+	{
+		// Uses an inclusive range so that sections that start and end on the same value aren't considered empty ranges.
+		return TRange<float>( TRange<float>::BoundsType::Inclusive( StartTime ), TRange<float>::BoundsType::Inclusive( EndTime ) );
+	}
+	
+	/**
+	 * Sets a new range of times for this section
+	 * 
+	 * @param NewRange	The new range of times
+	 */
+	void SetRange(TRange<float> NewRange)
+	{
+		check(NewRange.HasLowerBound() && NewRange.HasUpperBound());
+
+		Modify();
+
+		StartTime = NewRange.GetLowerBoundValue();
+		EndTime = NewRange.GetUpperBoundValue();
+	}
+
+	/**
+	 * Returns whether or not a provided position in time is within the timespan of the section 
+	 *
+	 * @param Position	The position to check
+	 * @return true if the position is within the timespan, false otherwise
+	 */
+	bool IsTimeWithinSection( float Position ) const 
+	{
+		return Position >= StartTime && Position <= EndTime;
+	}
+
+	/**
+	 * Moves the section by a specific amount of time
+	 *
+	 * @param DeltaTime	The distance in time to move the curve
+	 * @param KeyHandles The key handles to operate on
+	 */
+	virtual void MoveSection( float DeltaTime, TSet<FKeyHandle>& KeyHandles )
+	{
+		Modify();
+
+		StartTime += DeltaTime;
+		EndTime += DeltaTime;
+	}
+	
+	/**
+	 * Dilates the section by a specific factor
+	 *
+	 * @param DilationFactor The multiplier which scales this section
+	 * @param bFromStart Whether to dilate from the beginning or end (whichever stays put)
+	 * @param KeyHandles The key handles to operate on
+	 */
+	virtual void DilateSection( float DilationFactor, float Origin, TSet<FKeyHandle>& KeyHandles )
+	{
+		Modify();
+
+		StartTime = (StartTime - Origin) * DilationFactor + Origin;
+		EndTime = (EndTime - Origin) * DilationFactor + Origin;
+	}
+	
+	/**
+	 * Get the key handles for the keys on the curves within this section
+	 *
+	 * @param KeyHandles The key handles of the keys on the curves within this section
+	 */
+	virtual void GetKeyHandles(TSet<FKeyHandle>& KeyHandles) const {};
+
+	/**
+	 * Gets all snap times for this section
+	 *
+	 * @param OutSnapTimes The array of times we will to output
+	 * @param bGetSectionBorders Gets the section borders in addition to any custom snap times
+	 */
+	virtual void GetSnapTimes(TArray<float>& OutSnapTimes, bool bGetSectionBorders) const
+	{
+		if (bGetSectionBorders)
+		{
+			OutSnapTimes.Add(StartTime);
+			OutSnapTimes.Add(EndTime);
+		}
+	}
+
+	virtual void GetAllKeyTimes(TArray<float>& OutKeyTimes) const {}
+
+	/** Sets this section's new row index */
+	void SetRowIndex(int32 NewRowIndex) {RowIndex = NewRowIndex;}
+
+	/** Gets the row index for this section */
+	int32 GetRowIndex() const {return RowIndex;}
+	
+	/**
+	 * Adds a key to a rich curve, finding an existing key to modify or adding a new one
+	 *
+	 * @param InCurve	The curve to add keys to
+	 * @param Time		The time where the key should be added
+	 * @param Value		The value at the given time
+	 */
+	void MOVIESCENE_API AddKeyToCurve( FRichCurve& InCurve, float Time, float Value );
+
+	/**
+	 * Checks to see if this section overlaps with an array of other sections
+	 * given an optional time and track delta.
+	 *
+	 * @param Sections		Section array to check against
+	 * @param TrackDelta	Optional offset to this section's track index
+	 * @param TimeDelta		Optional offset to this section's time delta
+	 * @return				The first section that overlaps, or null if there is no overlap
+	 */
+	virtual MOVIESCENE_API const UMovieSceneSection* OverlapsWithSections(const TArray<UMovieSceneSection*>& Sections, int32 TrackDelta = 0, float TimeDelta = 0.f) const;
+	
+	/**
+	 * Places this section at the first valid row at the specified time. Good for placement upon creation.
+	 *
+	 * @param Sections		Sections that we can not overlap with
+	 * @param InStartTime	The new start time
+	 * @param InEndTime		The new end time
+	 * @param bAllowMultipleRows	If false, it will move the section in the time direction to make it fit, rather than the row direction
+	 */
+	virtual MOVIESCENE_API void InitialPlacement(const TArray<UMovieSceneSection*>& Sections, float InStartTime, float InEndTime, bool bAllowMultipleRows);
+
+private:
+	/** The start time of the section */
+	UPROPERTY()
+	float StartTime;
+
+	/** The end time of the section */
+	UPROPERTY()
+	float EndTime;
+
+	/** The row index that this section sits on */
+	UPROPERTY()
+	int32 RowIndex;
+};
diff --git a/Engine/Source/Runtime/MovieScene/Public/MovieSceneTrack.h b/Engine/Source/Runtime/MovieScene/Public/MovieSceneTrack.h
new file mode 100644
--- /dev/null
+++ b/Engine/Source/Runtime/MovieScene/Public/MovieSceneTrack.h
@@ -0,0 +1,85 @@
+// Copyright 1998-2015 Epic Games, Inc. All Rights Reserved.
+
+#pragma once
+
+#include "MovieSceneTrack.generated.h"
+
+class UMovieSceneSection;
+
+/**
+ * Base class for a track in a Movie Scene
+ */
+UCLASS( abstract, MinimalAPI )
+class UMovieSceneTrack : public UObject
+{
+	GENERATED_UCLASS_BODY()
+
+public:
+	/**
+	 * @return The name that makes this track unique from other track of the same class.
+	 */
+	virtual FName GetTrackName() const PURE_VIRTUAL ( UMovieSceneTrack::GetTrackName, return NAME_None; )
+
+	/**
+	 * Called when all the sections of the track need to be retrieved
+	 * 
+	 * @return List of all the sections in the track
+	 */
+	virtual TArray<UMovieSceneSection*> GetAllSections() const PURE_VIRTUAL( UMovieSceneTrack::GetAllSections, return TArray<UMovieSceneSection*>(); );
+
+	/**
+	 * Creates a runtime instance of this class
+	 * 
+	 * @return The created runtime instance.  This must not be null (note sharedptr not sharedref due to compatibility with uobjects not actually allowing pure virtuals)
+	 */
+	virtual TSharedPtr<class IMovieSceneTrackInstance> CreateInstance() PURE_VIRTUAL( UMovieSceneTrack::CreateInstance, check(false); return NULL; );
+
+	/**
+	 * Removes animation data
+	 */
+	virtual void RemoveAllAnimationData() {}
+
+	/**
+	 * @return Whether or not this track has any data in it
+	 */
+	virtual bool IsEmpty() const PURE_VIRTUAL( UMovieSceneTrack::IsEmpty, return true; );
+	
+	/**
+	 * @return True if the data of this track show generate display nodes by default
+	 */
+	virtual bool HasShowableData() const {return true;}
+
+	/**
+	 * @return Whether or not this track supports multiple row indices
+	 */
+	virtual bool SupportsMultipleRows() const { return false; }
+
+	/**
+	 * Generates a new section suitable for use with this track
+	 *
+	 * @return a new section suitable for use with this track
+	 */
+	virtual class UMovieSceneSection* CreateNewSection() PURE_VIRTUAL( UMovieSceneTrack::CreateNewSection, return NULL; );
+	
+	/**
+	 * Checks to see if the section is in this track
+	 *
+	 * @param Section	The section to query for
+	 * @return			True if the section is in this track
+	 */
+	virtual bool HasSection( UMovieSceneSection* Section ) const PURE_VIRTUAL( UMovieSceneSection::HasSection, return false; );
+
+	/**
+	 * Removes a section from this track
+	 *
+	 * @param Section	The section to remove
+	 */
+	virtual void RemoveSection( UMovieSceneSection* Section ) PURE_VIRTUAL( UMovieSceneSection::RemoveSection, );
+
+	/**
+	 * Gets the section boundaries of this track.
+	 * 
+	 * @return The range of time boundaries
+	 */
+	virtual TRange<float> GetSectionBoundaries() const PURE_VIRTUAL( UMovieSceneTrack::GetSectionBoundaries, return TRange<float>::Empty(); );
+};
diff --git a/Engine/Source/Runtime/MovieScene/Public/RuntimeMovieScenePlayer.h b/Engine/Source/Runtime/MovieScene/Public/RuntimeMovieScenePlayer.h
new file mode 100644
--- /dev/null
+++ b/Engine/Source/Runtime/MovieScene/Public/RuntimeMovieScenePlayer.h
@@ -0,0 +1,101 @@
+// Copyright 1998-2015 Epic Games, Inc. All Rights Reserved.
+
+#pragma once
+
+#include "IMovieScenePlayer.h"
+#include "MovieScene.h"
+#include "MovieSceneBindings.h"
+#include "Runtime/Engine/Classes/MovieScene/RuntimeMovieScenePlayerInterface.h"
+
+#include "RuntimeMovieScenePlayer.generated.h"
+
+
+/**
+ * RuntimeMovieScenePlayer is used to actually "play" a MovieScene asset at runtime.  It keeps track of playback
+ * state and provides functions for manipulating a MovieScene while its playing.
+ */
+UCLASS( MinimalAPI )
+class URuntimeMovieScenePlayer : public UObject, public IMovieScenePlayer, public IRuntimeMovieScenePlayerInterface
+{
+	GENERATED_UCLASS_BODY()
+
+public:
+
+	/**
+	 * Static:  Creates a RuntimeMovieScenePlayer.  This function is designed to be called through the expanded PlayMovieScene
+	 * Kismet node, and not called directly through visual script.
+	 *
+	 * @param	Level	The level to play this MovieScene in.  This level will "own" the instance of the returned RuntimeMovieScenePlayer.
+	 * @param	MovieSceneBindings	The MovieScene to play, along with the bindings to any actors in the level that may be possessed during playback.
+	 *
+	 * @return	Returns the newly created RuntimeMovieScenePlayer instance for this MovieScene
+	 */
+	UFUNCTION(BlueprintCallable, Category="Game|Cinematic")
+	static URuntimeMovieScenePlayer* CreateRuntimeMovieScenePlayer( class ULevel* Level, class UMovieSceneBindings* InMovieSceneBindings );
+
+	/**
+	 * Starts playback from the current time cursor position.  Called by the expanded PlayMovieScene node.
+	 */
+	UFUNCTION(BlueprintCallable, Category="Game|Cinematic")
+	void Play();
+
+	/**
+	 * Pauses playback.  Called by the expanded PlayMovieScene node.
+	 */
+	UFUNCTION(BlueprintCallable, Category="Game|Cinematic")
+	void Pause();
+
+	/** @return Returns true if the sequence is actively playing, or false if paused or if Play() was never called */
+	bool IsPlaying() const;
+
+
+protected:
+
+	/** Sets the bindings for this RuntimeMovieScenePlayer.  Should be called right after construction. */
+	void SetMovieSceneBindings( UMovieSceneBindings* NewMovieSceneBindings );
+
+	/** IMovieScenePlayer interface */
+	virtual void SpawnActorsForMovie( TSharedRef<FMovieSceneInstance> MovieSceneInstance );
+	virtual void DestroyActorsForMovie( TSharedRef<FMovieSceneInstance> MovieSceneInstance );
+	virtual void GetRuntimeObjects( TSharedRef<FMovieSceneInstance> MovieSceneInstance, const FGuid& ObjectHandle, TArray< UObject* >& OutObjects ) const override;
+	virtual void UpdatePreviewViewports(UObject* ObjectToViewThrough) const override {}
+	virtual EMovieScenePlayerStatus::Type GetPlaybackStatus() const override;
+	virtual void AddMovieSceneInstance( class UMovieSceneSection& MovieSceneSection, TSharedRef<FMovieSceneInstance> InstanceToAdd ) override;
+	virtual void RemoveMovieSceneInstance( class UMovieSceneSection& MovieSceneSection, TSharedRef<FMovieSceneInstance> InstanceToRemove ) override;
+	virtual TSharedRef<FMovieSceneInstance> GetRootMovieSceneInstance() const override { return RootMovieSceneInstance.ToSharedRef(); }
+	/** IRuntimeMovieScenePlayerInterface */
+	virtual void Tick( const float DeltaSeconds ) override;
+
+
+private:
+	struct FSpawnedActorInfo
+	{
+		/** Identifier that maps this actor to a movie scene spawnable */
+		FGuid RuntimeGuid;
+		/** The spawned actor */
+		TWeakObjectPtr<AActor> SpawnedActor;
+	};
+
+	/** The world this player will spawn actors in, if needed */
+	TWeakObjectPtr< UWorld > World;
+
+	/** The MovieScene we're playing as well as bindings for this MovieScene.  May contain sub-scenes. */
+	UPROPERTY()
+	UMovieSceneBindings* MovieSceneBindings;
+
+	/** The current time cursor position within the sequence (in seconds) */
+	UPROPERTY()
+	double TimeCursorPosition;
+
+	/** Whether we're currently playing.  If false, then sequence playback is paused or was never started. */
+	UPROPERTY()
+	bool bIsPlaying;
+
+	/** The root movie scene instance to update when playing. */
+	TSharedPtr<class FMovieSceneInstance> RootMovieSceneInstance;
+
+	/** Maps spawnable GUIDs to their spawned actor in the world */
+	TMap< TWeakPtr<FMovieSceneInstance>, TArray< FSpawnedActorInfo > > InstanceToSpawnedActorMap;
+};
+
+
diff --git a/Engine/Source/Runtime/MovieSceneCore/Classes/MovieScene.h b/Engine/Source/Runtime/MovieSceneCore/Classes/MovieScene.h
deleted file mode 100644
diff --git a/Engine/Source/Runtime/MovieSceneCore/Classes/MovieSceneBindings.h b/Engine/Source/Runtime/MovieSceneCore/Classes/MovieSceneBindings.h
deleted file mode 100644
diff --git a/Engine/Source/Runtime/MovieSceneCore/Classes/MovieScenePropertyTrack.h b/Engine/Source/Runtime/MovieSceneCore/Classes/MovieScenePropertyTrack.h
deleted file mode 100644
diff --git a/Engine/Source/Runtime/MovieSceneCore/Classes/MovieSceneSection.h b/Engine/Source/Runtime/MovieSceneCore/Classes/MovieSceneSection.h
deleted file mode 100644
diff --git a/Engine/Source/Runtime/MovieSceneCore/Classes/MovieSceneTrack.h b/Engine/Source/Runtime/MovieSceneCore/Classes/MovieSceneTrack.h
deleted file mode 100644
diff --git a/Engine/Source/Runtime/MovieSceneCore/Classes/RuntimeMovieScenePlayer.h b/Engine/Source/Runtime/MovieSceneCore/Classes/RuntimeMovieScenePlayer.h
deleted file mode 100644
diff --git a/Engine/Source/Runtime/MovieSceneCore/MovieSceneCore.Build.cs b/Engine/Source/Runtime/MovieSceneCore/MovieSceneCore.Build.cs
deleted file mode 100644
diff --git a/Engine/Source/Runtime/MovieSceneCore/Private/MovieScene.cpp b/Engine/Source/Runtime/MovieSceneCore/Private/MovieScene.cpp
deleted file mode 100644
diff --git a/Engine/Source/Runtime/MovieSceneCore/Private/MovieSceneBindings.cpp b/Engine/Source/Runtime/MovieSceneCore/Private/MovieSceneBindings.cpp
deleted file mode 100644
diff --git a/Engine/Source/Runtime/MovieSceneCore/Private/MovieSceneCommonHelpers.cpp b/Engine/Source/Runtime/MovieSceneCore/Private/MovieSceneCommonHelpers.cpp
deleted file mode 100644
diff --git a/Engine/Source/Runtime/MovieSceneCore/Private/MovieSceneCoreModule.cpp b/Engine/Source/Runtime/MovieSceneCore/Private/MovieSceneCoreModule.cpp
deleted file mode 100644
diff --git a/Engine/Source/Runtime/MovieSceneCore/Private/MovieSceneCorePrivatePCH.h b/Engine/Source/Runtime/MovieSceneCore/Private/MovieSceneCorePrivatePCH.h
deleted file mode 100644
diff --git a/Engine/Source/Runtime/MovieSceneCore/Private/MovieSceneInstance.cpp b/Engine/Source/Runtime/MovieSceneCore/Private/MovieSceneInstance.cpp
deleted file mode 100644
diff --git a/Engine/Source/Runtime/MovieSceneCore/Private/MovieScenePropertyTrack.cpp b/Engine/Source/Runtime/MovieSceneCore/Private/MovieScenePropertyTrack.cpp
deleted file mode 100644
diff --git a/Engine/Source/Runtime/MovieSceneCore/Private/MovieSceneSection.cpp b/Engine/Source/Runtime/MovieSceneCore/Private/MovieSceneSection.cpp
deleted file mode 100644
diff --git a/Engine/Source/Runtime/MovieSceneCore/Private/MovieSceneTrack.cpp b/Engine/Source/Runtime/MovieSceneCore/Private/MovieSceneTrack.cpp
deleted file mode 100644
diff --git a/Engine/Source/Runtime/MovieSceneCore/Private/RuntimeMovieScenePlayer.cpp b/Engine/Source/Runtime/MovieSceneCore/Private/RuntimeMovieScenePlayer.cpp
deleted file mode 100644
diff --git a/Engine/Source/Runtime/MovieSceneCore/Public/IMovieSceneCore.h b/Engine/Source/Runtime/MovieSceneCore/Public/IMovieSceneCore.h
deleted file mode 100644
diff --git a/Engine/Source/Runtime/MovieSceneCore/Public/IMovieScenePlayer.h b/Engine/Source/Runtime/MovieSceneCore/Public/IMovieScenePlayer.h
deleted file mode 100644
diff --git a/Engine/Source/Runtime/MovieSceneCore/Public/IMovieSceneTrackInstance.h b/Engine/Source/Runtime/MovieSceneCore/Public/IMovieSceneTrackInstance.h
deleted file mode 100644
diff --git a/Engine/Source/Runtime/MovieSceneCore/Public/MovieSceneCommonHelpers.h b/Engine/Source/Runtime/MovieSceneCore/Public/MovieSceneCommonHelpers.h
deleted file mode 100644
diff --git a/Engine/Source/Runtime/MovieSceneCore/Public/MovieSceneInstance.h b/Engine/Source/Runtime/MovieSceneCore/Public/MovieSceneInstance.h
deleted file mode 100644
diff --git a/Engine/Source/Runtime/MovieSceneTracks/MovieSceneTracks.Build.cs b/Engine/Source/Runtime/MovieSceneTracks/MovieSceneTracks.Build.cs
--- a/Engine/Source/Runtime/MovieSceneTracks/MovieSceneTracks.Build.cs
+++ b/Engine/Source/Runtime/MovieSceneTracks/MovieSceneTracks.Build.cs
@@ -24,7 +24,7 @@ public class MovieSceneTracks : ModuleRules
 
 		PrivateDependencyModuleNames.AddRange(
 			new string[] {
-				"MovieSceneCore",
+				"MovieScene",
 				"SlateCore"
 			}
 		);
diff --git a/Engine/Source/Runtime/UMG/UMG.Build.cs b/Engine/Source/Runtime/UMG/UMG.Build.cs
--- a/Engine/Source/Runtime/UMG/UMG.Build.cs
+++ b/Engine/Source/Runtime/UMG/UMG.Build.cs
@@ -22,7 +22,7 @@ public class UMG : ModuleRules
 				"SlateCore",
                 "ShaderCore",
 				"RenderCore",
-				"MovieSceneCore",
+				"MovieScene",
 				"RHI",
 			}
 		);
