MovieScene: Renamed module file.

diff --git a/Engine/Source/Runtime/MovieScene/Private/MovieSceneCoreModule.cpp b/Engine/Source/Runtime/MovieScene/Private/MovieSceneCoreModule.cpp
deleted file mode 100644
diff --git a/Engine/Source/Runtime/MovieScene/Private/MovieSceneModule.cpp b/Engine/Source/Runtime/MovieScene/Private/MovieSceneModule.cpp
new file mode 100644
--- /dev/null
+++ b/Engine/Source/Runtime/MovieScene/Private/MovieSceneModule.cpp
@@ -0,0 +1,39 @@
+// Copyright 1998-2015 Epic Games, Inc. All Rights Reserved.
+
+#include "MovieScenePrivatePCH.h"
+#include "ModuleManager.h"
+
+#include "MovieScene.h"
+#include "MovieSceneBindings.h"
+#include "RuntimeMovieScenePlayer.h"
+
+DEFINE_LOG_CATEGORY(LogSequencerRuntime);
+
+
+/**
+ * MovieScene module implementation (private)
+ */
+class FMovieSceneModule : public IMovieSceneModule
+{
+
+public:
+
+	/** IModuleInterface */
+	virtual void StartupModule() override;
+	virtual void ShutdownModule() override;
+
+};
+
+
+IMPLEMENT_MODULE( FMovieSceneModule, MovieScene );
+
+
+
+void FMovieSceneModule::StartupModule()
+{
+}
+
+
+void FMovieSceneModule::ShutdownModule()
+{
+}
\ No newline at end of file
