MovieScene: Scaffolding for movie scene refactor

diff --git a/Engine/Source/Editor/Sequencer/Private/SequencerBindingManager.cpp b/Engine/Source/Editor/Sequencer/Private/SequencerBindingManager.cpp
new file mode 100644
--- /dev/null
+++ b/Engine/Source/Editor/Sequencer/Private/SequencerBindingManager.cpp
@@ -0,0 +1,53 @@
+// Copyright 1998-2015 Epic Games, Inc. All Rights Reserved.
+
+#include "SequencerPrivatePCH.h"
+#include "MovieScene.h"
+#include "MovieSceneInstance.h"
+#include "SequencerBindingManager.h"
+
+
+/* USequencerBindingManager structors
+ *****************************************************************************/
+
+USequencerBindingManager::USequencerBindingManager(const FObjectInitializer& ObjectInitializer)
+	: Super(ObjectInitializer)
+{ }
+
+
+/* IMovieSceneBindingManager interface
+ *****************************************************************************/
+
+void USequencerBindingManager::BindPossessableObject(const FGuid& PossessableGuid, UObject& PossessedObject)
+{
+
+}
+
+
+bool USequencerBindingManager::CanPossessObject(UObject& Object) const
+{
+	return false;
+}
+
+
+FGuid USequencerBindingManager::FindGuidForObject(const UMovieScene& MovieScene, UObject& Object) const
+{
+	return FGuid();
+}
+
+
+void USequencerBindingManager::GetRuntimeObjects(const TSharedRef<FMovieSceneInstance>& MovieSceneInstance, const FGuid& ObjectGuid, TArray<UObject*>& OutRuntimeObjects) const
+{
+
+}
+
+
+bool USequencerBindingManager::TryGetObjectBindingDisplayName(const TSharedRef<FMovieSceneInstance>& MovieSceneInstance, const FGuid& ObjectGuid, FText& DisplayName) const
+{
+	return false;
+}
+
+
+void USequencerBindingManager::UnbindPossessableObjects(const FGuid& PossessableGuid)
+{
+
+}
diff --git a/Engine/Source/Editor/Sequencer/Private/SequencerBindingManager.h b/Engine/Source/Editor/Sequencer/Private/SequencerBindingManager.h
new file mode 100644
--- /dev/null
+++ b/Engine/Source/Editor/Sequencer/Private/SequencerBindingManager.h
@@ -0,0 +1,29 @@
+// Copyright 1998-2015 Epic Games, Inc. All Rights Reserved.
+
+#pragma once
+
+#include "MovieSceneBindingManager.h"
+#include "SequencerBindingManager.generated.h"
+
+
+/**
+ * Movie scene binding manager used by Sequencer.
+ */
+UCLASS(BlueprintType, MinimalAPI)
+class USequencerBindingManager
+	: public UObject
+	, public IMovieSceneBindingManager
+{
+	GENERATED_UCLASS_BODY()
+
+public:
+
+	// IMovieSceneBindingManager interface
+
+	virtual void BindPossessableObject(const FGuid& PossessableGuid, UObject& PossessedObject) override;
+	virtual bool CanPossessObject( UObject& Object ) const override;
+	virtual FGuid FindGuidForObject(const UMovieScene& MovieScene, UObject& Object) const override;
+	virtual void GetRuntimeObjects(const TSharedRef<FMovieSceneInstance>& MovieSceneInstance, const FGuid& ObjectGuid, TArray<UObject*>& OutRuntimeObjects) const override;
+	virtual bool TryGetObjectBindingDisplayName(const TSharedRef<FMovieSceneInstance>& MovieSceneInstance, const FGuid& ObjectGuid, FText& DisplayName) const override;
+	virtual void UnbindPossessableObjects(const FGuid& PossessableGuid) override;
+};
diff --git a/Engine/Source/Editor/UMGEditor/Private/Animation/UMGBindingManager.cpp b/Engine/Source/Editor/UMGEditor/Private/Animation/UMGBindingManager.cpp
new file mode 100644
--- /dev/null
+++ b/Engine/Source/Editor/UMGEditor/Private/Animation/UMGBindingManager.cpp
@@ -0,0 +1,53 @@
+// Copyright 1998-2015 Epic Games, Inc. All Rights Reserved.
+
+#include "UMGEditorPrivatePCH.h"
+#include "MovieScene.h"
+#include "MovieSceneInstance.h"
+#include "UMGBindingManager.h"
+
+
+/* UUMGBindingManager structors
+ *****************************************************************************/
+
+UUMGBindingManager::UUMGBindingManager(const FObjectInitializer& ObjectInitializer)
+	: Super(ObjectInitializer)
+{ }
+
+
+/* IMovieSceneBindingManager interface
+ *****************************************************************************/
+
+void UUMGBindingManager::BindPossessableObject(const FGuid& PossessableGuid, UObject& PossessedObject)
+{
+
+}
+
+
+bool UUMGBindingManager::CanPossessObject(UObject& Object) const
+{
+	return false;
+}
+
+
+FGuid UUMGBindingManager::FindGuidForObject(const UMovieScene& MovieScene, UObject& Object) const
+{
+	return FGuid();
+}
+
+
+void UUMGBindingManager::GetRuntimeObjects(const TSharedRef<FMovieSceneInstance>& MovieSceneInstance, const FGuid& ObjectGuid, TArray<UObject*>& OutRuntimeObjects) const
+{
+
+}
+
+
+bool UUMGBindingManager::TryGetObjectBindingDisplayName(const TSharedRef<FMovieSceneInstance>& MovieSceneInstance, const FGuid& ObjectGuid, FText& DisplayName) const
+{
+	return false;
+}
+
+
+void UUMGBindingManager::UnbindPossessableObjects(const FGuid& PossessableGuid)
+{
+
+}
diff --git a/Engine/Source/Editor/UMGEditor/Private/Animation/UMGBindingManager.h b/Engine/Source/Editor/UMGEditor/Private/Animation/UMGBindingManager.h
new file mode 100644
--- /dev/null
+++ b/Engine/Source/Editor/UMGEditor/Private/Animation/UMGBindingManager.h
@@ -0,0 +1,29 @@
+// Copyright 1998-2015 Epic Games, Inc. All Rights Reserved.
+
+#pragma once
+
+#include "MovieSceneBindingManager.h"
+#include "UMGBindingManager.generated.h"
+
+
+/**
+ * Movie scene binding manager used by Sequencer.
+ */
+UCLASS(BlueprintType, MinimalAPI)
+class UUMGBindingManager
+	: public UObject
+	, public IMovieSceneBindingManager
+{
+	GENERATED_UCLASS_BODY()
+
+public:
+
+	// IMovieSceneBindingManager interface
+
+	virtual void BindPossessableObject(const FGuid& PossessableGuid, UObject& PossessedObject) override;
+	virtual bool CanPossessObject( UObject& Object ) const override;
+	virtual FGuid FindGuidForObject(const UMovieScene& MovieScene, UObject& Object) const override;
+	virtual void GetRuntimeObjects(const TSharedRef<FMovieSceneInstance>& MovieSceneInstance, const FGuid& ObjectGuid, TArray<UObject*>& OutRuntimeObjects) const override;
+	virtual bool TryGetObjectBindingDisplayName(const TSharedRef<FMovieSceneInstance>& MovieSceneInstance, const FGuid& ObjectGuid, FText& DisplayName) const override;
+	virtual void UnbindPossessableObjects(const FGuid& PossessableGuid) override;
+};
diff --git a/Engine/Source/Runtime/MovieScene/Private/MovieSceneBindingManager.cpp b/Engine/Source/Runtime/MovieScene/Private/MovieSceneBindingManager.cpp
new file mode 100644
--- /dev/null
+++ b/Engine/Source/Runtime/MovieScene/Private/MovieSceneBindingManager.cpp
@@ -0,0 +1,12 @@
+// Copyright 1998-2015 Epic Games, Inc. All Rights Reserved.
+
+#include "MovieScenePrivatePCH.h"
+#include "MovieSceneBindingManager.h"
+
+
+/* UMovieSceneBindingManager structors
+ *****************************************************************************/
+
+UMovieSceneBindingManager::UMovieSceneBindingManager(const FObjectInitializer& ObjectInitializer)
+	: Super(ObjectInitializer)
+{ }
diff --git a/Engine/Source/Runtime/MovieScene/Public/MovieScene.h b/Engine/Source/Runtime/MovieScene/Public/MovieScene.h
--- a/Engine/Source/Runtime/MovieScene/Public/MovieScene.h
+++ b/Engine/Source/Runtime/MovieScene/Public/MovieScene.h
@@ -4,8 +4,13 @@
 
 #include "MovieScene.generated.h"
 
+
+class UMovieSceneBindingManager;
+
+
 MOVIESCENE_API DECLARE_LOG_CATEGORY_EXTERN(LogSequencerRuntime, Log, All);
 
+
 /**
  * MovieSceneSpawnable describes an object that can be spawned for this MovieScene
  */
@@ -409,6 +414,11 @@ private:
 	void RemoveObjectBinding( const FGuid& Guid );
 
 private:
+
+	/** The object binding manager. */
+	UPROPERTY()
+	TScriptInterface<UMovieSceneBindingManager> BindingManager;
+
 	/** Data-only blueprints for all of the objects that we we're able to spawn.  These describe objects and actors
 		that we may instantiate at runtime, or create proxy objects for previewing in the editor */
 	UPROPERTY()
diff --git a/Engine/Source/Runtime/MovieScene/Public/MovieSceneBindingManager.h b/Engine/Source/Runtime/MovieScene/Public/MovieSceneBindingManager.h
new file mode 100644
--- /dev/null
+++ b/Engine/Source/Runtime/MovieScene/Public/MovieSceneBindingManager.h
@@ -0,0 +1,68 @@
+// Copyright 1998-2015 Epic Games, Inc. All Rights Reserved.
+
+#pragma once
+
+#include "MovieSceneBindingManager.generated.h"
+
+
+class FMovieSceneInstance;
+class UMovieScene;
+class UObject;
+
+
+UINTERFACE(MinimalAPI, meta=(CannotImplementInterfaceInBlueprint))
+class UMovieSceneBindingManager
+	: public UInterface
+{
+	GENERATED_UINTERFACE_BODY()
+};
+
+
+class IMovieSceneBindingManager
+{
+	GENERATED_IINTERFACE_BODY()
+
+	/**
+	 * Called when Sequencer has created an object binding for a possessable object
+	 * 
+	 * @param PossessableGuid The guid used to map to the possessable object.  Note the guid can be bound to multiple objects at once
+	 * @param PossessedObject The runtime object which was possessed.
+	 */
+	virtual void BindPossessableObject(const FGuid& PossessableGuid, UObject& PossessedObject) PURE_VIRTUAL(IMovieSceneBindingManager::BindPossessableObject,);
+	
+	/**
+	 * @todo sequencer: add documentation
+	 */
+	virtual bool CanPossessObject( UObject& Object ) const PURE_VIRTUAL(IMovieSceneBindingManager::CanPossessObject, return false;);
+
+	/**
+	 * Finds an existing guid for a bound object
+	 *
+	 * @param MovieScene The movie scene that may contain the binding
+	 * @param The object to get the guid for
+	 * @return The object guid or an invalid guid if not found
+	 */
+	virtual FGuid FindGuidForObject(const UMovieScene& MovieScene, UObject& Object) const PURE_VIRTUAL(IMovieSceneBindingManager::FindGuidForObject, return FGuid(););
+
+	/**
+	 * @todo sequencer: add documentation
+	 */
+	virtual void GetRuntimeObjects(const TSharedRef<FMovieSceneInstance>& MovieSceneInstance, const FGuid& ObjectGuid, TArray<UObject*>& OutRuntimeObjects) const PURE_VIRTUAL(IMovieSceneBindingManager::GetRuntimeObjects,);
+
+	/**
+	 * 
+	 * Tries to get a display name for the binding represented by the guid.
+	 *
+	 * @param ObjectGuid The guid for the object binding.
+	 * @param DisplayName The display name for the object binding.
+	 * @returns true if DisplayName has been set to a valid display name, otherwise false.
+	 */
+	virtual bool TryGetObjectBindingDisplayName(const TSharedRef<FMovieSceneInstance>& MovieSceneInstance, const FGuid& ObjectGuid, FText& DisplayName) const PURE_VIRTUAL(IMovieSceneBindingManager::TryGetObjectBindingDisplayName, return false;);
+
+	/**
+	 * Unbinds all possessable objects from the provided guid
+	 *
+	 * @param PossessableGuid The guid bound to possessable objects that should be removed
+	 */
+	virtual void UnbindPossessableObjects(const FGuid& PossessableGuid) PURE_VIRTUAL(IMovieSceneBindingManager::UnbindPossessableObjects,);
+};
