MovieScene: Code and documentation cleanup pass.

diff --git a/Engine/Source/Runtime/MovieScene/Public/MovieScene.h b/Engine/Source/Runtime/MovieScene/Public/MovieScene.h
--- a/Engine/Source/Runtime/MovieScene/Public/MovieScene.h
+++ b/Engine/Source/Runtime/MovieScene/Public/MovieScene.h
@@ -5,7 +5,10 @@
 #include "MovieScene.generated.h"
 
 
+class UBlueprint;
 class UMovieSceneBindingManager;
+class UMovieSceneSection;
+class UMovieSceneTrack;
 
 
 MOVIESCENE_API DECLARE_LOG_CATEGORY_EXTERN(LogSequencerRuntime, Log, All);
@@ -157,7 +160,7 @@ struct FMovieSceneObjectBinding
 	FMovieSceneObjectBinding()
 	{}
 
-	FMovieSceneObjectBinding( const FGuid& InObjectGuid, const FString& InBindingName, const TArray<class UMovieSceneTrack*>& InTracks )
+	FMovieSceneObjectBinding( const FGuid& InObjectGuid, const FString& InBindingName, const TArray<UMovieSceneTrack*>& InTracks )
 		: ObjectGuid( InObjectGuid )
 		, BindingName( InBindingName )
 		, Tracks( InTracks )
@@ -202,7 +205,9 @@ struct FMovieSceneObjectBinding
 	 * @return All tracks in this binding
 	 */
 	const TArray<UMovieSceneTrack*>& GetTracks() const { return Tracks; }
+
 private:
+
 	/** Object binding guid for runtime objects */
 	UPROPERTY()
 	FGuid ObjectGuid;
@@ -213,205 +218,230 @@ private:
 
 	/** All tracks in this binding */
 	UPROPERTY()
-	TArray<class UMovieSceneTrack*> Tracks;
+	TArray<UMovieSceneTrack*> Tracks;
 };
 
 
 /**
- * MovieScene asset
+ * Implements a movie scene asset.
  */
-UCLASS( MinimalAPI )
-class UMovieScene : public UObject
+UCLASS()
+class MOVIESCENE_API UMovieScene
+	: public UObject
 {
 	GENERATED_UCLASS_BODY()
+
 public:
 
 #if WITH_EDITOR
 	/**
-	 * Adds a spawnable to this movie scene's list of owned blueprints.  These objects are stored as "inners"
-	 * of the MovieScene.
+	 * Add a spawnable to this movie scene's list of owned blueprints.
 	 *
-	 * @param	Name		Name of the spawnable
-	 * @param	Blueprint	The blueprint to add.
-	 * @param	CounterpartGamePreviewObject	Optional game preview object to map to this spawnable.  Only a weak pointer to this object is stored.
+	 * These objects are stored as "inners" of the MovieScene.
 	 *
-	 * @return	Guid of the newly-added spawnable
+	 * @param Name Name of the spawnable.
+	 * @param Blueprint	The blueprint to add.
+	 * @param CounterpartGamePreviewObject	Optional game preview object to map to this spawnable.  Only a weak pointer to this object is stored.
+	 * @return Guid of the newly-added spawnable.
 	 */
-	virtual FGuid AddSpawnable( const FString& Name, class UBlueprint* Blueprint, UObject* CounterpartGamePreviewObject );
+	FGuid AddSpawnable(const FString& Name, UBlueprint* Blueprint, UObject* CounterpartGamePreviewObject);
 
 	/**
 	 * Removes a spawnable from this movie scene.
 	 *
-	 * @param	Guid	The guid of a spawnable to find and remove 
-	 * @return	true if anything was removed
+	 * @param Guid The guid of a spawnable to find and remove.
+	 * @return true if anything was removed.
 	 */
-	virtual bool RemoveSpawnable( const FGuid& Guid );
+	bool RemoveSpawnable(const FGuid& Guid);
+
 #endif //WITH_EDITOR
 
 	/**
-	 * Grabs a reference to a specific spawnable by index
-	 *
-	 * @param	Index of spawnable to return
+	 * Grabs a reference to a specific spawnable by index.
 	 *
-	 * @return Returns the specified spawnable by index
+	 * @param Index of spawnable to return.
+	 * @return Returns the specified spawnable by index.
 	 */
-	virtual struct FMovieSceneSpawnable& GetSpawnable( const int32 Index );
+	FMovieSceneSpawnable& GetSpawnable(const int32 Index);
 
 	/**
-	 * Tries to locate a spawnable in this MovieScene for the specified spawnable guid
+	 * Tries to locate a spawnable in this MovieScene for the specified spawnable GUID.
 	 *
-	 * @param	Guid	The spawnable guid to search for
-	 * 
-	 * @return	Spawnable object that was found (or NULL if not found)
+	 * @param Guid The spawnable guid to search for.
+	 * @return Spawnable object that was found (or nullptr if not found).
 	 */
-	virtual struct FMovieSceneSpawnable* FindSpawnable( const FGuid& Guid );
+	FMovieSceneSpawnable* FindSpawnable(const FGuid& Guid);
 
 	/**
-	 * Tries to locate a spawnable for the specified game preview object (e.g. a PIE-world actor)
+	 * Tries to locate a spawnable for the specified game preview object (e.g. a PIE-world actor).
 	 *
-	 * @param	GamePreviewObject	The object to find a spawnable counterpart for
-	 *
-	 * @return	Spawnable object that was found (or NULL if not found)
+	 * @param GamePreviewObject The object to find a spawnable counterpart for.
+	 * @return Spawnable object that was found (or nullptr if not found).
 	 */
-	virtual const struct FMovieSceneSpawnable* FindSpawnableForCounterpart( UObject* GamePreviewObject ) const;
+	const FMovieSceneSpawnable* FindSpawnableForCounterpart(UObject* GamePreviewObject) const;
 
 	/**
-	 * @return Returns the number of spawnables in this MovieScene
+	 * Get the number of spawnable objects in this scene.
+	 *
+	 * @return Spawnable object count.
 	 */
-	virtual int32 GetSpawnableCount() const;
+	int32 GetSpawnableCount() const;
 	
+public:
+
 	/**
 	 * Adds a possessable to this movie scene.
 	 *
-	 * @param	Name		Name of the possessable
-	 * @param	Class		The class of object that will be possessed
-	 *
-	 * @return	Guid of the newly-added possessable
+	 * @param Name Name of the possessable.
+	 * @param Class The class of object that will be possessed.
+	 * @return Guid of the newly-added possessable.
 	 */
-	virtual FGuid AddPossessable( const FString& Name, UClass* Class );
+	FGuid AddPossessable(const FString& Name, UClass* Class);
 
 	/**
 	 * Removes a possessable from this movie scene.
 	 *
-	 * @param	PossessableGuid		Guid of possessable to remove
+	 * @param PossessableGuid Guid of possessable to remove.
 	 */
-	virtual bool RemovePossessable( const FGuid& PossessableGuid );
+	bool RemovePossessable(const FGuid& PossessableGuid);
 	
 	/**
-	 * Tries to locate a possessable in this MovieScene for the specified possessable guid
+	 * Tries to locate a possessable in this MovieScene for the specified possessable GUID.
 	 *
-	 * @param	Guid	The possessable guid to search for
-	 * 
-	 * @return	Possessable object that was found (or NULL if not found)
+	 * @param Guid The possessable guid to search for.
+	 * @return Possessable object that was found (or nullptr if not found).
 	 */
-	virtual struct FMovieScenePossessable* FindPossessable( const FGuid& Guid );
+	struct FMovieScenePossessable* FindPossessable(const FGuid& Guid);
 
 	/**
-	 * Grabs a reference to a specific possessable by index
-	 *
-	 * @param	Index of possessable to return
+	 * Grabs a reference to a specific possessable by index.
 	 *
-	 * @return Returns the specified possessable by index
+	 * @param Index of possessable to return.
+	 * @return Returns the specified possessable by index.
 	 */
-	virtual struct FMovieScenePossessable& GetPossessable( const int32 Index );
+	FMovieScenePossessable& GetPossessable(const int32 Index);
 
 	/**
-	 * @return Returns the number of possessables in this MovieScene
+	 * Get the number of possessable objects in this scene.
+	 *
+	 * @return Possessable object count.
 	 */
-	virtual int32 GetPossessableCount() const;
+	int32 GetPossessableCount() const;
+
+public:
 
 	/**
-	 * Finds a track 
+	 * Finds a track.
 	 *
-	 * @param TrackClass 	The class of the track to find
-	 * @param ObjectGuid		The runtime object guid that the track is bound to
-	 * @param UniqueTypeName	The unique name of the track to differentiate the one we are searching for from other tracks of the same class
-	 * @return The found track or NULL if one does not exist
+	 * @param TrackClass The class of the track to find.
+	 * @param ObjectGuid The runtime object guid that the track is bound to.
+	 * @param UniqueTypeName The unique name of the track to differentiate the one we are searching for from other tracks of the same class.
+	 * @return The found track or nullptr if one does not exist.
 	 */
-	virtual class UMovieSceneTrack* FindTrack( TSubclassOf<UMovieSceneTrack> TrackClass, const FGuid& ObjectGuid, FName UniqueTypeName ) const;
+	UMovieSceneTrack* FindTrack(TSubclassOf<UMovieSceneTrack> TrackClass, const FGuid& ObjectGuid, FName UniqueTypeName) const;
 
 	/**
-	 * Adds a track.  The type should not already exist
-	 * @param TrackClass The class of the track to create
-	 * @param ObjectGuid	The runtime object guid that the type should bind to
+	 * Adds a track.
 	 *
-	 * @param Type	The newly created type
+	 * Note: The type should not already exist.
+	 *
+	 * @param TrackClass The class of the track to create.
+	 * @param ObjectGuid The runtime object guid that the type should bind to.
+	 * @param Type The newly created type.
 	 */
-	virtual class UMovieSceneTrack* AddTrack( TSubclassOf<UMovieSceneTrack> TrackClass, const FGuid& ObjectGuid );
+	UMovieSceneTrack* AddTrack(TSubclassOf<UMovieSceneTrack> TrackClass, const FGuid& ObjectGuid);
 	
 	/**
-	 * Removes a track
+	 * Removes a track.
 	 *
-	 * @param Track The track to remove
-	 * @return true if anything was removed
+	 * @param Track The track to remove.
+	 * @return true if anything was removed.
 	 */
-	virtual bool RemoveTrack( UMovieSceneTrack* Track );
+	bool RemoveTrack( UMovieSceneTrack* Track );
 
 	/**
-	 * Finds a master track (one not bound to a runtime objects)
+	 * Finds a master track (one not bound to a runtime objects).
 	 *
-	 * @param TrackClass 	The class of the track to find
-	 * @return The found track or NULL if one does not exist
+	 * @param TrackClass The class of the track to find.
+	 * @return The found track or nullptr if one does not exist.
 	 */
-	virtual class UMovieSceneTrack* FindMasterTrack( TSubclassOf<UMovieSceneTrack> TrackClass ) const;
+	UMovieSceneTrack* FindMasterTrack(TSubclassOf<UMovieSceneTrack> TrackClass) const;
 
 	/**
-	 * Adds a master track.  The type should not already exist
-	 * @param TrackClass The class of the track to create
+	 * Adds a master track.
+	 *
+	 * Note: The type should not already exist.
 	 *
+	 * @param TrackClass The class of the track to create
 	 * @param Type	The newly created type
 	 */
-	virtual class UMovieSceneTrack* AddMasterTrack(  TSubclassOf<UMovieSceneTrack> TrackClass );
+	UMovieSceneTrack* AddMasterTrack(TSubclassOf<UMovieSceneTrack> TrackClass);
 	
 	/**
-	 * Removes a master track
+	 * Removes a master track.
 	 *
-	 * @param Track The track to remove
-	 * @return true if anything was removed
+	 * @param Track The track to remove.
+	 * @return true if anything was removed.
 	 */
-	virtual bool RemoveMasterTrack( UMovieSceneTrack* Track );
+	bool RemoveMasterTrack( UMovieSceneTrack* Track );
 
 	/**
-	 * @return True if the passed in track is a master track
+	 * Check whether the specified track is a master track in this scene.
+	 *
+	 * @return true if the track is a master track, false otherwise.
 	 */
-	virtual bool IsAMasterTrack(const UMovieSceneTrack* Track) const;
+	bool IsAMasterTrack(const UMovieSceneTrack* Track) const;
 
 	/**
-	 * @return All master tracks
+	 * Get all master tracks.
+	 *
+	 * @return Track collection.
 	 */
-	virtual const TArray<class UMovieSceneTrack*>& GetMasterTracks() const { return MasterTracks; }
+	const TArray<UMovieSceneTrack*>& GetMasterTracks() const
+	{
+		return MasterTracks;
+	}
+
+public:
 
 	/**
-	 * @return All object bindings
+	 * @return All object bindings.
 	 */
-	virtual const TArray<FMovieSceneObjectBinding>& GetObjectBindings() const { return ObjectBindings; }
+	const TArray<FMovieSceneObjectBinding>& GetObjectBindings() const
+	{
+		return ObjectBindings;
+	}
 
 	/**
 	 * @return The time range of the movie scene (defined by all sections in the scene)
 	 */
-	virtual TRange<float> GetTimeRange() const;
+	TRange<float> GetTimeRange() const;
 
 	/**
-	 * Returns all sections and their associated binding data
+	 * Returns all sections and their associated binding data.
 	 *
-	 * @return A list of sections with object bindings and names
+	 * @return A list of sections with object bindings and names.
 	 */
-	virtual TArray<class UMovieSceneSection*> GetAllSections() const;
+	TArray<UMovieSceneSection*> GetAllSections() const;
 
 #if WITH_EDITORONLY_DATA
 	/**
 	 * @return The editor only data for use with this movie scene
 	 */
-	FMovieSceneEditorData& GetEditorData() { return EditorData; }
+	FMovieSceneEditorData& GetEditorData()
+	{
+		return EditorData;
+	}
 #endif
 
-private:
+protected:
+
 	/**
-	 * Removes animation data bound to a guid
+	 * Removes animation data bound to a guid.
 	 *
-	 * @param Guid	The guid bound to animation data to remove
+	 * @param Guid The guid bound to animation data to remove
 	 */
-	void RemoveObjectBinding( const FGuid& Guid );
+	void RemoveObjectBinding(const FGuid& Guid);
 
 private:
 
@@ -419,24 +449,25 @@ private:
 	UPROPERTY()
 	TScriptInterface<UMovieSceneBindingManager> BindingManager;
 
-	/** Data-only blueprints for all of the objects that we we're able to spawn.  These describe objects and actors
-		that we may instantiate at runtime, or create proxy objects for previewing in the editor */
+	/**
+	 * Data-only blueprints for all of the objects that we we're able to spawn.
+	 * These describe objects and actors that we may instantiate at runtime,
+	 * or create proxy objects for previewing in the editor.
+	 */
 	UPROPERTY()
-	TArray< FMovieSceneSpawnable > Spawnables;
+	TArray<FMovieSceneSpawnable> Spawnables;
 
 	/** Typed slots for already-spawned objects that we are able to control with this MovieScene */
 	UPROPERTY()
-	TArray< FMovieScenePossessable > Possessables;
+	TArray<FMovieScenePossessable> Possessables;
 
 	/** Tracks bound to possessed or spawned objects */
 	UPROPERTY()
-	TArray< FMovieSceneObjectBinding > ObjectBindings;
+	TArray<FMovieSceneObjectBinding> ObjectBindings;
 
-	/**
-	 * Master tracks which are not bound to spawned or possessed objects
-	 */
+	/** Master tracks which are not bound to spawned or possessed objects */
 	UPROPERTY()
-	TArray<class UMovieSceneTrack*> MasterTracks;
+	TArray<UMovieSceneTrack*> MasterTracks;
 
 #if WITH_EDITORONLY_DATA
 	/** Editor only data that needs to be saved between sessions for editing but has no runtime purpose */
@@ -444,5 +475,3 @@ private:
 	FMovieSceneEditorData EditorData;
 #endif
 };
-
-
diff --git a/Engine/Source/Runtime/MovieSceneTracks/MovieSceneTracks.Build.cs b/Engine/Source/Runtime/MovieSceneTracks/MovieSceneTracks.Build.cs
--- a/Engine/Source/Runtime/MovieSceneTracks/MovieSceneTracks.Build.cs
+++ b/Engine/Source/Runtime/MovieSceneTracks/MovieSceneTracks.Build.cs
@@ -18,14 +18,14 @@ public class MovieSceneTracks : ModuleRules
 			new string[] {
 				"Core",
 				"CoreUObject",
-				"Engine"
+				"Engine",
 			}
 		);
 
 		PrivateDependencyModuleNames.AddRange(
 			new string[] {
 				"MovieScene",
-				"SlateCore"
+				"SlateCore",
 			}
 		);
 	}
