MovieScene: Added movie scene object identifier

diff --git a/Engine/Source/Editor/Sequencer/Private/SequencerBindingManager.cpp b/Engine/Source/Editor/Sequencer/Private/SequencerBindingManager.cpp
--- a/Engine/Source/Editor/Sequencer/Private/SequencerBindingManager.cpp
+++ b/Engine/Source/Editor/Sequencer/Private/SequencerBindingManager.cpp
@@ -17,7 +17,7 @@ USequencerBindingManager::USequencerBindingManager(const FObjectInitializer& Obj
 /* IMovieSceneBindingManager interface
  *****************************************************************************/
 
-void USequencerBindingManager::BindPossessableObject(const FGuid& PossessableGuid, UObject& PossessedObject)
+void USequencerBindingManager::BindPossessableObject(const FMovieSceneObjectId& PossessableGuid, UObject& PossessedObject)
 {
 
 }
@@ -29,9 +29,9 @@ bool USequencerBindingManager::CanPossessObject(UObject& Object) const
 }
 
 
-FGuid USequencerBindingManager::FindGuidForObject(const UMovieScene& MovieScene, UObject& Object) const
+FMovieSceneObjectId USequencerBindingManager::FindGuidForObject(const UMovieScene& MovieScene, UObject& Object) const
 {
-	return FGuid();
+	return FMovieSceneObjectId();
 }
 
 
@@ -47,7 +47,7 @@ bool USequencerBindingManager::TryGetObjectBindingDisplayName(const TSharedRef<F
 }
 
 
-void USequencerBindingManager::UnbindPossessableObjects(const FGuid& PossessableGuid)
+void USequencerBindingManager::UnbindPossessableObjects(const FMovieSceneObjectId& PossessableGuid)
 {
 
 }
diff --git a/Engine/Source/Editor/Sequencer/Private/SequencerBindingManager.h b/Engine/Source/Editor/Sequencer/Private/SequencerBindingManager.h
--- a/Engine/Source/Editor/Sequencer/Private/SequencerBindingManager.h
+++ b/Engine/Source/Editor/Sequencer/Private/SequencerBindingManager.h
@@ -20,10 +20,10 @@ public:
 
 	// IMovieSceneBindingManager interface
 
-	virtual void BindPossessableObject(const FGuid& PossessableGuid, UObject& PossessedObject) override;
+	virtual void BindPossessableObject(const FMovieSceneObjectId& PossessableGuid, UObject& PossessedObject) override;
 	virtual bool CanPossessObject( UObject& Object ) const override;
-	virtual FGuid FindGuidForObject(const UMovieScene& MovieScene, UObject& Object) const override;
+	virtual FMovieSceneObjectId FindGuidForObject(const UMovieScene& MovieScene, UObject& Object) const override;
 	virtual void GetRuntimeObjects(const TSharedRef<FMovieSceneInstance>& MovieSceneInstance, const FGuid& ObjectGuid, TArray<UObject*>& OutRuntimeObjects) const override;
 	virtual bool TryGetObjectBindingDisplayName(const TSharedRef<FMovieSceneInstance>& MovieSceneInstance, const FGuid& ObjectGuid, FText& DisplayName) const override;
-	virtual void UnbindPossessableObjects(const FGuid& PossessableGuid) override;
+	virtual void UnbindPossessableObjects(const FMovieSceneObjectId& PossessableGuid) override;
 };
diff --git a/Engine/Source/Editor/UMGEditor/Private/Animation/UMGBindingManager.cpp b/Engine/Source/Editor/UMGEditor/Private/Animation/UMGBindingManager.cpp
--- a/Engine/Source/Editor/UMGEditor/Private/Animation/UMGBindingManager.cpp
+++ b/Engine/Source/Editor/UMGEditor/Private/Animation/UMGBindingManager.cpp
@@ -17,7 +17,7 @@ UUMGBindingManager::UUMGBindingManager(const FObjectInitializer& ObjectInitializ
 /* IMovieSceneBindingManager interface
  *****************************************************************************/
 
-void UUMGBindingManager::BindPossessableObject(const FGuid& PossessableGuid, UObject& PossessedObject)
+void UUMGBindingManager::BindPossessableObject(const FMovieSceneObjectId& PossessableGuid, UObject& PossessedObject)
 {
 
 }
@@ -29,9 +29,9 @@ bool UUMGBindingManager::CanPossessObject(UObject& Object) const
 }
 
 
-FGuid UUMGBindingManager::FindGuidForObject(const UMovieScene& MovieScene, UObject& Object) const
+FMovieSceneObjectId UUMGBindingManager::FindGuidForObject(const UMovieScene& MovieScene, UObject& Object) const
 {
-	return FGuid();
+	return FMovieSceneObjectId();
 }
 
 
@@ -47,7 +47,7 @@ bool UUMGBindingManager::TryGetObjectBindingDisplayName(const TSharedRef<FMovieS
 }
 
 
-void UUMGBindingManager::UnbindPossessableObjects(const FGuid& PossessableGuid)
+void UUMGBindingManager::UnbindPossessableObjects(const FMovieSceneObjectId& PossessableGuid)
 {
 
 }
diff --git a/Engine/Source/Editor/UMGEditor/Private/Animation/UMGBindingManager.h b/Engine/Source/Editor/UMGEditor/Private/Animation/UMGBindingManager.h
--- a/Engine/Source/Editor/UMGEditor/Private/Animation/UMGBindingManager.h
+++ b/Engine/Source/Editor/UMGEditor/Private/Animation/UMGBindingManager.h
@@ -20,10 +20,10 @@ public:
 
 	// IMovieSceneBindingManager interface
 
-	virtual void BindPossessableObject(const FGuid& PossessableGuid, UObject& PossessedObject) override;
+	virtual void BindPossessableObject(const FMovieSceneObjectId& PossessableGuid, UObject& PossessedObject) override;
 	virtual bool CanPossessObject( UObject& Object ) const override;
-	virtual FGuid FindGuidForObject(const UMovieScene& MovieScene, UObject& Object) const override;
+	virtual FMovieSceneObjectId FindGuidForObject(const UMovieScene& MovieScene, UObject& Object) const override;
 	virtual void GetRuntimeObjects(const TSharedRef<FMovieSceneInstance>& MovieSceneInstance, const FGuid& ObjectGuid, TArray<UObject*>& OutRuntimeObjects) const override;
 	virtual bool TryGetObjectBindingDisplayName(const TSharedRef<FMovieSceneInstance>& MovieSceneInstance, const FGuid& ObjectGuid, FText& DisplayName) const override;
-	virtual void UnbindPossessableObjects(const FGuid& PossessableGuid) override;
+	virtual void UnbindPossessableObjects(const FMovieSceneObjectId& PossessableGuid) override;
 };
diff --git a/Engine/Source/Runtime/MovieScene/Public/MovieSceneBindingManager.h b/Engine/Source/Runtime/MovieScene/Public/MovieSceneBindingManager.h
--- a/Engine/Source/Runtime/MovieScene/Public/MovieSceneBindingManager.h
+++ b/Engine/Source/Runtime/MovieScene/Public/MovieSceneBindingManager.h
@@ -10,6 +10,82 @@ class UMovieScene;
 class UObject;
 
 
+/**
+ * Identifies an object (or embedded component) bound in a movie scene.
+ */
+USTRUCT()
+struct FMovieSceneObjectId
+{
+	GENERATED_USTRUCT_BODY()
+
+	/** Unique identifier of a bound object. */
+	UPROPERTY()
+	FGuid ObjectGuid;
+
+	/** Name of the bound object's component object being used (optional). */
+	UPROPERTY()
+	FName ComponentName;
+
+public:
+
+	/** Default constructor. */
+	FMovieSceneObjectId() { }
+
+	/** Creates and initializes a new instance from an object ID. */
+	FMovieSceneObjectId(const FGuid& InObjectGuid)
+		: ObjectGuid(InObjectGuid)
+	{ }
+
+	/** Creates and initializes a new instance from an object and component ID. */
+	FMovieSceneObjectId(const FGuid& InObjectGuid, const FName& InComponentName)
+		: ObjectGuid(InObjectGuid)
+		, ComponentName(InComponentName)
+	{ }
+
+public:
+
+	/**
+	 * Compares two object identifiers for equality.
+	 *
+	 * @param X The first identifier to compare.
+	 * @param Y The second identifier to compare.
+	 * @return true if the identifiers are equal, false otherwise.
+	 */
+	friend bool operator==( const FMovieSceneObjectId& X, const FMovieSceneObjectId& Y )
+	{
+		return ((X.ObjectGuid == Y.ObjectGuid) && (X.ComponentName == X.ComponentName));
+	}
+
+	/**
+	 * Compares two GUIDs for inequality.
+	 *
+	 * @param X The first GUID to compare.
+	 * @param Y The second GUID to compare.
+	 * @return true if the GUIDs are not equal, false otherwise.
+	 */
+	friend bool operator!=( const FMovieSceneObjectId& X, const FMovieSceneObjectId& Y )
+	{
+		return ((X.ObjectGuid != Y.ObjectGuid) || (X.ComponentName != X.ComponentName));
+	}
+
+	/**
+	 * Checks whether this GUID is valid or not.
+	 *
+	 * A GUID that has all its components set to zero is considered invalid.
+	 *
+	 * @return true if valid, false otherwise.
+	 * @see Invalidate
+	 */
+	bool IsValid() const
+	{
+		return ObjectGuid.IsValid();
+	}
+};
+
+
+/**
+ * Interface for movie scene object binding managers.
+ */
 UINTERFACE(MinimalAPI, meta=(CannotImplementInterfaceInBlueprint))
 class UMovieSceneBindingManager
 	: public UInterface
@@ -18,6 +94,9 @@ class UMovieSceneBindingManager
 };
 
 
+/**
+ * Interface for movie scene object binding managers.
+ */
 class IMovieSceneBindingManager
 {
 	GENERATED_IINTERFACE_BODY()
@@ -28,7 +107,7 @@ class IMovieSceneBindingManager
 	 * @param PossessableGuid The guid used to map to the possessable object.  Note the guid can be bound to multiple objects at once
 	 * @param PossessedObject The runtime object which was possessed.
 	 */
-	virtual void BindPossessableObject(const FGuid& PossessableGuid, UObject& PossessedObject) PURE_VIRTUAL(IMovieSceneBindingManager::BindPossessableObject,);
+	virtual void BindPossessableObject(const FMovieSceneObjectId& PossessableGuid, UObject& PossessedObject) PURE_VIRTUAL(IMovieSceneBindingManager::BindPossessableObject,);
 	
 	/**
 	 * @todo sequencer: add documentation
@@ -40,9 +119,9 @@ class IMovieSceneBindingManager
 	 *
 	 * @param MovieScene The movie scene that may contain the binding
 	 * @param The object to get the guid for
-	 * @return The object guid or an invalid guid if not found
+	 * @return The object identifier or an invalid guid if not found
 	 */
-	virtual FGuid FindGuidForObject(const UMovieScene& MovieScene, UObject& Object) const PURE_VIRTUAL(IMovieSceneBindingManager::FindGuidForObject, return FGuid(););
+	virtual FMovieSceneObjectId FindGuidForObject(const UMovieScene& MovieScene, UObject& Object) const PURE_VIRTUAL(IMovieSceneBindingManager::FindGuidForObject, return FGuid(););
 
 	/**
 	 * @todo sequencer: add documentation
@@ -64,5 +143,5 @@ class IMovieSceneBindingManager
 	 *
 	 * @param PossessableGuid The guid bound to possessable objects that should be removed
 	 */
-	virtual void UnbindPossessableObjects(const FGuid& PossessableGuid) PURE_VIRTUAL(IMovieSceneBindingManager::UnbindPossessableObjects,);
+	virtual void UnbindPossessableObjects(const FMovieSceneObjectId& PossessableGuid) PURE_VIRTUAL(IMovieSceneBindingManager::UnbindPossessableObjects,);
 };
