MovieScene: fixed comments

diff --git a/Engine/Source/Runtime/MovieScene/Public/MovieSceneBindingManager.h b/Engine/Source/Runtime/MovieScene/Public/MovieSceneBindingManager.h
--- a/Engine/Source/Runtime/MovieScene/Public/MovieSceneBindingManager.h
+++ b/Engine/Source/Runtime/MovieScene/Public/MovieSceneBindingManager.h
@@ -57,11 +57,11 @@ public:
 	}
 
 	/**
-	 * Compares two GUIDs for inequality.
+	 * Compares two object identifiers for inequality.
 	 *
-	 * @param X The first GUID to compare.
-	 * @param Y The second GUID to compare.
-	 * @return true if the GUIDs are not equal, false otherwise.
+	 * @param X The first identifier to compare.
+	 * @param Y The second identifier to compare.
+	 * @return true if the identifier are not equal, false otherwise.
 	 */
 	friend bool operator!=( const FMovieSceneObjectId& X, const FMovieSceneObjectId& Y )
 	{
@@ -69,12 +69,9 @@ public:
 	}
 
 	/**
-	 * Checks whether this GUID is valid or not.
-	 *
-	 * A GUID that has all its components set to zero is considered invalid.
+	 * Checks whether this object identifier is valid or not.
 	 *
 	 * @return true if valid, false otherwise.
-	 * @see Invalidate
 	 */
 	bool IsValid() const
 	{
