Sub-movie scene work
-Sub-movie scenes and shots can now be properly deleted
-Sub-movie scenes now properly update their spawnables when they are created/deleted
-Movie scene instances now have a chance to clean up their inner movie scene data when they are removed
-Movie scene shot tracks now properly evaluate any inner tracks they have
-Movie scene sections are now evaluated at t=0 instead of being skipped due to TRange not working if both start and end time are zero

diff --git a/Engine/Source/Editor/MovieSceneTools/Private/TrackEditors/ShotTrackEditor.cpp b/Engine/Source/Editor/MovieSceneTools/Private/TrackEditors/ShotTrackEditor.cpp
--- a/Engine/Source/Editor/MovieSceneTools/Private/TrackEditors/ShotTrackEditor.cpp
+++ b/Engine/Source/Editor/MovieSceneTools/Private/TrackEditors/ShotTrackEditor.cpp
@@ -58,7 +58,6 @@ void FShotThumbnailPool::DrawThumbnails()
 	if (ThumbnailsNeedingDraw.Num() > 0)
 	{
 		// save the global time
-		Sequencer.Pin()->SetPerspectiveViewportPossessionEnabled(false);
 		float SavedTime = Sequencer.Pin()->GetGlobalTime();
 
 		uint32 ThumbnailsDrawn = 0;
@@ -68,7 +67,13 @@ void FShotThumbnailPool::DrawThumbnails()
 			
 			if (Thumbnail->IsVisible())
 			{
+			
+				Sequencer.Pin()->SetPerspectiveViewportPossessionEnabled(false);
+			
 				Thumbnail->DrawThumbnail();
+
+			
+				Sequencer.Pin()->SetPerspectiveViewportPossessionEnabled(true);
 				++ThumbnailsDrawn;
 
 				ThumbnailsNeedingDraw.Remove(Thumbnail);
@@ -80,10 +85,13 @@ void FShotThumbnailPool::DrawThumbnails()
 				ThumbnailsNeedingDraw.Remove(Thumbnail);
 			}
 		}
+		
+		if( ThumbnailsDrawn > 0 )
+		{
+			// restore the global time
+			Sequencer.Pin()->SetGlobalTime(SavedTime);
+		}
 
-		// restore the global time
-		Sequencer.Pin()->SetGlobalTime(SavedTime);
-		Sequencer.Pin()->SetPerspectiveViewportPossessionEnabled(true);
 	}
 }
 
diff --git a/Engine/Source/Editor/NiagaraEditor/Private/NiagaraSequencer.h b/Engine/Source/Editor/NiagaraEditor/Private/NiagaraSequencer.h
--- a/Engine/Source/Editor/NiagaraEditor/Private/NiagaraSequencer.h
+++ b/Engine/Source/Editor/NiagaraEditor/Private/NiagaraSequencer.h
@@ -120,6 +120,10 @@ public:
 	}
 
 	virtual void RefreshInstance(const TArray<UObject*>& RuntimeObjects, class IMovieScenePlayer& Player);
+
+	virtual void ClearInstance( IMovieScenePlayer& Player ) override {}
+	virtual void SaveState(const TArray<UObject*>& RuntimeObjects) override {}
+	virtual void RestoreState(const TArray<UObject*>& RuntimeObjects) override {}
 private:
 	TSharedPtr<FNiagaraSimulation> Emitter;
 	class UEmitterMovieSceneTrack *Track;
diff --git a/Engine/Source/Editor/Sequencer/Private/Sequencer.cpp b/Engine/Source/Editor/Sequencer/Private/Sequencer.cpp
--- a/Engine/Source/Editor/Sequencer/Private/Sequencer.cpp
+++ b/Engine/Source/Editor/Sequencer/Private/Sequencer.cpp
@@ -718,9 +718,12 @@ EMovieScenePlayerStatus::Type FSequencer::GetPlaybackStatus() const
 	return PlaybackState;
 }
 
-void FSequencer::AddMovieSceneInstance( UMovieSceneSection& MovieSceneSection, TSharedRef<FMovieSceneInstance> InstanceToAdd )
+void FSequencer::AddOrUpdateMovieSceneInstance( UMovieSceneSection& MovieSceneSection, TSharedRef<FMovieSceneInstance> InstanceToAdd )
 {
-	MovieSceneSectionToInstanceMap.Add( &MovieSceneSection, InstanceToAdd );
+	if( !MovieSceneSectionToInstanceMap.Contains( &MovieSceneSection ) )
+	{
+		MovieSceneSectionToInstanceMap.Add( &MovieSceneSection, InstanceToAdd );
+	}
 
 	SpawnOrDestroyPuppetObjects( InstanceToAdd );
 }
@@ -1169,6 +1172,10 @@ void FSequencer::OnRequestNodeDeleted( TSharedRef<const FSequencerDisplayNode>&
 		{
 			FocusedMovieScene->RemoveMasterTrack( Track );
 		}
+		else if( FocusedMovieScene->GetShotTrack() == Track )
+		{
+			FocusedMovieScene->RemoveShotTrack();
+		}
 		else
 		{
 			FocusedMovieScene->RemoveTrack( Track );
@@ -1180,12 +1187,6 @@ void FSequencer::OnRequestNodeDeleted( TSharedRef<const FSequencerDisplayNode>&
 	
 	if( bAnythingRemoved )
 	{
-		if( bAnySpawnablesRemoved )
-		{
-			// @todo Sequencer Sub-MovieScenes needs to destroy objects for all movie scenes that had this node
-			SpawnOrDestroyPuppetObjects( MovieSceneInstance );
-		}
-
 		NotifyMovieSceneDataChanged();
 	}
 }
diff --git a/Engine/Source/Editor/Sequencer/Private/Sequencer.h b/Engine/Source/Editor/Sequencer/Private/Sequencer.h
--- a/Engine/Source/Editor/Sequencer/Private/Sequencer.h
+++ b/Engine/Source/Editor/Sequencer/Private/Sequencer.h
@@ -229,7 +229,7 @@ public:
 	virtual void GetRuntimeObjects( TSharedRef<FMovieSceneInstance> MovieSceneInstance, const FGuid& ObjectHandle, TArray< UObject* >& OutObjects ) const override;
 	virtual void UpdateCameraCut(UObject* ObjectToViewThrough, bool bNewCameraCut) const override;
 	virtual EMovieScenePlayerStatus::Type GetPlaybackStatus() const override;
-	virtual void AddMovieSceneInstance( class UMovieSceneSection& MovieSceneSection, TSharedRef<FMovieSceneInstance> InstanceToAdd ) override;
+	virtual void AddOrUpdateMovieSceneInstance( class UMovieSceneSection& MovieSceneSection, TSharedRef<FMovieSceneInstance> InstanceToAdd ) override;
 	virtual void RemoveMovieSceneInstance( class UMovieSceneSection& MovieSceneSection, TSharedRef<FMovieSceneInstance> InstanceToRemove ) override;
 
 	virtual void SpawnActorsForMovie( TSharedRef<FMovieSceneInstance> MovieSceneInstance );
diff --git a/Engine/Source/Runtime/MovieScene/Private/MovieScene.cpp b/Engine/Source/Runtime/MovieScene/Private/MovieScene.cpp
--- a/Engine/Source/Runtime/MovieScene/Private/MovieScene.cpp
+++ b/Engine/Source/Runtime/MovieScene/Private/MovieScene.cpp
@@ -384,6 +384,16 @@ UMovieSceneTrack* UMovieScene::GetShotTrack()
 	return ShotTrack;
 }
 
+void UMovieScene::RemoveShotTrack()
+{
+	if( ShotTrack )
+	{
+		Modify();
+
+		ShotTrack = nullptr;
+	}
+}
+
 bool UMovieScene::RemoveMasterTrack( UMovieSceneTrack* Track ) 
 {
 	Modify();
diff --git a/Engine/Source/Runtime/MovieScene/Private/MovieSceneCommonHelpers.cpp b/Engine/Source/Runtime/MovieScene/Private/MovieSceneCommonHelpers.cpp
--- a/Engine/Source/Runtime/MovieScene/Private/MovieSceneCommonHelpers.cpp
+++ b/Engine/Source/Runtime/MovieScene/Private/MovieSceneCommonHelpers.cpp
@@ -19,7 +19,7 @@ TArray<UMovieSceneSection*> MovieSceneHelpers::GetTraversedSections( const TArra
 	for( int32 SectionIndex = 0; SectionIndex < Sections.Num(); ++SectionIndex )
 	{
 		UMovieSceneSection* Section = Sections[SectionIndex];
-		if( SliderRange.Overlaps( TRange<float>( Section->GetRange() ) ) )
+		if( Section->GetStartTime() == CurrentTime || SliderRange.Overlaps( TRange<float>( Section->GetRange() ) ) )
 		{
 			TraversedSections.Add( Section );
 		}
diff --git a/Engine/Source/Runtime/MovieScene/Private/MovieSceneInstance.cpp b/Engine/Source/Runtime/MovieScene/Private/MovieSceneInstance.cpp
--- a/Engine/Source/Runtime/MovieScene/Private/MovieSceneInstance.cpp
+++ b/Engine/Source/Runtime/MovieScene/Private/MovieSceneInstance.cpp
@@ -70,10 +70,15 @@ void FMovieSceneInstance::RefreshInstance( IMovieScenePlayer& Player )
 
 	TSharedRef<FMovieSceneInstance> ThisInstance = AsShared();
 
+	FMovieSceneInstanceMap ShotTrackInstanceMap;
 	// Only if root movie scene. Any sub-movie scene that has a shot track is ignored
 	if( ShotTrack && Player.GetRootMovieSceneInstance() == ThisInstance )
 	{
-		FMovieSceneInstanceMap ShotTrackInstanceMap;
+		if( ShotTrackInstance.IsValid() )
+		{
+			ShotTrackInstanceMap.Add(  ShotTrack, ShotTrackInstance );
+		}
+
 		TArray<UObject*> Objects;
 		TArray<UMovieSceneTrack*> Tracks;
 		Tracks.Add(ShotTrack);
@@ -81,6 +86,11 @@ void FMovieSceneInstance::RefreshInstance( IMovieScenePlayer& Player )
 
 		ShotTrackInstance = ShotTrackInstanceMap.FindRef( ShotTrack );
 	}
+	else if( !ShotTrack && ShotTrackInstance.IsValid() )
+	{
+		ShotTrackInstance->ClearInstance( Player );
+		ShotTrackInstance.Reset();
+	}
 
 	// Get all the master tracks and create instances for them if needed
 	const TArray<UMovieSceneTrack*>& MasterTracks = MovieScene->GetMasterTracks();
@@ -162,6 +172,8 @@ void FMovieSceneInstance::RefreshInstanceMap( const TArray<UMovieSceneTrack*>& T
 	{
 		if( !FoundTracks.Contains( It.Key().Get() ) )
 		{
+			It.Value()->ClearInstance( Player );
+
 			// This track was not found in the moviescene's track list so it was removed.
 			It.RemoveCurrent();
 		}
diff --git a/Engine/Source/Runtime/MovieScene/Private/RuntimeMovieScenePlayer.cpp b/Engine/Source/Runtime/MovieScene/Private/RuntimeMovieScenePlayer.cpp
--- a/Engine/Source/Runtime/MovieScene/Private/RuntimeMovieScenePlayer.cpp
+++ b/Engine/Source/Runtime/MovieScene/Private/RuntimeMovieScenePlayer.cpp
@@ -211,7 +211,7 @@ EMovieScenePlayerStatus::Type URuntimeMovieScenePlayer::GetPlaybackStatus() cons
 	return bIsPlaying ? EMovieScenePlayerStatus::Playing : EMovieScenePlayerStatus::Stopped;
 }
 
-void URuntimeMovieScenePlayer::AddMovieSceneInstance( UMovieSceneSection& MovieSceneSection, TSharedRef<FMovieSceneInstance> InstanceToAdd )
+void URuntimeMovieScenePlayer::AddOrUpdateMovieSceneInstance( UMovieSceneSection& MovieSceneSection, TSharedRef<FMovieSceneInstance> InstanceToAdd )
 {
 	SpawnActorsForMovie( InstanceToAdd );
 }
diff --git a/Engine/Source/Runtime/MovieScene/Public/IMovieScenePlayer.h b/Engine/Source/Runtime/MovieScene/Public/IMovieScenePlayer.h
--- a/Engine/Source/Runtime/MovieScene/Public/IMovieScenePlayer.h
+++ b/Engine/Source/Runtime/MovieScene/Public/IMovieScenePlayer.h
@@ -47,7 +47,7 @@ public:
 	 * @param MovieSceneSection	The section owning the MovieScene being instanced. 
 	 * @param InstanceToAdd		The instance being added
 	 */
-	virtual void AddMovieSceneInstance( class UMovieSceneSection& MovieSceneSection, TSharedRef<FMovieSceneInstance> InstanceToAdd ) = 0;
+	virtual void AddOrUpdateMovieSceneInstance( class UMovieSceneSection& MovieSceneSection, TSharedRef<FMovieSceneInstance> InstanceToAdd ) = 0;
 
 	/**
 	 * Removes a MovieScene instance from the player. 
diff --git a/Engine/Source/Runtime/MovieScene/Public/IMovieSceneTrackInstance.h b/Engine/Source/Runtime/MovieScene/Public/IMovieSceneTrackInstance.h
--- a/Engine/Source/Runtime/MovieScene/Public/IMovieSceneTrackInstance.h
+++ b/Engine/Source/Runtime/MovieScene/Public/IMovieSceneTrackInstance.h
@@ -19,12 +19,12 @@ public:
 	/**
 	 * Save state of objects that this instance will be editing.
 	 */
-	virtual void SaveState(const TArray<UObject*>& RuntimeObjects) { }
+	virtual void SaveState(const TArray<UObject*>& RuntimeObjects) = 0;
 
 	/**
 	 * Restore state of objects that this instance edited.
 	 */
-	virtual void RestoreState(const TArray<UObject*>& RuntimeObjects) { }
+	virtual void RestoreState(const TArray<UObject*>& RuntimeObjects) = 0;
 
 	/**
 	 * Main update function for track instances.  Called in game and in editor when we update a moviescene.
@@ -40,4 +40,11 @@ public:
 	 * Refreshes the current instance
 	 */
 	virtual void RefreshInstance(const TArray<UObject*>& RuntimeObjects, IMovieScenePlayer& Player) = 0;
+
+	/**
+	 * Removes all instance data from this track instance
+	 *
+	 * Called before an instance is removed
+	 */
+	virtual void ClearInstance( IMovieScenePlayer& Player ) = 0;
 };
diff --git a/Engine/Source/Runtime/MovieScene/Public/MovieScene.h b/Engine/Source/Runtime/MovieScene/Public/MovieScene.h
--- a/Engine/Source/Runtime/MovieScene/Public/MovieScene.h
+++ b/Engine/Source/Runtime/MovieScene/Public/MovieScene.h
@@ -390,6 +390,11 @@ public:
 	/** @return The shot track if it exists */
 	UMovieSceneTrack* GetShotTrack();
 
+	/**
+	 * Removes the shot track if it exists
+	 */
+	void RemoveShotTrack();
+
 	/**
 	 * Removes a master track.
 	 *
diff --git a/Engine/Source/Runtime/MovieScene/Public/RuntimeMovieScenePlayer.h b/Engine/Source/Runtime/MovieScene/Public/RuntimeMovieScenePlayer.h
--- a/Engine/Source/Runtime/MovieScene/Public/RuntimeMovieScenePlayer.h
+++ b/Engine/Source/Runtime/MovieScene/Public/RuntimeMovieScenePlayer.h
@@ -60,7 +60,7 @@ protected:
 	virtual void GetRuntimeObjects( TSharedRef<FMovieSceneInstance> MovieSceneInstance, const FGuid& ObjectHandle, TArray< UObject* >& OutObjects ) const override;
 	virtual void UpdateCameraCut(UObject* ObjectToViewThrough, bool bNewCameraCut) const override;
 	virtual EMovieScenePlayerStatus::Type GetPlaybackStatus() const override;
-	virtual void AddMovieSceneInstance( class UMovieSceneSection& MovieSceneSection, TSharedRef<FMovieSceneInstance> InstanceToAdd ) override;
+	virtual void AddOrUpdateMovieSceneInstance( class UMovieSceneSection& MovieSceneSection, TSharedRef<FMovieSceneInstance> InstanceToAdd ) override;
 	virtual void RemoveMovieSceneInstance( class UMovieSceneSection& MovieSceneSection, TSharedRef<FMovieSceneInstance> InstanceToRemove ) override;
 	virtual TSharedRef<FMovieSceneInstance> GetRootMovieSceneInstance() const override { return RootMovieSceneInstance.ToSharedRef(); }
 	/** IRuntimeMovieScenePlayerInterface */
diff --git a/Engine/Source/Runtime/MovieSceneTracks/Private/TrackInstances/MovieSceneShotTrackInstance.cpp b/Engine/Source/Runtime/MovieSceneTracks/Private/TrackInstances/MovieSceneShotTrackInstance.cpp
--- a/Engine/Source/Runtime/MovieSceneTracks/Private/TrackInstances/MovieSceneShotTrackInstance.cpp
+++ b/Engine/Source/Runtime/MovieSceneTracks/Private/TrackInstances/MovieSceneShotTrackInstance.cpp
@@ -39,6 +39,8 @@ void FMovieSceneShotTrackInstance::RefreshInstance( const TArray<UObject*>& Runt
 
 void FMovieSceneShotTrackInstance::Update( float Position, float LastPosition, const TArray<UObject*>& RuntimeObjects, class IMovieScenePlayer& Player ) 
 {
+	FSubMovieSceneTrackInstance::Update( Position, LastPosition, RuntimeObjects, Player );
+
 	const TArray<UMovieSceneSection*>& ShotSections = SubMovieSceneTrack->GetAllSections();
 
 	for (int32 ShotIndex = 0; ShotIndex < ShotSections.Num(); ++ShotIndex)
@@ -61,3 +63,9 @@ void FMovieSceneShotTrackInstance::Update( float Position, float LastPosition, c
 		}
 	}
 }
+
+void FMovieSceneShotTrackInstance::ClearInstance( IMovieScenePlayer& Player )
+{
+	FSubMovieSceneTrackInstance::ClearInstance( Player );
+	Player.UpdateCameraCut( nullptr, false );
+}
diff --git a/Engine/Source/Runtime/MovieSceneTracks/Private/TrackInstances/SubMovieSceneTrackInstance.cpp b/Engine/Source/Runtime/MovieSceneTracks/Private/TrackInstances/SubMovieSceneTrackInstance.cpp
--- a/Engine/Source/Runtime/MovieSceneTracks/Private/TrackInstances/SubMovieSceneTrackInstance.cpp
+++ b/Engine/Source/Runtime/MovieSceneTracks/Private/TrackInstances/SubMovieSceneTrackInstance.cpp
@@ -37,11 +37,11 @@ void FSubMovieSceneTrackInstance::RefreshInstance( const TArray<UObject*>& Runti
 			{
 				Instance = MakeShareable(new FMovieSceneInstance(*Section->GetMovieScene()));
 
-				Player.AddMovieSceneInstance(*Section, Instance.ToSharedRef());
-
 				SubMovieSceneInstances.Add(Section, Instance.ToSharedRef());
 			}
 
+			Player.AddOrUpdateMovieSceneInstance(*Section, Instance.ToSharedRef());
+
 			// Refresh the existing instance
 			Instance->RefreshInstance(Player);
 		}
@@ -118,4 +118,13 @@ void FSubMovieSceneTrackInstance::RestoreState(const TArray<UObject*>& RuntimeOb
 			Instance->RestoreState();
 		}
 	}
+}
+
+void FSubMovieSceneTrackInstance::ClearInstance( IMovieScenePlayer& Player )
+{
+	TMap< TWeakObjectPtr<USubMovieSceneSection>, TSharedPtr<FMovieSceneInstance> >::TIterator It =  SubMovieSceneInstances.CreateIterator();
+	for( ; It; ++It )
+	{
+		Player.RemoveMovieSceneInstance( *It.Key().Get(), It.Value().ToSharedRef() );
+	}
 }
\ No newline at end of file
diff --git a/Engine/Source/Runtime/MovieSceneTracks/Public/TrackInstances/MovieScene3DTransformTrackInstance.h b/Engine/Source/Runtime/MovieSceneTracks/Public/TrackInstances/MovieScene3DTransformTrackInstance.h
--- a/Engine/Source/Runtime/MovieSceneTracks/Public/TrackInstances/MovieScene3DTransformTrackInstance.h
+++ b/Engine/Source/Runtime/MovieSceneTracks/Public/TrackInstances/MovieScene3DTransformTrackInstance.h
@@ -23,6 +23,7 @@ public:
 	virtual void RestoreState (const TArray<UObject*>& RuntimeObjects) override;
 	virtual void Update( float Position, float LastPosition, const TArray<UObject*>& RuntimeObjects, IMovieScenePlayer& Player ) override;
 	virtual void RefreshInstance( const TArray<UObject*>& RuntimeObjects, IMovieScenePlayer& Player ) override {}
+	virtual void ClearInstance( IMovieScenePlayer& Player ) override {}
 
 private:
 
diff --git a/Engine/Source/Runtime/MovieSceneTracks/Public/TrackInstances/MovieSceneAnimationTrackInstance.h b/Engine/Source/Runtime/MovieSceneTracks/Public/TrackInstances/MovieSceneAnimationTrackInstance.h
--- a/Engine/Source/Runtime/MovieSceneTracks/Public/TrackInstances/MovieSceneAnimationTrackInstance.h
+++ b/Engine/Source/Runtime/MovieSceneTracks/Public/TrackInstances/MovieSceneAnimationTrackInstance.h
@@ -20,8 +20,11 @@ public:
 	virtual ~FMovieSceneAnimationTrackInstance();
 
 	/** IMovieSceneTrackInstance interface */
+	virtual void SaveState(const TArray<UObject*>& RuntimeObjects) override {}
+	virtual void RestoreState(const TArray<UObject*>& RuntimeObjects) override {}
 	virtual void Update( float Position, float LastPosition, const TArray<UObject*>& RuntimeObjects, IMovieScenePlayer& Player ) override;
 	virtual void RefreshInstance( const TArray<UObject*>& RuntimeObjects, IMovieScenePlayer& Player ) override {}
+	virtual void ClearInstance( IMovieScenePlayer& Player ) override {}
 
 private:
 
diff --git a/Engine/Source/Runtime/MovieSceneTracks/Public/TrackInstances/MovieSceneAudioTrackInstance.h b/Engine/Source/Runtime/MovieSceneTracks/Public/TrackInstances/MovieSceneAudioTrackInstance.h
--- a/Engine/Source/Runtime/MovieSceneTracks/Public/TrackInstances/MovieSceneAudioTrackInstance.h
+++ b/Engine/Source/Runtime/MovieSceneTracks/Public/TrackInstances/MovieSceneAudioTrackInstance.h
@@ -20,8 +20,11 @@ public:
 	FMovieSceneAudioTrackInstance( UMovieSceneAudioTrack& InAudioTrack );
 
 	/** IMovieSceneTrackInstance interface */
+	virtual void SaveState(const TArray<UObject*>& RuntimeObjects) override {}
+	virtual void RestoreState(const TArray<UObject*>& RuntimeObjects) override {}
 	virtual void Update( float Position, float LastPosition, const TArray<UObject*>& RuntimeObjects, IMovieScenePlayer& Player ) override;
 	virtual void RefreshInstance( const TArray<UObject*>& RuntimeObjects, IMovieScenePlayer& Player ) override {}
+	virtual void ClearInstance( IMovieScenePlayer& Player ) override {}
 
 private:
 
diff --git a/Engine/Source/Runtime/MovieSceneTracks/Public/TrackInstances/MovieSceneBoolTrackInstance.h b/Engine/Source/Runtime/MovieSceneTracks/Public/TrackInstances/MovieSceneBoolTrackInstance.h
--- a/Engine/Source/Runtime/MovieSceneTracks/Public/TrackInstances/MovieSceneBoolTrackInstance.h
+++ b/Engine/Source/Runtime/MovieSceneTracks/Public/TrackInstances/MovieSceneBoolTrackInstance.h
@@ -24,6 +24,7 @@ public:
 	virtual void RestoreState (const TArray<UObject*>& RuntimeObjects) override;
 	virtual void Update( float Position, float LastPosition, const TArray<UObject*>& RuntimeObjects, IMovieScenePlayer& Player ) override;
 	virtual void RefreshInstance( const TArray<UObject*>& RuntimeObjects, IMovieScenePlayer& Player ) override;
+	virtual void ClearInstance( IMovieScenePlayer& Player ) override {}
 
 private:
 
diff --git a/Engine/Source/Runtime/MovieSceneTracks/Public/TrackInstances/MovieSceneByteTrackInstance.h b/Engine/Source/Runtime/MovieSceneTracks/Public/TrackInstances/MovieSceneByteTrackInstance.h
--- a/Engine/Source/Runtime/MovieSceneTracks/Public/TrackInstances/MovieSceneByteTrackInstance.h
+++ b/Engine/Source/Runtime/MovieSceneTracks/Public/TrackInstances/MovieSceneByteTrackInstance.h
@@ -24,6 +24,7 @@ public:
 	virtual void RestoreState (const TArray<UObject*>& RuntimeObjects) override;
 	virtual void Update( float Position, float LastPosition, const TArray<UObject*>& RuntimeObjects, IMovieScenePlayer& Player ) override;
 	virtual void RefreshInstance( const TArray<UObject*>& RuntimeObjects, IMovieScenePlayer& Player ) override;
+	virtual void ClearInstance( IMovieScenePlayer& Player ) override {}
 
 private:
 
diff --git a/Engine/Source/Runtime/MovieSceneTracks/Public/TrackInstances/MovieSceneColorTrackInstance.h b/Engine/Source/Runtime/MovieSceneTracks/Public/TrackInstances/MovieSceneColorTrackInstance.h
--- a/Engine/Source/Runtime/MovieSceneTracks/Public/TrackInstances/MovieSceneColorTrackInstance.h
+++ b/Engine/Source/Runtime/MovieSceneTracks/Public/TrackInstances/MovieSceneColorTrackInstance.h
@@ -24,6 +24,7 @@ public:
 	virtual void RestoreState (const TArray<UObject*>& RuntimeObjects) override;
 	virtual void Update( float Position, float LastPosition, const TArray<UObject*>& RuntimeObjects, IMovieScenePlayer& Player ) override;
 	virtual void RefreshInstance( const TArray<UObject*>& RuntimeObjects, IMovieScenePlayer& Player ) override;
+	virtual void ClearInstance( IMovieScenePlayer& Player ) override {}
 
 private:
 
diff --git a/Engine/Source/Runtime/MovieSceneTracks/Public/TrackInstances/MovieSceneFloatTrackInstance.h b/Engine/Source/Runtime/MovieSceneTracks/Public/TrackInstances/MovieSceneFloatTrackInstance.h
--- a/Engine/Source/Runtime/MovieSceneTracks/Public/TrackInstances/MovieSceneFloatTrackInstance.h
+++ b/Engine/Source/Runtime/MovieSceneTracks/Public/TrackInstances/MovieSceneFloatTrackInstance.h
@@ -24,6 +24,7 @@ public:
 	virtual void RestoreState (const TArray<UObject*>& RuntimeObjects) override;
 	virtual void Update( float Position, float LastPosition, const TArray<UObject*>& RuntimeObjects, IMovieScenePlayer& Player ) override;
 	virtual void RefreshInstance( const TArray<UObject*>& RuntimeObjects, IMovieScenePlayer& Player ) override;
+	virtual void ClearInstance( IMovieScenePlayer& Player ) override {}
 
 private:
 
diff --git a/Engine/Source/Runtime/MovieSceneTracks/Public/TrackInstances/MovieSceneParticleTrackInstance.h b/Engine/Source/Runtime/MovieSceneTracks/Public/TrackInstances/MovieSceneParticleTrackInstance.h
--- a/Engine/Source/Runtime/MovieSceneTracks/Public/TrackInstances/MovieSceneParticleTrackInstance.h
+++ b/Engine/Source/Runtime/MovieSceneTracks/Public/TrackInstances/MovieSceneParticleTrackInstance.h
@@ -23,8 +23,11 @@ public:
 	virtual ~FMovieSceneParticleTrackInstance();
 
 	/** IMovieSceneTrackInstance interface */
+	virtual void SaveState(const TArray<UObject*>& RuntimeObjects) override {}
+	virtual void RestoreState(const TArray<UObject*>& RuntimeObjects) override {}
 	virtual void Update( float Position, float LastPosition, const TArray<UObject*>& RuntimeObjects, IMovieScenePlayer& Player ) override;
 	virtual void RefreshInstance( const TArray<UObject*>& RuntimeObjects, IMovieScenePlayer& Player ) override {}
+	virtual void ClearInstance( IMovieScenePlayer& Player ) override {}
 
 private:
 
diff --git a/Engine/Source/Runtime/MovieSceneTracks/Public/TrackInstances/MovieSceneShotTrackInstance.h b/Engine/Source/Runtime/MovieSceneTracks/Public/TrackInstances/MovieSceneShotTrackInstance.h
--- a/Engine/Source/Runtime/MovieSceneTracks/Public/TrackInstances/MovieSceneShotTrackInstance.h
+++ b/Engine/Source/Runtime/MovieSceneTracks/Public/TrackInstances/MovieSceneShotTrackInstance.h
@@ -20,7 +20,7 @@ public:
 	/** IMovieSceneTrackInstance interface */
 	virtual void RefreshInstance( const TArray<UObject*>& RuntimeObjects, IMovieScenePlayer& Player ) override;
 	virtual void Update( float Position, float LastPosition, const TArray<UObject*>& RuntimeObjects, IMovieScenePlayer& Player ) override;
-
+	virtual void ClearInstance( IMovieScenePlayer& Player ) override;
 private:
 	/** Runtime camera objects.  One for each shot.  Must be the same number of entries as sections */
 	TArray< TWeakObjectPtr<UObject> > RuntimeCameraObjects;
diff --git a/Engine/Source/Runtime/MovieSceneTracks/Public/TrackInstances/MovieSceneVectorTrackInstance.h b/Engine/Source/Runtime/MovieSceneTracks/Public/TrackInstances/MovieSceneVectorTrackInstance.h
--- a/Engine/Source/Runtime/MovieSceneTracks/Public/TrackInstances/MovieSceneVectorTrackInstance.h
+++ b/Engine/Source/Runtime/MovieSceneTracks/Public/TrackInstances/MovieSceneVectorTrackInstance.h
@@ -24,6 +24,7 @@ public:
 	virtual void RestoreState (const TArray<UObject*>& RuntimeObjects) override;
 	virtual void Update( float Position, float LastPosition, const TArray<UObject*>& RuntimeObjects, IMovieScenePlayer& Player ) override;
 	virtual void RefreshInstance( const TArray<UObject*>& RuntimeObjects, IMovieScenePlayer& Player ) override;
+	virtual void ClearInstance( IMovieScenePlayer& Player ) override {}
 
 private:
 
diff --git a/Engine/Source/Runtime/MovieSceneTracks/Public/TrackInstances/MovieSceneVisibilityTrackInstance.h b/Engine/Source/Runtime/MovieSceneTracks/Public/TrackInstances/MovieSceneVisibilityTrackInstance.h
--- a/Engine/Source/Runtime/MovieSceneTracks/Public/TrackInstances/MovieSceneVisibilityTrackInstance.h
+++ b/Engine/Source/Runtime/MovieSceneTracks/Public/TrackInstances/MovieSceneVisibilityTrackInstance.h
@@ -23,6 +23,8 @@ public:
 	virtual void RestoreState (const TArray<UObject*>& RuntimeObjects) override;
 	virtual void Update( float Position, float LastPosition, const TArray<UObject*>& RuntimeObjects, class IMovieScenePlayer& Player ) override;
 	virtual void RefreshInstance( const TArray<UObject*>& RuntimeObjects, class IMovieScenePlayer& Player ) override;
+	virtual void ClearInstance( IMovieScenePlayer& Player ) override {}
+
 private:
 	/** Visibility track that is being instanced */
 	UMovieSceneVisibilityTrack* VisibilityTrack;
diff --git a/Engine/Source/Runtime/MovieSceneTracks/Public/TrackInstances/SubMovieSceneTrackInstance.h b/Engine/Source/Runtime/MovieSceneTracks/Public/TrackInstances/SubMovieSceneTrackInstance.h
--- a/Engine/Source/Runtime/MovieSceneTracks/Public/TrackInstances/SubMovieSceneTrackInstance.h
+++ b/Engine/Source/Runtime/MovieSceneTracks/Public/TrackInstances/SubMovieSceneTrackInstance.h
@@ -25,6 +25,7 @@ public:
 	virtual void RefreshInstance( const TArray<UObject*>& RuntimeObjects, IMovieScenePlayer& Player ) override;
 	virtual void SaveState (const TArray<UObject*>& RuntimeObjects) override;
 	virtual void RestoreState (const TArray<UObject*>& RuntimeObjects) override;
+	virtual void ClearInstance( IMovieScenePlayer& Player ) override;
 protected:
 
 	/** Track that is being instanced */
diff --git a/Engine/Source/Runtime/UMG/Private/Animation/MovieScene2DTransformTrackInstance.h b/Engine/Source/Runtime/UMG/Private/Animation/MovieScene2DTransformTrackInstance.h
--- a/Engine/Source/Runtime/UMG/Private/Animation/MovieScene2DTransformTrackInstance.h
+++ b/Engine/Source/Runtime/UMG/Private/Animation/MovieScene2DTransformTrackInstance.h
@@ -13,8 +13,11 @@ public:
 	FMovieScene2DTransformTrackInstance( class UMovieScene2DTransformTrack& InTransformTrack );
 
 	/** IMovieSceneTrackInstance interface */
+	virtual void SaveState(const TArray<UObject*>& RuntimeObjects) override {}
+	virtual void RestoreState(const TArray<UObject*>& RuntimeObjects) override {}
 	virtual void Update( float Position, float LastPosition, const TArray<UObject*>& RuntimeObjects, class IMovieScenePlayer& Player ) override;
 	virtual void RefreshInstance( const TArray<UObject*>& RuntimeObjects, class IMovieScenePlayer& Player ) override;
+	virtual void ClearInstance( IMovieScenePlayer& Player ) override {};
 private:
 	/** Track that is being instanced */
 	UMovieScene2DTransformTrack* TransformTrack;
diff --git a/Engine/Source/Runtime/UMG/Private/Animation/MovieSceneMarginTrackInstance.h b/Engine/Source/Runtime/UMG/Private/Animation/MovieSceneMarginTrackInstance.h
--- a/Engine/Source/Runtime/UMG/Private/Animation/MovieSceneMarginTrackInstance.h
+++ b/Engine/Source/Runtime/UMG/Private/Animation/MovieSceneMarginTrackInstance.h
@@ -13,8 +13,11 @@ public:
 	FMovieSceneMarginTrackInstance( class UMovieSceneMarginTrack& InMarginTrack );
 
 	/** IMovieSceneTrackInstance interface */
+	virtual void SaveState( const TArray<UObject*>& RuntimeObjects ) override {}
+	virtual void RestoreState( const TArray<UObject*>& RuntimeObjects ) override {}
 	virtual void Update( float Position, float LastPosition, const TArray<UObject*>& RuntimeObjects, class IMovieScenePlayer& Player ) override;
 	virtual void RefreshInstance( const TArray<UObject*>& RuntimeObjects, class IMovieScenePlayer& Player ) override;
+	virtual void ClearInstance( IMovieScenePlayer& Player ) override {};
 private:
 	/** The track being instanced */
 	UMovieSceneMarginTrack* MarginTrack;
diff --git a/Engine/Source/Runtime/UMG/Public/Animation/UMGSequencePlayer.h b/Engine/Source/Runtime/UMG/Public/Animation/UMGSequencePlayer.h
--- a/Engine/Source/Runtime/UMG/Public/Animation/UMGSequencePlayer.h
+++ b/Engine/Source/Runtime/UMG/Public/Animation/UMGSequencePlayer.h
@@ -37,7 +37,7 @@ public:
 	/** IMovieScenePlayer interface */
 	virtual void GetRuntimeObjects( TSharedRef<FMovieSceneInstance> MovieSceneInstance, const FGuid& ObjectHandle, TArray< UObject* >& OutObjects ) const override;
 	virtual void UpdateCameraCut(UObject* ObjectToViewThrough, bool bNewCameraCut) const override {}
-	virtual void AddMovieSceneInstance( class UMovieSceneSection& MovieSceneSection, TSharedRef<FMovieSceneInstance> InstanceToAdd ) override {}
+	virtual void AddOrUpdateMovieSceneInstance( class UMovieSceneSection& MovieSceneSection, TSharedRef<FMovieSceneInstance> InstanceToAdd ) override {}
 	virtual void RemoveMovieSceneInstance( class UMovieSceneSection& MovieSceneSection, TSharedRef<FMovieSceneInstance> InstanceToRemove ) override {}
 	virtual TSharedRef<FMovieSceneInstance> GetRootMovieSceneInstance() const override { return RootMovieSceneInstance.ToSharedRef(); }
 	virtual EMovieScenePlayerStatus::Type GetPlaybackStatus() const override;
