MovieScene: moved FMovieSceneObjectBinding into its own files

diff --git a/Engine/Source/Runtime/MovieScene/Private/MovieScene.cpp b/Engine/Source/Runtime/MovieScene/Private/MovieScene.cpp
--- a/Engine/Source/Runtime/MovieScene/Private/MovieScene.cpp
+++ b/Engine/Source/Runtime/MovieScene/Private/MovieScene.cpp
@@ -22,26 +22,6 @@ FMovieScenePossessable::FMovieScenePossessable( const FString& InitName, UClass*
 }
 
 
-TRange<float> FMovieSceneObjectBinding::GetTimeRange() const
-{
-	TArray< TRange<float> > Bounds;
-	for (int32 TypeIndex = 0; TypeIndex < Tracks.Num(); ++TypeIndex)
-	{
-		Bounds.Add(Tracks[TypeIndex]->GetSectionBoundaries());
-	}
-	return TRange<float>::Hull(Bounds);
-}
-
-void FMovieSceneObjectBinding::AddTrack( UMovieSceneTrack& NewTrack )
-{
-	Tracks.Add( &NewTrack );
-}
-
-bool FMovieSceneObjectBinding::RemoveTrack( UMovieSceneTrack& Track )
-{
-	return Tracks.RemoveSingle( &Track ) != 0;
-}
-
 UMovieScene::UMovieScene( const FObjectInitializer& ObjectInitializer )
 	: Super( ObjectInitializer )
 {
diff --git a/Engine/Source/Runtime/MovieScene/Private/MovieSceneBinding.cpp b/Engine/Source/Runtime/MovieScene/Private/MovieSceneBinding.cpp
new file mode 100644
--- /dev/null
+++ b/Engine/Source/Runtime/MovieScene/Private/MovieSceneBinding.cpp
@@ -0,0 +1,33 @@
+// Copyright 1998-2015 Epic Games, Inc. All Rights Reserved.
+
+#include "MovieScenePrivatePCH.h"
+#include "MovieSceneBinding.h"
+#include "MovieSceneTrack.h"
+
+
+/* FMovieSceneBinding interface
+ *****************************************************************************/
+
+TRange<float> FMovieSceneObjectBinding::GetTimeRange() const
+{
+	TArray<TRange<float>> Bounds;
+
+	for (int32 TypeIndex = 0; TypeIndex < Tracks.Num(); ++TypeIndex)
+	{
+		Bounds.Add(Tracks[TypeIndex]->GetSectionBoundaries());
+	}
+
+	return TRange<float>::Hull(Bounds);
+}
+
+
+void FMovieSceneObjectBinding::AddTrack(UMovieSceneTrack& NewTrack)
+{
+	Tracks.Add( &NewTrack );
+}
+
+
+bool FMovieSceneObjectBinding::RemoveTrack(UMovieSceneTrack& Track)
+{
+	return (Tracks.RemoveSingle(&Track) != 0);
+}
diff --git a/Engine/Source/Runtime/MovieScene/Public/MovieScene.h b/Engine/Source/Runtime/MovieScene/Public/MovieScene.h
--- a/Engine/Source/Runtime/MovieScene/Public/MovieScene.h
+++ b/Engine/Source/Runtime/MovieScene/Public/MovieScene.h
@@ -2,6 +2,7 @@
 
 #pragma once
 
+#include "MovieSceneBinding.h"
 #include "MovieScene.generated.h"
 
 
@@ -43,7 +44,10 @@ public:
 	}
 
 	/** @return Returns the blueprint associated with this spawnable */
-	UClass* GetClass() { return GeneratedClass; }
+	UClass* GetClass()
+	{
+		return GeneratedClass;
+	}
 
 	/** @return	Returns the game preview counterpart object for this spawnable, if it has one. */
 	const FWeakObjectPtr& GetCounterpartGamePreviewObject() const
@@ -75,6 +79,7 @@ private:
 	FWeakObjectPtr CounterpartGamePreviewObject;
 };
 
+
 /**
  * MovieScenePossessable is a "typed slot" used to allow the MovieScene to control an already-existing object
  */
@@ -147,78 +152,6 @@ struct FMovieSceneEditorData
 	TArray<FString> CollapsedSequencerNodes;
 };
 
-/**
- * A set of tracks bound to runtime objects
- */
-USTRUCT()
-struct FMovieSceneObjectBinding
-{
-	GENERATED_USTRUCT_BODY()
-
-	FMovieSceneObjectBinding()
-	{}
-
-	FMovieSceneObjectBinding( const FGuid& InObjectGuid, const FString& InBindingName, const TArray<UMovieSceneTrack*>& InTracks )
-		: ObjectGuid( InObjectGuid )
-		, BindingName( InBindingName )
-		, Tracks( InTracks )
-	{}
-
-	FMovieSceneObjectBinding( const FGuid& InObjectGuid, const FString& InBindingName )
-		: ObjectGuid( InObjectGuid )
-		, BindingName( InBindingName )
-	{}
-
-	/**
-	 * @return The time range of all tracks in this binding
-	 */
-	TRange<float> GetTimeRange() const;
-
-	/**
-	 * @return The guid of runtime objects in this binding
-	 */
-	const FGuid& GetObjectGuid() const { return ObjectGuid; }
-
-	/**
-	 * @return The display name of the binding
-	 */
-	const FString& GetName() const { return BindingName; } 
-
-	/**
-	 * Adds a new track to this binding
-	 *
-	 * @param NewTrack	The track to add
-	 */
-	void AddTrack( UMovieSceneTrack& NewTrack );
-
-	/**
-	 * Removes a track from this binding
-	 * 
-	 * @param Track	The track to remove
-	 * @return true if the track was successfully removed, false if the track could not be found
-	 */
-	bool RemoveTrack( UMovieSceneTrack& Track );
-
-	/**
-	 * @return All tracks in this binding
-	 */
-	const TArray<UMovieSceneTrack*>& GetTracks() const { return Tracks; }
-
-private:
-
-	/** Object binding guid for runtime objects */
-	UPROPERTY()
-	FGuid ObjectGuid;
-	
-	/** Display name */
-	UPROPERTY()
-	FString BindingName;
-
-	/** All tracks in this binding */
-	UPROPERTY()
-	TArray<UMovieSceneTrack*> Tracks;
-};
-
 
 /**
  * Implements a movie scene asset.
diff --git a/Engine/Source/Runtime/MovieScene/Public/MovieSceneBinding.h b/Engine/Source/Runtime/MovieScene/Public/MovieSceneBinding.h
new file mode 100644
--- /dev/null
+++ b/Engine/Source/Runtime/MovieScene/Public/MovieSceneBinding.h
@@ -0,0 +1,103 @@
+// Copyright 1998-2015 Epic Games, Inc. All Rights Reserved.
+
+#pragma once
+
+#include "MovieSceneBinding.generated.h"
+
+
+class UMovieSceneTrack;
+
+
+/**
+ * A set of tracks bound to runtime objects
+ */
+USTRUCT()
+struct FMovieSceneObjectBinding
+{
+	GENERATED_USTRUCT_BODY()
+
+	/** Default constructor. */
+	FMovieSceneObjectBinding() { }
+
+	/**
+	 * Creates and initializes a new instance.
+	 *
+	 * @param InObjectGuid
+	 * @param InBindingName
+	 * @param InTracks
+	 */
+	FMovieSceneObjectBinding(const FGuid& InObjectGuid, const FString& InBindingName, const TArray<UMovieSceneTrack*>& InTracks)
+		: ObjectGuid(InObjectGuid)
+		, BindingName(InBindingName)
+		, Tracks(InTracks)
+	{ }
+
+	/**
+	 * Creates and initializes a new instance.
+	 *
+	 * @param InObjectGuid
+	 * @param InBindingName
+	 */
+	FMovieSceneObjectBinding(const FGuid& InObjectGuid, const FString& InBindingName)
+		: ObjectGuid(InObjectGuid)
+		, BindingName(InBindingName)
+	{ }
+
+	/**
+	 * @return The time range of all tracks in this binding
+	 */
+	TRange<float> GetTimeRange() const;
+
+	/**
+	 * @return The guid of runtime objects in this binding
+	 */
+	const FGuid& GetObjectGuid() const
+	{
+		return ObjectGuid;
+	}
+
+	/**
+	 * @return The display name of the binding
+	 */
+	const FString& GetName() const
+	{
+		return BindingName;
+	} 
+
+	/**
+	 * Adds a new track to this binding
+	 *
+	 * @param NewTrack	The track to add
+	 */
+	void AddTrack(UMovieSceneTrack& NewTrack);
+
+	/**
+	 * Removes a track from this binding
+	 * 
+	 * @param Track	The track to remove
+	 * @return true if the track was successfully removed, false if the track could not be found
+	 */
+	bool RemoveTrack(UMovieSceneTrack& Track);
+
+	/**
+	 * @return All tracks in this binding
+	 */
+	const TArray<UMovieSceneTrack*>& GetTracks() const
+	{
+		return Tracks;
+	}
+
+private:
+
+	/** Object binding guid for runtime objects */
+	UPROPERTY()
+	FGuid ObjectGuid;
+	
+	/** Display name */
+	UPROPERTY()
+	FString BindingName;
+
+	/** All tracks in this binding */
+	UPROPERTY()
+	TArray<UMovieSceneTrack*> Tracks;
+};
