MovieScene: NULL --> nullptr

diff --git a/Engine/Source/Runtime/MovieScene/Private/MovieScene.cpp b/Engine/Source/Runtime/MovieScene/Private/MovieScene.cpp
--- a/Engine/Source/Runtime/MovieScene/Private/MovieScene.cpp
+++ b/Engine/Source/Runtime/MovieScene/Private/MovieScene.cpp
@@ -13,7 +13,6 @@ FMovieSceneSpawnable::FMovieSceneSpawnable( const FString& InitName, UClass* Ini
 }
 
 
-
 FMovieScenePossessable::FMovieScenePossessable( const FString& InitName, UClass* InitPossessedObjectClass )
 {
 	Guid = FGuid::NewGuid();
@@ -24,14 +23,15 @@ FMovieScenePossessable::FMovieScenePossessable( const FString& InitName, UClass*
 
 UMovieScene::UMovieScene( const FObjectInitializer& ObjectInitializer )
 	: Super( ObjectInitializer )
-{
-}
+{ }
+
 
 #if WITH_EDITOR
+
 // @todo sequencer: Some of these methods should only be used by tools, and should probably move out of MovieScene!
 FGuid UMovieScene::AddSpawnable( const FString& Name, UBlueprint* Blueprint, UObject* CounterpartGamePreviewObject )
 {
-	check( (Blueprint != NULL) && (Blueprint->GeneratedClass) );
+	check( (Blueprint != nullptr) && (Blueprint->GeneratedClass) );
 
 	Modify();
 
@@ -44,6 +44,7 @@ FGuid UMovieScene::AddSpawnable( const FString& Name, UBlueprint* Blueprint, UOb
 	return NewSpawnable.GetGuid();
 }
 
+
 bool UMovieScene::RemoveSpawnable( const FGuid& Guid )
 {
 	bool bAnythingRemoved = false;
@@ -58,8 +59,8 @@ bool UMovieScene::RemoveSpawnable( const FGuid& Guid )
 
 				{
 					UClass* GeneratedClass = CurSpawnable.GetClass();
-					UBlueprint* Blueprint = GeneratedClass ? Cast<UBlueprint>(GeneratedClass->ClassGeneratedBy) : NULL;
-					check(NULL != Blueprint);
+					UBlueprint* Blueprint = GeneratedClass ? Cast<UBlueprint>(GeneratedClass->ClassGeneratedBy) : nullptr;
+					check(nullptr != Blueprint);
 					// @todo sequencer: Also remove created Blueprint inner object.  Is this sufficient?  Needs to work with Undo too!
 					Blueprint->ClearFlags( RF_Standalone );	// @todo sequencer: Probably not needed for Blueprint
 					Blueprint->MarkPendingKill();
@@ -105,13 +106,13 @@ FMovieSceneSpawnable* UMovieScene::FindSpawnable( const FGuid& Guid )
 		}
 	}
 
-	return NULL;
+	return nullptr;
 }
 
 
 const FMovieSceneSpawnable* UMovieScene::FindSpawnableForCounterpart( UObject* GamePreviewObject ) const
 {
-	check( GamePreviewObject != NULL );
+	check( GamePreviewObject != nullptr );
 
 	// Must only be called for objects in game preview worlds
 	const bool bIsGamePreviewObject = !!( GamePreviewObject->GetOutermost()->PackageFlags & PKG_PlayInEditor );
@@ -126,7 +127,7 @@ const FMovieSceneSpawnable* UMovieScene::FindSpawnableForCounterpart( UObject* G
 		}
 	}
 
-	return NULL;
+	return nullptr;
 }
 
 
@@ -181,7 +182,7 @@ FMovieScenePossessable* UMovieScene::FindPossessable( const FGuid& Guid )
 		}
 	}
 
-	return NULL;
+	return nullptr;
 }
 
 
@@ -193,24 +194,29 @@ int32 UMovieScene::GetPossessableCount() const
 
 FMovieScenePossessable& UMovieScene::GetPossessable( const int32 Index )
 {
-	return Possessables[ Index ];
+	return Possessables[Index];
 }
 
+
 TRange<float> UMovieScene::GetTimeRange() const
 {
 	// Get the range of all sections combined
 	TArray< TRange<float> > Bounds;
+
 	for (int32 TypeIndex = 0; TypeIndex < MasterTracks.Num(); ++TypeIndex)
 	{
 		Bounds.Add(MasterTracks[TypeIndex]->GetSectionBoundaries());
 	}
+
 	for (int32 BindingIndex = 0; BindingIndex < ObjectBindings.Num(); ++BindingIndex)
 	{
 		Bounds.Add(ObjectBindings[BindingIndex].GetTimeRange());
 	}
+
 	return TRange<float>::Hull(Bounds);
 }
 
+
 TArray<UMovieSceneSection*> UMovieScene::GetAllSections() const
 {
 	TArray<UMovieSceneSection*> OutSections;
@@ -236,6 +242,7 @@ TArray<UMovieSceneSection*> UMovieScene::GetAllSections() const
 	return OutSections;
 }
 
+
 void UMovieScene::RemoveBinding( const FGuid& Guid )
 {
 	// Update each type
@@ -249,9 +256,10 @@ void UMovieScene::RemoveBinding( const FGuid& Guid )
 	}
 }
 
+
 UMovieSceneTrack* UMovieScene::FindTrack( TSubclassOf<UMovieSceneTrack> TrackClass, const FGuid& ObjectGuid, FName UniqueTrackName ) const
 {
-	UMovieSceneTrack* FoundTrack = NULL;
+	UMovieSceneTrack* FoundTrack = nullptr;
 	
 	check( UniqueTrackName != NAME_None );
 	check( ObjectGuid.IsValid() );
@@ -278,9 +286,10 @@ UMovieSceneTrack* UMovieScene::FindTrack( TSubclassOf<UMovieSceneTrack> TrackCla
 	return FoundTrack;
 }
 
+
 class UMovieSceneTrack* UMovieScene::AddTrack( TSubclassOf<UMovieSceneTrack> TrackClass, const FGuid& ObjectGuid )
 {
-	UMovieSceneTrack* CreatedType = NULL;
+	UMovieSceneTrack* CreatedType = nullptr;
 
 	check( ObjectGuid.IsValid() )
 
@@ -293,7 +302,6 @@ class UMovieSceneTrack* UMovieScene::AddTrack( TSubclassOf<UMovieSceneTrack> Tra
 			Modify();
 
 			CreatedType = NewObject<UMovieSceneTrack>(this, TrackClass, NAME_None, RF_Transactional);
-			
 			Binding.AddTrack( *CreatedType );
 		}
 	}
@@ -301,9 +309,11 @@ class UMovieSceneTrack* UMovieScene::AddTrack( TSubclassOf<UMovieSceneTrack> Tra
 	return CreatedType;
 }
 
+
 bool UMovieScene::RemoveTrack( UMovieSceneTrack* Track )
 {
 	Modify();
+
 	bool bAnythingRemoved = false;
 	for( int32 ObjectBindingIndex = 0; ObjectBindingIndex < ObjectBindings.Num(); ++ObjectBindingIndex )
 	{
@@ -320,9 +330,10 @@ bool UMovieScene::RemoveTrack( UMovieSceneTrack* Track )
 	return bAnythingRemoved;
 }
 
+
 UMovieSceneTrack* UMovieScene::FindMasterTrack( TSubclassOf<UMovieSceneTrack> TrackClass ) const
 {
-	UMovieSceneTrack* FoundTrack = NULL;
+	UMovieSceneTrack* FoundTrack = nullptr;
 	for( int32 TrackIndex = 0; TrackIndex < MasterTracks.Num(); ++TrackIndex )
 	{
 		UMovieSceneTrack* Track = MasterTracks[TrackIndex];
@@ -336,6 +347,7 @@ UMovieSceneTrack* UMovieScene::FindMasterTrack( TSubclassOf<UMovieSceneTrack> Tr
 	return FoundTrack;
 }
 
+
 class UMovieSceneTrack* UMovieScene::AddMasterTrack( TSubclassOf<UMovieSceneTrack> TrackClass )
 {
 	Modify();
@@ -347,6 +359,7 @@ class UMovieSceneTrack* UMovieScene::AddMasterTrack( TSubclassOf<UMovieSceneTrac
 	return CreatedType;
 }
 
+
 UMovieSceneTrack* UMovieScene::AddShotTrack( TSubclassOf<UMovieSceneTrack> TrackClass )
 {
 	if( !ShotTrack )
@@ -359,11 +372,13 @@ UMovieSceneTrack* UMovieScene::AddShotTrack( TSubclassOf<UMovieSceneTrack> Track
 	return ShotTrack;
 }
 
+
 UMovieSceneTrack* UMovieScene::GetShotTrack()
 {
 	return ShotTrack;
 }
 
+
 void UMovieScene::RemoveShotTrack()
 {
 	if( ShotTrack )
@@ -374,21 +389,24 @@ void UMovieScene::RemoveShotTrack()
 	}
 }
 
+
 bool UMovieScene::RemoveMasterTrack( UMovieSceneTrack* Track ) 
 {
 	Modify();
 
-	return MasterTracks.RemoveSingle( Track ) != 0;
+	return MasterTracks.RemoveSingle(Track) != 0;
 }
 
+
 bool UMovieScene::IsAMasterTrack(const UMovieSceneTrack* Track) const
 {
-	for( int32 TrackIndex = 0; TrackIndex < MasterTracks.Num(); ++TrackIndex )
+	for (int32 TrackIndex = 0; TrackIndex < MasterTracks.Num(); ++TrackIndex)
 	{
 		if (Track == MasterTracks[TrackIndex])
 		{
 			return true;
 		}
 	}
+
 	return false;
 }
diff --git a/Engine/Source/Runtime/MovieScene/Private/MovieSceneCommonHelpers.cpp b/Engine/Source/Runtime/MovieScene/Private/MovieSceneCommonHelpers.cpp
--- a/Engine/Source/Runtime/MovieScene/Private/MovieSceneCommonHelpers.cpp
+++ b/Engine/Source/Runtime/MovieScene/Private/MovieSceneCommonHelpers.cpp
@@ -28,6 +28,7 @@ TArray<UMovieSceneSection*> MovieSceneHelpers::GetTraversedSections( const TArra
 	return TraversedSections;
 }
 
+
 UMovieSceneSection* MovieSceneHelpers::FindSectionAtTime( const TArray<UMovieSceneSection*>& Sections, float Time )
 {
 	if( Sections.Num() == 1 )
@@ -41,14 +42,15 @@ UMovieSceneSection* MovieSceneHelpers::FindSectionAtTime( const TArray<UMovieSce
 
 }
 
+
 UMovieSceneSection* MovieSceneHelpers::FindNearestSectionAtTime( const TArray<UMovieSceneSection*>& Sections, float Time )
 {
 	// Only update the section if the position is within the time span of the section
 	// Or if there are no sections at the time, the left closest section to the time
 	// Or in the case that Time is before all sections, take the section with the earliest start time
-	UMovieSceneSection* ClosestSection = NULL;
+	UMovieSceneSection* ClosestSection = nullptr;
 	float ClosestSectionTime = 0.f;
-	UMovieSceneSection* EarliestSection = NULL;
+	UMovieSceneSection* EarliestSection = nullptr;
 	float EarliestSectionTime = 0.f;
 	for( int32 SectionIndex = 0; SectionIndex < Sections.Num(); ++SectionIndex )
 	{
@@ -82,15 +84,16 @@ UMovieSceneSection* MovieSceneHelpers::FindNearestSectionAtTime( const TArray<UM
 	// if we get here, we are off of any section
 	// if ClosestSection, then we take the closest to left of this time
 	// else, we take the EarliestSection
-	// if that's NULL, then there are no sections
+	// if that's nullptr, then there are no sections
 	return ClosestSection ? ClosestSection : EarliestSection;
 }
 
+
 USceneComponent* MovieSceneHelpers::SceneComponentFromRuntimeObject(UObject* Object)
 {
 	AActor* Actor = Cast<AActor>(Object);
 
-	USceneComponent* SceneComponent = NULL;
+	USceneComponent* SceneComponent = nullptr;
 	if (Actor && Actor->GetRootComponent())
 	{
 		// If there is an actor, modify its root component
@@ -104,6 +107,7 @@ USceneComponent* MovieSceneHelpers::SceneComponentFromRuntimeObject(UObject* Obj
 	return SceneComponent;
 }
 
+
 FTrackInstancePropertyBindings::FTrackInstancePropertyBindings( FName InPropertyName, const FString& InPropertyPath, const FName& InFunctionName )
     : PropertyPath( InPropertyPath )
 	, PropertyName( InPropertyName )
@@ -122,6 +126,7 @@ FTrackInstancePropertyBindings::FTrackInstancePropertyBindings( FName InProperty
 	}
 }
 
+
 void FTrackInstancePropertyBindings::CallFunction( UObject* InRuntimeObject, void* FunctionParams )
 {
 	FPropertyAndFunction PropAndFunction = RuntimeObjectToFunctionMap.FindRef(InRuntimeObject);
@@ -164,6 +169,7 @@ FTrackInstancePropertyBindings::FPropertyAddress FTrackInstancePropertyBindings:
 
 }
 
+
 FTrackInstancePropertyBindings::FPropertyAddress FTrackInstancePropertyBindings::FindProperty( UObject* InObject, const FString& InPropertyPath ) const
 {
 	TArray<FString> PropertyNames;
@@ -181,6 +187,7 @@ FTrackInstancePropertyBindings::FPropertyAddress FTrackInstancePropertyBindings:
 
 }
 
+
 void FTrackInstancePropertyBindings::UpdateBindings( const TArray<UObject*>& InRuntimeObjects )
 {
 	for(UObject* Object : InRuntimeObjects)
diff --git a/Engine/Source/Runtime/MovieScene/Private/MovieSceneModule.cpp b/Engine/Source/Runtime/MovieScene/Private/MovieSceneModule.cpp
--- a/Engine/Source/Runtime/MovieScene/Private/MovieSceneModule.cpp
+++ b/Engine/Source/Runtime/MovieScene/Private/MovieSceneModule.cpp
@@ -2,38 +2,26 @@
 
 #include "MovieScenePrivatePCH.h"
 #include "ModuleManager.h"
-
 #include "MovieScene.h"
 #include "MovieSceneBindings.h"
 #include "RuntimeMovieScenePlayer.h"
 
+
 DEFINE_LOG_CATEGORY(LogSequencerRuntime);
 
 
 /**
  * MovieScene module implementation (private)
  */
-class FMovieSceneModule : public IMovieSceneModule
+class FMovieSceneModule
+	: public IMovieSceneModule
 {
-
 public:
 
 	/** IModuleInterface */
-	virtual void StartupModule() override;
-	virtual void ShutdownModule() override;
-
+	virtual void StartupModule() override { }
+	virtual void ShutdownModule() override { }
 };
 
 
 IMPLEMENT_MODULE( FMovieSceneModule, MovieScene );
-
-
-
-void FMovieSceneModule::StartupModule()
-{
-}
-
-
-void FMovieSceneModule::ShutdownModule()
-{
-}
\ No newline at end of file
diff --git a/Engine/Source/Runtime/MovieScene/Private/MovieSceneSection.cpp b/Engine/Source/Runtime/MovieScene/Private/MovieSceneSection.cpp
--- a/Engine/Source/Runtime/MovieScene/Private/MovieSceneSection.cpp
+++ b/Engine/Source/Runtime/MovieScene/Private/MovieSceneSection.cpp
@@ -27,7 +27,7 @@ const UMovieSceneSection* UMovieSceneSection::OverlapsWithSections(const TArray<
 			}
 		}
 	}
-	return NULL;
+	return nullptr;
 }
 
 void UMovieSceneSection::InitialPlacement(const TArray<UMovieSceneSection*>& Sections, float InStartTime, float InEndTime, bool bAllowMultipleRows)
@@ -40,7 +40,7 @@ void UMovieSceneSection::InitialPlacement(const TArray<UMovieSceneSection*>& Sec
 	RowIndex = 0;
 	if (bAllowMultipleRows)
 	{
-		while (OverlapsWithSections(Sections) != NULL)
+		while (OverlapsWithSections(Sections) != nullptr)
 		{
 			++RowIndex;
 		}
@@ -51,7 +51,7 @@ void UMovieSceneSection::InitialPlacement(const TArray<UMovieSceneSection*>& Sec
 		{
 			const UMovieSceneSection* OverlappedSection = OverlapsWithSections(Sections);
 
-			if (OverlappedSection == NULL)
+			if (OverlappedSection == nullptr)
 			{
 				break;
 			}
diff --git a/Engine/Source/Runtime/MovieScene/Private/RuntimeMovieScenePlayer.cpp b/Engine/Source/Runtime/MovieScene/Private/RuntimeMovieScenePlayer.cpp
--- a/Engine/Source/Runtime/MovieScene/Private/RuntimeMovieScenePlayer.cpp
+++ b/Engine/Source/Runtime/MovieScene/Private/RuntimeMovieScenePlayer.cpp
@@ -4,13 +4,14 @@
 #include "RuntimeMovieScenePlayer.h"
 #include "MovieSceneInstance.h"
 
+
 URuntimeMovieScenePlayer* URuntimeMovieScenePlayer::CreateRuntimeMovieScenePlayer( ULevel* Level, UMovieSceneBindings* MovieSceneBindings )
 {	
 	URuntimeMovieScenePlayer* NewRuntimeMovieScenePlayer = NewObject<URuntimeMovieScenePlayer>(
 		(UObject*)GetTransientPackage(), 
 		NAME_None, 
 		RF_Transient );
-	check( NewRuntimeMovieScenePlayer != NULL );
+	check( NewRuntimeMovieScenePlayer != nullptr );
 
 	// Associate the player with its world
 	NewRuntimeMovieScenePlayer->World = Level->OwningWorld;
@@ -22,7 +23,7 @@ URuntimeMovieScenePlayer* URuntimeMovieScenePlayer::CreateRuntimeMovieScenePlaye
 	// Spawn actors for each spawnable!
 
 	// Add this to the level's list of active RuntimeMovieScenePlayers.  The level will own this player from now on.
-	check( Level != NULL );
+	check( Level != nullptr );
 	Level->AddActiveRuntimeMovieScenePlayer( NewRuntimeMovieScenePlayer );
 
 	return NewRuntimeMovieScenePlayer;
@@ -33,7 +34,7 @@ URuntimeMovieScenePlayer::URuntimeMovieScenePlayer( const FObjectInitializer& Ob
 	: Super( ObjectInitializer )
 {
 
-	MovieSceneBindings = NULL;
+	MovieSceneBindings = nullptr;
 	TimeCursorPosition = 0.0f;
 	bIsPlaying = false;
 }
@@ -42,10 +43,10 @@ URuntimeMovieScenePlayer::URuntimeMovieScenePlayer( const FObjectInitializer& Ob
 void URuntimeMovieScenePlayer::SpawnActorsForMovie( TSharedRef<FMovieSceneInstance> MovieSceneInstance  )
 {
 	UWorld* WorldPtr = World.Get();
-	if( WorldPtr != NULL && MovieSceneBindings != NULL )
+	if( WorldPtr != nullptr && MovieSceneBindings != nullptr )
 	{
 		UMovieScene* MovieScene = MovieSceneInstance->GetMovieScene();
-		if( MovieScene != NULL )
+		if( MovieScene != nullptr )
 		{
 			TArray<FSpawnedActorInfo>* FoundSpawnedActors = InstanceToSpawnedActorMap.Find( MovieSceneInstance );
 			if( FoundSpawnedActors )
@@ -61,7 +62,7 @@ void URuntimeMovieScenePlayer::SpawnActorsForMovie( TSharedRef<FMovieSceneInstan
 				auto& Spawnable = MovieScene->GetSpawnable( SpawnableIndex );
 
 				UClass* GeneratedClass = Spawnable.GetClass();
-				if ( GeneratedClass != NULL )
+				if ( GeneratedClass != nullptr )
 				{
 					const bool bIsActorBlueprint = GeneratedClass->IsChildOf( AActor::StaticClass() );
 					if ( bIsActorBlueprint )
@@ -90,10 +91,11 @@ void URuntimeMovieScenePlayer::SpawnActorsForMovie( TSharedRef<FMovieSceneInstan
 	}
 }
 
+
 void URuntimeMovieScenePlayer::DestroyActorsForMovie( TSharedRef<FMovieSceneInstance> MovieSceneInstance )
 {
 	UWorld* WorldPtr = World.Get();
-	if( WorldPtr != NULL && MovieSceneBindings != NULL )
+	if( WorldPtr != nullptr && MovieSceneBindings != nullptr )
 	{
 		TArray<FSpawnedActorInfo>* SpawnedActors = InstanceToSpawnedActorMap.Find( MovieSceneInstance );
 		if( SpawnedActors )
@@ -117,6 +119,7 @@ void URuntimeMovieScenePlayer::DestroyActorsForMovie( TSharedRef<FMovieSceneInst
 	}
 }
 
+
 bool URuntimeMovieScenePlayer::IsPlaying() const
 {
 	return bIsPlaying;
@@ -126,7 +129,7 @@ bool URuntimeMovieScenePlayer::IsPlaying() const
 void URuntimeMovieScenePlayer::Play()
 {
 	// Init the root movie scene instance
-	if( MovieSceneBindings != NULL )
+	if( MovieSceneBindings != nullptr )
 	{
 		UMovieScene* MovieScene = MovieSceneBindings->GetRootMovieScene();
 	
@@ -165,7 +168,7 @@ void URuntimeMovieScenePlayer::Tick( const float DeltaSeconds )
 	}
 
 	// Update the movie scene!
-	if( MovieSceneBindings != NULL && RootMovieSceneInstance.IsValid() )
+	if( MovieSceneBindings != nullptr && RootMovieSceneInstance.IsValid() )
 	{
 		RootMovieSceneInstance->Update( TimeCursorPosition, LastTimePosition, *this );
 	}
@@ -196,7 +199,7 @@ void URuntimeMovieScenePlayer::GetRuntimeObjects( TSharedRef<FMovieSceneInstance
 	
 	}
 	// Otherwise, check to see if we have one or more possessed actors that are mapped to this handle
-	else if( MovieSceneBindings != NULL )
+	else if( MovieSceneBindings != nullptr )
 	{
 		for (FMovieSceneBoundObjectInfo& BoundObject : MovieSceneBindings->FindBoundObjects(ObjectHandle))
 		{
@@ -222,22 +225,26 @@ void URuntimeMovieScenePlayer::GetRuntimeObjects( TSharedRef<FMovieSceneInstance
 	}
 }
 
+
 void URuntimeMovieScenePlayer::UpdateCameraCut(UObject* ObjectToViewThrough, bool bNewCameraCut) const
 {
 }
 
+
 EMovieScenePlayerStatus::Type URuntimeMovieScenePlayer::GetPlaybackStatus() const
 {
 	return bIsPlaying ? EMovieScenePlayerStatus::Playing : EMovieScenePlayerStatus::Stopped;
 }
 
+
 void URuntimeMovieScenePlayer::AddOrUpdateMovieSceneInstance( UMovieSceneSection& MovieSceneSection, TSharedRef<FMovieSceneInstance> InstanceToAdd )
 {
 	SpawnActorsForMovie( InstanceToAdd );
 }
 
+
 void URuntimeMovieScenePlayer::RemoveMovieSceneInstance( UMovieSceneSection& MovieSceneSection, TSharedRef<FMovieSceneInstance> InstanceToRemove )
 {
 	const bool bDestroyAll = true;
 	DestroyActorsForMovie( InstanceToRemove );
-}
\ No newline at end of file
+}
