MovieScene: Moved possessables and spawnables into their own files

diff --git a/Engine/Source/Runtime/MovieScene/Private/MovieScene.cpp b/Engine/Source/Runtime/MovieScene/Private/MovieScene.cpp
--- a/Engine/Source/Runtime/MovieScene/Private/MovieScene.cpp
+++ b/Engine/Source/Runtime/MovieScene/Private/MovieScene.cpp
@@ -4,22 +4,6 @@
 #include "MovieScene.h"
 
 
-FMovieSceneSpawnable::FMovieSceneSpawnable( const FString& InitName, UClass* InitClass, UObject* InitCounterpartGamePreviewObject )
-{
-	Guid = FGuid::NewGuid();
-	Name = InitName;
-	GeneratedClass = InitClass;
-	CounterpartGamePreviewObject = InitCounterpartGamePreviewObject;
-}
-
-
-FMovieScenePossessable::FMovieScenePossessable( const FString& InitName, UClass* InitPossessedObjectClass )
-{
-	Guid = FGuid::NewGuid();
-	Name = InitName;
-	PossessedObjectClass = InitPossessedObjectClass;
-}
-
 
 UMovieScene::UMovieScene( const FObjectInitializer& ObjectInitializer )
 	: Super( ObjectInitializer )
diff --git a/Engine/Source/Runtime/MovieScene/Public/MovieScene.h b/Engine/Source/Runtime/MovieScene/Public/MovieScene.h
--- a/Engine/Source/Runtime/MovieScene/Public/MovieScene.h
+++ b/Engine/Source/Runtime/MovieScene/Public/MovieScene.h
@@ -3,6 +3,8 @@
 #pragma once
 
 #include "MovieSceneBinding.h"
+#include "MovieScenePossessable.h"
+#include "MovieSceneSpawnable.h"
 #include "MovieScene.generated.h"
 
 
@@ -15,130 +17,6 @@ class UMovieSceneTrack;
 MOVIESCENE_API DECLARE_LOG_CATEGORY_EXTERN(LogSequencerRuntime, Log, All);
 
 
-/**
- * MovieSceneSpawnable describes an object that can be spawned for this MovieScene
- */
-USTRUCT()
-struct FMovieSceneSpawnable 
-{
-	GENERATED_USTRUCT_BODY()
-
-public:
-
-	/** FMovieSceneSpawnable default constructor */
-	FMovieSceneSpawnable() { }
-
-	/** FMovieSceneSpawnable initialization constructor */
-	FMovieSceneSpawnable( const FString& InitName, UClass* InitClass, UObject* InitCounterpartGamePreviewObject );
-
-	/** @return Returns the guid for this spawnable */
-	const FGuid& GetGuid() const
-	{
-		return Guid;
-	}
-
-	/** @return Returns the name of this spawnable */
-	const FString& GetName() const
-	{
-		return Name;
-	}
-
-	/** @return Returns the blueprint associated with this spawnable */
-	UClass* GetClass()
-	{
-		return GeneratedClass;
-	}
-
-	/** @return	Returns the game preview counterpart object for this spawnable, if it has one. */
-	const FWeakObjectPtr& GetCounterpartGamePreviewObject() const
-	{
-		return CounterpartGamePreviewObject;
-	}
-
-private:
-
-	/** Guid */
-	// @todo sequencer: Guids need to be handled carefully when the asset is duplicated (or loaded after being copied on disk).
-	//					Sometimes we'll need to generate fresh Guids.
-	UPROPERTY()
-	FGuid Guid;
-
-	/** Name label */
-	// @todo sequencer: Should be editor-only probably
-	UPROPERTY()
-	FString Name;
-
-	/** Data-only blueprint-generated-class for this object */
-	// @todo sequencer: Could be weak object ptr here, IF blueprints that are inners are housekept properly without references
-	UPROPERTY()
-	UClass* GeneratedClass;
-
-	/** Optional transient weak pointer to the game preview object this spawnable was created to capture data for.  This is
-	    used in the editor when capturing keyframe data from a live simulation */
-	// @todo sequencer data: Should be editor only
-	FWeakObjectPtr CounterpartGamePreviewObject;
-};
-
-
-/**
- * MovieScenePossessable is a "typed slot" used to allow the MovieScene to control an already-existing object
- */
-USTRUCT()
-struct FMovieScenePossessable
-{
-	GENERATED_USTRUCT_BODY( FMovieScenePossessable )
-
-public:
-
-	/** FMovieScenePossessable default constructor */
-	FMovieScenePossessable()
-	{
-	}
-
-	/** FMovieScenePossessable initialization constructor */
-	FMovieScenePossessable( const FString& InitName, UClass* InitPossessedObjectClass );
-
-	/** @return Returns the guid for this possessable */
-	const FGuid& GetGuid() const
-	{
-		return Guid;
-	}
-
-	/** @return Returns the name of this spawnable */
-	const FString& GetName() const
-	{
-		return Name;
-	}
-
-	/** @return Returns the class of the object we'll possess */
-	const UClass* GetPossessedObjectClass() const
-	{
-		return PossessedObjectClass;
-	}
-
-private:
-
-	/** Guid */
-	// @todo sequencer: Guids need to be handled carefully when the asset is duplicated (or loaded after being copied on disk).
-	//					Sometimes we'll need to generate fresh Guids.
-	UPROPERTY()
-	FGuid Guid;
-
-	/** Name label for this slot */
-	// @todo sequencer: Should be editor-only probably
-	UPROPERTY()
-	FString Name;
-
-	/** Type of the object we'll be possessing */
-	// @todo sequencer: Might be able to be editor-only.  We'll see.
-	// @todo sequencer: This isn't used for anything yet.  We could use it to gate which types of objects can be bound to a
-	// possessable "slot" though.  Or we could use it to generate a "preview" spawnable puppet when previewing with no
-	// possessable object available.
-	UPROPERTY()
-	UClass* PossessedObjectClass;
-};
-
-
 /**
  * Editor only data that needs to be saved between sessions for editing but has no runtime purpose
  */
diff --git a/Engine/Source/Runtime/MovieScene/Public/MovieScenePossessable.h b/Engine/Source/Runtime/MovieScene/Public/MovieScenePossessable.h
new file mode 100644
--- /dev/null
+++ b/Engine/Source/Runtime/MovieScene/Public/MovieScenePossessable.h
@@ -0,0 +1,88 @@
+// Copyright 1998-2015 Epic Games, Inc. All Rights Reserved.
+
+#pragma once
+
+#include "MovieScenePossessable.generated.h"
+
+
+/**
+ * MovieScenePossessable is a "typed slot" used to allow the MovieScene to control an already-existing object
+ */
+USTRUCT()
+struct FMovieScenePossessable
+{
+	GENERATED_USTRUCT_BODY(FMovieScenePossessable)
+
+public:
+
+	/** Default constructor. */
+	FMovieScenePossessable() { }
+
+	/**
+	 * Creates and initializes a new instance.
+	 *
+	 * @param InitName
+	 * @param InitPossessedObjectClass
+	 */
+	FMovieScenePossessable(const FString& InitName, UClass* InitPossessedObjectClass)
+		: Guid(FGuid::NewGuid())
+		, Name(InitName)
+		, PossessedObjectClass(InitPossessedObjectClass)
+	{ }
+
+public:
+
+	/**
+	 * Get the unique identifier of the possessed object.
+	 *
+	 * @return Object GUID.
+	 * @see GetName, GetPossessedObjectClass
+	 */
+	const FGuid& GetGuid() const
+	{
+		return Guid;
+	}
+
+	/**
+	 * Get the name of the possessed object.
+	 *
+	 * @return Object name.
+	 * @see GetGuid, GetPossessedObjectClass
+	 */
+	const FString& GetName() const
+	{
+		return Name;
+	}
+
+	/**
+	 * Get the class of the possessed object.
+	 *
+	 * @return Object class.
+	 * @see GetGuid, GetName
+	 */
+	const UClass* GetPossessedObjectClass() const
+	{
+		return PossessedObjectClass;
+	}
+
+private:
+
+	/** Unique identifier of the possessable object. */
+	// @todo sequencer: Guids need to be handled carefully when the asset is duplicated (or loaded after being copied on disk).
+	//					Sometimes we'll need to generate fresh Guids.
+	UPROPERTY()
+	FGuid Guid;
+
+	/** Name label for this slot */
+	// @todo sequencer: Should be editor-only probably
+	UPROPERTY()
+	FString Name;
+
+	/** Type of the object we'll be possessing */
+	// @todo sequencer: Might be able to be editor-only.  We'll see.
+	// @todo sequencer: This isn't used for anything yet.  We could use it to gate which types of objects can be bound to a
+	// possessable "slot" though.  Or we could use it to generate a "preview" spawnable puppet when previewing with no
+	// possessable object available.
+	UPROPERTY()
+	UClass* PossessedObjectClass;
+};
diff --git a/Engine/Source/Runtime/MovieScene/Public/MovieSceneSpawnable.h b/Engine/Source/Runtime/MovieScene/Public/MovieSceneSpawnable.h
new file mode 100644
--- /dev/null
+++ b/Engine/Source/Runtime/MovieScene/Public/MovieSceneSpawnable.h
@@ -0,0 +1,97 @@
+// Copyright 1998-2015 Epic Games, Inc. All Rights Reserved.
+
+#pragma once
+
+#include "MovieSceneSpawnable.generated.h"
+
+
+/**
+ * MovieSceneSpawnable describes an object that can be spawned for this MovieScene
+ */
+USTRUCT()
+struct FMovieSceneSpawnable 
+{
+	GENERATED_USTRUCT_BODY()
+
+public:
+
+	/** FMovieSceneSpawnable default constructor */
+	FMovieSceneSpawnable() { }
+
+	/** FMovieSceneSpawnable initialization constructor */
+	FMovieSceneSpawnable(const FString& InitName, UClass* InitClass, UObject* InitCounterpartGamePreviewObject)
+		: Guid(FGuid::NewGuid())
+		, Name(InitName)
+		, GeneratedClass(InitClass)
+		, CounterpartGamePreviewObject(InitCounterpartGamePreviewObject)
+	{ }
+
+public:
+
+	/**
+	 * Get the Blueprint associated with the spawnable object.
+	 *
+	 * @return Blueprint class
+	 * @see GetCounterpartGamePreviewObject, GetGuid, GetName
+	 */
+	UClass* GetClass()
+	{
+		return GeneratedClass;
+	}
+
+	/**
+	 * Get the game preview counterpart object for the spawnable object, if it has one.
+	 *
+	 * @return Counterpart object, or nullptr if it doesn't have one.
+	 * @see GetClass, GetGuid, GetName
+	 */
+	const FWeakObjectPtr& GetCounterpartGamePreviewObject() const
+	{
+		return CounterpartGamePreviewObject;
+	}
+
+	/**
+	 * Get the unique identifier of the spawnable object.
+	 *
+	 * @return Object GUID.
+	 * @see GetClass, GetCounterpartGamePreviewObject, GetName
+	 */
+	const FGuid& GetGuid() const
+	{
+		return Guid;
+	}
+
+	/**
+	 * Get the name of the spawnable object.
+	 *
+	 * @return Object name.
+	 * @see GetClass, GetCounterpartGamePreviewObject, GetGuid
+	 */
+	const FString& GetName() const
+	{
+		return Name;
+	}
+
+private:
+
+	/** Unique identifier of the spawnable object. */
+	// @todo sequencer: Guids need to be handled carefully when the asset is duplicated (or loaded after being copied on disk).
+	//					Sometimes we'll need to generate fresh Guids.
+	UPROPERTY()
+	FGuid Guid;
+
+	/** Name label */
+	// @todo sequencer: Should be editor-only probably
+	UPROPERTY()
+	FString Name;
+
+	/** Data-only blueprint-generated-class for this object */
+	// @todo sequencer: Could be weak object ptr here, IF blueprints that are inners are housekept properly without references
+	UPROPERTY()
+	UClass* GeneratedClass;
+
+	/** Optional transient weak pointer to the game preview object this spawnable was created to capture data for.  This is
+	    used in the editor when capturing keyframe data from a live simulation */
+	// @todo sequencer data: Should be editor only
+	FWeakObjectPtr CounterpartGamePreviewObject;
+};
