MovieScene: Reduced nesting

diff --git a/Engine/Source/Runtime/MovieScene/Private/RuntimeMovieScenePlayer.cpp b/Engine/Source/Runtime/MovieScene/Private/RuntimeMovieScenePlayer.cpp
--- a/Engine/Source/Runtime/MovieScene/Private/RuntimeMovieScenePlayer.cpp
+++ b/Engine/Source/Runtime/MovieScene/Private/RuntimeMovieScenePlayer.cpp
@@ -81,51 +81,66 @@ void URuntimeMovieScenePlayer::Play()
 
 void URuntimeMovieScenePlayer::SpawnActorsForMovie(TSharedRef<FMovieSceneInstance> MovieSceneInstance)
 {
+	if (MovieSceneBindings == nullptr)
+	{
+		return;
+	}
+
 	UWorld* WorldPtr = World.Get();
-	if( WorldPtr != nullptr && MovieSceneBindings != nullptr )
+
+	if (WorldPtr == nullptr)
 	{
-		UMovieScene* MovieScene = MovieSceneInstance->GetMovieScene();
-		if( MovieScene != nullptr )
+		return;
+	}
+
+	UMovieScene* MovieScene = MovieSceneInstance->GetMovieScene();
+
+	if (MovieScene == nullptr)
+	{
+		return;
+	}
+
+	TArray<FSpawnedActorInfo>* FoundSpawnedActors = InstanceToSpawnedActorMap.Find(MovieSceneInstance);
+
+	if (FoundSpawnedActors != nullptr)
+	{
+		// Remove existing spawned actors for this movie
+		DestroyActorsForMovie( MovieSceneInstance );
+	}
+
+	TArray<FSpawnedActorInfo>& SpawnedActorList = InstanceToSpawnedActorMap.Add(MovieSceneInstance, TArray<FSpawnedActorInfo>());
+
+	for (auto SpawnableIndex = 0; SpawnableIndex < MovieScene->GetSpawnableCount(); ++SpawnableIndex)
+	{
+		auto& Spawnable = MovieScene->GetSpawnable(SpawnableIndex);
+		UClass* GeneratedClass = Spawnable.GetClass();
+		
+		if ((GeneratedClass == nullptr) || !GeneratedClass->IsChildOf(AActor::StaticClass()))
 		{
-			TArray<FSpawnedActorInfo>* FoundSpawnedActors = InstanceToSpawnedActorMap.Find( MovieSceneInstance );
-			if( FoundSpawnedActors )
-			{
-				// Remove existing spawned actors for this movie
-				DestroyActorsForMovie( MovieSceneInstance );
-			}
+			continue;
+		}
 
-			TArray<FSpawnedActorInfo>& SpawnedActorList = InstanceToSpawnedActorMap.Add( MovieSceneInstance, TArray<FSpawnedActorInfo>() );
+		AActor* ActorCDO = CastChecked<AActor>(GeneratedClass->ClassDefaultObject);
+		const FVector SpawnLocation = ActorCDO->GetRootComponent()->RelativeLocation;
+		const FRotator SpawnRotation = ActorCDO->GetRootComponent()->RelativeRotation;
 
-			for( auto SpawnableIndex = 0; SpawnableIndex < MovieScene->GetSpawnableCount(); ++SpawnableIndex )
-			{
-				auto& Spawnable = MovieScene->GetSpawnable( SpawnableIndex );
+		FActorSpawnParameters SpawnInfo;
+		{
+			SpawnInfo.ObjectFlags = RF_NoFlags;
+		}
 
-				UClass* GeneratedClass = Spawnable.GetClass();
-				if ( GeneratedClass != nullptr )
-				{
-					const bool bIsActorBlueprint = GeneratedClass->IsChildOf( AActor::StaticClass() );
-					if ( bIsActorBlueprint )
-					{
-						AActor* ActorCDO = CastChecked< AActor >( GeneratedClass->ClassDefaultObject );
-
-						const FVector SpawnLocation = ActorCDO->GetRootComponent()->RelativeLocation;
-						const FRotator SpawnRotation = ActorCDO->GetRootComponent()->RelativeRotation;
-
-						FActorSpawnParameters SpawnInfo;
-						SpawnInfo.ObjectFlags = RF_NoFlags;
-						AActor* NewActor = WorldPtr->SpawnActor( GeneratedClass, &SpawnLocation, &SpawnRotation, SpawnInfo );
-						if( NewActor )
-						{	
-							// Actor was spawned OK!
-							FSpawnedActorInfo NewInfo;
-							NewInfo.RuntimeGuid = Spawnable.GetGuid();
-							NewInfo.SpawnedActor = NewActor;
-						
-							SpawnedActorList.Add( NewInfo );
-						}
-					}
-				}
+		AActor* NewActor = WorldPtr->SpawnActor(GeneratedClass, &SpawnLocation, &SpawnRotation, SpawnInfo);
+
+		if (NewActor)
+		{	
+			// Actor was spawned OK!
+			FSpawnedActorInfo NewInfo;
+			{
+				NewInfo.RuntimeGuid = Spawnable.GetGuid();
+				NewInfo.SpawnedActor = NewActor;
 			}
+						
+			SpawnedActorList.Add(NewInfo);
 		}
 	}
 }
@@ -133,29 +148,43 @@ void URuntimeMovieScenePlayer::SpawnActorsForMovie(TSharedRef<FMovieSceneInstanc
 
 void URuntimeMovieScenePlayer::DestroyActorsForMovie(TSharedRef<FMovieSceneInstance> MovieSceneInstance)
 {
+	if (MovieSceneBindings == nullptr)
+	{
+		return;
+	}
+
 	UWorld* WorldPtr = World.Get();
-	if( WorldPtr != nullptr && MovieSceneBindings != nullptr )
+
+	if (WorldPtr == nullptr)
 	{
-		TArray<FSpawnedActorInfo>* SpawnedActors = InstanceToSpawnedActorMap.Find( MovieSceneInstance );
-		if( SpawnedActors )
+		return;
+	}
+
+	TArray<FSpawnedActorInfo>* SpawnedActors = InstanceToSpawnedActorMap.Find(MovieSceneInstance);
+
+	if (SpawnedActors == nullptr)
+	{
+		return;
+	}
+
+	TArray<FSpawnedActorInfo>& SpawnedActorsRef = *SpawnedActors;
+
+	for(int32 ActorIndex = 0; ActorIndex < SpawnedActors->Num(); ++ActorIndex)
+	{
+		AActor* Actor = SpawnedActorsRef[ActorIndex].SpawnedActor.Get();
+		
+		if (Actor != nullptr)
 		{
-			TArray<FSpawnedActorInfo>& SpawnedActorsRef = *SpawnedActors;
-			for( int32 ActorIndex = 0; ActorIndex < SpawnedActors->Num(); ++ActorIndex )
-			{
-				AActor* Actor = SpawnedActorsRef[ActorIndex].SpawnedActor.Get();
-				if( Actor )
-				{
-					// @todo Sequencer figure this out.  Defaults to false.
-					const bool bNetForce = false;
+			// @todo Sequencer figure this out.  Defaults to false.
+			const bool bNetForce = false;
 
-					// At runtime, level modification is not needed
-					const bool bShouldModifyLevel = false;
-					Actor->GetWorld()->DestroyActor( Actor, bNetForce, bShouldModifyLevel );
-				}
-			}
-			InstanceToSpawnedActorMap.Remove( MovieSceneInstance );
+			// At runtime, level modification is not needed
+			const bool bShouldModifyLevel = false;
+			Actor->GetWorld()->DestroyActor(Actor, bNetForce, bShouldModifyLevel);
 		}
 	}
+
+	InstanceToSpawnedActorMap.Remove(MovieSceneInstance);
 }
 
 
@@ -171,26 +200,26 @@ void URuntimeMovieScenePlayer::GetRuntimeObjects(TSharedRef<FMovieSceneInstance>
 	//    This may be important to reduce the up-front hitch when spawning actors for the entire movie, but
 	//    may introduce smaller hitches during playback.  Needs experimentation.
 
-	// First, check to see if we have a spawned actor for this handle
-	const TArray<FSpawnedActorInfo>* SpawnedActors = InstanceToSpawnedActorMap.Find( MovieSceneInstance );
-	if( SpawnedActors && SpawnedActors->Num() > 0 )
+	const TArray<FSpawnedActorInfo>* SpawnedActors = InstanceToSpawnedActorMap.Find(MovieSceneInstance);
+
+	if ((SpawnedActors != nullptr) && (SpawnedActors->Num() > 0))
 	{
+		// we have a spawned actor for this handle
 		const TArray<FSpawnedActorInfo>& SpawnedActorInfoRef = *SpawnedActors;
 
-		for( int32 SpawnedActorIndex = 0; SpawnedActorIndex < SpawnedActorInfoRef.Num(); ++SpawnedActorIndex )
+		for (int32 SpawnedActorIndex = 0; SpawnedActorIndex < SpawnedActorInfoRef.Num(); ++SpawnedActorIndex)
 		{
 			const FSpawnedActorInfo ActorInfo = SpawnedActorInfoRef[SpawnedActorIndex];
 
-			if( ActorInfo.RuntimeGuid == ObjectHandle && ActorInfo.SpawnedActor.IsValid() )
+			if ((ActorInfo.RuntimeGuid == ObjectHandle) && ActorInfo.SpawnedActor.IsValid())
 			{
-				OutObjects.Add( ActorInfo.SpawnedActor.Get() );
+				OutObjects.Add(ActorInfo.SpawnedActor.Get());
 			}
 		}
-	
 	}
-	// Otherwise, check to see if we have one or more possessed actors that are mapped to this handle
-	else if( MovieSceneBindings != nullptr )
+	else if (MovieSceneBindings != nullptr)
 	{
+		// otherwise, check whether we have one or more possessed actors that are mapped to this handle
 		for (FMovieSceneBoundObjectInfo& BoundObject : MovieSceneBindings->FindBoundObjects(ObjectHandle))
 		{
 			if (BoundObject.Tag.IsEmpty() == false)
