MovieScene: Changed IMovieSceneAnimation interface to UMovieSceneAnimation base class to work around UInterface limitations

diff --git a/Engine/Source/Editor/Sequencer/Private/Sequencer.cpp b/Engine/Source/Editor/Sequencer/Private/Sequencer.cpp
--- a/Engine/Source/Editor/Sequencer/Private/Sequencer.cpp
+++ b/Engine/Source/Editor/Sequencer/Private/Sequencer.cpp
@@ -307,7 +307,7 @@ UMovieScene* FSequencer::GetFocusedMovieScene() const
 	return MovieSceneStack.Top()->GetMovieScene();
 }
 
-void FSequencer::ResetToNewAnimation(IMovieSceneAnimation& NewAnimation)
+void FSequencer::ResetToNewAnimation(UMovieSceneAnimation& NewAnimation)
 {
 	DestroySpawnablesForAllMovieScenes();
 
@@ -1507,7 +1507,7 @@ void FSequencer::KeyProperty(FKeyPropertyParams KeyPropertyParams)
 	ObjectChangeListener->KeyProperty(KeyPropertyParams);
 }
 
-IMovieSceneAnimation* FSequencer::GetAnimation()
+UMovieSceneAnimation* FSequencer::GetAnimation()
 {
 	return Animation;
 }
diff --git a/Engine/Source/Editor/Sequencer/Private/Sequencer.h b/Engine/Source/Editor/Sequencer/Private/Sequencer.h
--- a/Engine/Source/Editor/Sequencer/Private/Sequencer.h
+++ b/Engine/Source/Editor/Sequencer/Private/Sequencer.h
@@ -9,7 +9,7 @@
 
 class UMovieScene;
 class IToolkitHost;
-class IMovieSceneAnimation;
+class UMovieSceneAnimation;
 class ISequencerObjectChangeListener;
 
 
@@ -63,7 +63,7 @@ public:
 	virtual TSharedRef<SWidget> GetSequencerWidget() const override { return SequencerWidget.ToSharedRef(); }
 	virtual UMovieScene* GetRootMovieScene() const override;
 	virtual UMovieScene* GetFocusedMovieScene() const override;
-	virtual void ResetToNewAnimation(IMovieSceneAnimation& NewAnimation) override;
+	virtual void ResetToNewAnimation(UMovieSceneAnimation& NewAnimation) override;
 	virtual TSharedRef<FMovieSceneInstance> GetRootMovieSceneInstance() const override;
 	virtual TSharedRef<FMovieSceneInstance> GetFocusedMovieSceneInstance() const override;
 	virtual void FocusSubMovieScene( TSharedRef<FMovieSceneInstance> SubMovieSceneInstance ) override;
@@ -86,7 +86,7 @@ public:
 	virtual void FilterToSelectedShotSections(bool bZoomToShotBounds = true) override;
 	virtual bool CanKeyProperty(FCanKeyPropertyParams CanKeyPropertyParams) const override;
 	virtual void KeyProperty(FKeyPropertyParams KeyPropertyParams) override;
-	virtual IMovieSceneAnimation* GetAnimation() override;
+	virtual UMovieSceneAnimation* GetAnimation() override;
 	virtual FSequencerSelection& GetSelection() override;
 	virtual void NotifyMapChanged(class UWorld* NewWorld, EMapChangeType MapChangeType) override;
 
@@ -404,7 +404,7 @@ private:
 	TWeakPtr<FMovieSceneTrackEditor> AnimationTrackEditor;
 
 	/** The current animation being edited. */
-	IMovieSceneAnimation* Animation;
+	UMovieSceneAnimation* Animation;
 
 	/** Listener for object changes being made while this sequencer is open*/
 	TSharedPtr< class ISequencerObjectChangeListener > ObjectChangeListener;
diff --git a/Engine/Source/Editor/Sequencer/Private/SequencerNodeTree.cpp b/Engine/Source/Editor/Sequencer/Private/SequencerNodeTree.cpp
--- a/Engine/Source/Editor/Sequencer/Private/SequencerNodeTree.cpp
+++ b/Engine/Source/Editor/Sequencer/Private/SequencerNodeTree.cpp
@@ -201,7 +201,7 @@ TSharedRef<FObjectBindingNode> FSequencerNodeTree::AddObjectBinding(const FStrin
 		// Try to get the parent object node if there is one.
 		TSharedPtr<FObjectBindingNode> ParentNode;
 		TArray<UObject*> RuntimeObjects;
-		IMovieSceneAnimation* Animation = Sequencer.GetAnimation();
+		UMovieSceneAnimation* Animation = Sequencer.GetAnimation();
 		UObject* RuntimeObject = Animation->FindObject(ObjectBinding);
 		if ( RuntimeObject != nullptr)
 		{
diff --git a/Engine/Source/Editor/Sequencer/Public/ISequencer.h b/Engine/Source/Editor/Sequencer/Public/ISequencer.h
--- a/Engine/Source/Editor/Sequencer/Public/ISequencer.h
+++ b/Engine/Source/Editor/Sequencer/Public/ISequencer.h
@@ -9,7 +9,7 @@
 
 
 class FSequencerSelection;
-class IMovieSceneAnimation;
+class UMovieSceneAnimation;
 class UAnimSequence;
 class UMovieScene;
 class UMovieSceneSection;
@@ -43,7 +43,7 @@ public:
 	virtual TSharedRef<FMovieSceneInstance> GetFocusedMovieSceneInstance() const = 0;
 
 	/** Resets sequencer with a new animation */
-	virtual void ResetToNewAnimation(IMovieSceneAnimation& NewAnimation) = 0;
+	virtual void ResetToNewAnimation(UMovieSceneAnimation& NewAnimation) = 0;
 
 	/**
 	 * Focuses a sub-movie scene (MovieScene within a MovieScene) in the sequencer.
@@ -161,7 +161,7 @@ public:
 	 *
 	 * @return The current animation, or nullptr if no animation is playing.
 	 */
-	virtual IMovieSceneAnimation* GetAnimation() = 0;
+	virtual UMovieSceneAnimation* GetAnimation() = 0;
 
 	virtual FSequencerSelection& GetSelection() = 0;
 
diff --git a/Engine/Source/Editor/Sequencer/Public/ISequencerModule.h b/Engine/Source/Editor/Sequencer/Public/ISequencerModule.h
--- a/Engine/Source/Editor/Sequencer/Public/ISequencerModule.h
+++ b/Engine/Source/Editor/Sequencer/Public/ISequencerModule.h
@@ -56,7 +56,7 @@ struct FSequencerViewParams
 struct FSequencerInitParams
 {
 	/** The animation being edited. */
-	IMovieSceneAnimation* Animation;
+	UMovieSceneAnimation* Animation;
 
 	/** The asset editor created for this (if any) */
 	TSharedPtr<IToolkitHost> ToolkitHost;
diff --git a/Engine/Source/Runtime/ActorAnimation/Private/ActorAnimation.cpp b/Engine/Source/Runtime/ActorAnimation/Private/ActorAnimation.cpp
--- a/Engine/Source/Runtime/ActorAnimation/Private/ActorAnimation.cpp
+++ b/Engine/Source/Runtime/ActorAnimation/Private/ActorAnimation.cpp
@@ -25,7 +25,7 @@ void UActorAnimation::Initialize()
 }
 
 
-/* IMovieSceneAnimation interface
+/* UMovieSceneAnimation overrides
  *****************************************************************************/
 
 bool UActorAnimation::AllowsSpawnableObjects() const
diff --git a/Engine/Source/Runtime/ActorAnimation/Public/ActorAnimation.h b/Engine/Source/Runtime/ActorAnimation/Public/ActorAnimation.h
--- a/Engine/Source/Runtime/ActorAnimation/Public/ActorAnimation.h
+++ b/Engine/Source/Runtime/ActorAnimation/Public/ActorAnimation.h
@@ -12,8 +12,7 @@
  */
 UCLASS(BlueprintType, MinimalAPI)
 class UActorAnimation
-	: public UObject
-	, public IMovieSceneAnimation
+	: public UMovieSceneAnimation
 {
 	GENERATED_UCLASS_BODY()
 
@@ -22,7 +21,7 @@ class UActorAnimation
 
 public:
 
-	// IMovieSceneAnimation interface
+	// UMovieSceneAnimation overrides
 
 	virtual bool AllowsSpawnableObjects() const override;
 	virtual void BindPossessableObject(const FGuid& ObjectId, UObject& PossessedObject) override;
diff --git a/Engine/Source/Runtime/MovieScene/Private/MovieSceneAnimation.cpp b/Engine/Source/Runtime/MovieScene/Private/MovieSceneAnimation.cpp
deleted file mode 100644
diff --git a/Engine/Source/Runtime/MovieScene/Public/MovieSceneAnimation.h b/Engine/Source/Runtime/MovieScene/Public/MovieSceneAnimation.h
--- a/Engine/Source/Runtime/MovieScene/Public/MovieSceneAnimation.h
+++ b/Engine/Source/Runtime/MovieScene/Public/MovieSceneAnimation.h
@@ -11,29 +11,22 @@ class UObject;
 
 
 /**
- * Interface for movie scene animations (UObject version).
+ * Abstract base class for movie scene animations (C++ version).
  */
-UINTERFACE(MinimalAPI, meta=(CannotImplementInterfaceInBlueprint))
+UCLASS(MinimalAPI)
 class UMovieSceneAnimation
-	: public UInterface
+	: public UObject
 {
-	GENERATED_UINTERFACE_BODY()
-};
-
+	GENERATED_BODY()
 
-/**
- * Interface for movie scene animations (C++ version).
- */
-class IMovieSceneAnimation
-{
-	GENERATED_IINTERFACE_BODY()
+public:
 
 	/**
 	 * Whether objects can be spawned at run-time.
 	 *
 	 * @return true if objects can be spawned by sequencer, false if only existing objects can be possessed.
 	 */
-	virtual bool AllowsSpawnableObjects() const PURE_VIRTUAL(IMovieSceneAnimation::AllowsSpawnableObjects, return false;);
+	virtual bool AllowsSpawnableObjects() const PURE_VIRTUAL(UMovieSceneAnimation::AllowsSpawnableObjects, return false;);
 
 	/**
 	 * Called when Sequencer has created an object binding for a possessable object
@@ -42,7 +35,7 @@ class IMovieSceneAnimation
 	 * @param PossessedObject The runtime object which was possessed.
 	 * @see UnbindPossessableObjects
 	 */
-	virtual void BindPossessableObject(const FGuid& ObjectId, UObject& PossessedObject) PURE_VIRTUAL(IMovieSceneAnimation::BindPossessableObject,);
+	virtual void BindPossessableObject(const FGuid& ObjectId, UObject& PossessedObject) PURE_VIRTUAL(UMovieSceneAnimation::BindPossessableObject,);
 	
 	/**
 	 * Check whether the given object can be possessed by this animation.
@@ -50,14 +43,14 @@ class IMovieSceneAnimation
 	 * @param Object The object to check.
 	 * @return true if the object can be possessed, false otherwise.
 	 */
-	virtual bool CanPossessObject(UObject& Object) const PURE_VIRTUAL(IMovieSceneAnimation::CanPossessObject, return false;);
+	virtual bool CanPossessObject(UObject& Object) const PURE_VIRTUAL(UMovieSceneAnimation::CanPossessObject, return false;);
 
 	/**
 	 * Destroy all objects that were spawned by this animation.
 	 *
 	 * @param SubMovieScene The (sub) movie scene to destroy objects for.
 	 */
-	virtual void DestroyAllSpawnedObjects(UMovieScene& SubMovieScene) PURE_VIRTUAL(IMovieSceneAnimation::DestroyAllSpawnedObjects,);
+	virtual void DestroyAllSpawnedObjects(UMovieScene& SubMovieScene) PURE_VIRTUAL(UMovieSceneAnimation::DestroyAllSpawnedObjects,);
 
 	/**
 	 * Gets the possessed or spawned object for the specified identifier.
@@ -65,7 +58,7 @@ class IMovieSceneAnimation
 	 * @param ObjectId The unique identifier of the object.
 	 * @return The object, or nullptr if not found.
 	 */
-	virtual UObject* FindObject(const FGuid& ObjectId) const PURE_VIRTUAL(IMovieSceneAnimation::FindObject, return nullptr;);
+	virtual UObject* FindObject(const FGuid& ObjectId) const PURE_VIRTUAL(UMovieSceneAnimation::FindObject, return nullptr;);
 
 	/**
 	 * Find the identifier for a possessed or spawned object.
@@ -73,7 +66,7 @@ class IMovieSceneAnimation
 	 * @param The object to get the identifier for.
 	 * @return The object identifier, or an invalid GUID if not found.
 	 */
-	virtual FGuid FindObjectId( UObject& Object) const PURE_VIRTUAL(IMovieSceneAnimation::FindGuidForObject, return FGuid(););
+	virtual FGuid FindObjectId( UObject& Object) const PURE_VIRTUAL(UMovieSceneAnimation::FindGuidForObject, return FGuid(););
 
 	/**
 	 * Get the movie scene that controls this animation.
@@ -83,7 +76,7 @@ class IMovieSceneAnimation
 	 *
 	 * @return The movie scene.
 	 */
-	virtual UMovieScene* GetMovieScene() const PURE_VIRTUAL(IMovieSceneAnimation::GetMovieScene(), return nullptr;);
+	virtual UMovieScene* GetMovieScene() const PURE_VIRTUAL(UMovieSceneAnimation::GetMovieScene(), return nullptr;);
 
 	/**
 	 * Get the logical parent object for the supplied object (not necessarily its outer).
@@ -91,7 +84,7 @@ class IMovieSceneAnimation
 	 * @param Object The object whose parent to get.
 	 * @return The parent object, or nullptr if the object has no logical parent.
 	 */
-	virtual UObject* GetParentObject(UObject* Object) const = 0;
+	virtual UObject* GetParentObject(UObject* Object) const PURE_VIRTUAL(UMovieSceneAnimation::GetParentObject(), return nullptr;);
 
 	/**
 	 * Called to spawn or destroy objects for a movie instance.
@@ -99,7 +92,7 @@ class IMovieSceneAnimation
 	 * @param SubMovieScene The (sub) movie scene to spawn or destroy objects for.
 	 * @param DestroyAll If true, destroy all spawned objects for the instance, if false only destroy unused objects.
 	 */
-	virtual void SpawnOrDestroyObjects(UMovieScene* SubMovieScene, bool DestroyAll) PURE_VIRTUAL(IMovieSceneAnimation::SpawnOrDestroyObjects,);
+	virtual void SpawnOrDestroyObjects(UMovieScene* SubMovieScene, bool DestroyAll) PURE_VIRTUAL(UMovieSceneAnimation::SpawnOrDestroyObjects,);
 
 	/**
 	 * Unbinds all possessable objects from the provided GUID.
@@ -107,7 +100,7 @@ class IMovieSceneAnimation
 	 * @param ObjectId The guid bound to possessable objects that should be removed.
 	 * @see BindPossessableObject
 	 */
-	virtual void UnbindPossessableObjects(const FGuid& ObjectId) PURE_VIRTUAL(IMovieSceneAnimation::UnbindPossessableObjects,);
+	virtual void UnbindPossessableObjects(const FGuid& ObjectId) PURE_VIRTUAL(UMovieSceneAnimation::UnbindPossessableObjects,);
 
 public:
 
@@ -120,6 +113,6 @@ public:
 	 * @param DisplayName The display name for the object binding.
 	 * @returns true if DisplayName has been set to a valid display name, otherwise false.
 	 */
-	virtual bool TryGetObjectDisplayName(const FGuid& ObjectId, FText& OutDisplayName) const PURE_VIRTUAL(IMovieSceneAnimation::TryGetObjectDisplayName, return false;);
+	virtual bool TryGetObjectDisplayName(const FGuid& ObjectId, FText& OutDisplayName) const PURE_VIRTUAL(UMovieSceneAnimation::TryGetObjectDisplayName, return false;);
 #endif
 };
diff --git a/Engine/Source/Runtime/Niagara/Private/NiagaraAnimation.cpp b/Engine/Source/Runtime/Niagara/Private/NiagaraAnimation.cpp
--- a/Engine/Source/Runtime/Niagara/Private/NiagaraAnimation.cpp
+++ b/Engine/Source/Runtime/Niagara/Private/NiagaraAnimation.cpp
@@ -13,7 +13,7 @@ UNiagaraAnimation::UNiagaraAnimation(const FObjectInitializer& ObjectInitializer
 { }
 
 
-/* IMovieSceneAnimation interface
+/* UMovieSceneAnimation overrides
  *****************************************************************************/
 
 bool UNiagaraAnimation::AllowsSpawnableObjects() const
diff --git a/Engine/Source/Runtime/Niagara/Public/NiagaraAnimation.h b/Engine/Source/Runtime/Niagara/Public/NiagaraAnimation.h
--- a/Engine/Source/Runtime/Niagara/Public/NiagaraAnimation.h
+++ b/Engine/Source/Runtime/Niagara/Public/NiagaraAnimation.h
@@ -11,14 +11,13 @@
  */
 UCLASS(BlueprintType, MinimalAPI)
 class UNiagaraAnimation
-	: public UObject
-	, public IMovieSceneAnimation
+	: public UMovieSceneAnimation
 {
 	GENERATED_UCLASS_BODY()
 
 public:
 
-	// IMovieSceneAnimation interface
+	// UMovieSceneAnimation overrides
 
 	virtual bool AllowsSpawnableObjects() const override;
 	virtual void BindPossessableObject(const FGuid& ObjectId, UObject& PossessedObject) override;
diff --git a/Engine/Source/Runtime/UMG/Private/Animation/WidgetAnimation.cpp b/Engine/Source/Runtime/UMG/Private/Animation/WidgetAnimation.cpp
--- a/Engine/Source/Runtime/UMG/Private/Animation/WidgetAnimation.cpp
+++ b/Engine/Source/Runtime/UMG/Private/Animation/WidgetAnimation.cpp
@@ -94,7 +94,7 @@ void UWidgetAnimation::Initialize(UUserWidget* InPreviewWidget)
 }
 
 
-/* IMovieSceneAnimation interface
+/* UMovieSceneAnimation overrides
  *****************************************************************************/
 
 bool UWidgetAnimation::AllowsSpawnableObjects() const
diff --git a/Engine/Source/Runtime/UMG/Public/Animation/WidgetAnimation.h b/Engine/Source/Runtime/UMG/Public/Animation/WidgetAnimation.h
--- a/Engine/Source/Runtime/UMG/Public/Animation/WidgetAnimation.h
+++ b/Engine/Source/Runtime/UMG/Public/Animation/WidgetAnimation.h
@@ -13,8 +13,7 @@ class UUserWidget;
 
 UCLASS(BlueprintType, MinimalAPI)
 class UWidgetAnimation
-	: public UObject
-	, public IMovieSceneAnimation
+	: public UMovieSceneAnimation
 {
 	GENERATED_UCLASS_BODY()
 
@@ -56,7 +55,7 @@ public:
 
 public:
 
-	// IMovieSceneAnimation interface
+	// UMovieSceneAnimation overrides
 
 	virtual bool AllowsSpawnableObjects() const override;
 	virtual void BindPossessableObject(const FGuid& ObjectId, UObject& PossessedObject) override;
