Sequencer: Fix for invalid time range in gameplay. The time range is set in the editor when it's loaded, but wasn't set for loading movie scenes in gameplay. It now reverts back to the previous behavior of all the section ranges.

diff --git a/Engine/Source/Runtime/MovieScene/Private/MovieScene.cpp b/Engine/Source/Runtime/MovieScene/Private/MovieScene.cpp
--- a/Engine/Source/Runtime/MovieScene/Private/MovieScene.cpp
+++ b/Engine/Source/Runtime/MovieScene/Private/MovieScene.cpp
@@ -217,6 +217,24 @@ FMovieScenePossessable& UMovieScene::GetPossessable( const int32 Index )
 
 TRange<float> UMovieScene::GetTimeRange() const
 {
+	if (InTime == FLT_MAX || OutTime == -FLT_MAX)
+	{
+		// Get the range of all sections combined
+		TArray< TRange<float> > Bounds;
+
+		for (int32 TypeIndex = 0; TypeIndex < MasterTracks.Num(); ++TypeIndex)
+		{
+			Bounds.Add(MasterTracks[TypeIndex]->GetSectionBoundaries());
+		}
+
+		for (int32 BindingIndex = 0; BindingIndex < ObjectBindings.Num(); ++BindingIndex)
+		{
+			Bounds.Add(ObjectBindings[BindingIndex].GetTimeRange());
+		}
+
+		return TRange<float>::Hull(Bounds);
+	}
+
 	return TRange<float>(InTime, OutTime);
 }
 
