Sequencer: Forgot missing file.

diff --git a/Engine/Source/Runtime/MovieScene/Public/KeyParams.h b/Engine/Source/Runtime/MovieScene/Public/KeyParams.h
new file mode 100644
--- /dev/null
+++ b/Engine/Source/Runtime/MovieScene/Public/KeyParams.h
@@ -0,0 +1,26 @@
+// Copyright 1998-2015 Epic Games, Inc. All Rights Reserved.
+
+#pragma once
+
+/**
+ * Parameters for determining keying behavior
+ */
+struct MOVIESCENE_API FKeyParams
+{
+	FKeyParams()
+	{
+		bCreateHandleIfMissing = false;
+		bCreateTrackIfMissing = false;
+		bAddKeyEvenIfUnchanged = false;
+		bAutoKeying = false;
+	}
+
+	/** Create handle if it doesn't exist. */
+	bool bCreateHandleIfMissing;
+	/** Create track if it doesn't exist. */
+	bool bCreateTrackIfMissing;
+	/** Create a key even if it's unchanged. */
+	bool bAddKeyEvenIfUnchanged;
+	/** Auto keying on.*/
+	bool bAutoKeying;
+};
