Sequencer: Set actor component mobility to movable when animating with the transform track.

diff --git a/Engine/Source/Editor/MovieSceneTools/Private/TrackEditors/AttachTrackEditor.cpp b/Engine/Source/Editor/MovieSceneTools/Private/TrackEditors/AttachTrackEditor.cpp
--- a/Engine/Source/Editor/MovieSceneTools/Private/TrackEditors/AttachTrackEditor.cpp
+++ b/Engine/Source/Editor/MovieSceneTools/Private/TrackEditors/AttachTrackEditor.cpp
@@ -9,6 +9,7 @@
 #include "ISequencerSection.h"
 #include "ISectionLayoutBuilder.h"
 #include "IKeyArea.h"
+#include "MovieSceneCommonHelpers.h"
 #include "MovieSceneToolHelpers.h"
 #include "MovieSceneTrackEditor.h"
 #include "ActorEditorUtils.h"
@@ -330,6 +331,8 @@ void F3DAttachTrackEditor::AddKeyInternal( float KeyTime, const TArray<UObject*>
 			FGuid ObjectHandle = FindOrCreateHandleToObject( Object );
 			if (ObjectHandle.IsValid())
 			{
+				MovieSceneHelpers::SetRuntimeObjectMobility(Object);
+
 				UMovieSceneTrack* Track = GetTrackForObject( ObjectHandle, UMovieScene3DAttachTrack::StaticClass(), FName("Attach"));
 
 				if (ensure(Track))
diff --git a/Engine/Source/Editor/MovieSceneTools/Private/TrackEditors/PathTrackEditor.cpp b/Engine/Source/Editor/MovieSceneTools/Private/TrackEditors/PathTrackEditor.cpp
--- a/Engine/Source/Editor/MovieSceneTools/Private/TrackEditors/PathTrackEditor.cpp
+++ b/Engine/Source/Editor/MovieSceneTools/Private/TrackEditors/PathTrackEditor.cpp
@@ -9,6 +9,7 @@
 #include "ISequencerSection.h"
 #include "ISectionLayoutBuilder.h"
 #include "IKeyArea.h"
+#include "MovieSceneCommonHelpers.h"
 #include "MovieSceneToolHelpers.h"
 #include "MovieSceneTrackEditor.h"
 #include "PathTrackEditor.h"
@@ -311,6 +312,8 @@ void F3DPathTrackEditor::AddKeyInternal( float KeyTime, const TArray<UObject*> O
 			FGuid ObjectHandle = FindOrCreateHandleToObject( Object );
 			if (ObjectHandle.IsValid())
 			{
+				MovieSceneHelpers::SetRuntimeObjectMobility(Object);
+
 				UMovieSceneTrack* Track = GetTrackForObject( ObjectHandle, UMovieScene3DPathTrack::StaticClass(), FName("Path"));
 
 				if (ensure(Track))
diff --git a/Engine/Source/Editor/MovieSceneTools/Private/TrackEditors/TransformTrackEditor.cpp b/Engine/Source/Editor/MovieSceneTools/Private/TrackEditors/TransformTrackEditor.cpp
--- a/Engine/Source/Editor/MovieSceneTools/Private/TrackEditors/TransformTrackEditor.cpp
+++ b/Engine/Source/Editor/MovieSceneTools/Private/TrackEditors/TransformTrackEditor.cpp
@@ -3,6 +3,7 @@
 #include "MovieSceneToolsPrivatePCH.h"
 #include "ScopedTransaction.h"
 #include "MovieScene.h"
+#include "MovieSceneCommonHelpers.h"
 #include "MovieSceneSection.h"
 #include "MovieScene3DTransformTrack.h"
 #include "MovieScene3DTransformSection.h"
@@ -312,6 +313,8 @@ void F3DTransformTrackEditor::AddKeyInternal(const FGuid& ObjectGuid, UObject* A
 
 		if( SceneComponent )
 		{
+			MovieSceneHelpers::SetRuntimeObjectMobility(Object);
+
 			// Cache off the existing transform so we can detect which components have changed
 			// and keys only when something has changed
 			FTransformData CurrentTransform( SceneComponent );
diff --git a/Engine/Source/Runtime/MovieScene/Private/MovieSceneCommonHelpers.cpp b/Engine/Source/Runtime/MovieScene/Private/MovieSceneCommonHelpers.cpp
--- a/Engine/Source/Runtime/MovieScene/Private/MovieSceneCommonHelpers.cpp
+++ b/Engine/Source/Runtime/MovieScene/Private/MovieSceneCommonHelpers.cpp
@@ -113,6 +113,19 @@ USceneComponent* MovieSceneHelpers::SceneComponentFromRuntimeObject(UObject* Obj
 	return SceneComponent;
 }
 
+void MovieSceneHelpers::SetRuntimeObjectMobility(UObject* Object, EComponentMobility::Type ComponentMobility)
+{
+	USceneComponent* SceneComponent = SceneComponentFromRuntimeObject(Object);
+	if (SceneComponent)
+	{
+		if (SceneComponent->Mobility != ComponentMobility)
+		{
+			SceneComponent->Modify();
+
+			SceneComponent->SetMobility(ComponentMobility);
+		}
+	}
+}
 
 FTrackInstancePropertyBindings::FTrackInstancePropertyBindings( FName InPropertyName, const FString& InPropertyPath, const FName& InFunctionName )
     : PropertyPath( InPropertyPath )
diff --git a/Engine/Source/Runtime/MovieScene/Public/MovieSceneCommonHelpers.h b/Engine/Source/Runtime/MovieScene/Public/MovieSceneCommonHelpers.h
--- a/Engine/Source/Runtime/MovieScene/Public/MovieSceneCommonHelpers.h
+++ b/Engine/Source/Runtime/MovieScene/Public/MovieSceneCommonHelpers.h
@@ -2,7 +2,6 @@
 
 #pragma once
 
-
 class MOVIESCENE_API MovieSceneHelpers
 {
 public:
@@ -12,28 +11,36 @@ public:
 	static TArray<UMovieSceneSection*> GetTraversedSections( const TArray<UMovieSceneSection*>& Sections, float CurrentTime, float PreviousTime );
 
 	/**
-	* Finds a section that exists at a given time
-	*
-	* @param Time	The time to find a section at
-	* @return The found section or null
-	*/
+	 * Finds a section that exists at a given time
+	 *
+	 * @param Time	The time to find a section at
+	 * @return The found section or null
+	 */
 	static UMovieSceneSection* FindSectionAtTime( const TArray<UMovieSceneSection*>& Sections, float Time );
 
 	/**
-	* Finds the nearest section to the given time
-	*
-	* @param Time	The time to find a section at
-	* @return The found section or null
-	*/
+	 * Finds the nearest section to the given time
+	 *
+	 * @param Time	The time to find a section at
+	 * @return The found section or null
+	 */
 	static UMovieSceneSection* FindNearestSectionAtTime( const TArray<UMovieSceneSection*>& Sections, float Time );
 
 	/**
-	* Get the scene component from the runtime object
-	*
-	* @param Object The object to get the scene component for
-	* @return The found scene component
-	*/	
+	 * Get the scene component from the runtime object
+	 *
+	 * @param Object The object to get the scene component for
+	 * @return The found scene component
+	 */	
 	static USceneComponent* SceneComponentFromRuntimeObject(UObject* Object);
+
+	/*
+	 * Set the runtime object movable
+	 *
+	 * @param Object The object to set the mobility for
+	 * @param Mobility The mobility of the runtime object
+	 */
+	static void SetRuntimeObjectMobility(UObject* Object, EComponentMobility::Type ComponentMobility = EComponentMobility::Movable);
 };
 
 /**
