Sequencer: Add IMovieScenePlayer to RestoreState and SaveState and ensure that master tracks are also stored and restored. This fixes the fade track so that it restores the fade when closing out.

diff --git a/Engine/Source/Editor/NiagaraEditor/Private/NiagaraSequencer.h b/Engine/Source/Editor/NiagaraEditor/Private/NiagaraSequencer.h
--- a/Engine/Source/Editor/NiagaraEditor/Private/NiagaraSequencer.h
+++ b/Engine/Source/Editor/NiagaraEditor/Private/NiagaraSequencer.h
@@ -121,8 +121,8 @@ public:
 	virtual void RefreshInstance(const TArray<UObject*>& RuntimeObjects, class IMovieScenePlayer& Player);
 
 	virtual void ClearInstance( IMovieScenePlayer& Player ) override {}
-	virtual void SaveState(const TArray<UObject*>& RuntimeObjects) override {}
-	virtual void RestoreState(const TArray<UObject*>& RuntimeObjects) override {}
+	virtual void SaveState(const TArray<UObject*>& RuntimeObjects, IMovieScenePlayer& Player) override {}
+	virtual void RestoreState(const TArray<UObject*>& RuntimeObjects, IMovieScenePlayer& Player) override {}
 
 private:
 	TSharedPtr<FNiagaraSimulation> Emitter;
diff --git a/Engine/Source/Editor/Sequencer/Private/Sequencer.cpp b/Engine/Source/Editor/Sequencer/Private/Sequencer.cpp
--- a/Engine/Source/Editor/Sequencer/Private/Sequencer.cpp
+++ b/Engine/Source/Editor/Sequencer/Private/Sequencer.cpp
@@ -189,7 +189,7 @@ void FSequencer::Close()
 {
 	if (ActiveSequencer.IsValid() && ActiveSequencer.Get() == this)
 	{
-		RootMovieSceneSequenceInstance->RestoreState();
+		RootMovieSceneSequenceInstance->RestoreState(*this);
 
 		FEditorDelegates::PreSaveWorld.RemoveAll(this);
 		FEditorDelegates::PostSaveWorld.RemoveAll(this);
@@ -1664,7 +1664,7 @@ void FSequencer::PostUndo(bool bSuccess)
 void FSequencer::OnPreSaveWorld(uint32 SaveFlags, class UWorld* World)
 {
 	// Restore the saved state so that the level save can save that instead of the animated state.
-	RootMovieSceneSequenceInstance->RestoreState();
+	RootMovieSceneSequenceInstance->RestoreState(*this);
 }
 
 void FSequencer::OnPostSaveWorld(uint32 SaveFlags, class UWorld* World, bool bSuccess)
diff --git a/Engine/Source/Runtime/MovieScene/Private/MovieSceneSequenceInstance.cpp b/Engine/Source/Runtime/MovieScene/Private/MovieSceneSequenceInstance.cpp
--- a/Engine/Source/Runtime/MovieScene/Private/MovieSceneSequenceInstance.cpp
+++ b/Engine/Source/Runtime/MovieScene/Private/MovieSceneSequenceInstance.cpp
@@ -11,7 +11,7 @@ FMovieSceneSequenceInstance::FMovieSceneSequenceInstance(const UMovieSceneSequen
 }
 
 
-void FMovieSceneSequenceInstance::SaveState()
+void FMovieSceneSequenceInstance::SaveState(class IMovieScenePlayer& Player)
 {
 	TMap<FGuid, FMovieSceneObjectBindingInstance>::TIterator ObjectIt = ObjectBindingInstances.CreateIterator();
 	for (; ObjectIt; ++ObjectIt)
@@ -20,13 +20,18 @@ void FMovieSceneSequenceInstance::SaveState()
 
 		for (FMovieSceneInstanceMap::TIterator It = ObjectBindingInstance.TrackInstances.CreateIterator(); It; ++It)
 		{
-			It.Value()->SaveState(ObjectBindingInstance.RuntimeObjects);
+			It.Value()->SaveState(ObjectBindingInstance.RuntimeObjects, Player);
+		}
+
+		for( FMovieSceneInstanceMap::TIterator It( MasterTrackInstances ); It; ++It )
+		{
+			It.Value()->SaveState( ObjectBindingInstance.RuntimeObjects, Player );
 		}
 	}
 }
 
 
-void FMovieSceneSequenceInstance::RestoreState()
+void FMovieSceneSequenceInstance::RestoreState(class IMovieScenePlayer& Player)
 {
 	TMap<FGuid, FMovieSceneObjectBindingInstance>::TIterator ObjectIt = ObjectBindingInstances.CreateIterator();
 	for (; ObjectIt; ++ObjectIt)
@@ -35,7 +40,12 @@ void FMovieSceneSequenceInstance::RestoreState()
 
 		for (FMovieSceneInstanceMap::TIterator It = ObjectBindingInstance.TrackInstances.CreateIterator(); It; ++It)
 		{
-			It.Value()->RestoreState(ObjectBindingInstance.RuntimeObjects);
+			It.Value()->RestoreState(ObjectBindingInstance.RuntimeObjects, Player);
+		}
+
+		for( FMovieSceneInstanceMap::TIterator It( MasterTrackInstances ); It; ++It )
+		{
+			It.Value()->RestoreState( ObjectBindingInstance.RuntimeObjects, Player );
 		}
 	}
 }
@@ -172,7 +182,7 @@ void FMovieSceneSequenceInstance::RefreshInstanceMap( const TArray<UMovieSceneTr
 			// The track does not have an instance, create one
 			Instance = Track->CreateInstance();
 			Instance->RefreshInstance( RuntimeObjects, Player );
-			Instance->SaveState(RuntimeObjects);
+			Instance->SaveState(RuntimeObjects, Player);
 
 			TrackInstances.Add( Track, Instance );
 		}
diff --git a/Engine/Source/Runtime/MovieScene/Public/IMovieSceneTrackInstance.h b/Engine/Source/Runtime/MovieScene/Public/IMovieSceneTrackInstance.h
--- a/Engine/Source/Runtime/MovieScene/Public/IMovieSceneTrackInstance.h
+++ b/Engine/Source/Runtime/MovieScene/Public/IMovieSceneTrackInstance.h
@@ -20,13 +20,13 @@ public:
 	 * Save state of objects that this instance will be editing.
 	 * @todo Sequencer: This is likely editor only
 	 */
-	virtual void SaveState(const TArray<UObject*>& RuntimeObjects) = 0;
+	virtual void SaveState(const TArray<UObject*>& RuntimeObjects, IMovieScenePlayer& Player) = 0;
 
 	/**
 	 * Restore state of objects that this instance edited.
 	 * @todo Sequencer: This is likely editor only
 	 */
-	virtual void RestoreState(const TArray<UObject*>& RuntimeObjects) = 0;
+	virtual void RestoreState(const TArray<UObject*>& RuntimeObjects, IMovieScenePlayer& Player) = 0;
 
 	/**
 	 * Main update function for track instances.  Called in game and in editor when we update a moviescene.
diff --git a/Engine/Source/Runtime/MovieScene/Public/MovieSceneSequenceInstance.h b/Engine/Source/Runtime/MovieScene/Public/MovieSceneSequenceInstance.h
--- a/Engine/Source/Runtime/MovieScene/Public/MovieSceneSequenceInstance.h
+++ b/Engine/Source/Runtime/MovieScene/Public/MovieSceneSequenceInstance.h
@@ -27,10 +27,10 @@ public:
 	MOVIESCENE_API FMovieSceneSequenceInstance(const UMovieSceneSequence& InMovieSceneSequence);
 
 	/** Save state of the objects that this movie scene controls. */
-	MOVIESCENE_API void SaveState();
+	MOVIESCENE_API void SaveState(class IMovieScenePlayer& Player);
 
 	/** Restore state of the objects that this movie scene controls. */
-	MOVIESCENE_API void RestoreState();
+	MOVIESCENE_API void RestoreState(class IMovieScenePlayer& Player);
 
 	/**
 	 * Updates this movie scene.
diff --git a/Engine/Source/Runtime/MovieSceneTracks/Private/TrackInstances/MovieScene3DConstraintTrackInstance.cpp b/Engine/Source/Runtime/MovieSceneTracks/Private/TrackInstances/MovieScene3DConstraintTrackInstance.cpp
--- a/Engine/Source/Runtime/MovieSceneTracks/Private/TrackInstances/MovieScene3DConstraintTrackInstance.cpp
+++ b/Engine/Source/Runtime/MovieSceneTracks/Private/TrackInstances/MovieScene3DConstraintTrackInstance.cpp
@@ -15,7 +15,7 @@ FMovieScene3DConstraintTrackInstance::FMovieScene3DConstraintTrackInstance( UMov
 }
 
 
-void FMovieScene3DConstraintTrackInstance::SaveState(const TArray<UObject*>& RuntimeObjects)
+void FMovieScene3DConstraintTrackInstance::SaveState(const TArray<UObject*>& RuntimeObjects, IMovieScenePlayer& Player)
 {
 	for (int32 ObjIndex = 0; ObjIndex < RuntimeObjects.Num(); ++ObjIndex)
 	{
@@ -28,7 +28,7 @@ void FMovieScene3DConstraintTrackInstance::SaveState(const TArray<UObject*>& Run
 }
 
 
-void FMovieScene3DConstraintTrackInstance::RestoreState(const TArray<UObject*>& RuntimeObjects)
+void FMovieScene3DConstraintTrackInstance::RestoreState(const TArray<UObject*>& RuntimeObjects, IMovieScenePlayer& Player)
 {
 	for (int32 ObjIndex = 0; ObjIndex < RuntimeObjects.Num(); ++ObjIndex)
 	{
diff --git a/Engine/Source/Runtime/MovieSceneTracks/Private/TrackInstances/MovieScene3DTransformTrackInstance.cpp b/Engine/Source/Runtime/MovieSceneTracks/Private/TrackInstances/MovieScene3DTransformTrackInstance.cpp
--- a/Engine/Source/Runtime/MovieSceneTracks/Private/TrackInstances/MovieScene3DTransformTrackInstance.cpp
+++ b/Engine/Source/Runtime/MovieSceneTracks/Private/TrackInstances/MovieScene3DTransformTrackInstance.cpp
@@ -12,7 +12,7 @@ FMovieScene3DTransformTrackInstance::FMovieScene3DTransformTrackInstance( UMovie
 }
 
 
-void FMovieScene3DTransformTrackInstance::SaveState(const TArray<UObject*>& RuntimeObjects)
+void FMovieScene3DTransformTrackInstance::SaveState(const TArray<UObject*>& RuntimeObjects, IMovieScenePlayer& Player)
 {
 	for (int32 ObjIndex = 0; ObjIndex < RuntimeObjects.Num(); ++ObjIndex)
 	{
@@ -25,7 +25,7 @@ void FMovieScene3DTransformTrackInstance::SaveState(const TArray<UObject*>& Runt
 }
 
 
-void FMovieScene3DTransformTrackInstance::RestoreState(const TArray<UObject*>& RuntimeObjects)
+void FMovieScene3DTransformTrackInstance::RestoreState(const TArray<UObject*>& RuntimeObjects, IMovieScenePlayer& Player)
 {
 	for (int32 ObjIndex = 0; ObjIndex < RuntimeObjects.Num(); ++ObjIndex)
 	{
diff --git a/Engine/Source/Runtime/MovieSceneTracks/Private/TrackInstances/MovieSceneBoolTrackInstance.cpp b/Engine/Source/Runtime/MovieSceneTracks/Private/TrackInstances/MovieSceneBoolTrackInstance.cpp
--- a/Engine/Source/Runtime/MovieSceneTracks/Private/TrackInstances/MovieSceneBoolTrackInstance.cpp
+++ b/Engine/Source/Runtime/MovieSceneTracks/Private/TrackInstances/MovieSceneBoolTrackInstance.cpp
@@ -13,7 +13,7 @@ FMovieSceneBoolTrackInstance::FMovieSceneBoolTrackInstance( UMovieSceneBoolTrack
 }
 
 
-void FMovieSceneBoolTrackInstance::SaveState(const TArray<UObject*>& RuntimeObjects)
+void FMovieSceneBoolTrackInstance::SaveState(const TArray<UObject*>& RuntimeObjects, IMovieScenePlayer& Player)
 {
 	for( UObject* Object : RuntimeObjects )
 	{
@@ -23,7 +23,7 @@ void FMovieSceneBoolTrackInstance::SaveState(const TArray<UObject*>& RuntimeObje
 }
 
 
-void FMovieSceneBoolTrackInstance::RestoreState(const TArray<UObject*>& RuntimeObjects)
+void FMovieSceneBoolTrackInstance::RestoreState(const TArray<UObject*>& RuntimeObjects, IMovieScenePlayer& Player)
 {
 	for( UObject* Object : RuntimeObjects )
 	{
diff --git a/Engine/Source/Runtime/MovieSceneTracks/Private/TrackInstances/MovieSceneByteTrackInstance.cpp b/Engine/Source/Runtime/MovieSceneTracks/Private/TrackInstances/MovieSceneByteTrackInstance.cpp
--- a/Engine/Source/Runtime/MovieSceneTracks/Private/TrackInstances/MovieSceneByteTrackInstance.cpp
+++ b/Engine/Source/Runtime/MovieSceneTracks/Private/TrackInstances/MovieSceneByteTrackInstance.cpp
@@ -13,7 +13,7 @@ FMovieSceneByteTrackInstance::FMovieSceneByteTrackInstance( UMovieSceneByteTrack
 }
 
 
-void FMovieSceneByteTrackInstance::SaveState(const TArray<UObject*>& RuntimeObjects)
+void FMovieSceneByteTrackInstance::SaveState(const TArray<UObject*>& RuntimeObjects, IMovieScenePlayer& Player)
 {
 	for( UObject* Object : RuntimeObjects )
 	{
@@ -23,7 +23,7 @@ void FMovieSceneByteTrackInstance::SaveState(const TArray<UObject*>& RuntimeObje
 }
 
 
-void FMovieSceneByteTrackInstance::RestoreState(const TArray<UObject*>& RuntimeObjects)
+void FMovieSceneByteTrackInstance::RestoreState(const TArray<UObject*>& RuntimeObjects, IMovieScenePlayer& Player)
 {
 	for( UObject* Object : RuntimeObjects )
 	{
diff --git a/Engine/Source/Runtime/MovieSceneTracks/Private/TrackInstances/MovieSceneColorTrackInstance.cpp b/Engine/Source/Runtime/MovieSceneTracks/Private/TrackInstances/MovieSceneColorTrackInstance.cpp
--- a/Engine/Source/Runtime/MovieSceneTracks/Private/TrackInstances/MovieSceneColorTrackInstance.cpp
+++ b/Engine/Source/Runtime/MovieSceneTracks/Private/TrackInstances/MovieSceneColorTrackInstance.cpp
@@ -15,7 +15,7 @@ FMovieSceneColorTrackInstance::FMovieSceneColorTrackInstance( UMovieSceneColorTr
 }
 
 
-void FMovieSceneColorTrackInstance::SaveState(const TArray<UObject*>& RuntimeObjects)
+void FMovieSceneColorTrackInstance::SaveState(const TArray<UObject*>& RuntimeObjects, IMovieScenePlayer& Player)
 {
 	for( UObject* Object : RuntimeObjects )
 	{
@@ -45,7 +45,7 @@ void FMovieSceneColorTrackInstance::SaveState(const TArray<UObject*>& RuntimeObj
 }
 
 
-void FMovieSceneColorTrackInstance::RestoreState(const TArray<UObject*>& RuntimeObjects)
+void FMovieSceneColorTrackInstance::RestoreState(const TArray<UObject*>& RuntimeObjects, IMovieScenePlayer& Player)
 {
 	for( UObject* Object : RuntimeObjects )
 	{
diff --git a/Engine/Source/Runtime/MovieSceneTracks/Private/TrackInstances/MovieSceneFadeTrackInstance.cpp b/Engine/Source/Runtime/MovieSceneTracks/Private/TrackInstances/MovieSceneFadeTrackInstance.cpp
--- a/Engine/Source/Runtime/MovieSceneTracks/Private/TrackInstances/MovieSceneFadeTrackInstance.cpp
+++ b/Engine/Source/Runtime/MovieSceneTracks/Private/TrackInstances/MovieSceneFadeTrackInstance.cpp
@@ -17,6 +17,23 @@ FMovieSceneFadeTrackInstance::FMovieSceneFadeTrackInstance(UMovieSceneFadeTrack&
 /* IMovieSceneTrackInstance interface
  *****************************************************************************/
 
+void FMovieSceneFadeTrackInstance::RestoreState(const TArray<UObject*>& RuntimeObjects, IMovieScenePlayer& Player)
+{
+	// Reset editor preview/fade
+	EMovieSceneViewportParams ViewportParams;
+	ViewportParams.SetWhichViewportParam = (EMovieSceneViewportParams::SetViewportParam)(EMovieSceneViewportParams::SVP_FadeAmount | EMovieSceneViewportParams::SVP_FadeColor);
+	ViewportParams.FadeAmount = 0.f;
+	ViewportParams.FadeColor = FLinearColor::Black;
+
+	TMap<FViewportClient*, EMovieSceneViewportParams> ViewportParamsMap;
+	Player.GetViewportSettings(ViewportParamsMap);
+	for (auto ViewportParamsPair : ViewportParamsMap)
+	{
+		ViewportParamsMap[ViewportParamsPair.Key] = ViewportParams;
+	}
+	Player.SetViewportSettings(ViewportParamsMap);
+}
+
 void FMovieSceneFadeTrackInstance::Update(float Position, float LastPosition, const TArray<UObject*>& RuntimeObjects, IMovieScenePlayer& Player)
 {
 	FLinearColor FadeColor = FLinearColor::Black;
diff --git a/Engine/Source/Runtime/MovieSceneTracks/Private/TrackInstances/MovieSceneFloatTrackInstance.cpp b/Engine/Source/Runtime/MovieSceneTracks/Private/TrackInstances/MovieSceneFloatTrackInstance.cpp
--- a/Engine/Source/Runtime/MovieSceneTracks/Private/TrackInstances/MovieSceneFloatTrackInstance.cpp
+++ b/Engine/Source/Runtime/MovieSceneTracks/Private/TrackInstances/MovieSceneFloatTrackInstance.cpp
@@ -13,7 +13,7 @@ FMovieSceneFloatTrackInstance::FMovieSceneFloatTrackInstance( UMovieSceneFloatTr
 }
 
 
-void FMovieSceneFloatTrackInstance::SaveState(const TArray<UObject*>& RuntimeObjects)
+void FMovieSceneFloatTrackInstance::SaveState(const TArray<UObject*>& RuntimeObjects, IMovieScenePlayer& Player)
 {
 	for( UObject* Object : RuntimeObjects )
 	{
@@ -23,7 +23,7 @@ void FMovieSceneFloatTrackInstance::SaveState(const TArray<UObject*>& RuntimeObj
 }
 
 
-void FMovieSceneFloatTrackInstance::RestoreState(const TArray<UObject*>& RuntimeObjects)
+void FMovieSceneFloatTrackInstance::RestoreState(const TArray<UObject*>& RuntimeObjects, IMovieScenePlayer& Player)
 {
 	for( UObject* Object : RuntimeObjects )
 	{
diff --git a/Engine/Source/Runtime/MovieSceneTracks/Private/TrackInstances/MovieSceneMaterialTrackInstance.cpp b/Engine/Source/Runtime/MovieSceneTracks/Private/TrackInstances/MovieSceneMaterialTrackInstance.cpp
--- a/Engine/Source/Runtime/MovieSceneTracks/Private/TrackInstances/MovieSceneMaterialTrackInstance.cpp
+++ b/Engine/Source/Runtime/MovieSceneTracks/Private/TrackInstances/MovieSceneMaterialTrackInstance.cpp
@@ -9,7 +9,7 @@ FMovieSceneMaterialTrackInstance::FMovieSceneMaterialTrackInstance( UMovieSceneM
 	MaterialTrack = &InMaterialTrack;
 }
 
-void FMovieSceneMaterialTrackInstance::SaveState(const TArray<UObject*>& RuntimeObjects)
+void FMovieSceneMaterialTrackInstance::SaveState(const TArray<UObject*>& RuntimeObjects, IMovieScenePlayer& Player)
 {
 	for( UObject* Object : RuntimeObjects )
 	{
@@ -33,7 +33,7 @@ void FMovieSceneMaterialTrackInstance::SaveState(const TArray<UObject*>& Runtime
 	}
 }
 
-void FMovieSceneMaterialTrackInstance::RestoreState(const TArray<UObject*>& RuntimeObjects)
+void FMovieSceneMaterialTrackInstance::RestoreState(const TArray<UObject*>& RuntimeObjects, IMovieScenePlayer& Player)
 {
 	for( UObject* Object : RuntimeObjects )
 	{
diff --git a/Engine/Source/Runtime/MovieSceneTracks/Private/TrackInstances/MovieSceneVectorTrackInstance.cpp b/Engine/Source/Runtime/MovieSceneTracks/Private/TrackInstances/MovieSceneVectorTrackInstance.cpp
--- a/Engine/Source/Runtime/MovieSceneTracks/Private/TrackInstances/MovieSceneVectorTrackInstance.cpp
+++ b/Engine/Source/Runtime/MovieSceneTracks/Private/TrackInstances/MovieSceneVectorTrackInstance.cpp
@@ -12,7 +12,7 @@ FMovieSceneVectorTrackInstance::FMovieSceneVectorTrackInstance( UMovieSceneVecto
 }
 
 
-void FMovieSceneVectorTrackInstance::SaveState(const TArray<UObject*>& RuntimeObjects)
+void FMovieSceneVectorTrackInstance::SaveState(const TArray<UObject*>& RuntimeObjects, IMovieScenePlayer& Player)
 {
 	int32 NumChannelsUsed = VectorTrack->GetNumChannelsUsed();
 	for( UObject* Object : RuntimeObjects )
@@ -42,7 +42,7 @@ void FMovieSceneVectorTrackInstance::SaveState(const TArray<UObject*>& RuntimeOb
 }
 
 
-void FMovieSceneVectorTrackInstance::RestoreState(const TArray<UObject*>& RuntimeObjects)
+void FMovieSceneVectorTrackInstance::RestoreState(const TArray<UObject*>& RuntimeObjects, IMovieScenePlayer& Player)
 {
 	int32 NumChannelsUsed = VectorTrack->GetNumChannelsUsed();
 	for( UObject* Object : RuntimeObjects )
diff --git a/Engine/Source/Runtime/MovieSceneTracks/Private/TrackInstances/MovieSceneVisibilityTrackInstance.cpp b/Engine/Source/Runtime/MovieSceneTracks/Private/TrackInstances/MovieSceneVisibilityTrackInstance.cpp
--- a/Engine/Source/Runtime/MovieSceneTracks/Private/TrackInstances/MovieSceneVisibilityTrackInstance.cpp
+++ b/Engine/Source/Runtime/MovieSceneTracks/Private/TrackInstances/MovieSceneVisibilityTrackInstance.cpp
@@ -12,7 +12,7 @@ FMovieSceneVisibilityTrackInstance::FMovieSceneVisibilityTrackInstance(UMovieSce
 }
 
 
-void FMovieSceneVisibilityTrackInstance::SaveState(const TArray<UObject*>& RuntimeObjects)
+void FMovieSceneVisibilityTrackInstance::SaveState(const TArray<UObject*>& RuntimeObjects, IMovieScenePlayer& Player)
 {
  	for( UObject* Object : RuntimeObjects )
  	{
@@ -33,7 +33,7 @@ void FMovieSceneVisibilityTrackInstance::SaveState(const TArray<UObject*>& Runti
 }
 
 
-void FMovieSceneVisibilityTrackInstance::RestoreState(const TArray<UObject*>& RuntimeObjects)
+void FMovieSceneVisibilityTrackInstance::RestoreState(const TArray<UObject*>& RuntimeObjects, IMovieScenePlayer& Player)
 {
  	for( UObject* Object : RuntimeObjects )
  	{
diff --git a/Engine/Source/Runtime/MovieSceneTracks/Private/TrackInstances/SubMovieSceneTrackInstance.cpp b/Engine/Source/Runtime/MovieSceneTracks/Private/TrackInstances/SubMovieSceneTrackInstance.cpp
--- a/Engine/Source/Runtime/MovieSceneTracks/Private/TrackInstances/SubMovieSceneTrackInstance.cpp
+++ b/Engine/Source/Runtime/MovieSceneTracks/Private/TrackInstances/SubMovieSceneTrackInstance.cpp
@@ -90,7 +90,7 @@ void FSubMovieSceneTrackInstance::Update( float Position, float LastPosition, co
 }
 
 
-void FSubMovieSceneTrackInstance::SaveState(const TArray<UObject*>& RuntimeObjects)
+void FSubMovieSceneTrackInstance::SaveState(const TArray<UObject*>& RuntimeObjects, IMovieScenePlayer& Player)
 {
 	const TArray<UMovieSceneSection*>& AllSections = SubMovieSceneTrack->GetAllSections();
 
@@ -101,13 +101,13 @@ void FSubMovieSceneTrackInstance::SaveState(const TArray<UObject*>& RuntimeObjec
 		TSharedPtr<FMovieSceneSequenceInstance> Instance = SubMovieSceneInstances.FindRef( Section );
 		if( Instance.IsValid() )
 		{
-			Instance->SaveState();
+			Instance->SaveState(Player);
 		}
 	}
 }
 
 
-void FSubMovieSceneTrackInstance::RestoreState(const TArray<UObject*>& RuntimeObjects)
+void FSubMovieSceneTrackInstance::RestoreState(const TArray<UObject*>& RuntimeObjects, IMovieScenePlayer& Player)
 {
 	const TArray<UMovieSceneSection*>& AllSections = SubMovieSceneTrack->GetAllSections();
 
@@ -118,7 +118,7 @@ void FSubMovieSceneTrackInstance::RestoreState(const TArray<UObject*>& RuntimeOb
 		TSharedPtr<FMovieSceneSequenceInstance> Instance = SubMovieSceneInstances.FindRef( Section );
 		if( Instance.IsValid() )
 		{
-			Instance->RestoreState();
+			Instance->RestoreState(Player);
 		}
 	}
 }
diff --git a/Engine/Source/Runtime/MovieSceneTracks/Public/TrackInstances/MovieScene3DConstraintTrackInstance.h b/Engine/Source/Runtime/MovieSceneTracks/Public/TrackInstances/MovieScene3DConstraintTrackInstance.h
--- a/Engine/Source/Runtime/MovieSceneTracks/Public/TrackInstances/MovieScene3DConstraintTrackInstance.h
+++ b/Engine/Source/Runtime/MovieSceneTracks/Public/TrackInstances/MovieScene3DConstraintTrackInstance.h
@@ -17,8 +17,8 @@ public:
 	FMovieScene3DConstraintTrackInstance( UMovieScene3DConstraintTrack& InConstraintTrack );
 
 	/** IMovieSceneTrackInstance interface */
-	virtual void SaveState (const TArray<UObject*>& RuntimeObjects) override;
-	virtual void RestoreState (const TArray<UObject*>& RuntimeObjects) override;
+	virtual void SaveState (const TArray<UObject*>& RuntimeObjects, IMovieScenePlayer& Player) override;
+	virtual void RestoreState (const TArray<UObject*>& RuntimeObjects, IMovieScenePlayer& Player) override;
 	virtual void Update( float Position, float LastPosition, const TArray<UObject*>& RuntimeObjects, IMovieScenePlayer& Player ) override;
 	virtual void RefreshInstance( const TArray<UObject*>& RuntimeObjects, IMovieScenePlayer& Player ) override {}
 	virtual void ClearInstance( IMovieScenePlayer& Player ) override {}
diff --git a/Engine/Source/Runtime/MovieSceneTracks/Public/TrackInstances/MovieScene3DTransformTrackInstance.h b/Engine/Source/Runtime/MovieSceneTracks/Public/TrackInstances/MovieScene3DTransformTrackInstance.h
--- a/Engine/Source/Runtime/MovieSceneTracks/Public/TrackInstances/MovieScene3DTransformTrackInstance.h
+++ b/Engine/Source/Runtime/MovieSceneTracks/Public/TrackInstances/MovieScene3DTransformTrackInstance.h
@@ -19,8 +19,8 @@ public:
 	FMovieScene3DTransformTrackInstance( UMovieScene3DTransformTrack& InTransformTrack );
 
 	/** IMovieSceneTrackInstance interface */
-	virtual void SaveState (const TArray<UObject*>& RuntimeObjects) override;
-	virtual void RestoreState (const TArray<UObject*>& RuntimeObjects) override;
+	virtual void SaveState (const TArray<UObject*>& RuntimeObjects, IMovieScenePlayer& Player) override;
+	virtual void RestoreState (const TArray<UObject*>& RuntimeObjects, IMovieScenePlayer& Player) override;
 	virtual void Update( float Position, float LastPosition, const TArray<UObject*>& RuntimeObjects, IMovieScenePlayer& Player ) override;
 	virtual void RefreshInstance( const TArray<UObject*>& RuntimeObjects, IMovieScenePlayer& Player ) override {}
 	virtual void ClearInstance( IMovieScenePlayer& Player ) override {}
diff --git a/Engine/Source/Runtime/MovieSceneTracks/Public/TrackInstances/MovieSceneAudioTrackInstance.h b/Engine/Source/Runtime/MovieSceneTracks/Public/TrackInstances/MovieSceneAudioTrackInstance.h
--- a/Engine/Source/Runtime/MovieSceneTracks/Public/TrackInstances/MovieSceneAudioTrackInstance.h
+++ b/Engine/Source/Runtime/MovieSceneTracks/Public/TrackInstances/MovieSceneAudioTrackInstance.h
@@ -20,8 +20,8 @@ public:
 	FMovieSceneAudioTrackInstance( UMovieSceneAudioTrack& InAudioTrack );
 
 	/** IMovieSceneTrackInstance interface */
-	virtual void SaveState(const TArray<UObject*>& RuntimeObjects) override {}
-	virtual void RestoreState(const TArray<UObject*>& RuntimeObjects) override {}
+	virtual void SaveState(const TArray<UObject*>& RuntimeObjects, IMovieScenePlayer& Player) override {}
+	virtual void RestoreState(const TArray<UObject*>& RuntimeObjects, IMovieScenePlayer& Player) override {}
 	virtual void Update( float Position, float LastPosition, const TArray<UObject*>& RuntimeObjects, IMovieScenePlayer& Player ) override;
 	virtual void RefreshInstance( const TArray<UObject*>& RuntimeObjects, IMovieScenePlayer& Player ) override {}
 	virtual void ClearInstance( IMovieScenePlayer& Player ) override {}
diff --git a/Engine/Source/Runtime/MovieSceneTracks/Public/TrackInstances/MovieSceneBoolTrackInstance.h b/Engine/Source/Runtime/MovieSceneTracks/Public/TrackInstances/MovieSceneBoolTrackInstance.h
--- a/Engine/Source/Runtime/MovieSceneTracks/Public/TrackInstances/MovieSceneBoolTrackInstance.h
+++ b/Engine/Source/Runtime/MovieSceneTracks/Public/TrackInstances/MovieSceneBoolTrackInstance.h
@@ -20,8 +20,8 @@ public:
 	FMovieSceneBoolTrackInstance( UMovieSceneBoolTrack& InBoolTrack );
 
 	/** IMovieSceneTrackInstance interface */
-	virtual void SaveState (const TArray<UObject*>& RuntimeObjects) override;
-	virtual void RestoreState (const TArray<UObject*>& RuntimeObjects) override;
+	virtual void SaveState (const TArray<UObject*>& RuntimeObjects, IMovieScenePlayer& Player) override;
+	virtual void RestoreState (const TArray<UObject*>& RuntimeObjects, IMovieScenePlayer& Player) override;
 	virtual void Update( float Position, float LastPosition, const TArray<UObject*>& RuntimeObjects, IMovieScenePlayer& Player ) override;
 	virtual void RefreshInstance( const TArray<UObject*>& RuntimeObjects, IMovieScenePlayer& Player ) override;
 	virtual void ClearInstance( IMovieScenePlayer& Player ) override {}
diff --git a/Engine/Source/Runtime/MovieSceneTracks/Public/TrackInstances/MovieSceneByteTrackInstance.h b/Engine/Source/Runtime/MovieSceneTracks/Public/TrackInstances/MovieSceneByteTrackInstance.h
--- a/Engine/Source/Runtime/MovieSceneTracks/Public/TrackInstances/MovieSceneByteTrackInstance.h
+++ b/Engine/Source/Runtime/MovieSceneTracks/Public/TrackInstances/MovieSceneByteTrackInstance.h
@@ -20,8 +20,8 @@ public:
 	FMovieSceneByteTrackInstance( UMovieSceneByteTrack& InByteTrack );
 
 	/** IMovieSceneTrackInstance interface */
-	virtual void SaveState (const TArray<UObject*>& RuntimeObjects) override;
-	virtual void RestoreState (const TArray<UObject*>& RuntimeObjects) override;
+	virtual void SaveState (const TArray<UObject*>& RuntimeObjects, IMovieScenePlayer& Player) override;
+	virtual void RestoreState (const TArray<UObject*>& RuntimeObjects, IMovieScenePlayer& Player) override;
 	virtual void Update( float Position, float LastPosition, const TArray<UObject*>& RuntimeObjects, IMovieScenePlayer& Player ) override;
 	virtual void RefreshInstance( const TArray<UObject*>& RuntimeObjects, IMovieScenePlayer& Player ) override;
 	virtual void ClearInstance( IMovieScenePlayer& Player ) override {}
diff --git a/Engine/Source/Runtime/MovieSceneTracks/Public/TrackInstances/MovieSceneColorTrackInstance.h b/Engine/Source/Runtime/MovieSceneTracks/Public/TrackInstances/MovieSceneColorTrackInstance.h
--- a/Engine/Source/Runtime/MovieSceneTracks/Public/TrackInstances/MovieSceneColorTrackInstance.h
+++ b/Engine/Source/Runtime/MovieSceneTracks/Public/TrackInstances/MovieSceneColorTrackInstance.h
@@ -19,8 +19,8 @@ public:
 	FMovieSceneColorTrackInstance( UMovieSceneColorTrack& InColorTrack );
 
 	/** IMovieSceneTrackInstance interface */
-	virtual void SaveState (const TArray<UObject*>& RuntimeObjects) override;
-	virtual void RestoreState (const TArray<UObject*>& RuntimeObjects) override;
+	virtual void SaveState (const TArray<UObject*>& RuntimeObjects, IMovieScenePlayer& Player) override;
+	virtual void RestoreState (const TArray<UObject*>& RuntimeObjects, IMovieScenePlayer& Player) override;
 	virtual void Update( float Position, float LastPosition, const TArray<UObject*>& RuntimeObjects, IMovieScenePlayer& Player ) override;
 	virtual void RefreshInstance( const TArray<UObject*>& RuntimeObjects, IMovieScenePlayer& Player ) override;
 	virtual void ClearInstance( IMovieScenePlayer& Player ) override {}
diff --git a/Engine/Source/Runtime/MovieSceneTracks/Public/TrackInstances/MovieSceneEventTrackInstance.h b/Engine/Source/Runtime/MovieSceneTracks/Public/TrackInstances/MovieSceneEventTrackInstance.h
--- a/Engine/Source/Runtime/MovieSceneTracks/Public/TrackInstances/MovieSceneEventTrackInstance.h
+++ b/Engine/Source/Runtime/MovieSceneTracks/Public/TrackInstances/MovieSceneEventTrackInstance.h
@@ -32,8 +32,8 @@ public:
 
 	virtual void ClearInstance(IMovieScenePlayer& Player) override { }
 	virtual void RefreshInstance(const TArray<UObject*>& RuntimeObjects, IMovieScenePlayer& Player) override { }
-	virtual void RestoreState(const TArray<UObject*>& RuntimeObjects) override { }
-	virtual void SaveState(const TArray<UObject*>& RuntimeObjects) override { }
+	virtual void RestoreState(const TArray<UObject*>& RuntimeObjects, IMovieScenePlayer& Player) override { }
+	virtual void SaveState(const TArray<UObject*>& RuntimeObjects, IMovieScenePlayer& Player) override { }
 	virtual void Update(float Position, float LastPosition, const TArray<UObject*>& RuntimeObjects, IMovieScenePlayer& Player) override;
 
 private:
diff --git a/Engine/Source/Runtime/MovieSceneTracks/Public/TrackInstances/MovieSceneFadeTrackInstance.h b/Engine/Source/Runtime/MovieSceneTracks/Public/TrackInstances/MovieSceneFadeTrackInstance.h
--- a/Engine/Source/Runtime/MovieSceneTracks/Public/TrackInstances/MovieSceneFadeTrackInstance.h
+++ b/Engine/Source/Runtime/MovieSceneTracks/Public/TrackInstances/MovieSceneFadeTrackInstance.h
@@ -32,8 +32,8 @@ public:
 
 	virtual void ClearInstance(IMovieScenePlayer& Player) override { }
 	virtual void RefreshInstance(const TArray<UObject*>& RuntimeObjects, IMovieScenePlayer& Player) override { }
-	virtual void RestoreState(const TArray<UObject*>& RuntimeObjects) override { }
-	virtual void SaveState(const TArray<UObject*>& RuntimeObjects) override { }
+	virtual void RestoreState(const TArray<UObject*>& RuntimeObjects, IMovieScenePlayer& Player) override;
+	virtual void SaveState(const TArray<UObject*>& RuntimeObjects, IMovieScenePlayer& Player) override { }
 	virtual void Update(float Position, float LastPosition, const TArray<UObject*>& RuntimeObjects, IMovieScenePlayer& Player) override;
 
 private:
diff --git a/Engine/Source/Runtime/MovieSceneTracks/Public/TrackInstances/MovieSceneFloatTrackInstance.h b/Engine/Source/Runtime/MovieSceneTracks/Public/TrackInstances/MovieSceneFloatTrackInstance.h
--- a/Engine/Source/Runtime/MovieSceneTracks/Public/TrackInstances/MovieSceneFloatTrackInstance.h
+++ b/Engine/Source/Runtime/MovieSceneTracks/Public/TrackInstances/MovieSceneFloatTrackInstance.h
@@ -20,8 +20,8 @@ public:
 	FMovieSceneFloatTrackInstance( UMovieSceneFloatTrack& InFloatTrack );
 
 	/** IMovieSceneTrackInstance interface */
-	virtual void SaveState (const TArray<UObject*>& RuntimeObjects) override;
-	virtual void RestoreState (const TArray<UObject*>& RuntimeObjects) override;
+	virtual void SaveState (const TArray<UObject*>& RuntimeObjects, IMovieScenePlayer& Player) override;
+	virtual void RestoreState (const TArray<UObject*>& RuntimeObjects, IMovieScenePlayer& Player) override;
 	virtual void Update( float Position, float LastPosition, const TArray<UObject*>& RuntimeObjects, IMovieScenePlayer& Player ) override;
 	virtual void RefreshInstance( const TArray<UObject*>& RuntimeObjects, IMovieScenePlayer& Player ) override;
 	virtual void ClearInstance( IMovieScenePlayer& Player ) override {}
diff --git a/Engine/Source/Runtime/MovieSceneTracks/Public/TrackInstances/MovieSceneMaterialTrackInstance.h b/Engine/Source/Runtime/MovieSceneTracks/Public/TrackInstances/MovieSceneMaterialTrackInstance.h
--- a/Engine/Source/Runtime/MovieSceneTracks/Public/TrackInstances/MovieSceneMaterialTrackInstance.h
+++ b/Engine/Source/Runtime/MovieSceneTracks/Public/TrackInstances/MovieSceneMaterialTrackInstance.h
@@ -19,8 +19,8 @@ public:
 	FMovieSceneMaterialTrackInstance( UMovieSceneMaterialTrack& InMaterialTrack );
 
 	/** IMovieSceneTrackInstance interface */
-	virtual void SaveState (const TArray<UObject*>& RuntimeObjects) override;
-	virtual void RestoreState (const TArray<UObject*>& RuntimeObjects) override;
+	virtual void SaveState (const TArray<UObject*>& RuntimeObjects, IMovieScenePlayer& Player) override;
+	virtual void RestoreState (const TArray<UObject*>& RuntimeObjects, IMovieScenePlayer& Player) override;
 	virtual void Update( float Position, float LastPosition, const TArray<UObject*>& RuntimeObjects, IMovieScenePlayer& Player ) override;
 	virtual void RefreshInstance( const TArray<UObject*>& RuntimeObjects, IMovieScenePlayer& Player ) override;
 	virtual void ClearInstance( IMovieScenePlayer& Player ) override;
diff --git a/Engine/Source/Runtime/MovieSceneTracks/Public/TrackInstances/MovieSceneParticleTrackInstance.h b/Engine/Source/Runtime/MovieSceneTracks/Public/TrackInstances/MovieSceneParticleTrackInstance.h
--- a/Engine/Source/Runtime/MovieSceneTracks/Public/TrackInstances/MovieSceneParticleTrackInstance.h
+++ b/Engine/Source/Runtime/MovieSceneTracks/Public/TrackInstances/MovieSceneParticleTrackInstance.h
@@ -23,8 +23,8 @@ public:
 	virtual ~FMovieSceneParticleTrackInstance();
 
 	/** IMovieSceneTrackInstance interface */
-	virtual void SaveState(const TArray<UObject*>& RuntimeObjects) override {}
-	virtual void RestoreState(const TArray<UObject*>& RuntimeObjects) override {}
+	virtual void SaveState(const TArray<UObject*>& RuntimeObjects, IMovieScenePlayer& Player) override {}
+	virtual void RestoreState(const TArray<UObject*>& RuntimeObjects, IMovieScenePlayer& Player) override {}
 	virtual void Update( float Position, float LastPosition, const TArray<UObject*>& RuntimeObjects, IMovieScenePlayer& Player ) override;
 	virtual void RefreshInstance( const TArray<UObject*>& RuntimeObjects, IMovieScenePlayer& Player ) override {}
 	virtual void ClearInstance( IMovieScenePlayer& Player ) override {}
diff --git a/Engine/Source/Runtime/MovieSceneTracks/Public/TrackInstances/MovieSceneSkeletalAnimationTrackInstance.h b/Engine/Source/Runtime/MovieSceneTracks/Public/TrackInstances/MovieSceneSkeletalAnimationTrackInstance.h
--- a/Engine/Source/Runtime/MovieSceneTracks/Public/TrackInstances/MovieSceneSkeletalAnimationTrackInstance.h
+++ b/Engine/Source/Runtime/MovieSceneTracks/Public/TrackInstances/MovieSceneSkeletalAnimationTrackInstance.h
@@ -20,8 +20,8 @@ public:
 	virtual ~FMovieSceneSkeletalAnimationTrackInstance();
 
 	/** IMovieSceneTrackInstance interface */
-	virtual void SaveState(const TArray<UObject*>& RuntimeObjects) override {}
-	virtual void RestoreState(const TArray<UObject*>& RuntimeObjects) override {}
+	virtual void SaveState(const TArray<UObject*>& RuntimeObjects, IMovieScenePlayer& Player) override {}
+	virtual void RestoreState(const TArray<UObject*>& RuntimeObjects, IMovieScenePlayer& Player) override {}
 	virtual void Update( float Position, float LastPosition, const TArray<UObject*>& RuntimeObjects, IMovieScenePlayer& Player ) override;
 	virtual void RefreshInstance( const TArray<UObject*>& RuntimeObjects, IMovieScenePlayer& Player ) override {}
 	virtual void ClearInstance( IMovieScenePlayer& Player ) override {}
diff --git a/Engine/Source/Runtime/MovieSceneTracks/Public/TrackInstances/MovieSceneSlomoTrackInstance.h b/Engine/Source/Runtime/MovieSceneTracks/Public/TrackInstances/MovieSceneSlomoTrackInstance.h
--- a/Engine/Source/Runtime/MovieSceneTracks/Public/TrackInstances/MovieSceneSlomoTrackInstance.h
+++ b/Engine/Source/Runtime/MovieSceneTracks/Public/TrackInstances/MovieSceneSlomoTrackInstance.h
@@ -32,8 +32,8 @@ public:
 
 	virtual void ClearInstance(IMovieScenePlayer& Player) override { }
 	virtual void RefreshInstance(const TArray<UObject*>& RuntimeObjects, IMovieScenePlayer& Player) override { }
-	virtual void RestoreState(const TArray<UObject*>& RuntimeObjects) override { }
-	virtual void SaveState(const TArray<UObject*>& RuntimeObjects) override { }
+	virtual void RestoreState(const TArray<UObject*>& RuntimeObjects, IMovieScenePlayer& Player) override { }
+	virtual void SaveState(const TArray<UObject*>& RuntimeObjects, IMovieScenePlayer& Player) override { }
 	virtual void Update(float Position, float LastPosition, const TArray<UObject*>& RuntimeObjects, IMovieScenePlayer& Player) override;
 
 protected:
diff --git a/Engine/Source/Runtime/MovieSceneTracks/Public/TrackInstances/MovieSceneVectorTrackInstance.h b/Engine/Source/Runtime/MovieSceneTracks/Public/TrackInstances/MovieSceneVectorTrackInstance.h
--- a/Engine/Source/Runtime/MovieSceneTracks/Public/TrackInstances/MovieSceneVectorTrackInstance.h
+++ b/Engine/Source/Runtime/MovieSceneTracks/Public/TrackInstances/MovieSceneVectorTrackInstance.h
@@ -20,8 +20,8 @@ public:
 	FMovieSceneVectorTrackInstance( UMovieSceneVectorTrack& InVectorTrack );
 
 	/** IMovieSceneTrackInstance interface */
-	virtual void SaveState (const TArray<UObject*>& RuntimeObjects) override;
-	virtual void RestoreState (const TArray<UObject*>& RuntimeObjects) override;
+	virtual void SaveState (const TArray<UObject*>& RuntimeObjects, IMovieScenePlayer& Player) override;
+	virtual void RestoreState (const TArray<UObject*>& RuntimeObjects, IMovieScenePlayer& Player) override;
 	virtual void Update( float Position, float LastPosition, const TArray<UObject*>& RuntimeObjects, IMovieScenePlayer& Player ) override;
 	virtual void RefreshInstance( const TArray<UObject*>& RuntimeObjects, IMovieScenePlayer& Player ) override;
 	virtual void ClearInstance( IMovieScenePlayer& Player ) override {}
diff --git a/Engine/Source/Runtime/MovieSceneTracks/Public/TrackInstances/MovieSceneVisibilityTrackInstance.h b/Engine/Source/Runtime/MovieSceneTracks/Public/TrackInstances/MovieSceneVisibilityTrackInstance.h
--- a/Engine/Source/Runtime/MovieSceneTracks/Public/TrackInstances/MovieSceneVisibilityTrackInstance.h
+++ b/Engine/Source/Runtime/MovieSceneTracks/Public/TrackInstances/MovieSceneVisibilityTrackInstance.h
@@ -19,8 +19,8 @@ public:
 	FMovieSceneVisibilityTrackInstance( UMovieSceneVisibilityTrack& InVisibilityTrack );
 
 	/** IMovieSceneTrackInstance interface */
-	virtual void SaveState (const TArray<UObject*>& RuntimeObjects) override;
-	virtual void RestoreState (const TArray<UObject*>& RuntimeObjects) override;
+	virtual void SaveState (const TArray<UObject*>& RuntimeObjects, IMovieScenePlayer& Player) override;
+	virtual void RestoreState (const TArray<UObject*>& RuntimeObjects, IMovieScenePlayer& Player) override;
 	virtual void Update( float Position, float LastPosition, const TArray<UObject*>& RuntimeObjects, class IMovieScenePlayer& Player ) override;
 	virtual void RefreshInstance( const TArray<UObject*>& RuntimeObjects, class IMovieScenePlayer& Player ) override;
 	virtual void ClearInstance( IMovieScenePlayer& Player ) override {}
diff --git a/Engine/Source/Runtime/MovieSceneTracks/Public/TrackInstances/SubMovieSceneTrackInstance.h b/Engine/Source/Runtime/MovieSceneTracks/Public/TrackInstances/SubMovieSceneTrackInstance.h
--- a/Engine/Source/Runtime/MovieSceneTracks/Public/TrackInstances/SubMovieSceneTrackInstance.h
+++ b/Engine/Source/Runtime/MovieSceneTracks/Public/TrackInstances/SubMovieSceneTrackInstance.h
@@ -23,8 +23,8 @@ public:
 	/** IMovieSceneTrackInstance interface */
 	virtual void Update( float Position, float LastPosition, const TArray<UObject*>& RuntimeObjects, IMovieScenePlayer& Player ) override;
 	virtual void RefreshInstance( const TArray<UObject*>& RuntimeObjects, IMovieScenePlayer& Player ) override;
-	virtual void SaveState (const TArray<UObject*>& RuntimeObjects) override;
-	virtual void RestoreState (const TArray<UObject*>& RuntimeObjects) override;
+	virtual void SaveState (const TArray<UObject*>& RuntimeObjects, IMovieScenePlayer& Player) override;
+	virtual void RestoreState (const TArray<UObject*>& RuntimeObjects, IMovieScenePlayer& Player) override;
 	virtual void ClearInstance( IMovieScenePlayer& Player ) override;
 protected:
 
diff --git a/Engine/Source/Runtime/UMG/Private/Animation/MovieScene2DTransformTrackInstance.h b/Engine/Source/Runtime/UMG/Private/Animation/MovieScene2DTransformTrackInstance.h
--- a/Engine/Source/Runtime/UMG/Private/Animation/MovieScene2DTransformTrackInstance.h
+++ b/Engine/Source/Runtime/UMG/Private/Animation/MovieScene2DTransformTrackInstance.h
@@ -13,8 +13,8 @@ public:
 	FMovieScene2DTransformTrackInstance( class UMovieScene2DTransformTrack& InTransformTrack );
 
 	/** IMovieSceneTrackInstance interface */
-	virtual void SaveState(const TArray<UObject*>& RuntimeObjects) override {}
-	virtual void RestoreState(const TArray<UObject*>& RuntimeObjects) override {}
+	virtual void SaveState(const TArray<UObject*>& RuntimeObjects, IMovieScenePlayer& Player) override {}
+	virtual void RestoreState(const TArray<UObject*>& RuntimeObjects, IMovieScenePlayer& Player) override {}
 	virtual void Update( float Position, float LastPosition, const TArray<UObject*>& RuntimeObjects, class IMovieScenePlayer& Player ) override;
 	virtual void RefreshInstance( const TArray<UObject*>& RuntimeObjects, class IMovieScenePlayer& Player ) override;
 	virtual void ClearInstance( IMovieScenePlayer& Player ) override {};
diff --git a/Engine/Source/Runtime/UMG/Private/Animation/MovieSceneMarginTrackInstance.h b/Engine/Source/Runtime/UMG/Private/Animation/MovieSceneMarginTrackInstance.h
--- a/Engine/Source/Runtime/UMG/Private/Animation/MovieSceneMarginTrackInstance.h
+++ b/Engine/Source/Runtime/UMG/Private/Animation/MovieSceneMarginTrackInstance.h
@@ -13,8 +13,8 @@ public:
 	FMovieSceneMarginTrackInstance( class UMovieSceneMarginTrack& InMarginTrack );
 
 	/** IMovieSceneTrackInstance interface */
-	virtual void SaveState( const TArray<UObject*>& RuntimeObjects ) override {}
-	virtual void RestoreState( const TArray<UObject*>& RuntimeObjects ) override {}
+	virtual void SaveState( const TArray<UObject*>& RuntimeObjects, IMovieScenePlayer& Player ) override {}
+	virtual void RestoreState( const TArray<UObject*>& RuntimeObjects, IMovieScenePlayer& Player ) override {}
 	virtual void Update( float Position, float LastPosition, const TArray<UObject*>& RuntimeObjects, class IMovieScenePlayer& Player ) override;
 	virtual void RefreshInstance( const TArray<UObject*>& RuntimeObjects, class IMovieScenePlayer& Player ) override;
 	virtual void ClearInstance( IMovieScenePlayer& Player ) override {};
