Copying //UE4/Dev-Sequencer to //UE4/Main

diff --git a/Engine/Source/Runtime/Engine/Classes/Matinee/InterpGroupInst.h b/Engine/Source/Runtime/Engine/Classes/Matinee/InterpGroupInst.h
--- a/Engine/Source/Runtime/Engine/Classes/Matinee/InterpGroupInst.h
+++ b/Engine/Source/Runtime/Engine/Classes/Matinee/InterpGroupInst.h
@@ -33,11 +33,14 @@ class UInterpGroupInst : public UObject
 
 
 	/** 
-	 *	Return the AActor  that this GroupInstance is working on. 
+	 *	Return the AActor that this GroupInstance is working on.
 	 *	Should use this instead of just referencing GroupActor, as it check IsPendingKill() for you.
 	 */
 	virtual AActor* GetGroupActor() const;
 
+	/** Sets the GroupActor tha this GroupInstance should work on. */
+	void SetGroupActor(AActor* Actor);
+
 	/** Called before Interp editing to save original state of Actor. @see UInterpTrackInst::SaveActorState */
 	virtual void SaveGroupActorState();
 
diff --git a/Engine/Source/Runtime/Engine/Private/Animation/AnimCurveTypes.cpp b/Engine/Source/Runtime/Engine/Private/Animation/AnimCurveTypes.cpp
--- a/Engine/Source/Runtime/Engine/Private/Animation/AnimCurveTypes.cpp
+++ b/Engine/Source/Runtime/Engine/Private/Animation/AnimCurveTypes.cpp
@@ -250,6 +250,34 @@ void FRawCurveTracks::DeleteAllCurveData(ESupportedCurveType SupportedCurveType
 	}
 }
 
+#if WITH_EDITOR
+void FRawCurveTracks::AddFloatCurveKey(const USkeleton::AnimCurveUID Uid, int32 CurveFlags, float Time, float Value)
+{
+	FFloatCurve* FloatCurve = GetCurveDataImpl<FFloatCurve>(FloatCurves, Uid);
+	if (FloatCurve == nullptr)
+	{
+		AddCurveData(Uid, CurveFlags, FloatType);
+		FloatCurve = GetCurveDataImpl<FFloatCurve>(FloatCurves, Uid);
+	}
+
+	if (FloatCurve->GetCurveTypeFlags() != CurveFlags)
+	{
+		FloatCurve->SetCurveTypeFlags(FloatCurve->GetCurveTypeFlags() | CurveFlags);
+	}
+
+	FloatCurve->UpdateOrAddKey(Value, Time);
+}
+
+void FRawCurveTracks::RemoveRedundantKeys()
+{
+	for (auto CurveIter = FloatCurves.CreateIterator(); CurveIter; ++CurveIter)
+	{
+		FFloatCurve& Curve = *CurveIter;
+		Curve.FloatCurve.RemoveRedundantKeys(SMALL_NUMBER);
+	}
+}
+#endif
+
 bool FRawCurveTracks::AddCurveData(USkeleton::AnimCurveUID Uid, int32 CurveFlags /*= ACF_DefaultCurve*/, ESupportedCurveType SupportedCurveType /*= FloatType*/)
 {
 	switch(SupportedCurveType)
diff --git a/Engine/Source/Runtime/Engine/Private/Animation/AnimInstance.cpp b/Engine/Source/Runtime/Engine/Private/Animation/AnimInstance.cpp
--- a/Engine/Source/Runtime/Engine/Private/Animation/AnimInstance.cpp
+++ b/Engine/Source/Runtime/Engine/Private/Animation/AnimInstance.cpp
@@ -1805,6 +1805,30 @@ void UAnimInstance::Montage_Pause(UAnimMontage* Montage)
 	}
 }
 
+void UAnimInstance::Montage_Resume(UAnimMontage* Montage)
+{
+	if (Montage)
+	{
+		FAnimMontageInstance* MontageInstance = GetActiveInstanceForMontage(*Montage);
+		if (MontageInstance && !MontageInstance->IsPlaying())
+		{
+			MontageInstance->SetPlaying(true);
+		}
+	}
+	else
+	{
+		// If no Montage reference, do it on all active ones.
+		for (int32 InstanceIndex = 0; InstanceIndex < MontageInstances.Num(); InstanceIndex++)
+		{
+			FAnimMontageInstance* MontageInstance = MontageInstances[InstanceIndex];
+			if (MontageInstance && MontageInstance->IsActive() && !MontageInstance->IsPlaying())
+			{
+				MontageInstance->SetPlaying(true);
+			}
+		}
+	}
+}
+
 void UAnimInstance::Montage_JumpToSection(FName SectionName, UAnimMontage* Montage)
 {
 	if (Montage)
diff --git a/Engine/Source/Runtime/Engine/Private/Camera/CameraAnimInst.cpp b/Engine/Source/Runtime/Engine/Private/Camera/CameraAnimInst.cpp
--- a/Engine/Source/Runtime/Engine/Private/Camera/CameraAnimInst.cpp
+++ b/Engine/Source/Runtime/Engine/Private/Camera/CameraAnimInst.cpp
@@ -2,6 +2,7 @@
 
 
 #include "EnginePrivate.h"
+#include "Camera/CameraComponent.h"
 #include "Camera/CameraActor.h"
 #include "Camera/CameraAnim.h"
 #include "Camera/CameraAnimInst.h"
@@ -18,6 +19,7 @@ UCameraAnimInst::UCameraAnimInst(const FObjectInitializer& ObjectInitializer)
 	: Super(ObjectInitializer)
 {
 	bFinished = true;
+	bStopAutomatically = true;
 	PlayRate = 1.0f;
 	TransientScaleModifier = 1.0f;
 	PlaySpace = ECameraAnimPlaySpace::CameraLocal;
@@ -74,7 +76,7 @@ void UCameraAnimInst::AdvanceAnim(float DeltaTime, bool bJump)
 
 		if (bBlendingIn)
 		{
-			if (CurBlendInTime > BlendInTime)
+			if ((CurBlendInTime > BlendInTime) || (BlendInTime == 0.f))
 			{
 				// done blending in!
 				bBlendingIn = false;
@@ -95,6 +97,9 @@ void UCameraAnimInst::AdvanceAnim(float DeltaTime, bool bJump)
 			float const BlendInWeight = (bBlendingIn) ? (CurBlendInTime / BlendInTime) : 1.f;
 			float const BlendOutWeight = (bBlendingOut) ? (1.f - CurBlendOutTime / BlendOutTime) : 1.f;
 			CurrentBlendWeight = FMath::Min(BlendInWeight, BlendOutWeight) * BasePlayScale * TransientScaleModifier;
+			
+			// this is intended to be applied only once
+			TransientScaleModifier = 1.f;
 		}
 
 		// this will update tracks and apply the effects to the group actor (except move tracks)
@@ -103,24 +108,34 @@ void UCameraAnimInst::AdvanceAnim(float DeltaTime, bool bJump)
 			InterpGroupInst->Group->UpdateGroup(CurTime, InterpGroupInst, false, bJump);
 		}
 
-		if (bAnimJustFinished)
+		if (bStopAutomatically)
 		{
-			// completely finished
-			Stop(true);
-		}
-		else if (RemainingTime > 0.f)
-		{
-			// handle any specified duration
-			RemainingTime -= DeltaTime;
-			if (RemainingTime <= 0.f)
+			if (bAnimJustFinished)
+			{
+				// completely finished
+				Stop(true);
+			}
+			else if (RemainingTime > 0.f)
 			{
-				// stop with blend out
-				Stop();
+				// handle any specified duration
+				RemainingTime -= DeltaTime;
+				if (RemainingTime <= 0.f)
+				{
+					// stop with blend out
+					Stop();
+				}
 			}
 		}
 	}
 }
 
+void UCameraAnimInst::SetCurrentTime(float NewTime)
+{
+	float const TimeDelta = NewTime - (CurTime / PlayRate);
+	AdvanceAnim(TimeDelta, true);
+}
+
+
 void UCameraAnimInst::Update(float NewRate, float NewScale, float NewBlendInTime, float NewBlendOutTime, float NewDuration)
 {
 	if (bFinished == false)
@@ -172,6 +187,14 @@ void UCameraAnimInst::SetScale(float NewScale)
 	BasePlayScale = NewScale;
 }
 
+void UCameraAnimInst::SetCameraActor(class AActor* Actor)
+{
+	if (InterpGroupInst)
+	{
+		InterpGroupInst->SetGroupActor(Actor);
+	}
+}
+
 static const FName NAME_CameraComponentFieldOfViewPropertyName(TEXT("CameraComponent.FieldOfView"));
 
 void UCameraAnimInst::Play(UCameraAnim* Anim, class AActor* CamActor, float InRate, float InScale, float InBlendInTime, float InBlendOutTime, bool bInLooping, bool bRandomStartTime, float Duration)
@@ -200,16 +223,15 @@ void UCameraAnimInst::Play(UCameraAnim* Anim, class AActor* CamActor, float InRa
 		RemainingTime = (Duration > 0.f) ? (Duration - BlendOutTime) : 0.f;
 
 		// init the interpgroup
-		if ( CamActor->IsA(ACameraActor::StaticClass()) )
+		if ( CamActor && CamActor->IsA(ACameraActor::StaticClass()) )
 		{
+			// #fixme jf: I don't think this is necessary anymore
 			// ensure CameraActor is zeroed, so RelativeToInitial anims get proper InitialTM
 			CamActor->SetActorLocation(FVector::ZeroVector, false);
 			CamActor->SetActorRotation(FRotator::ZeroRotator);
 		}
 		InterpGroupInst->InitGroupInst(CamAnim->CameraInterpGroup, CamActor);
 
-		checkf(CamAnim->CameraInterpGroup->InterpTracks.Num() == InterpGroupInst->TrackInst.Num(), TEXT("Track count mismatch! Outer = %s"), *GetNameSafe(CamAnim));
-
 		// cache move track refs
 		for (int32 Idx = 0; Idx < InterpGroupInst->TrackInst.Num(); ++Idx)
 		{
@@ -286,4 +308,78 @@ void UCameraAnimInst::SetPlaySpace(ECameraAnimPlaySpace::Type NewSpace, FRotator
 
 
 /** Returns InterpGroupInst subobject **/
-UInterpGroupInst* UCameraAnimInst::GetInterpGroupInst() const { return InterpGroupInst; }
+UInterpGroupInst* UCameraAnimInst::GetInterpGroupInst() const
+{ 
+	return InterpGroupInst;
+}
+
+void UCameraAnimInst::ApplyToView(FMinimalViewInfo& InOutPOV) const
+{
+	if (CurrentBlendWeight > 0.f)
+	{
+		ACameraActor const* AnimatedCamActor = dynamic_cast<ACameraActor*>(InterpGroupInst->GetGroupActor());
+		if (AnimatedCamActor)
+		{
+
+			if (CamAnim->bRelativeToInitialTransform)
+			{
+				// move animated cam actor to initial-relative position
+				FTransform const AnimatedCamToWorld = AnimatedCamActor->GetTransform();
+				FTransform const AnimatedCamToInitialCam = AnimatedCamToWorld * InitialCamToWorld.Inverse();
+				ACameraActor* const MutableCamActor = const_cast<ACameraActor*>(AnimatedCamActor);
+				MutableCamActor->SetActorTransform(AnimatedCamToInitialCam);
+			}
+
+			float const Scale = CurrentBlendWeight;
+			FRotationMatrix const CameraToWorld(InOutPOV.Rotation);
+
+			if (PlaySpace == ECameraAnimPlaySpace::CameraLocal)
+			{
+				// the code in the else block will handle this just fine, but this path provides efficiency and simplicity for the most common case
+
+				// loc
+				FVector const LocalOffset = CameraToWorld.TransformVector(AnimatedCamActor->GetActorLocation()*Scale);
+				InOutPOV.Location += LocalOffset;
+
+				// rot
+				FRotationMatrix const AnimRotMat(AnimatedCamActor->GetActorRotation()*Scale);
+				InOutPOV.Rotation = (AnimRotMat * CameraToWorld).Rotator();
+			}
+			else
+			{
+				// handle playing the anim in an arbitrary space relative to the camera
+
+				// find desired space
+				FMatrix const PlaySpaceToWorld = (PlaySpace == ECameraAnimPlaySpace::UserDefined) ? UserPlaySpaceMatrix : FMatrix::Identity;
+
+				// loc
+				FVector const LocalOffset = PlaySpaceToWorld.TransformVector(AnimatedCamActor->GetActorLocation()*Scale);
+				InOutPOV.Location += LocalOffset;
+
+				// rot
+				// find transform from camera to the "play space"
+				FMatrix const CameraToPlaySpace = CameraToWorld * PlaySpaceToWorld.Inverse();	// CameraToWorld * WorldToPlaySpace
+
+				// find transform from anim (applied in playspace) back to camera
+				FRotationMatrix const AnimToPlaySpace(AnimatedCamActor->GetActorRotation()*Scale);
+				FMatrix const AnimToCamera = AnimToPlaySpace * CameraToPlaySpace.Inverse();			// AnimToPlaySpace * PlaySpaceToCamera
+
+				// RCS = rotated camera space, meaning camera space after it's been animated
+				// this is what we're looking for, the diff between rotated cam space and regular cam space.
+				// apply the transform back to camera space from the post-animated transform to get the RCS
+				FMatrix const RCSToCamera = CameraToPlaySpace * AnimToCamera;
+
+				// now apply to real camera
+				FRotationMatrix const RealCamToWorld(InOutPOV.Rotation);
+				InOutPOV.Rotation = (RCSToCamera * RealCamToWorld).Rotator();
+			}
+
+			// fov
+			const float FOVMin = 5.f;
+			const float FOVMax = 170.f;
+			InOutPOV.FOV += (AnimatedCamActor->GetCameraComponent()->FieldOfView - InitialFOV) * Scale;
+			InOutPOV.FOV = FMath::Clamp<float>(InOutPOV.FOV, FOVMin, FOVMax);
+		}
+	}
+}
+
diff --git a/Engine/Source/Runtime/Engine/Private/Camera/CameraComponent.cpp b/Engine/Source/Runtime/Engine/Private/Camera/CameraComponent.cpp
--- a/Engine/Source/Runtime/Engine/Private/Camera/CameraComponent.cpp
+++ b/Engine/Source/Runtime/Engine/Private/Camera/CameraComponent.cpp
@@ -139,6 +139,16 @@ void UCameraComponent::PostLoad()
 		 }
 	 }
  }
+
+void UCameraComponent::ResetProxyMeshTransform()
+{
+	if (ProxyMeshComponent != nullptr)
+	{
+		ProxyMeshComponent->ResetRelativeTransform();
+	}
+}
+
+
 void UCameraComponent::RefreshVisualRepresentation()
 {
 	if (DrawFrustum != nullptr)
@@ -159,10 +169,8 @@ void UCameraComponent::RefreshVisualRepresentation()
 		DrawFrustum->FrustumAspectRatio = AspectRatio;
 		DrawFrustum->MarkRenderStateDirty();
 	}
-	if (ProxyMeshComponent != nullptr)
-	{
-		ProxyMeshComponent->ResetRelativeTransform();
-	}
+
+	ResetProxyMeshTransform();
 }
 
 void UCameraComponent::OverrideFrustumColor(FColor OverrideColor)
@@ -235,10 +243,34 @@ void UCameraComponent::GetCameraView(float DeltaTime, FMinimalViewInfo& DesiredV
 		}
 	}
 
-	DesiredView.Location = GetComponentLocation();
-	DesiredView.Rotation = GetComponentRotation();
+	if (bUseAdditiveOffset)
+	{
+		FTransform OffsetCamToBaseCam = AdditiveOffset;
+		FTransform BaseCamToWorld = GetComponentToWorld();
+		FTransform OffsetCamToWorld = OffsetCamToBaseCam * BaseCamToWorld;
+
+		DesiredView.Location = OffsetCamToWorld.GetLocation();
+		DesiredView.Rotation = OffsetCamToWorld.Rotator();
+
+#if WITH_EDITORONLY_DATA
+		if (ProxyMeshComponent)
+		{
+			ResetProxyMeshTransform();
+
+			FTransform LocalTransform = ProxyMeshComponent->GetRelativeTransform();
+			FTransform WorldTransform = LocalTransform * OffsetCamToWorld;
+
+			ProxyMeshComponent->SetWorldTransform(WorldTransform);
+		}
+#endif
+	}
+	else
+	{
+		DesiredView.Location = GetComponentLocation();
+		DesiredView.Rotation = GetComponentRotation();
+	}
 
-	DesiredView.FOV = FieldOfView;
+	DesiredView.FOV = bUseAdditiveOffset ? (FieldOfView + AdditiveFOVOffset) : FieldOfView;
 	DesiredView.AspectRatio = AspectRatio;
 	DesiredView.bConstrainAspectRatio = bConstrainAspectRatio;
 	DesiredView.bUseFieldOfViewForLOD = bUseFieldOfViewForLOD;
@@ -255,10 +287,7 @@ void UCameraComponent::GetCameraView(float DeltaTime, FMinimalViewInfo& DesiredV
 	}
 
 #if WITH_EDITOR
-	if (ProxyMeshComponent != nullptr)
-	{
-		ProxyMeshComponent->ResetRelativeTransform();
-	}
+	ResetProxyMeshTransform();
 #endif //WITH_EDITOR
 }
 
@@ -285,6 +314,18 @@ void SetDeprecatedControllerViewRotation(UCameraComponent& Component, bool bValu
 	PRAGMA_ENABLE_DEPRECATION_WARNINGS
 }
 
+void UCameraComponent::NotifyCameraCut()
+{
+	// if we are owned by a camera actor, notify it too
+	// note: many camera components are not part of camera actors, so notification should begin at the
+	// component level.
+	ACameraActor* const OwningCamera = Cast<ACameraActor>(GetOwner());
+	if (OwningCamera)
+	{
+		OwningCamera->NotifyCameraCut();
+	}
+};
+
 
 #undef LOCTEXT_NAMESPACE
 
diff --git a/Engine/Source/Runtime/Engine/Private/Camera/CameraShake.cpp b/Engine/Source/Runtime/Engine/Private/Camera/CameraShake.cpp
--- a/Engine/Source/Runtime/Engine/Private/Camera/CameraShake.cpp
+++ b/Engine/Source/Runtime/Engine/Private/Camera/CameraShake.cpp
@@ -26,6 +26,12 @@ float FFOscillator::GetInitialOffset(FFOscillator const& Osc)
 		: 0.f;
 }
 
+// static
+float FFOscillator::GetOffsetAtTime(FFOscillator const& Osc, float InitialOffset, float Time)
+{
+	return InitialOffset + (Time * Osc.Frequency);
+}
+
 //////////////////////////////////////////////////////////////////////////
 // UCameraShake
 
@@ -45,9 +51,16 @@ void UCameraShake::StopShake(bool bImmediately)
 	if (bImmediately)
 	{
 		// stop cam anim if playing
-		if ((AnimInst != nullptr) && !AnimInst->bFinished)
+		if (AnimInst && !AnimInst->bFinished)
 		{
-			CameraOwner->StopCameraAnimInst(AnimInst, true);
+			if (CameraOwner)
+			{
+				CameraOwner->StopCameraAnimInst(AnimInst, true);
+			}
+			else
+			{
+				AnimInst->Stop(true);
+			}
 		}
 
 		AnimInst = nullptr;
@@ -60,9 +73,17 @@ void UCameraShake::StopShake(bool bImmediately)
 		// advance to the blend out time
 		OscillatorTimeRemaining = FMath::Min(OscillatorTimeRemaining, OscillationBlendOutTime);
 
-		if ((AnimInst != nullptr) && !AnimInst->bFinished)
+		if (AnimInst && !AnimInst->bFinished)
 		{
-			CameraOwner->StopCameraAnimInst(AnimInst, false);
+			if (CameraOwner)
+			{
+				CameraOwner->StopCameraAnimInst(AnimInst, false);
+			}
+			else
+			{
+				// playing without a cameramanager, stop it ourselves
+				AnimInst->Stop(false);
+			}
 		}
 	}
 
@@ -113,6 +134,10 @@ void UCameraShake::PlayShake(APlayerCameraManager* Camera, float Scale, ECameraA
 
 			FOVSinOffset = FFOscillator::GetInitialOffset(FOVOscillation);
 
+			InitialLocSinOffset = LocSinOffset;
+			InitialRotSinOffset = RotSinOffset;
+			InitialFOVSinOffset = FOVSinOffset;
+
 			OscillatorTimeRemaining = OscillationDuration;
 
 			if (OscillationBlendInTime > 0.f)
@@ -147,7 +172,22 @@ void UCameraShake::PlayShake(APlayerCameraManager* Camera, float Scale, ECameraA
 			float const FinalAnimScale = Scale * AnimScale;
 			if (FinalAnimScale > 0.f)
 			{
-				AnimInst = CameraOwner->PlayCameraAnim(Anim, AnimPlayRate, FinalAnimScale, AnimBlendInTime, AnimBlendOutTime, bLoop, bRandomStart, Duration, InPlaySpace, UserPlaySpaceRot);
+				if (CameraOwner)
+				{
+					AnimInst = CameraOwner->PlayCameraAnim(Anim, AnimPlayRate, FinalAnimScale, AnimBlendInTime, AnimBlendOutTime, bLoop, bRandomStart, Duration, InPlaySpace, UserPlaySpaceRot);
+				}
+				else
+				{
+					// allocate our own instance and start it
+					AnimInst = NewObject<UCameraAnimInst>(this);
+					if (AnimInst)
+					{
+						// note: we don't have a temp camera actor necessary for evaluating a camera anim.
+						// caller is responsible in this case for providing one by calling SetCameraACtor(??) on the CamAnimInst
+						AnimInst->Play(Anim, nullptr, AnimPlayRate, FinalAnimScale, AnimBlendInTime, AnimBlendOutTime, bLoop, bRandomStart, Duration);
+						AnimInst->SetPlaySpace(InPlaySpace, UserPlaySpaceRot);
+					}
+				}
 			}
 		}
 	}
@@ -315,3 +355,35 @@ bool UCameraShake::IsLooping() const
 	return OscillationDuration < 0.0f;
 }
 
+void UCameraShake::SetCurrentTimeAndApplyShake(float NewTime, FMinimalViewInfo& POV)
+{
+	// reset to start and advance to desired point
+	LocSinOffset = InitialLocSinOffset;
+	RotSinOffset = InitialRotSinOffset;
+	FOVSinOffset = InitialFOVSinOffset;
+
+	OscillatorTimeRemaining = OscillationDuration;
+
+	if (OscillationBlendInTime > 0.f)
+	{
+		bBlendingIn = true;
+		CurrentBlendInTime = 0.f;
+	}
+
+	if (OscillationDuration > 0.f)
+	{
+		if ((OscillationBlendOutTime > 0.f) && (NewTime > OscillationBlendOutTime))
+		{
+			bBlendingOut = true;
+			CurrentBlendOutTime = OscillationBlendOutTime - (OscillationDuration - NewTime);
+		}
+	}
+
+	UpdateAndApplyCameraShake(NewTime, 1.f, POV);
+
+	if (AnimInst)
+	{
+		AnimInst->SetCurrentTime(NewTime);
+		AnimInst->ApplyToView(POV);
+	}
+}
diff --git a/Engine/Source/Runtime/Engine/Private/Components/SkeletalMeshComponent.cpp b/Engine/Source/Runtime/Engine/Private/Components/SkeletalMeshComponent.cpp
--- a/Engine/Source/Runtime/Engine/Private/Components/SkeletalMeshComponent.cpp
+++ b/Engine/Source/Runtime/Engine/Private/Components/SkeletalMeshComponent.cpp
@@ -1354,6 +1354,9 @@ void USkeletalMeshComponent::PostAnimEvaluation(FAnimationEvaluationContext& Eva
 
 	if(AnimScriptInstance)
 	{
+#if WITH_EDITOR
+		GetEditableAnimationCurves() = EvaluationContext.Curve;
+#endif 
 		// curve update happens first
 		AnimScriptInstance->UpdateCurves(EvaluationContext.Curve);
 	}
diff --git a/Engine/Source/Runtime/Engine/Private/CurveBase.cpp b/Engine/Source/Runtime/Engine/Private/CurveBase.cpp
--- a/Engine/Source/Runtime/Engine/Private/CurveBase.cpp
+++ b/Engine/Source/Runtime/Engine/Private/CurveBase.cpp
@@ -524,6 +524,29 @@ FKeyHandle::FKeyHandle()
 }
 
 
+/* FCurveOwnerInterface
+ *****************************************************************************/
+
+FLinearColor FCurveOwnerInterface::GetCurveColor(FRichCurveEditInfo CurveInfo) const
+{
+	FString CurveName = CurveInfo.CurveName.ToString();
+	if (CurveName == TEXT("X") || CurveName == TEXT("R"))
+	{
+		return FLinearColor(1.0f, 0.0f, 0.0f);
+	}
+	else if (CurveName == TEXT("Y") || CurveName == TEXT("G"))
+	{
+		return FLinearColor(0.0f, 1.0f, 0.0f);
+	}
+	else if (CurveName == TEXT("Z") || CurveName == TEXT("B"))
+	{
+		return FLinearColor(0.05f, 0.05f, 1.0f);
+	}
+
+	return FLinearColor::Gray;
+}
+
+
 /* FRichCurve
  *****************************************************************************/
 
diff --git a/Engine/Source/Runtime/Engine/Private/Interpolation.cpp b/Engine/Source/Runtime/Engine/Private/Interpolation.cpp
--- a/Engine/Source/Runtime/Engine/Private/Interpolation.cpp
+++ b/Engine/Source/Runtime/Engine/Private/Interpolation.cpp
@@ -2570,6 +2570,11 @@ AActor* UInterpGroupInst::GetGroupActor() const
 	}
 }
 
+void UInterpGroupInst::SetGroupActor(AActor* Actor)
+{
+	GroupActor = Actor;
+}
+
 void UInterpGroupInst::SaveGroupActorState()
 {
 	check(Group);
diff --git a/Engine/Source/Runtime/Engine/Private/KismetMathLibrary.cpp b/Engine/Source/Runtime/Engine/Private/KismetMathLibrary.cpp
--- a/Engine/Source/Runtime/Engine/Private/KismetMathLibrary.cpp
+++ b/Engine/Source/Runtime/Engine/Private/KismetMathLibrary.cpp
@@ -388,6 +388,11 @@ float UKismetMathLibrary::GetPI()
 	return PI;
 }
 
+float UKismetMathLibrary::GetTAU()
+{
+	return 2.f * PI;
+}
+
 float UKismetMathLibrary::DegreesToRadians(float A)
 {
 	return FMath::DegreesToRadians(A);
diff --git a/Engine/Source/Runtime/Engine/Private/PlayerCameraManager.cpp b/Engine/Source/Runtime/Engine/Private/PlayerCameraManager.cpp
--- a/Engine/Source/Runtime/Engine/Private/PlayerCameraManager.cpp
+++ b/Engine/Source/Runtime/Engine/Private/PlayerCameraManager.cpp
@@ -280,69 +280,12 @@ void APlayerCameraManager::GetCachedPostProcessBlends(TArray<FPostProcessSetting
 
 void APlayerCameraManager::ApplyAnimToCamera(ACameraActor const* AnimatedCamActor, UCameraAnimInst const* AnimInst, FMinimalViewInfo& InOutPOV)
 {
-	if (AnimInst->CamAnim->bRelativeToInitialTransform)
-	{
-		// move animated cam actor to initial-relative position
-		FTransform const AnimatedCamToWorld = AnimatedCamActor->GetTransform();
-		FTransform const AnimatedCamToInitialCam = AnimatedCamToWorld * AnimInst->InitialCamToWorld.Inverse();
-		ACameraActor* const MutableCamActor = const_cast<ACameraActor*>(AnimatedCamActor);
-		MutableCamActor->SetActorTransform(AnimatedCamToInitialCam);
-	}
-
-	float const Scale = AnimInst->CurrentBlendWeight;
-	FRotationMatrix const CameraToWorld(InOutPOV.Rotation);
-
-	if (AnimInst->PlaySpace == ECameraAnimPlaySpace::CameraLocal)
-	{
-		// the code in the else block will handle this just fine, but this path provides efficiency and simplicity for the most common case
-
-		// loc
-		FVector const LocalOffset = CameraToWorld.TransformVector( AnimatedCamActor->GetActorLocation()*Scale );
-		InOutPOV.Location += LocalOffset;
-
-		// rot
-		FRotationMatrix const AnimRotMat( AnimatedCamActor->GetActorRotation()*Scale );
-		InOutPOV.Rotation = (AnimRotMat * CameraToWorld).Rotator();
-	}
-	else
-	{
-		// handle playing the anim in an arbitrary space relative to the camera
-
-		// find desired space
-		FMatrix const PlaySpaceToWorld = (AnimInst->PlaySpace == ECameraAnimPlaySpace::UserDefined) ? AnimInst->UserPlaySpaceMatrix : FMatrix::Identity;
-
-		// loc
-		FVector const LocalOffset = PlaySpaceToWorld.TransformVector( AnimatedCamActor->GetActorLocation()*Scale );
-		InOutPOV.Location += LocalOffset;
-
-		// rot
-		// find transform from camera to the "play space"
-		FMatrix const CameraToPlaySpace = CameraToWorld * PlaySpaceToWorld.Inverse();	// CameraToWorld * WorldToPlaySpace
-
-		// find transform from anim (applied in playspace) back to camera
-		FRotationMatrix const AnimToPlaySpace(AnimatedCamActor->GetActorRotation()*Scale);
-		FMatrix const AnimToCamera = AnimToPlaySpace * CameraToPlaySpace.Inverse();			// AnimToPlaySpace * PlaySpaceToCamera
-
-		// RCS = rotated camera space, meaning camera space after it's been animated
-		// this is what we're looking for, the diff between rotated cam space and regular cam space.
-		// apply the transform back to camera space from the post-animated transform to get the RCS
-		FMatrix const RCSToCamera = CameraToPlaySpace * AnimToCamera;
-
-		// now apply to real camera
-		FRotationMatrix const RealCamToWorld(InOutPOV.Rotation);
-		InOutPOV.Rotation = (RCSToCamera * RealCamToWorld).Rotator();
-	}
-
-	// fov
-	const float FOVMin = 5.f;
-	const float FOVMax = 170.f;
-	InOutPOV.FOV += (AnimatedCamActor->GetCameraComponent()->FieldOfView - AnimInst->InitialFOV) * Scale;
-	InOutPOV.FOV = FMath::Clamp<float>(InOutPOV.FOV, FOVMin, FOVMax);
+	AnimInst->ApplyToView(InOutPOV);
 
 	// postprocess
 	if (AnimatedCamActor->GetCameraComponent()->PostProcessBlendWeight > 0.f)
 	{
-		AddCachedPPBlend(AnimatedCamActor->GetCameraComponent()->PostProcessSettings, AnimatedCamActor->GetCameraComponent()->PostProcessBlendWeight * Scale);
+		AddCachedPPBlend(AnimatedCamActor->GetCameraComponent()->PostProcessSettings, AnimatedCamActor->GetCameraComponent()->PostProcessBlendWeight * AnimInst->CurrentBlendWeight);
 	}
 }
 
diff --git a/Engine/Source/Runtime/Engine/Private/SceneView.cpp b/Engine/Source/Runtime/Engine/Private/SceneView.cpp
--- a/Engine/Source/Runtime/Engine/Private/SceneView.cpp
+++ b/Engine/Source/Runtime/Engine/Private/SceneView.cpp
@@ -1026,6 +1026,7 @@ void FSceneView::OverridePostProcessSettings(const FPostProcessSettings& Src, fl
 		LERP_PP(IndirectLightingIntensity);
 		LERP_PP(DepthOfFieldFocalDistance);
 		LERP_PP(DepthOfFieldFstop);
+		LERP_PP(DepthOfFieldSensorWidth);
 		LERP_PP(DepthOfFieldDepthBlurRadius);
 		LERP_PP(DepthOfFieldDepthBlurAmount);
 		LERP_PP(DepthOfFieldFocalRegion);
diff --git a/Engine/Source/Runtime/Engine/Private/UnrealEngine.cpp b/Engine/Source/Runtime/Engine/Private/UnrealEngine.cpp
--- a/Engine/Source/Runtime/Engine/Private/UnrealEngine.cpp
+++ b/Engine/Source/Runtime/Engine/Private/UnrealEngine.cpp
@@ -6721,12 +6721,12 @@ float UEngine::GetMaxTickRate(float DeltaTime, bool bAllowFrameRateSmoothing) co
 	return MaxTickRate;
 }
 
-int32 UEngine::GetMaxFPS() const
+float UEngine::GetMaxFPS() const
 {
 	return CVarMaxFPS.GetValueOnAnyThread();
 }
 
-void UEngine::SetMaxFPS(const int32 MaxFPS)
+void UEngine::SetMaxFPS(const float MaxFPS)
 {
 	IConsoleVariable* ConsoleVariable = CVarMaxFPS.AsVariable();
 	ConsoleVariable->Set(MaxFPS);
diff --git a/Engine/Source/Runtime/Engine/Public/Animation/AnimCurveTypes.h b/Engine/Source/Runtime/Engine/Public/Animation/AnimCurveTypes.h
--- a/Engine/Source/Runtime/Engine/Public/Animation/AnimCurveTypes.h
+++ b/Engine/Source/Runtime/Engine/Public/Animation/AnimCurveTypes.h
@@ -572,6 +572,13 @@ struct FRawCurveTracks
 	 *	Evaluate transform curves 
 	 */
 	ENGINE_API void EvaluateTransformCurveData(USkeleton * Skeleton, TMap<FName, FTransform>&OutCurves, float CurrentTime, float BlendWeight) const;
+
+	/**
+	* Add new float curve from the given UID if not existing and add the key with time/value
+	*/
+	ENGINE_API void AddFloatCurveKey(const USkeleton::AnimCurveUID Uid, int32 CurveFlags, float Time, float Value);
+	ENGINE_API void RemoveRedundantKeys();
+
 #endif // WITH_EDITOR
 	/**
 	 * Find curve data based on the curve UID
diff --git a/Engine/Source/Runtime/Engine/Public/SceneView.h b/Engine/Source/Runtime/Engine/Public/SceneView.h
--- a/Engine/Source/Runtime/Engine/Public/SceneView.h
+++ b/Engine/Source/Runtime/Engine/Public/SceneView.h
@@ -500,6 +500,7 @@ BEGIN_UNIFORM_BUFFER_STRUCT_WITH_CONSTRUCTOR(FFrameUniformShaderParameters, ENGI
 	DECLARE_UNIFORM_BUFFER_STRUCT_MEMBER_ARRAY(FVector4, TranslucencyLightingVolumeInvSize, [TVC_MAX])
 	DECLARE_UNIFORM_BUFFER_STRUCT_MEMBER(FVector4, TemporalAAParams)
 	DECLARE_UNIFORM_BUFFER_STRUCT_MEMBER(FVector4, CircleDOFParams)
+	DECLARE_UNIFORM_BUFFER_STRUCT_MEMBER(float, DepthOfFieldSensorWidth)
 	DECLARE_UNIFORM_BUFFER_STRUCT_MEMBER(float, DepthOfFieldFocalDistance)
 	DECLARE_UNIFORM_BUFFER_STRUCT_MEMBER(float, DepthOfFieldScale)
 	DECLARE_UNIFORM_BUFFER_STRUCT_MEMBER(float, DepthOfFieldFocalLength)
diff --git a/Engine/Source/Runtime/GameplayAbilities/GameplayAbilities.Build.cs b/Engine/Source/Runtime/GameplayAbilities/GameplayAbilities.Build.cs
--- a/Engine/Source/Runtime/GameplayAbilities/GameplayAbilities.Build.cs
+++ b/Engine/Source/Runtime/GameplayAbilities/GameplayAbilities.Build.cs
@@ -51,6 +51,7 @@ namespace UnrealBuildTool.Rules
 				PrivateDependencyModuleNames.Add("UnrealEd");
                 PrivateDependencyModuleNames.Add("GameplayTagsEditor");
                 PrivateDependencyModuleNames.Add("Slate");
+				PrivateDependencyModuleNames.Add("SequenceRecorder");
 			}
 
             if (UEBuildConfiguration.bBuildDeveloperTools && Target.Configuration != UnrealTargetConfiguration.Shipping && Target.Configuration != UnrealTargetConfiguration.Test)
diff --git a/Engine/Source/Runtime/GameplayAbilities/Private/GameplayCueManager.cpp b/Engine/Source/Runtime/GameplayAbilities/Private/GameplayCueManager.cpp
--- a/Engine/Source/Runtime/GameplayAbilities/Private/GameplayCueManager.cpp
+++ b/Engine/Source/Runtime/GameplayAbilities/Private/GameplayCueManager.cpp
@@ -16,6 +16,7 @@
 #include "UnrealEd.h"
 #include "SNotificationList.h"
 #include "NotificationManager.h"
+#include "ISequenceRecorder.h"
 #define LOCTEXT_NAMESPACE "GameplayCueManager"
 #endif
 
@@ -212,6 +213,12 @@ AGameplayCueNotify_Actor* UGameplayCueManager::GetInstancedCueActor(AActor* Targ
 				SpawnedCue->SetActorHiddenInGame(false);
 				SpawnedCue->SetOwner(NewOwnerActor);
 				SpawnedCue->SetActorLocationAndRotation(TargetActor->GetActorLocation(), TargetActor->GetActorRotation());
+
+#if WITH_EDITOR
+				// let things know that we 'spawned'
+				ISequenceRecorder& SequenceRecorder	= FModuleManager::LoadModuleChecked<ISequenceRecorder>("SequenceRecorder");
+				SequenceRecorder.NotifyActorStartRecording(SpawnedCue);
+#endif
 			}
 		}
 
@@ -264,6 +271,12 @@ void UGameplayCueManager::NotifyGameplayCueActorFinished(AGameplayCueNotify_Acto
 			TArray<AGameplayCueNotify_Actor*>& PreAllocatedList = Info.PreallocatedInstances.FindOrAdd(Actor->GetClass());
 			PreAllocatedList.Push(Actor);
 
+#if WITH_EDITOR
+			// let things know that we 'de-spawned'
+			ISequenceRecorder& SequenceRecorder	= FModuleManager::LoadModuleChecked<ISequenceRecorder>("SequenceRecorder");
+			SequenceRecorder.NotifyActorStopRecording(Actor);
+#endif
+
 			return;
 		}
 	}	
diff --git a/Engine/Source/Runtime/Launch/Private/LaunchEngineLoop.cpp b/Engine/Source/Runtime/Launch/Private/LaunchEngineLoop.cpp
--- a/Engine/Source/Runtime/Launch/Private/LaunchEngineLoop.cpp
+++ b/Engine/Source/Runtime/Launch/Private/LaunchEngineLoop.cpp
@@ -2224,6 +2224,8 @@ int32 FEngineLoop::Init()
 	{
 		FModuleManager::Get().LoadModule(TEXT("ProfilerClient"));
 	}
+
+	FModuleManager::Get().LoadModule(TEXT("SequenceRecorder"));
 #endif
 
 	GIsRunning = true;
diff --git a/Engine/Source/Runtime/LevelSequence/Private/LevelSequencePlayer.cpp b/Engine/Source/Runtime/LevelSequence/Private/LevelSequencePlayer.cpp
--- a/Engine/Source/Runtime/LevelSequence/Private/LevelSequencePlayer.cpp
+++ b/Engine/Source/Runtime/LevelSequence/Private/LevelSequencePlayer.cpp
@@ -236,7 +236,7 @@ void ULevelSequencePlayer::Initialize(ULevelSequence* InLevelSequence, UWorld* I
 /* IMovieScenePlayer interface
  *****************************************************************************/
 
-void ULevelSequencePlayer::GetRuntimeObjects(TSharedRef<FMovieSceneSequenceInstance> MovieSceneInstance, const FGuid& ObjectId, TArray<UObject*>& OutObjects) const
+void ULevelSequencePlayer::GetRuntimeObjects(TSharedRef<FMovieSceneSequenceInstance> MovieSceneInstance, const FGuid& ObjectId, TArray<TWeakObjectPtr<UObject>>& OutObjects) const
 {
 	UObject* FoundObject = MovieSceneInstance->FindObject(ObjectId, *this);
 	if (FoundObject)
@@ -246,7 +246,7 @@ void ULevelSequencePlayer::GetRuntimeObjects(TSharedRef<FMovieSceneSequenceInsta
 }
 
 
-void ULevelSequencePlayer::UpdateCameraCut(UObject* CameraObject, UObject* UnlockIfCameraObject) const
+void ULevelSequencePlayer::UpdateCameraCut(UObject* CameraObject, UObject* UnlockIfCameraObject, bool bJumpCut) const
 {
 	// skip missing player controller
 	APlayerController* PC = World->GetGameInstance()->GetFirstLocalPlayerController();
@@ -335,7 +335,9 @@ void ULevelSequencePlayer::Update(const float DeltaSeconds)
 		UpdateMovieSceneInstance(TimeCursorPosition, LastTimePosition);
 
 		bHasCleanedUpSequence = true;
-		SpawnRegister->DestroyAllOwnedObjects(*this);
+		
+		SpawnRegister->ForgetExternallyOwnedSpawnedObjects(*this);
+		SpawnRegister->CleanUp(*this);
 	}
 }
 
diff --git a/Engine/Source/Runtime/LevelSequence/Private/LevelSequenceSpawnRegister.cpp b/Engine/Source/Runtime/LevelSequence/Private/LevelSequenceSpawnRegister.cpp
--- a/Engine/Source/Runtime/LevelSequence/Private/LevelSequenceSpawnRegister.cpp
+++ b/Engine/Source/Runtime/LevelSequence/Private/LevelSequenceSpawnRegister.cpp
@@ -8,6 +8,9 @@
 #include "MovieScene.h"
 #include "MovieSceneCommonHelpers.h"
 #include "Particles/ParticleSystem.h"
+#include "IMovieScenePlayer.h"
+
+static const FName SequencerActorTag(TEXT("SequencerActor"));
 
 UObject* FLevelSequenceSpawnRegister::SpawnObject(const FGuid& BindingId, FMovieSceneSequenceInstance& SequenceInstance, IMovieScenePlayer& Player)
 {
@@ -90,16 +93,25 @@ UObject* FLevelSequenceSpawnRegister::SpawnObject(const FGuid& BindingId, FMovie
 		}
 	}
 
-	UObject* SpawnedObject = GWorld->SpawnActor(SpawnableClass, &NewTransform, SpawnInfo);
-	if (!SpawnedObject)
+	UWorld* WorldContext = Cast<UWorld>(Player.GetPlaybackContext());
+	if(WorldContext == nullptr)
+	{
+		WorldContext = GWorld;
+	}
+
+	AActor* SpawnedActor = WorldContext->SpawnActor(SpawnableClass, &NewTransform, SpawnInfo);
+	if (!SpawnedActor)
 	{
 		return nullptr;
 	}
 
-	Register.Add(Key, FSpawnedObject(*SpawnedObject, Spawnable->GetSpawnOwnership()));
+	// tag this actor so we know it was spawned by sequencer
+	SpawnedActor->Tags.Add(SequencerActorTag);
+
+	Register.Add(Key, FSpawnedObject(BindingId, *SpawnedActor, Spawnable->GetSpawnOwnership()));
 
-	SequenceInstance.OnObjectSpawned(BindingId, *SpawnedObject, Player);
-	return SpawnedObject;
+	SequenceInstance.OnObjectSpawned(BindingId, *SpawnedActor, Player);
+	return SpawnedActor;
 }
 
 void FLevelSequenceSpawnRegister::DestroySpawnedObject(const FGuid& BindingId, FMovieSceneSequenceInstance& SequenceInstance, IMovieScenePlayer& Player)
@@ -155,12 +167,12 @@ void FLevelSequenceSpawnRegister::ForgetExternallyOwnedSpawnedObjects(IMovieScen
 	}
 }
 
-void FLevelSequenceSpawnRegister::DestroyObjects(IMovieScenePlayer& Player, TFunctionRef<bool(FMovieSceneSequenceInstance&, FSpawnedObject&)> Pred)
+void FLevelSequenceSpawnRegister::DestroyObjectsByPredicate(IMovieScenePlayer& Player, const TFunctionRef<bool(const FGuid&, ESpawnOwnership, FMovieSceneSequenceInstance&)>& Predicate)
 {
 	for (auto It = Register.CreateIterator(); It; ++It)
 	{
 		FMovieSceneSequenceInstance* ThisInstance = It.Key().SequenceInstance.Pin().Get();
-		if (!ThisInstance || Pred(*ThisInstance, It.Value()))
+		if (!ThisInstance || Predicate(It.Value().Guid, It.Value().Ownership, *ThisInstance))
 		{
 			UObject* SpawnedObject = It.Value().Object.Get();
 			if (SpawnedObject)
@@ -181,38 +193,6 @@ void FLevelSequenceSpawnRegister::DestroyObjects(IMovieScenePlayer& Player, TFun
 	}
 }
 
-void FLevelSequenceSpawnRegister::DestroyObjectsOwnedByInstance(FMovieSceneSequenceInstance& SequenceInstance, IMovieScenePlayer& Player)
-{
-	DestroyObjects(Player, [&](FMovieSceneSequenceInstance& ThisInstance, FSpawnedObject& SpawnedObject){
-		if (SpawnedObject.Ownership == ESpawnOwnership::InnerSequence)
-		{
-			return &ThisInstance == &SequenceInstance;
-		}
-		return false;
-	});
-}
-
-void FLevelSequenceSpawnRegister::DestroyObjectsSpawnedByInstance(FMovieSceneSequenceInstance& SequenceInstance, IMovieScenePlayer& Player)
-{
-	DestroyObjects(Player, [&](FMovieSceneSequenceInstance& ThisInstance, FSpawnedObject& SpawnedObject){
-		return &ThisInstance == &SequenceInstance;
-	});
-}
-
-void FLevelSequenceSpawnRegister::DestroyAllOwnedObjects(IMovieScenePlayer& Player)
-{
-	DestroyObjects(Player, [&](FMovieSceneSequenceInstance&, FSpawnedObject& SpawnedObject){
-		return SpawnedObject.Ownership != ESpawnOwnership::External;
-	});
-}
-
-void FLevelSequenceSpawnRegister::DestroyAllObjects(IMovieScenePlayer& Player)
-{
-	DestroyObjects(Player, [&](FMovieSceneSequenceInstance&, FSpawnedObject&){
-		return true;
-	});
-}
-
 void FLevelSequenceSpawnRegister::PreUpdateSequenceInstance(FMovieSceneSequenceInstance& Instance, IMovieScenePlayer& Player)
 {
 	++CurrentlyUpdatingSequenceCount;
@@ -228,7 +208,7 @@ void FLevelSequenceSpawnRegister::PostUpdateSequenceInstance(FMovieSceneSequence
 			TSharedPtr<FMovieSceneSequenceInstance> Instance = WeakInstance.Pin();
 			if (Instance.IsValid() && !ActiveInstances.Contains(Instance))
 			{
-				DestroyObjectsOwnedByInstance(*Instance, Player);
+				OnSequenceExpired(*Instance, Player);
 			}
 		}
 
diff --git a/Engine/Source/Runtime/LevelSequence/Public/LevelSequencePlayer.h b/Engine/Source/Runtime/LevelSequence/Public/LevelSequencePlayer.h
--- a/Engine/Source/Runtime/LevelSequence/Public/LevelSequencePlayer.h
+++ b/Engine/Source/Runtime/LevelSequence/Public/LevelSequencePlayer.h
@@ -133,8 +133,8 @@ public:
 protected:
 
 	// IMovieScenePlayer interface
-	virtual void GetRuntimeObjects(TSharedRef<FMovieSceneSequenceInstance> MovieSceneInstance, const FGuid& ObjectHandle, TArray<UObject*>& OutObjects) const override;
-	virtual void UpdateCameraCut(UObject* CameraObject, UObject* UnlockIfCameraObject) const override;
+	virtual void GetRuntimeObjects(TSharedRef<FMovieSceneSequenceInstance> MovieSceneInstance, const FGuid& ObjectHandle, TArray<TWeakObjectPtr<UObject>>& OutObjects) const override;
+	virtual void UpdateCameraCut(UObject* CameraObject, UObject* UnlockIfCameraObject, bool bJumpCut) const override;
 	virtual void SetViewportSettings(const TMap<FViewportClient*, EMovieSceneViewportParams>& ViewportParamsMap) override;
 	virtual void GetViewportSettings(TMap<FViewportClient*, EMovieSceneViewportParams>& ViewportParamsMap) const override;
 	virtual EMovieScenePlayerStatus::Type GetPlaybackStatus() const override;
diff --git a/Engine/Source/Runtime/LevelSequence/Public/LevelSequenceSpawnRegister.h b/Engine/Source/Runtime/LevelSequence/Public/LevelSequenceSpawnRegister.h
--- a/Engine/Source/Runtime/LevelSequence/Public/LevelSequenceSpawnRegister.h
+++ b/Engine/Source/Runtime/LevelSequence/Public/LevelSequenceSpawnRegister.h
@@ -18,11 +18,8 @@ public:
 	/** ~ IMovieSceneSpawnRegister interface */
 	virtual UObject* SpawnObject(const FGuid& BindingId, FMovieSceneSequenceInstance& SequenceInstance, IMovieScenePlayer&) override;
 	virtual void DestroySpawnedObject(const FGuid& BindingId, FMovieSceneSequenceInstance& SequenceInstance, IMovieScenePlayer&) override;
-	virtual void DestroyObjectsOwnedByInstance(FMovieSceneSequenceInstance&, IMovieScenePlayer&) override;
-	virtual void DestroyObjectsSpawnedByInstance(FMovieSceneSequenceInstance&, IMovieScenePlayer&) override;
-	virtual void DestroyAllOwnedObjects(IMovieScenePlayer&) override;
+	virtual void DestroyObjectsByPredicate(IMovieScenePlayer& Player, const TFunctionRef<bool(const FGuid&, ESpawnOwnership, FMovieSceneSequenceInstance&)>& Predicate) override;
 	virtual void ForgetExternallyOwnedSpawnedObjects(IMovieScenePlayer& Player) override;
-	virtual void DestroyAllObjects(IMovieScenePlayer&) override;
 	virtual void PreUpdateSequenceInstance(FMovieSceneSequenceInstance& Instance, IMovieScenePlayer&) override;
 	virtual void PostUpdateSequenceInstance(FMovieSceneSequenceInstance& Instance, IMovieScenePlayer&) override;
 
@@ -31,11 +28,15 @@ protected:
 	/** Structure holding information pertaining to a spawned object */
 	struct FSpawnedObject
 	{
-		FSpawnedObject(UObject& InObject, ESpawnOwnership InOwnership)
-			: Object(&InObject)
+		FSpawnedObject(const FGuid& InGuid, UObject& InObject, ESpawnOwnership InOwnership)
+			: Guid(InGuid)
+			, Object(&InObject)
 			, Ownership(InOwnership)
 		{}
 
+		/** The ID of the sequencer object binding that this object relates to */
+		FGuid Guid;
+
 		/** The object that has been spawned */
 		TWeakObjectPtr<UObject> Object;
 
diff --git a/Engine/Source/Runtime/MovieScene/Private/MovieScene.cpp b/Engine/Source/Runtime/MovieScene/Private/MovieScene.cpp
--- a/Engine/Source/Runtime/MovieScene/Private/MovieScene.cpp
+++ b/Engine/Source/Runtime/MovieScene/Private/MovieScene.cpp
@@ -2,13 +2,14 @@
 
 #include "MovieScenePrivatePCH.h"
 #include "MovieScene.h"
-
+#include "MovieSceneFolder.h"
 
 /* UMovieScene interface
  *****************************************************************************/
 
 UMovieScene::UMovieScene(const FObjectInitializer& ObjectInitializer)
 	: Super(ObjectInitializer)
+	, InOutRange(FFloatRange::Empty())
 	, PlaybackRange(FFloatRange::Empty())
 	, InTime_DEPRECATED(FLT_MAX)
 	, OutTime_DEPRECATED(-FLT_MAX)
@@ -55,10 +56,13 @@ bool UMovieScene::RemoveSpawnable( const FGuid& Guid )
 				{
 					UClass* GeneratedClass = CurSpawnable.GetClass();
 					UBlueprint* Blueprint = GeneratedClass ? Cast<UBlueprint>(GeneratedClass->ClassGeneratedBy) : nullptr;
-					check(nullptr != Blueprint);
-					// @todo sequencer: Also remove created Blueprint inner object.  Is this sufficient?  Needs to work with Undo too!
-					Blueprint->ClearFlags( RF_Standalone );	// @todo sequencer: Probably not needed for Blueprint
-					Blueprint->MarkPendingKill();
+
+					if (Blueprint)
+					{
+						// @todo sequencer: Also remove created Blueprint inner object.  Is this sufficient?  Needs to work with Undo too!
+						Blueprint->ClearFlags( RF_Standalone );	// @todo sequencer: Probably not needed for Blueprint
+						Blueprint->MarkPendingKill();
+					}
 				}
 
 				RemoveBinding( Guid );
@@ -248,17 +252,16 @@ void UMovieScene::SetObjectDisplayName(const FGuid& ObjectId, const FText& Displ
 		ObjectsToDisplayNames.Add(ObjectId.ToString(), DisplayName);
 	}
 }
-#endif
 
 
-TRange<float> UMovieScene::GetPlaybackRange() const
+TArray<UMovieSceneFolder*>&  UMovieScene::GetRootFolders()
 {
-	check(PlaybackRange.HasLowerBound() && PlaybackRange.HasUpperBound());
-	return PlaybackRange;
+	return RootFolders;
 }
+#endif
 
 
-void UMovieScene::SetPlaybackRange(float Start, float End)
+void UMovieScene::SetPlaybackRange(float Start, float End, bool bAlwaysMarkDirty)
 {
 	if (ensure(End >= Start))
 	{
@@ -269,8 +272,11 @@ void UMovieScene::SetPlaybackRange(float Start, float End)
 			return;
 		}
 
-		Modify();
-		
+		if (bAlwaysMarkDirty)
+		{
+			Modify();
+		}
+
 		PlaybackRange = NewRange;
 
 #if WITH_EDITORONLY_DATA
@@ -335,7 +341,7 @@ UMovieSceneTrack* UMovieScene::FindTrack(TSubclassOf<UMovieSceneTrack> TrackClas
 }
 
 
-class UMovieSceneTrack* UMovieScene::AddTrack( TSubclassOf<UMovieSceneTrack> TrackClass, const FGuid& ObjectGuid )
+UMovieSceneTrack* UMovieScene::AddTrack( TSubclassOf<UMovieSceneTrack> TrackClass, const FGuid& ObjectGuid )
 {
 	UMovieSceneTrack* CreatedType = nullptr;
 
@@ -379,6 +385,22 @@ bool bAnythingRemoved = false;
 	return bAnythingRemoved;
 }
 
+bool UMovieScene::FindTrackBinding(const UMovieSceneTrack& InTrack, FGuid& OutGuid) const
+{
+	for (auto& Binding : ObjectBindings)
+	{
+		for(auto& Track : Binding.GetTracks())
+		{
+			if(Track == &InTrack)
+			{
+				OutGuid = Binding.GetObjectGuid();
+				return true;
+			}
+		}
+	}
+
+	return false;
+}
 
 UMovieSceneTrack* UMovieScene::FindMasterTrack( TSubclassOf<UMovieSceneTrack> TrackClass ) const
 {
@@ -397,7 +419,7 @@ UMovieSceneTrack* UMovieScene::FindMasterTrack( TSubclassOf<UMovieSceneTrack> Tr
 }
 
 
-class UMovieSceneTrack* UMovieScene::AddMasterTrack( TSubclassOf<UMovieSceneTrack> TrackClass )
+UMovieSceneTrack* UMovieScene::AddMasterTrack( TSubclassOf<UMovieSceneTrack> TrackClass )
 {
 	Modify();
 
diff --git a/Engine/Source/Runtime/MovieScene/Private/MovieSceneCommonHelpers.cpp b/Engine/Source/Runtime/MovieScene/Private/MovieSceneCommonHelpers.cpp
--- a/Engine/Source/Runtime/MovieScene/Private/MovieSceneCommonHelpers.cpp
+++ b/Engine/Source/Runtime/MovieScene/Private/MovieSceneCommonHelpers.cpp
@@ -212,6 +212,29 @@ UCameraComponent* MovieSceneHelpers::CameraComponentFromActor(const AActor* InAc
 	return nullptr;
 }
 
+UCameraComponent* MovieSceneHelpers::CameraComponentFromRuntimeObject(UObject* RuntimeObject)
+{
+	if (RuntimeObject)
+	{
+		// find camera we want to control
+		UCameraComponent* const CameraComponent = dynamic_cast<UCameraComponent*>(RuntimeObject);
+		if (CameraComponent)
+		{
+			return CameraComponent;
+		}
+
+		// see if it's an actor that has a camera component
+		AActor* const Actor = dynamic_cast<AActor*>(RuntimeObject);
+		if (Actor)
+		{
+			return CameraComponentFromActor(Actor);
+		}
+	}
+
+	return nullptr;
+}
+
+
 void MovieSceneHelpers::SetKeyInterpolation(FRichCurve& InCurve, FKeyHandle InKeyHandle, EMovieSceneKeyInterpolation InKeyInterpolation)
 {
 	switch (InKeyInterpolation)
@@ -317,18 +340,26 @@ FTrackInstancePropertyBindings::FPropertyAddress FTrackInstancePropertyBindings:
 }
 
 
-void FTrackInstancePropertyBindings::UpdateBindings( const TArray<UObject*>& InRuntimeObjects )
+void FTrackInstancePropertyBindings::UpdateBindings( const TArray<TWeakObjectPtr<UObject>>& InRuntimeObjects )
 {
-	for(UObject* Object : InRuntimeObjects)
+	for (auto ObjectPtr : InRuntimeObjects)
 	{
-		FPropertyAndFunction PropAndFunction;
+		UObject* Object = ObjectPtr.Get();
+
+		if (Object != nullptr)
+		{
+			FPropertyAndFunction PropAndFunction;
+			{
+				PropAndFunction.Function = Object->FindFunction(FunctionName);
+				PropAndFunction.PropertyAddress = FindProperty(Object, PropertyPath);
+			}
 
-		PropAndFunction.Function = Object->FindFunction(FunctionName);
-		PropAndFunction.PropertyAddress = FindProperty( Object, PropertyPath );
-		RuntimeObjectToFunctionMap.Add(Object, PropAndFunction);
+			RuntimeObjectToFunctionMap.Add(Object, PropAndFunction);
+		}
 	}
 }
 
+
 UProperty* FTrackInstancePropertyBindings::GetProperty(const UObject* Object) const
 {
 	UProperty* Prop = nullptr;
diff --git a/Engine/Source/Runtime/MovieScene/Private/MovieSceneFolder.cpp b/Engine/Source/Runtime/MovieScene/Private/MovieSceneFolder.cpp
new file mode 100644
--- /dev/null
+++ b/Engine/Source/Runtime/MovieScene/Private/MovieSceneFolder.cpp
@@ -0,0 +1,98 @@
+// Copyright 1998-2016 Epic Games, Inc. All Rights Reserved.
+
+#include "MovieScenePrivatePCH.h"
+#include "MovieSceneFolder.h"
+
+UMovieSceneFolder::UMovieSceneFolder( const FObjectInitializer& ObjectInitializer )
+	: Super( ObjectInitializer )
+{
+}
+
+FName UMovieSceneFolder::GetFolderName() const
+{
+	return FolderName;
+}
+
+
+void UMovieSceneFolder::SetFolderName( FName InFolderName )
+{
+	FolderName = InFolderName;
+}
+
+
+const TArray<UMovieSceneFolder*>& UMovieSceneFolder::GetChildFolders() const
+{
+	return ChildFolders;
+}
+
+
+void UMovieSceneFolder::AddChildFolder( UMovieSceneFolder* InChildFolder )
+{
+	ChildFolders.Add( InChildFolder );
+}
+
+
+void UMovieSceneFolder::RemoveChildFolder( UMovieSceneFolder* InChildFolder )
+{
+	ChildFolders.Remove(InChildFolder);
+}
+
+
+const TArray<UMovieSceneTrack*>& UMovieSceneFolder::GetChildMasterTracks() const
+{
+	return ChildMasterTracks;
+}
+
+
+void UMovieSceneFolder::AddChildMasterTrack( UMovieSceneTrack* InMasterTrack )
+{
+	ChildMasterTracks.Add( InMasterTrack );
+}
+
+
+void UMovieSceneFolder::RemoveChildMasterTrack( UMovieSceneTrack* InMasterTrack )
+{
+	ChildMasterTracks.Remove( InMasterTrack );
+}
+
+
+const TArray<FGuid>& UMovieSceneFolder::GetChildObjectBindings() const
+{
+	return ChildObjectBindings;
+}
+
+
+void UMovieSceneFolder::AddChildObjectBinding(const FGuid& InObjectBinding )
+{
+	ChildObjectBindings.Add( InObjectBinding );
+}
+
+
+void UMovieSceneFolder::RemoveChildObjectBinding( const FGuid& InObjectBinding )
+{
+	ChildObjectBindings.Remove( InObjectBinding );
+}
+
+
+void UMovieSceneFolder::PreSave()
+{
+	ChildObjectBindingStrings.Empty();
+	for ( const FGuid& ChildObjectBinding : ChildObjectBindings )
+	{
+		ChildObjectBindingStrings.Add( ChildObjectBinding.ToString() );
+	}
+	Super::PreSave();
+}
+
+
+void UMovieSceneFolder::PostLoad()
+{
+	Super::PostLoad();
+	ChildObjectBindings.Empty();
+	for ( const FString& ChildObjectBindingString : ChildObjectBindingStrings )
+	{
+		FGuid ChildObjectBinding;
+		FGuid::Parse( ChildObjectBindingString, ChildObjectBinding );
+		ChildObjectBindings.Add( ChildObjectBinding );
+	}
+}
diff --git a/Engine/Source/Runtime/MovieScene/Private/MovieSceneSection.cpp b/Engine/Source/Runtime/MovieScene/Private/MovieSceneSection.cpp
--- a/Engine/Source/Runtime/MovieScene/Private/MovieSceneSection.cpp
+++ b/Engine/Source/Runtime/MovieScene/Private/MovieSceneSection.cpp
@@ -31,14 +31,16 @@ bool UMovieSceneSection::TryModify(bool bAlwaysMarkDirty)
 
 const UMovieSceneSection* UMovieSceneSection::OverlapsWithSections(const TArray<UMovieSceneSection*>& Sections, int32 TrackDelta, float TimeDelta) const
 {
+	// Check overlaps with exclusive ranges so that sections can butt up against each other
 	int32 NewTrackIndex = RowIndex + TrackDelta;
-	TRange<float> NewSectionRange = TRange<float>(StartTime + TimeDelta, EndTime + TimeDelta);
+	TRange<float> NewSectionRange = TRange<float>(TRange<float>::BoundsType::Exclusive(StartTime + TimeDelta), TRange<float>::BoundsType::Exclusive(EndTime + TimeDelta));
 
 	for (const auto Section : Sections)
 	{
 		if ((this != Section) && (Section->GetRowIndex() == NewTrackIndex))
 		{
-			if (NewSectionRange.Overlaps(Section->GetRange()))
+			TRange<float> ExclusiveSectionRange = TRange<float>(TRange<float>::BoundsType::Exclusive(Section->GetRange().GetLowerBoundValue()), TRange<float>::BoundsType::Exclusive(Section->GetRange().GetUpperBoundValue()));
+			if (NewSectionRange.Overlaps(ExclusiveSectionRange))
 			{
 				return Section;
 			}
diff --git a/Engine/Source/Runtime/MovieScene/Private/MovieSceneSequenceInstance.cpp b/Engine/Source/Runtime/MovieScene/Private/MovieSceneSequenceInstance.cpp
--- a/Engine/Source/Runtime/MovieScene/Private/MovieSceneSequenceInstance.cpp
+++ b/Engine/Source/Runtime/MovieScene/Private/MovieSceneSequenceInstance.cpp
@@ -3,6 +3,7 @@
 #include "MovieScenePrivatePCH.h"
 #include "MovieSceneSequenceInstance.h"
 #include "MovieSceneSequence.h"
+#include "MovieSceneTrack.h"
 
 FMovieSceneSequenceInstance::FMovieSceneSequenceInstance(const UMovieSceneSequence& InMovieSceneSequence)
 	: MovieSceneSequence( &InMovieSceneSequence )
@@ -10,6 +11,11 @@ FMovieSceneSequenceInstance::FMovieSceneSequenceInstance(const UMovieSceneSequen
 	TimeRange = MovieSceneSequence->GetMovieScene()->GetPlaybackRange();
 }
 
+FMovieSceneSequenceInstance::FMovieSceneSequenceInstance(const UMovieSceneTrack& InMovieSceneTrack)
+{
+	TimeRange = CastChecked<UMovieScene>(InMovieSceneTrack.GetOuter())->GetPlaybackRange();
+}
+
 FMovieSceneSequenceInstance::~FMovieSceneSequenceInstance()
 {
 	for (auto& Pair : SpawnedObjects)
@@ -20,49 +26,60 @@ FMovieSceneSequenceInstance::~FMovieSceneSequenceInstance()
 
 FGuid FMovieSceneSequenceInstance::FindObjectId(UObject& Object) const
 {
-	for (auto& Pair : ObjectBindingInstances)
+	if(MovieSceneSequence.IsValid())
 	{
-		if (Pair.Value.RuntimeObjects.Contains(&Object))
+		for (auto& Pair : ObjectBindingInstances)
 		{
-			return Pair.Key;
+			if (Pair.Value.RuntimeObjects.Contains(&Object))
+			{
+				return Pair.Key;
+			}
 		}
+
+		// At this point the only possibility left is that we have not cached the object
+		// in ObjectBindingInstances, so we to see if the sequence itself can tell us the GUID
+		return MovieSceneSequence->FindPossessableObjectId(Object);
 	}
 
-	// At this point the only possibility left is that we have not cached the object
-	// in ObjectBindingInstances, so we to see if the sequence itself can tell us the GUID
-	return MovieSceneSequence->FindPossessableObjectId(Object);
+	return FGuid();
 }
 
 FGuid FMovieSceneSequenceInstance::FindParentObjectId(UObject& Object) const
 {
-	UObject* ParentObject = MovieSceneSequence->GetParentObject(&Object);
-	if (ParentObject)
+	if(MovieSceneSequence.IsValid())
 	{
-		return FindObjectId(*ParentObject);
+		UObject* ParentObject = MovieSceneSequence->GetParentObject(&Object);
+		if (ParentObject)
+		{
+			return FindObjectId(*ParentObject);
+		}
 	}
 	return FGuid();
 }
 
 UObject* FMovieSceneSequenceInstance::FindObject(const FGuid& ObjectId, const IMovieScenePlayer& Player) const
 {
-	// Attempt to find a possessable first
-	UMovieScene* MovieScene = MovieSceneSequence->GetMovieScene();
-	if (MovieScene != nullptr)
+	if(MovieSceneSequence.IsValid())
 	{
-		FMovieScenePossessable* Possessable = MovieScene->FindPossessable(ObjectId);
-		if (Possessable)
+		// Attempt to find a possessable first
+		UMovieScene* MovieScene = MovieSceneSequence->GetMovieScene();
+		if (MovieScene != nullptr)
 		{
-			UObject* ParentObject = Player.GetPlaybackContext();
-			if (Possessable->GetParent().IsValid())
+			FMovieScenePossessable* Possessable = MovieScene->FindPossessable(ObjectId);
+			if (Possessable)
 			{
-				ParentObject = FindObject(Possessable->GetParent(), Player);
-			}
+				UObject* ParentObject = Player.GetPlaybackContext();
+				if (Possessable->GetParent().IsValid())
+				{
+					ParentObject = FindObject(Possessable->GetParent(), Player);
+				}
 
-			return MovieSceneSequence->FindPossessableObject(ObjectId, ParentObject);
-		}
-		else
-		{
-			return FindSpawnedObject(ObjectId);
+				return MovieSceneSequence->FindPossessableObject(ObjectId, ParentObject);
+			}
+			else
+			{
+				return FindSpawnedObject(ObjectId);
+			}
 		}
 	}
 
@@ -77,7 +94,7 @@ UObject* FMovieSceneSequenceInstance::FindSpawnedObject(const FGuid& ObjectId) c
 
 void FMovieSceneSequenceInstance::SaveState(class IMovieScenePlayer& Player)
 {
-	TArray<UObject*> NoObjects;
+	TArray<TWeakObjectPtr<UObject>> NoObjects;
 	if (CameraCutTrackInstance.IsValid())
 	{
 		CameraCutTrackInstance->SaveState(NoObjects, Player, *this);
@@ -108,7 +125,7 @@ void FMovieSceneSequenceInstance::SaveState(class IMovieScenePlayer& Player)
 
 void FMovieSceneSequenceInstance::RestoreState(class IMovieScenePlayer& Player)
 {
-	Player.GetSpawnRegister().DestroyObjectsSpawnedByInstance(*this, Player);
+	Player.GetSpawnRegister().CleanUpSequence(*this, Player);
 
 	TMap<FGuid, FMovieSceneObjectBindingInstance>::TIterator ObjectIt = ObjectBindingInstances.CreateIterator();
 	for (; ObjectIt; ++ObjectIt)
@@ -126,7 +143,7 @@ void FMovieSceneSequenceInstance::RestoreState(class IMovieScenePlayer& Player)
 		}
 	}
 
-	TArray<UObject*> NoObjects;
+	TArray<TWeakObjectPtr<UObject>> NoObjects;
 	if (CameraCutTrackInstance.IsValid())
 	{
 		CameraCutTrackInstance->RestoreState(NoObjects, Player, *this);
@@ -138,8 +155,19 @@ void FMovieSceneSequenceInstance::RestoreState(class IMovieScenePlayer& Player)
 	}
 }
 
-
 void FMovieSceneSequenceInstance::Update( EMovieSceneUpdateData& UpdateData, class IMovieScenePlayer& Player )
+{
+	if(MovieSceneSequence.IsValid())
+	{
+		PreUpdate(Player);
+
+		UpdatePasses(UpdateData, Player);
+
+		PostUpdate(Player);
+	}
+}
+
+void FMovieSceneSequenceInstance::PreUpdate(class IMovieScenePlayer& Player)
 {
 	// Remove any stale runtime objects
 	TMap<FGuid, FMovieSceneObjectBindingInstance>::TIterator ObjectIt = ObjectBindingInstances.CreateIterator();
@@ -148,7 +176,7 @@ void FMovieSceneSequenceInstance::Update( EMovieSceneUpdateData& UpdateData, cla
 		FMovieSceneObjectBindingInstance& ObjectBindingInstance = ObjectIt.Value();
 		for (int32 ObjectIndex = 0; ObjectIndex < ObjectBindingInstance.RuntimeObjects.Num(); )
 		{
-			UObject* RuntimeObject = ObjectBindingInstance.RuntimeObjects[ObjectIndex];
+			UObject* RuntimeObject = ObjectBindingInstance.RuntimeObjects[ObjectIndex].Get();
 			if (RuntimeObject == nullptr || RuntimeObject->HasAnyFlags(RF_BeginDestroyed|RF_FinishDestroyed) || RuntimeObject->IsPendingKill())
 			{
 				ObjectBindingInstance.RuntimeObjects.RemoveAt(ObjectIndex);
@@ -161,55 +189,64 @@ void FMovieSceneSequenceInstance::Update( EMovieSceneUpdateData& UpdateData, cla
 	}
 
 	Player.GetSpawnRegister().PreUpdateSequenceInstance(*this, Player);
+}
 
+void FMovieSceneSequenceInstance::UpdatePasses(EMovieSceneUpdateData& UpdateData, class IMovieScenePlayer& Player)
+{
 	UpdateData.UpdatePass = MSUP_PreUpdate;
-	UpdateInternal(UpdateData, Player);
+	UpdatePassSingle(UpdateData, Player);
 	UpdateData.UpdatePass = MSUP_Update;
-	UpdateInternal(UpdateData, Player);
+	UpdatePassSingle(UpdateData, Player);
 	UpdateData.UpdatePass = MSUP_PostUpdate;
-	UpdateInternal(UpdateData, Player);
+	UpdatePassSingle(UpdateData, Player);
+}
 
+void FMovieSceneSequenceInstance::PostUpdate(class IMovieScenePlayer& Player)
+{
 	Player.GetSpawnRegister().PostUpdateSequenceInstance(*this, Player);
 }
 
-void FMovieSceneSequenceInstance::UpdateInternal( EMovieSceneUpdateData& UpdateData, class IMovieScenePlayer& Player )
+void FMovieSceneSequenceInstance::UpdatePassSingle( EMovieSceneUpdateData& UpdateData, class IMovieScenePlayer& Player )
 {
-	// Refresh time range so that spawnables can be created if they fall within the playback range, or destroyed if not
-	UMovieScene* MovieScene = MovieSceneSequence->GetMovieScene();
-	TimeRange = MovieScene->GetPlaybackRange();
+	if(MovieSceneSequence.IsValid())
+	{
+		// Refresh time range so that spawnables can be created if they fall within the playback range, or destroyed if not
+		UMovieScene* MovieScene = MovieSceneSequence->GetMovieScene();
+		TimeRange = MovieScene->GetPlaybackRange();
 
-	TArray<UObject*> NoObjects;
+		TArray<TWeakObjectPtr<UObject>> NoObjects;
 
-	// update each master track
-	for( FMovieSceneInstanceMap::TIterator It( MasterTrackInstances ); It; ++It )
-	{
-		if (It.Value()->HasUpdatePasses() & UpdateData.UpdatePass)
-		{
-			It.Value()->Update( UpdateData, NoObjects, Player, *this);
+		// update each master track
+		for( FMovieSceneInstanceMap::TIterator It( MasterTrackInstances ); It; ++It )
+		{		
+			if (It.Value()->HasUpdatePasses() & UpdateData.UpdatePass && It.Value() != CameraCutTrackInstance)
+			{
+				It.Value()->Update( UpdateData, NoObjects, Player, *this);
+			}
 		}
-	}
 
-	// update tracks bound to objects
-	TMap<FGuid, FMovieSceneObjectBindingInstance>::TIterator ObjectIt = ObjectBindingInstances.CreateIterator();
-	for(; ObjectIt; ++ObjectIt )
-	{
-		FMovieSceneObjectBindingInstance& ObjectBindingInstance = ObjectIt.Value();
-		
-		for( FMovieSceneInstanceMap::TIterator It = ObjectBindingInstance.TrackInstances.CreateIterator(); It; ++It )
+		// update tracks bound to objects
+		TMap<FGuid, FMovieSceneObjectBindingInstance>::TIterator ObjectIt = ObjectBindingInstances.CreateIterator();
+		for(; ObjectIt; ++ObjectIt )
 		{
-			if (It.Value()->HasUpdatePasses() & UpdateData.UpdatePass)
+			FMovieSceneObjectBindingInstance& ObjectBindingInstance = ObjectIt.Value();
+		
+			for( FMovieSceneInstanceMap::TIterator It = ObjectBindingInstance.TrackInstances.CreateIterator(); It; ++It )
 			{
-				It.Value()->Update( UpdateData, ObjectBindingInstance.RuntimeObjects, Player, *this );
+				if (It.Value()->HasUpdatePasses() & UpdateData.UpdatePass && It.Value() != CameraCutTrackInstance)
+				{
+					It.Value()->Update( UpdateData, ObjectBindingInstance.RuntimeObjects, Player, *this );
+				}
 			}
 		}
-	}
 
-	// update camera cut track last to make sure spawnable cameras are there, and to override sub-shots
-	if (CameraCutTrackInstance.IsValid())
-	{
-		if (CameraCutTrackInstance->HasUpdatePasses() & UpdateData.UpdatePass)
+		// update camera cut track last to make sure spawnable cameras are there, and to override sub-shots
+		if (CameraCutTrackInstance.IsValid())
 		{
-			CameraCutTrackInstance->Update(UpdateData, NoObjects, Player, *this );
+			if (CameraCutTrackInstance->HasUpdatePasses() & UpdateData.UpdatePass)
+			{
+				CameraCutTrackInstance->Update(UpdateData, NoObjects, Player, *this );
+			}
 		}
 	}
 }
@@ -217,73 +254,76 @@ void FMovieSceneSequenceInstance::UpdateInternal( EMovieSceneUpdateData& UpdateD
 
 void FMovieSceneSequenceInstance::RefreshInstance( IMovieScenePlayer& Player )
 {
-	UMovieScene* MovieScene = MovieSceneSequence->GetMovieScene();
-	TimeRange = MovieScene->GetPlaybackRange();
-
-	UMovieSceneTrack* CameraCutTrack = MovieScene->GetCameraCutTrack();
-
-	if (CameraCutTrack != nullptr)
+	if(MovieSceneSequence.IsValid())
 	{
-		FMovieSceneInstanceMap CameraCutTrackInstanceMap;
+		UMovieScene* MovieScene = MovieSceneSequence->GetMovieScene();
+		TimeRange = MovieScene->GetPlaybackRange();
 
-		if (CameraCutTrackInstance.IsValid())
+		UMovieSceneTrack* CameraCutTrack = MovieScene->GetCameraCutTrack();
+
+		if (CameraCutTrack != nullptr)
 		{
-			CameraCutTrackInstanceMap.Add(CameraCutTrack, CameraCutTrackInstance);
-		}
+			FMovieSceneInstanceMap CameraCutTrackInstanceMap;
 
-		TArray<UObject*> Objects;
-		TArray<UMovieSceneTrack*> Tracks;
-		Tracks.Add(CameraCutTrack);
-		RefreshInstanceMap(Tracks, Objects, CameraCutTrackInstanceMap, Player);
+			if (CameraCutTrackInstance.IsValid())
+			{
+				CameraCutTrackInstanceMap.Add(CameraCutTrack, CameraCutTrackInstance);
+			}
 
-		CameraCutTrackInstance = CameraCutTrackInstanceMap.FindRef(CameraCutTrack);
-	}
-	else if(CameraCutTrackInstance.IsValid())
-	{
-		CameraCutTrackInstance->ClearInstance(Player, *this);
-		CameraCutTrackInstance.Reset();
-	}
+			TArray<TWeakObjectPtr<UObject>> Objects;
+			TArray<UMovieSceneTrack*> Tracks;
+			Tracks.Add(CameraCutTrack);
+			RefreshInstanceMap(Tracks, Objects, CameraCutTrackInstanceMap, Player);
+
+			CameraCutTrackInstance = CameraCutTrackInstanceMap.FindRef(CameraCutTrack);
+		}
+		else if(CameraCutTrackInstance.IsValid())
+		{
+			CameraCutTrackInstance->ClearInstance(Player, *this);
+			CameraCutTrackInstance.Reset();
+		}
 
-	// Get all the master tracks and create instances for them if needed
-	const TArray<UMovieSceneTrack*>& MasterTracks = MovieScene->GetMasterTracks();
-	TArray<UObject*> Objects;
-	RefreshInstanceMap( MasterTracks, Objects, MasterTrackInstances, Player );
+		// Get all the master tracks and create instances for them if needed
+		const TArray<UMovieSceneTrack*>& MasterTracks = MovieScene->GetMasterTracks();
+		TArray<TWeakObjectPtr<UObject>> Objects;
+		RefreshInstanceMap( MasterTracks, Objects, MasterTrackInstances, Player );
 
-	TSet< FGuid > FoundObjectBindings;
-	// Get all tracks for each object binding and create instances for them if needed
-	const TArray<FMovieSceneBinding>& ObjectBindings = MovieScene->GetBindings();
-	for( int32 BindingIndex = 0; BindingIndex < ObjectBindings.Num(); ++BindingIndex )
-	{
-		const FMovieSceneBinding& ObjectBinding = ObjectBindings[BindingIndex];
+		TSet< FGuid > FoundObjectBindings;
+		// Get all tracks for each object binding and create instances for them if needed
+		const TArray<FMovieSceneBinding>& ObjectBindings = MovieScene->GetBindings();
+		for( int32 BindingIndex = 0; BindingIndex < ObjectBindings.Num(); ++BindingIndex )
+		{
+			const FMovieSceneBinding& ObjectBinding = ObjectBindings[BindingIndex];
 
-		// Create an instance for this object binding
-		FMovieSceneObjectBindingInstance& BindingInstance = ObjectBindingInstances.FindOrAdd( ObjectBinding.GetObjectGuid() );
-		BindingInstance.ObjectGuid = ObjectBinding.GetObjectGuid();
+			// Create an instance for this object binding
+			FMovieSceneObjectBindingInstance& BindingInstance = ObjectBindingInstances.FindOrAdd( ObjectBinding.GetObjectGuid() );
+			BindingInstance.ObjectGuid = ObjectBinding.GetObjectGuid();
 
-		FoundObjectBindings.Add( ObjectBinding.GetObjectGuid() );
+			FoundObjectBindings.Add( ObjectBinding.GetObjectGuid() );
 
-		// Populate the runtime objects for this instance of the binding.
-		// @todo sequencer: SubSequences: We need to know which actors were removed and which actors were added so we know which saved actor state to restore/create
-		BindingInstance.RuntimeObjects.Empty();
-		Player.GetRuntimeObjects( SharedThis( this ), BindingInstance.ObjectGuid, BindingInstance.RuntimeObjects );
+			// Populate the runtime objects for this instance of the binding.
+			// @todo sequencer: SubSequences: We need to know which actors were removed and which actors were added so we know which saved actor state to restore/create
+			BindingInstance.RuntimeObjects.Empty();
+			Player.GetRuntimeObjects( SharedThis( this ), BindingInstance.ObjectGuid, BindingInstance.RuntimeObjects );
 
-		// Refresh the instance's tracks
-		const TArray<UMovieSceneTrack*>& Tracks = ObjectBinding.GetTracks();
-		RefreshInstanceMap( Tracks, BindingInstance.RuntimeObjects, BindingInstance.TrackInstances, Player );
-	}
+			// Refresh the instance's tracks
+			const TArray<UMovieSceneTrack*>& Tracks = ObjectBinding.GetTracks();
+			RefreshInstanceMap( Tracks, BindingInstance.RuntimeObjects, BindingInstance.TrackInstances, Player );
+		}
 
-	IMovieSceneSpawnRegister& SpawnRegister = Player.GetSpawnRegister();
+		IMovieSceneSpawnRegister& SpawnRegister = Player.GetSpawnRegister();
 
-	// Remove object binding instances which are no longer bound
-	TMap<FGuid, FMovieSceneObjectBindingInstance>::TIterator It = ObjectBindingInstances.CreateIterator();
-	for( ; It; ++It )
-	{
-		if( !FoundObjectBindings.Contains( It.Key() ) )
+		// Remove object binding instances which are no longer bound
+		TMap<FGuid, FMovieSceneObjectBindingInstance>::TIterator It = ObjectBindingInstances.CreateIterator();
+		for( ; It; ++It )
 		{
-			SpawnRegister.DestroySpawnedObject(It.Key(), *this, Player);
+			if( !FoundObjectBindings.Contains( It.Key() ) )
+			{
+				SpawnRegister.DestroySpawnedObject(It.Key(), *this, Player);
 
-			// The instance no longer is bound to an existing guid
-			It.RemoveCurrent();
+				// The instance no longer is bound to an existing guid
+				It.RemoveCurrent();
+			}
 		}
 	}
 }
@@ -298,7 +338,7 @@ struct FTrackInstanceEvalSorter
 };
 
 
-void FMovieSceneSequenceInstance::RefreshInstanceMap( const TArray<UMovieSceneTrack*>& Tracks, const TArray<UObject*>& RuntimeObjects, FMovieSceneInstanceMap& TrackInstances, IMovieScenePlayer& Player  )
+void FMovieSceneSequenceInstance::RefreshInstanceMap( const TArray<UMovieSceneTrack*>& Tracks, const TArray<TWeakObjectPtr<UObject>>& RuntimeObjects, FMovieSceneInstanceMap& TrackInstances, IMovieScenePlayer& Player  )
 {
 	// All the tracks we found during this pass
 	TSet< UMovieSceneTrack* > FoundTracks;
@@ -350,89 +390,108 @@ void FMovieSceneSequenceInstance::RefreshInstanceMap( const TArray<UMovieSceneTr
 
 void FMovieSceneSequenceInstance::UpdateObjectBinding(const FGuid& ObjectId, IMovieScenePlayer& Player)
 {
-	UMovieSceneSequence* Sequence = MovieSceneSequence.Get();
-	auto* BindingInstance = ObjectBindingInstances.Find(ObjectId);
-
-	if (!BindingInstance || !Sequence)
+	if(MovieSceneSequence.IsValid())
 	{
-		return;
-	}
+		UMovieSceneSequence* Sequence = MovieSceneSequence.Get();
+		auto* BindingInstance = ObjectBindingInstances.Find(ObjectId);
 
-	// Update the runtime objects
-	BindingInstance->RuntimeObjects.Reset();
+		if (!BindingInstance || !Sequence)
+		{
+			return;
+		}
 
-	TWeakObjectPtr<UObject>* WeakSpawnedObject = SpawnedObjects.Find(ObjectId);
-	if (WeakSpawnedObject)
-	{
-		UObject* SpawnedObject = WeakSpawnedObject->Get();
-		if (SpawnedObject)
+		// Update the runtime objects
+		BindingInstance->RuntimeObjects.Reset();
+
+		TWeakObjectPtr<UObject>* WeakSpawnedObject = SpawnedObjects.Find(ObjectId);
+		if (WeakSpawnedObject)
 		{
-			BindingInstance->RuntimeObjects.Add(SpawnedObject);
+			UObject* SpawnedObject = WeakSpawnedObject->Get();
+			if (SpawnedObject)
+			{
+				BindingInstance->RuntimeObjects.Add(SpawnedObject);
+			}
+		}
+		else
+		{
+			Player.GetRuntimeObjects(SharedThis(this), BindingInstance->ObjectGuid, BindingInstance->RuntimeObjects);
 		}
-	}
-	else
-	{
-		Player.GetRuntimeObjects(SharedThis(this), BindingInstance->ObjectGuid, BindingInstance->RuntimeObjects);
-	}
 
-	const FMovieSceneBinding* ObjectBinding = Sequence->GetMovieScene()->GetBindings().FindByPredicate([&](const FMovieSceneBinding& In){
-		return In.GetObjectGuid() == ObjectId;
-	});
+		const FMovieSceneBinding* ObjectBinding = Sequence->GetMovieScene()->GetBindings().FindByPredicate([&](const FMovieSceneBinding& In){
+			return In.GetObjectGuid() == ObjectId;
+		});
 
-	// Refresh the instance map, if we found the binding itself
-	if (ObjectBinding)
-	{
-		RefreshInstanceMap(ObjectBinding->GetTracks(), BindingInstance->RuntimeObjects, BindingInstance->TrackInstances, Player);
+		// Refresh the instance map, if we found the binding itself
+		if (ObjectBinding)
+		{
+			RefreshInstanceMap(ObjectBinding->GetTracks(), BindingInstance->RuntimeObjects, BindingInstance->TrackInstances, Player);
+		}
 	}
 }
 
 void FMovieSceneSequenceInstance::OnObjectSpawned(const FGuid& ObjectId, UObject& SpawnedObject, IMovieScenePlayer& Player)
 {
-	auto* BindingInstance = ObjectBindingInstances.Find(ObjectId);
-
-	if (!BindingInstance)
+	if(MovieSceneSequence.IsValid())
 	{
-		return;
-	}
+		auto* BindingInstance = ObjectBindingInstances.Find(ObjectId);
+
+		if (!BindingInstance)
+		{
+			return;
+		}
 
-	SpawnedObjects.Add(ObjectId, &SpawnedObject);
+		SpawnedObjects.Add(ObjectId, &SpawnedObject);
 
-	// Add it to the instance's runtime objects array, and update any child possessable binding instances
-	BindingInstance->RuntimeObjects.Reset();
-	BindingInstance->RuntimeObjects.Emplace(&SpawnedObject);
+		// Add it to the instance's runtime objects array, and update any child possessable binding instances
+		BindingInstance->RuntimeObjects.Reset();
+		BindingInstance->RuntimeObjects.Emplace(&SpawnedObject);
 
-	UMovieSceneSequence* Sequence = GetSequence();
-	FMovieSceneSpawnable* Spawnable = Sequence ? Sequence->GetMovieScene()->FindSpawnable(ObjectId) : nullptr;
-	if (Spawnable)
-	{
-		for (const FGuid& Child : Spawnable->GetChildPossessables())
+		UMovieSceneSequence* Sequence = GetSequence();
+		FMovieSceneSpawnable* Spawnable = Sequence ? Sequence->GetMovieScene()->FindSpawnable(ObjectId) : nullptr;
+		if (Spawnable)
 		{
-			UpdateObjectBinding(Child, Player);
+			UpdateObjectBinding(ObjectId, Player);
+			for (const FGuid& Child : Spawnable->GetChildPossessables())
+			{
+				UpdateObjectBinding(Child, Player);
+			}
 		}
 	}
 }
 
 void FMovieSceneSequenceInstance::OnSpawnedObjectDestroyed(const FGuid& ObjectId, IMovieScenePlayer& Player)
 {
-	auto* BindingInstance = ObjectBindingInstances.Find(ObjectId);
-	if (!BindingInstance)
+	if(MovieSceneSequence.IsValid())
 	{
-		return;
-	}
+		auto* BindingInstance = ObjectBindingInstances.Find(ObjectId);
+		if (!BindingInstance)
+		{
+			return;
+		}
 
-	SpawnedObjects.Remove(ObjectId);
+		SpawnedObjects.Remove(ObjectId);
 
-	// Destroy the object
-	BindingInstance->RuntimeObjects.Reset();
+		// Destroy the object
+		BindingInstance->RuntimeObjects.Reset();
 
-	// Update any child possessable object bindings
-	UMovieSceneSequence* Sequence = GetSequence();
-	FMovieSceneSpawnable* Spawnable = Sequence ? Sequence->GetMovieScene()->FindSpawnable(ObjectId) : nullptr;
-	if (Spawnable)
-	{
-		for (const FGuid& Child : Spawnable->GetChildPossessables())
+		// Update any child possessable object bindings
+		UMovieSceneSequence* Sequence = GetSequence();
+		FMovieSceneSpawnable* Spawnable = Sequence ? Sequence->GetMovieScene()->FindSpawnable(ObjectId) : nullptr;
+		if (Spawnable)
 		{
-			UpdateObjectBinding(Child, Player);
+			for (const FGuid& Child : Spawnable->GetChildPossessables())
+			{
+				UpdateObjectBinding(Child, Player);
+			}
 		}
 	}
+}
+
+void FMovieSceneSequenceInstance::HandleSequenceSectionChanged(UMovieSceneSequence* Sequence)
+{
+	MovieSceneSequence = Sequence;
+	if(MovieSceneSequence.IsValid())
+	{
+		TimeRange = MovieSceneSequence->GetMovieScene()->GetPlaybackRange();
+	}
 }
\ No newline at end of file
diff --git a/Engine/Source/Runtime/MovieScene/Public/IMovieScenePlayer.h b/Engine/Source/Runtime/MovieScene/Public/IMovieScenePlayer.h
--- a/Engine/Source/Runtime/MovieScene/Public/IMovieScenePlayer.h
+++ b/Engine/Source/Runtime/MovieScene/Public/IMovieScenePlayer.h
@@ -3,6 +3,11 @@
 #pragma once
 
 #include "IMovieSceneSpawnRegister.h"
+#include "UnrealClient.h"
+
+
+class FMovieSceneSequenceInstance;
+
 
 namespace EMovieScenePlayerStatus
 {
@@ -13,10 +18,12 @@ namespace EMovieScenePlayerStatus
 		Recording,
 		Scrubbing,
 		Jumping,
+		Stepping,
 		MAX
 	};
 }
 
+
 struct EMovieSceneViewportParams
 {
 	EMovieSceneViewportParams()
@@ -42,7 +49,6 @@ struct EMovieSceneViewportParams
 	bool bEnableColorScaling;
 };
 
-class FMovieSceneSequenceInstance;
 
 /**
  * Interface for movie scene players
@@ -58,16 +64,16 @@ public:
 	 * @param ObjectHandle The handle to runtime objects
 	 * @Param The list of runtime objects that will be populated
 	 */
-	virtual void GetRuntimeObjects( TSharedRef<FMovieSceneSequenceInstance> MovieSceneInstance, const FGuid& ObjectHandle, TArray< UObject* >& OutObjects ) const = 0;
-
+	virtual void GetRuntimeObjects(TSharedRef<FMovieSceneSequenceInstance> MovieSceneInstance, const FGuid& ObjectHandle, TArray<TWeakObjectPtr<UObject>>& OutObjects) const = 0;
 
 	/**
 	 * Updates the perspective viewports with the actor to view through
 	 *
 	 * @param CameraObject The object, probably a camera, that the viewports should lock to
 	 * @param UnlockIfCameraObject If this is not nullptr, release actor lock only if currently locked to this object.
+	 * @param bJumpCut Whether this is a jump cut, ie. the cut jumps from one shot to another shot
 	 */
-	virtual void UpdateCameraCut(UObject* CameraObject, UObject* UnlockIfCameraObject = nullptr) const = 0;
+	virtual void UpdateCameraCut(UObject* CameraObject, UObject* UnlockIfCameraObject = nullptr, bool bJumpCut = false) const = 0;
 
 	/*
 	 * Set the perspective viewport settings
@@ -122,6 +128,11 @@ public:
 	 */
 	virtual UObject* GetPlaybackContext() const { return nullptr; }
 
+	/**
+	 * Test whether this is a preview player or not. As such, playback range becomes insignificant for things like spawnables
+	 */
+	virtual bool IsPreview() const { return false; }
+
 private:
 	/** Null register that asserts on use */
 	FNullMovieSceneSpawnRegister NullRegister;
diff --git a/Engine/Source/Runtime/MovieScene/Public/IMovieSceneSpawnRegister.h b/Engine/Source/Runtime/MovieScene/Public/IMovieSceneSpawnRegister.h
--- a/Engine/Source/Runtime/MovieScene/Public/IMovieSceneSpawnRegister.h
+++ b/Engine/Source/Runtime/MovieScene/Public/IMovieSceneSpawnRegister.h
@@ -3,11 +3,15 @@
 #pragma once
 
 #include "MovieSceneSequenceInstance.h"
+#include "MovieSceneSpawnable.h"
 #include "ValueOrError.h"
 
+ enum class ESpawnOwnership : uint8;
+
 class IMovieScenePlayer;
 class FMovieSceneSequenceInstance;
 class UMovieScene;
+struct FMovieSceneSpawnable;
 
 /** Struct used for defining a new spawnable type */
 struct FNewSpawnable
@@ -50,21 +54,12 @@ public:
 	virtual void DestroySpawnedObject(const FGuid& Object, FMovieSceneSequenceInstance& Instance, IMovieScenePlayer& Player) = 0;
 
 	/**
-	 * Destroy all objects spawned, and owned, by the specified sequence instance
-	 * @note this only destorys objects that the specified sequence instance has direct ownership of (ownership == InnerSequence)
-	 *
-	 * @param Instance 		Sequence instance that is destroying the objects
-	 * @param Player 		Movie scene player that is ultimately responsible for destroying the objects
-	 */
-	virtual void DestroyObjectsOwnedByInstance(FMovieSceneSequenceInstance& Instance, IMovieScenePlayer& Player) = 0;
-
-	/**
-	 * Destroy all objects that have ever been spawned by the specified sequence instance (regardless of ownership)
+	 * Destroy spawned objects using a custom predicate
 	 *
-	 * @param Instance 		Sequence instance that is destroying the objects
 	 * @param Player 		Movie scene player that is ultimately responsible for destroying the objects
+	 * @param Predicate		Predicate used for testing whether an object should be destroyed. Returns true for destruction, false to skip.
 	 */
-	virtual void DestroyObjectsSpawnedByInstance(FMovieSceneSequenceInstance& Instance, IMovieScenePlayer& Player) = 0;
+	virtual void DestroyObjectsByPredicate(IMovieScenePlayer& Player, const TFunctionRef<bool(const FGuid&, ESpawnOwnership, FMovieSceneSequenceInstance&)>& Predicate) = 0;
 
 	/**
 	 * Purge any memory of any objects that are considered externally owned
@@ -73,20 +68,6 @@ public:
 	 */
 	virtual void ForgetExternallyOwnedSpawnedObjects(IMovieScenePlayer& Player) = 0;
 
-	/**
-	 * Destroy all objects that this spawn register has responsibility for (Ownership != External)
-	 *
-	 * @param Player 		Movie scene player that is ultimately responsible for destroying the objects
-	 */
-	virtual void DestroyAllOwnedObjects(IMovieScenePlayer& Player) = 0;
-
-	/**
-	 * Destroy all objects that we have ever spawned
-	 *
-	 * @param Player 		Movie scene player that is ultimately responsible for destroying the objects
-	 */
-	virtual void DestroyAllObjects(IMovieScenePlayer& Player) = 0;
-
 	/**
 	 * Called before a sequence instance is about to update
 	 *
@@ -109,24 +90,79 @@ public:
 	 *
 	 * @param SourceObject		The source object to create the spawnable from
 	 * @param OwnerMovieScene	The owner movie scene that this spawnable type should reside in
-	 * @return the new spawnable type, or nullptr
+	 * @return the new spawnable type, or error
 	 */
 	virtual TValueOrError<FNewSpawnable, FText> CreateNewSpawnableType(UObject& SourceObject, UMovieScene& OwnerMovieScene) { return MakeError(NSLOCTEXT("SpawnRegister", "NotSupported", "Not supported")); }
+
+	/**
+	 * Called to test whether an expired object can be destroyed or not
+	 *
+	 * @param Object 			ID of the object to destroy
+	 * @param Ownership 		The ownership of the object
+	 * @param SequenceInstance	The sequence instance that spawned this object
+	 * @return true if it can be destroyed, false otherwise
+	 */
+	virtual bool CanDestroyExpiredObject(const FGuid& ObjectId, ESpawnOwnership Ownership, FMovieSceneSequenceInstance& SequenceInstance) { return true; }
+
 #endif
-};
 
+public:
+
+	/**
+	 * Called to indiscriminately clean up any spawned objects
+	 */
+	void CleanUp(IMovieScenePlayer& Player)
+	{
+		DestroyObjectsByPredicate(Player, [&](const FGuid&, ESpawnOwnership, FMovieSceneSequenceInstance&){
+			return true;
+		});
+	}
+
+	/**
+	 * Called to clean up any non-externally owned spawnables that were spawned from the specified instance
+	 */
+	void CleanUpSequence(FMovieSceneSequenceInstance& Instance, IMovieScenePlayer& Player)
+	{
+		DestroyObjectsByPredicate(Player, [&](const FGuid&, ESpawnOwnership, FMovieSceneSequenceInstance& SequenceInstance){
+			return &Instance == &SequenceInstance;
+		});
+	}
+
+	/**
+	 * Called when the current time has moved beyond the specified sequence's play range
+	 */
+	void OnSequenceExpired(FMovieSceneSequenceInstance& Instance, IMovieScenePlayer& Player)
+	{
+		DestroyObjectsByPredicate(Player, [&](const FGuid& ObjectId, ESpawnOwnership Ownership, FMovieSceneSequenceInstance& SequenceInstance){
+			return	(Ownership == ESpawnOwnership::InnerSequence)
+				&&	(&Instance == &SequenceInstance)
+#if WITH_EDITOR
+				&&	CanDestroyExpiredObject(ObjectId, Ownership, SequenceInstance)
+#endif
+				;
+		});
+	}
+
+#if WITH_EDITOR
+	/**
+	 * Called when the current time has moved beyond the master sequence's play range
+	 */
+	void OnMasterSequenceExpired(IMovieScenePlayer& Player)
+	{
+		DestroyObjectsByPredicate(Player, [&](const FGuid& ObjectId, ESpawnOwnership Ownership, FMovieSceneSequenceInstance& SequenceInstance){
+			return Ownership != ESpawnOwnership::External && CanDestroyExpiredObject(ObjectId, Ownership, SequenceInstance);
+		});
+	}
+#endif
+};
 
 class FNullMovieSceneSpawnRegister : public IMovieSceneSpawnRegister
 {
 public:
-	// These 2 functions ensure because they should never be called for sequences that don't support spawnables
 	virtual UObject* SpawnObject(const FGuid& Object, FMovieSceneSequenceInstance&, IMovieScenePlayer&) override { check(false); return nullptr; }
 	virtual void DestroySpawnedObject(const FGuid& Object, FMovieSceneSequenceInstance&, IMovieScenePlayer&) override { }
-	virtual void DestroyObjectsOwnedByInstance(FMovieSceneSequenceInstance&, IMovieScenePlayer&) override { }
-	virtual void DestroyObjectsSpawnedByInstance(FMovieSceneSequenceInstance&, IMovieScenePlayer&) override { }
-	virtual void DestroyAllOwnedObjects(IMovieScenePlayer&) override { }
+	virtual void DestroyObjectsByPredicate(IMovieScenePlayer& Player, const TFunctionRef<bool(const FGuid&, ESpawnOwnership, FMovieSceneSequenceInstance&)>& Predicate) { }
 	virtual void ForgetExternallyOwnedSpawnedObjects(IMovieScenePlayer& Player) override { }
-	virtual void DestroyAllObjects(IMovieScenePlayer&) override { }
 };
 
 /**
diff --git a/Engine/Source/Runtime/MovieScene/Public/IMovieSceneTrackInstance.h b/Engine/Source/Runtime/MovieScene/Public/IMovieSceneTrackInstance.h
--- a/Engine/Source/Runtime/MovieScene/Public/IMovieSceneTrackInstance.h
+++ b/Engine/Source/Runtime/MovieScene/Public/IMovieSceneTrackInstance.h
@@ -10,7 +10,8 @@ enum EMovieSceneUpdatePass
 {
 	MSUP_PreUpdate = 0x00000001,
 	MSUP_Update = 0x00000002,
-	MSUP_PostUpdate = 0x00000004
+	MSUP_PostUpdate = 0x00000004,
+	MSUP_All = MSUP_PreUpdate | MSUP_Update | MSUP_PostUpdate
 };
 
 struct EMovieSceneUpdateData
@@ -53,13 +54,13 @@ public:
 	 * Save state of objects that this instance will be editing.
 	 * @todo Sequencer: This is likely editor only
 	 */
-	virtual void SaveState(const TArray<UObject*>& RuntimeObjects, IMovieScenePlayer& Player, FMovieSceneSequenceInstance& SequenceInstance) = 0;
+	virtual void SaveState(const TArray<TWeakObjectPtr<UObject>>& RuntimeObjects, IMovieScenePlayer& Player, FMovieSceneSequenceInstance& SequenceInstance) = 0;
 
 	/**
 	 * Restore state of objects that this instance edited.
 	 * @todo Sequencer: This is likely editor only
 	 */
-	virtual void RestoreState(const TArray<UObject*>& RuntimeObjects, IMovieScenePlayer& Player, FMovieSceneSequenceInstance& SequenceInstance) = 0;
+	virtual void RestoreState(const TArray<TWeakObjectPtr<UObject>>& RuntimeObjects, IMovieScenePlayer& Player, FMovieSceneSequenceInstance& SequenceInstance) = 0;
 
 	/**
 	 * Main update function for track instances.  Called in game and in editor when we update a moviescene.
@@ -68,7 +69,7 @@ public:
 	 * @param RuntimeObjects	Runtime objects bound to this instance (if any)
 	 * @param Player			The playback interface.  Contains state and some other functionality for runtime playback
 	 */
-	virtual void Update(EMovieSceneUpdateData& UpdateData, const TArray<UObject*>& RuntimeObjects, IMovieScenePlayer& Player, FMovieSceneSequenceInstance& SequenceInstance) = 0;
+	virtual void Update(EMovieSceneUpdateData& UpdateData, const TArray<TWeakObjectPtr<UObject>>& RuntimeObjects, IMovieScenePlayer& Player, FMovieSceneSequenceInstance& SequenceInstance) = 0;
 	
 	/*
 	 * Which update passes does this track instance evaluate in?
@@ -78,7 +79,7 @@ public:
 	/**
 	 * Refreshes the current instance
 	 */
-	virtual void RefreshInstance(const TArray<UObject*>& RuntimeObjects, IMovieScenePlayer& Player, FMovieSceneSequenceInstance& SequenceInstance) = 0;
+	virtual void RefreshInstance(const TArray<TWeakObjectPtr<UObject>>& RuntimeObjects, IMovieScenePlayer& Player, FMovieSceneSequenceInstance& SequenceInstance) = 0;
 
 	/**
 	 * Removes all instance data from this track instance
diff --git a/Engine/Source/Runtime/MovieScene/Public/MovieScene.h b/Engine/Source/Runtime/MovieScene/Public/MovieScene.h
--- a/Engine/Source/Runtime/MovieScene/Public/MovieScene.h
+++ b/Engine/Source/Runtime/MovieScene/Public/MovieScene.h
@@ -11,6 +11,7 @@
 class UBlueprint;
 class UMovieSceneSection;
 class UMovieSceneTrack;
+class UMovieSceneFolder;
 
 
 MOVIESCENE_API DECLARE_LOG_CATEGORY_EXTERN(LogMovieScene, Log, All);
@@ -257,6 +258,15 @@ public:
 	 */
 	bool RemoveTrack(UMovieSceneTrack& Track);
 
+	/**
+	 * Find a track binding Guid from a UMovieSceneTrack
+	 * 
+	 * @param	InTrack		The track to find the binding for.
+	 * @param	OutGuid		The binding's Guid if one was found.
+	 * @return true if a binding was found for this track.
+	 */
+	bool FindTrackBinding(const UMovieSceneTrack& InTrack, FGuid& OutGuid) const;
+
 public:
 
 	/**
@@ -372,6 +382,12 @@ public:
 		return ObjectBindings;
 	}
 
+	/** Get the current in/out range. */
+	TRange<float> GetInOutRange() const
+	{
+		return InOutRange;
+	}
+
 	/**
 	 * Get the display name of the object with the specified identifier.
 	 *
@@ -380,10 +396,11 @@ public:
 	 */
 	FText GetObjectDisplayName(const FGuid& ObjectId);
 
-	/**
-	 * @return The playback time range of this movie scene, relative to its 0-time offset.
-	 */
-	TRange<float> GetPlaybackRange() const;
+	/** Get the playback time range of this movie scene, relative to its 0-time offset. */
+	TRange<float> GetPlaybackRange() const
+	{
+		return PlaybackRange;
+	}
 
 	/*
 	* Replace an existing binding with another 
@@ -398,6 +415,12 @@ public:
 		return ObjectsToLabels;
 	}
 
+	/** Set the in/out selection range. */
+	void SetInOutRange(TRange<float> Range)
+	{
+		InOutRange = Range;
+	}
+
 	/**
 	 * Get the display name of the object with the specified identifier.
 	 *
@@ -405,6 +428,11 @@ public:
 	 * @result The object's display name.
 	 */
 	void SetObjectDisplayName(const FGuid& ObjectId, const FText& DisplayName);
+
+	/**
+	 * Gets the root folders for this movie scene.
+	 */
+	TArray<UMovieSceneFolder*>& GetRootFolders();
 #endif
 
 	/**
@@ -412,8 +440,10 @@ public:
 	 *
 	 * @param Start The offset from 0-time to start playback of this movie scene
 	 * @param End The offset from 0-time to end playback of this movie scene
+	 * @param bAlwaysMarkDirty Whether to always mark the playback range dirty when changing it. 
+	 *        In the case where the playback range is dynamic and based on section bounds, the playback range doesn't need to be dirtied when set
 	 */
-	void SetPlaybackRange(float Start, float End);
+	void SetPlaybackRange(float Start, float End, bool bAlwaysMarkDirty = true);
 
 public:
 	
@@ -473,6 +503,14 @@ private:
 	UPROPERTY()
 	UMovieSceneTrack* CameraCutTrack;
 
+	/** User-defined in/out selection range. */
+	UPROPERTY()
+	FFloatRange InOutRange;
+
+	/** User-defined playback range for this movie scene. Must be a finite range. Relative to this movie-scene's 0-time origin. */
+	UPROPERTY()
+	FFloatRange PlaybackRange;
+
 #if WITH_EDITORONLY_DATA
 	/** Maps object GUIDs to user defined display names. */
 	UPROPERTY()
@@ -485,13 +523,13 @@ private:
 	/** Editor only data that needs to be saved between sessions for editing but has no runtime purpose */
 	UPROPERTY()
 	FMovieSceneEditorData EditorData;
+
+	/** The root folders for this movie scene. */
+	UPROPERTY()
+	TArray<UMovieSceneFolder*> RootFolders;
 #endif
 
 private:
-
-	/** User-defined playback range for this movie scene. Must be a finite range. Relative to this movie-scene's 0-time origin. */
-	UPROPERTY()
-	FFloatRange PlaybackRange;
 	
 	UPROPERTY()
 	float InTime_DEPRECATED;
diff --git a/Engine/Source/Runtime/MovieScene/Public/MovieSceneCommonHelpers.h b/Engine/Source/Runtime/MovieScene/Public/MovieSceneCommonHelpers.h
--- a/Engine/Source/Runtime/MovieScene/Public/MovieSceneCommonHelpers.h
+++ b/Engine/Source/Runtime/MovieScene/Public/MovieSceneCommonHelpers.h
@@ -64,7 +64,15 @@ public:
 	 */
 	static UCameraComponent* CameraComponentFromActor(const AActor* InActor);
 
-	/*
+	/**
+	 * Find and return camera component from the runtime object
+	 *
+	 * @param Object The object to get the camera component for
+	 * @return The found camera component
+	 */	
+	static UCameraComponent* CameraComponentFromRuntimeObject(UObject* RuntimeObject);
+
+	/**
 	 * Set the runtime object movable
 	 *
 	 * @param Object The object to set the mobility for
@@ -118,7 +126,7 @@ public:
 	 *
 	 * @param InRuntimeObjects	The objects to rebuild mappings for
 	 */
-	void UpdateBindings( const TArray<UObject*>& InRuntimeObjects );
+	void UpdateBindings( const TArray<TWeakObjectPtr<UObject>>& InRuntimeObjects );
 
 	/**
 	 * Gets the UProperty that is bound to the track instance
diff --git a/Engine/Source/Runtime/MovieScene/Public/MovieSceneFolder.h b/Engine/Source/Runtime/MovieScene/Public/MovieSceneFolder.h
new file mode 100644
--- /dev/null
+++ b/Engine/Source/Runtime/MovieScene/Public/MovieSceneFolder.h
@@ -0,0 +1,72 @@
+// Copyright 1998-2016 Epic Games, Inc. All Rights Reserved.
+
+#pragma once
+
+
+#include "MovieSceneFolder.generated.h"
+
+/** Reprents a folder used for organizing objects in tracks in a movie scene. */
+UCLASS()
+class MOVIESCENE_API UMovieSceneFolder : public UObject
+{
+	GENERATED_UCLASS_BODY()
+
+	/** Gets the name of this folder. */
+	FName GetFolderName() const;
+
+	/** Sets the name of this folder. */
+	void SetFolderName( FName InFolderName );
+
+	/** Gets the folders contained by this folder. */
+	const TArray<UMovieSceneFolder*>& GetChildFolders() const;
+
+	/** Adds a child folder to this folder. */
+	void AddChildFolder( UMovieSceneFolder* InChildFolder );
+
+	/** Removes a child folder from this folder. */
+	void RemoveChildFolder( UMovieSceneFolder* InChildFolder );
+
+	/** Gets the master tracks contained by this folder. */
+	const TArray<UMovieSceneTrack*>& GetChildMasterTracks() const;
+
+	/** Adds a master track to this folder. */
+	void AddChildMasterTrack( UMovieSceneTrack* InMasterTrack );
+
+	/** Removes a master track from this folder. */
+	void RemoveChildMasterTrack( UMovieSceneTrack* InMasterTrack );
+
+	/** Gets the guids for the object bindings contained by this folder. */
+	const TArray<FGuid>& GetChildObjectBindings() const;
+
+	/** Adds a guid for an object binding to this folder. */
+	void AddChildObjectBinding(const FGuid& InObjectBinding );
+
+	/** Removes a guid for an object binding from this folder. */
+	void RemoveChildObjectBinding( const FGuid& InObjectBinding );
+
+
+	virtual void PreSave() override;
+	virtual void PostLoad() override;
+
+
+private:
+	/** The name of this folder. */
+	UPROPERTY()
+	FName FolderName;
+
+	/** The folders contained by this folder. */
+	UPROPERTY()
+	TArray<UMovieSceneFolder*> ChildFolders;
+
+	/** The master tracks contained by this folder. */
+	UPROPERTY()
+	TArray<UMovieSceneTrack*> ChildMasterTracks;
+
+	/** The guid strings used to serialize the guids for the object bindings contained by this folder. */
+	UPROPERTY()
+	TArray<FString> ChildObjectBindingStrings;
+
+	/** The guids for the object bindings contained by this folder. */
+	TArray<FGuid> ChildObjectBindings;
+};
+
diff --git a/Engine/Source/Runtime/MovieScene/Public/MovieSceneNameableTrack.h b/Engine/Source/Runtime/MovieScene/Public/MovieSceneNameableTrack.h
--- a/Engine/Source/Runtime/MovieScene/Public/MovieSceneNameableTrack.h
+++ b/Engine/Source/Runtime/MovieScene/Public/MovieSceneNameableTrack.h
@@ -25,6 +25,13 @@ public:
 	 */
 	MOVIESCENE_API void SetDisplayName(const FText& NewDisplayName);
 
+	/**
+	 * Can rename this track.
+	 *
+	 * @return Whether this track can be renamed.
+	 */
+	MOVIESCENE_API virtual bool CanRename() const { return true; }
+
 public:
 
 	// UMovieSceneTrack interface
diff --git a/Engine/Source/Runtime/MovieScene/Public/MovieSceneSequenceInstance.h b/Engine/Source/Runtime/MovieScene/Public/MovieSceneSequenceInstance.h
--- a/Engine/Source/Runtime/MovieScene/Public/MovieSceneSequenceInstance.h
+++ b/Engine/Source/Runtime/MovieScene/Public/MovieSceneSequenceInstance.h
@@ -28,6 +28,14 @@ public:
 	 * @param InMovieSceneSequence The sequence that this instance represents.
 	 */
 	MOVIESCENE_API FMovieSceneSequenceInstance(const UMovieSceneSequence& InMovieSceneSequence);
+
+	/** 
+	 * Creates and initializes a new instance that is primed for recording 
+	 * 
+	 * @param InMovieSceneSubTrack The subtrack we are primed for recording in
+	 */
+	MOVIESCENE_API FMovieSceneSequenceInstance(const UMovieSceneTrack& InMovieSceneSubTrack);
+
 	MOVIESCENE_API ~FMovieSceneSequenceInstance();
 
 	/**
@@ -71,11 +79,31 @@ public:
 	/**
 	 * Updates this movie scene.
 	 *
-	 * @param UpdateData The current and previous position of the moviescene that is playing. The update pass.
+	 * @param UpdateData The current and previous position of the moviescene that is playing.
 	 * @param Player Movie scene player interface for interaction with runtime data.
 	 */
 	MOVIESCENE_API void Update(EMovieSceneUpdateData& UpdateData, IMovieScenePlayer& Player);
 
+	/*
+	 * Pre update - spawnables, any stale runtime objects.
+	 */
+	MOVIESCENE_API void PreUpdate(IMovieScenePlayer& Player);
+	
+	/*
+	 * Update all the passes
+	 */
+	MOVIESCENE_API void UpdatePasses(EMovieSceneUpdateData& UpdateData, IMovieScenePlayer& Player);
+
+	/*
+	 * Post update - clean up spawnables.
+	 */
+	MOVIESCENE_API void PostUpdate(IMovieScenePlayer& Player);
+
+	/*
+	 * Update a single pass
+	 */
+	MOVIESCENE_API void UpdatePassSingle(EMovieSceneUpdateData& UpdateData, IMovieScenePlayer& Player);
+
 	/**
 	 * Refreshes the existing instance.
 	 *
@@ -124,11 +152,12 @@ public:
 		return TimeRange;
 	}
 
-protected:
+	/** Handle the sequence this instance refers to changing */
+	MOVIESCENE_API void HandleSequenceSectionChanged(UMovieSceneSequence* Sequence);
 
-	void RefreshInstanceMap(const TArray<UMovieSceneTrack*>& Tracks, const TArray<UObject*>& RuntimeObjects, FMovieSceneInstanceMap& TrackInstances, IMovieScenePlayer& Player);
+protected:
 
-	void UpdateInternal(EMovieSceneUpdateData& UpdateData, IMovieScenePlayer& Player);
+	void RefreshInstanceMap(const TArray<UMovieSceneTrack*>& Tracks, const TArray<TWeakObjectPtr<UObject>>& RuntimeObjects, FMovieSceneInstanceMap& TrackInstances, IMovieScenePlayer& Player);
 
 	/** Update the object binding instance for the specified object */
 	void UpdateObjectBinding(const FGuid& ObjectId, IMovieScenePlayer& Player);
@@ -142,7 +171,7 @@ private:
 		FGuid ObjectGuid;
 
 		/** Actual runtime objects */	
-		TArray<UObject*> RuntimeObjects;
+		TArray<TWeakObjectPtr<UObject>> RuntimeObjects;
 
 		/** Instances that animate the runtime objects */
 		FMovieSceneInstanceMap TrackInstances;
@@ -150,7 +179,7 @@ private:
 
 
 	/** MovieScene that is instanced */
-	const TWeakObjectPtr<UMovieSceneSequence> MovieSceneSequence;
+	TWeakObjectPtr<UMovieSceneSequence> MovieSceneSequence;
 
 	/** The camera cut track instance map */
 	TSharedPtr<IMovieSceneTrackInstance> CameraCutTrackInstance;
diff --git a/Engine/Source/Runtime/MovieScene/Public/MovieSceneSpawnable.h b/Engine/Source/Runtime/MovieScene/Public/MovieSceneSpawnable.h
--- a/Engine/Source/Runtime/MovieScene/Public/MovieSceneSpawnable.h
+++ b/Engine/Source/Runtime/MovieScene/Public/MovieSceneSpawnable.h
@@ -36,6 +36,9 @@ public:
 		, Name(InitName)
 		, GeneratedClass(InitClass)
 		, Ownership(ESpawnOwnership::InnerSequence)
+#if WITH_EDITORONLY_DATA
+		, bIgnoreOwnershipInEditor(false)
+#endif
 	{ }
 
 public:
@@ -128,6 +131,24 @@ public:
 		Ownership = InOwnership;
 	}
 
+#if WITH_EDITORONLY_DATA
+	/**
+	 * Check whether this spawnable should remain spawned when outside the play-range, regardless of ownership
+	 */
+	bool ShouldIgnoreOwnershipInEditor() const
+	{
+		return bIgnoreOwnershipInEditor;
+	}
+
+	/**
+	 * Set whether this spawnable should remain spawned when outside the play-range, regardless of ownership
+	 */
+	void SetIgnoreOwnershipInEditor(bool bInIgnoreOwnershipInEditor)
+	{
+		bIgnoreOwnershipInEditor = bInIgnoreOwnershipInEditor;
+	}
+#endif
+
 private:
 
 	/** Unique identifier of the spawnable object. */
@@ -154,4 +175,10 @@ private:
 	/** Property indicating where ownership responsibility for this object lies */
 	UPROPERTY()
 	ESpawnOwnership Ownership;
+
+#if WITH_EDITORONLY_DATA
+	/** When true, this spawnable will not respect its ownership sematics outside of the playback range, when spawned from inside a currently editing sequence */
+	UPROPERTY()
+	bool bIgnoreOwnershipInEditor;
+#endif
 };
diff --git a/Engine/Source/Runtime/MovieSceneCapture/Private/MovieSceneCapture.cpp b/Engine/Source/Runtime/MovieSceneCapture/Private/MovieSceneCapture.cpp
--- a/Engine/Source/Runtime/MovieSceneCapture/Private/MovieSceneCapture.cpp
+++ b/Engine/Source/Runtime/MovieSceneCapture/Private/MovieSceneCapture.cpp
@@ -36,7 +36,7 @@ FMovieSceneCaptureSettings::FMovieSceneCaptureSettings()
 	GameModeOverride = nullptr;
 	OutputFormat = TEXT("{world}");
 	FrameRate = 24;
-	bEnableTextureStreaming = true;
+	bEnableTextureStreaming = false;
 	bCinematicMode = true;
 	bAllowMovement = false;
 	bAllowTurning = false;
diff --git a/Engine/Source/Runtime/MovieSceneCapture/Public/MovieSceneCapture.h b/Engine/Source/Runtime/MovieSceneCapture/Public/MovieSceneCapture.h
--- a/Engine/Source/Runtime/MovieSceneCapture/Public/MovieSceneCapture.h
+++ b/Engine/Source/Runtime/MovieSceneCapture/Public/MovieSceneCapture.h
@@ -49,7 +49,7 @@ public:
 public:
 
 	/** The type of capture protocol to use */
-	UPROPERTY(config, EditAnywhere, Category=CaptureSettings)
+	UPROPERTY(config, EditAnywhere, Category=CaptureSettings, DisplayName="Output Format")
 	FCaptureProtocolID CaptureType;
 
 	/** Settings specific to the capture protocol */
