Copying //UE4/Dev-Sequencer to //UE4/Dev-Main (Source: //UE4/Dev-Sequencer @ 3178529)

diff --git a/Engine/Source/Runtime/MovieScene/Private/Compilation/MovieSceneEvaluationTemplateGenerator.cpp b/Engine/Source/Runtime/MovieScene/Private/Compilation/MovieSceneEvaluationTemplateGenerator.cpp
new file mode 100644
--- /dev/null
+++ b/Engine/Source/Runtime/MovieScene/Private/Compilation/MovieSceneEvaluationTemplateGenerator.cpp
@@ -0,0 +1,346 @@
+// Copyright 1998-2016 Epic Games, Inc. All Rights Reserved.
+
+#include "MovieScenePrivatePCH.h"
+#include "MovieSceneEvaluationTemplateGenerator.h"
+
+namespace 
+{
+	void AddPtrsToGroup(FMovieSceneEvaluationGroup& Group, bool bRequiresImmediateFlush, TArray<FMovieSceneEvaluationFieldSegmentPtr>& InitPtrs, TArray<FMovieSceneEvaluationFieldSegmentPtr>& EvalPtrs)
+	{
+		if (!InitPtrs.Num() && !EvalPtrs.Num())
+		{
+			return;
+		}
+
+		FMovieSceneEvaluationGroupLUTIndex Index;
+		
+		Index.bRequiresImmediateFlush = bRequiresImmediateFlush;
+
+		Index.LUTOffset = Group.SegmentPtrLUT.Num();
+		Index.NumInitPtrs = InitPtrs.Num();
+		Index.NumEvalPtrs = EvalPtrs.Num();
+
+		Group.LUTIndices.Add(Index);
+		Group.SegmentPtrLUT.Append(InitPtrs);
+		Group.SegmentPtrLUT.Append(EvalPtrs);
+
+		InitPtrs.Reset();
+		EvalPtrs.Reset();
+	}
+}
+
+FMovieSceneEvaluationTemplateGenerator::FMovieSceneEvaluationTemplateGenerator(UMovieSceneSequence& InSequence, FMovieSceneEvaluationTemplate& OutTemplate, FMovieSceneSequenceTemplateStore& InStore)
+	: SourceSequence(InSequence)
+	, Template(OutTemplate)
+	, TransientArgs(*this, InStore)
+{
+}
+
+void FMovieSceneEvaluationTemplateGenerator::AddLegacyTrack(FMovieSceneEvaluationTrack&& InTrackTemplate, const UMovieSceneTrack& SourceTrack)
+{
+	AddOwnedTrack(MoveTemp(InTrackTemplate), SourceTrack);
+	Template.bHasLegacyTrackInstances = true;
+}
+
+void FMovieSceneEvaluationTemplateGenerator::AddOwnedTrack(FMovieSceneEvaluationTrack&& InTrackTemplate, const UMovieSceneTrack& SourceTrack)
+{
+	// Add the track to the template
+	Template.AddTrack(SourceTrack.GetSignature(), MoveTemp(InTrackTemplate));
+	CompiledSignatures.Add(SourceTrack.GetSignature());
+}
+
+void FMovieSceneEvaluationTemplateGenerator::AddSharedTrack(FMovieSceneEvaluationTrack&& InTrackTemplate, FMovieSceneSharedDataId SharedId, const UMovieSceneTrack& SourceTrack)
+{
+	FSharedPersistentDataKey Key(SharedId, FMovieSceneEvaluationOperand(MovieSceneSequenceID::Root, InTrackTemplate.GetObjectBindingID()));
+
+	if (AddedSharedTracks.Contains(Key))
+	{
+		return;
+	}
+
+	AddedSharedTracks.Add(Key);
+	AddOwnedTrack(MoveTemp(InTrackTemplate), SourceTrack);
+}
+
+void FMovieSceneEvaluationTemplateGenerator::AddExternalSegments(TRange<float> RootRange, TArrayView<const FMovieSceneEvaluationFieldSegmentPtr> SegmentPtrs)
+{
+	if (RootRange.IsEmpty())
+	{
+		return;
+	}
+
+	SegmentData.Reserve(SegmentData.Num() + SegmentPtrs.Num());
+
+	for (const FMovieSceneEvaluationFieldSegmentPtr& SegmentPtr : SegmentPtrs)
+	{
+		if (!ExternalSegmentLookup.Contains(SegmentPtr))
+		{
+			ExternalSegmentLookup.Add(SegmentPtr, TrackLUT.Num());
+			TrackLUT.Add(SegmentPtr);
+		}
+
+		SegmentData.Add(FMovieSceneSectionData(RootRange, ExternalSegmentLookup.FindRef(SegmentPtr)));
+	}
+}
+
+FMovieSceneSequenceTransform FMovieSceneEvaluationTemplateGenerator::GetSequenceTransform(FMovieSceneSequenceIDRef InSequenceID) const
+{
+	if (InSequenceID == MovieSceneSequenceID::Root)
+	{
+		return FMovieSceneSequenceTransform();
+	}
+
+	const FMovieSceneSubSequenceData* Data = Template.Hierarchy.FindSubData(InSequenceID);
+	return ensure(Data) ? Data->RootToSequenceTransform : FMovieSceneSequenceTransform();
+}
+
+FMovieSceneSequenceID FMovieSceneEvaluationTemplateGenerator::GenerateSequenceID(FMovieSceneSubSequenceData SequenceData, FMovieSceneSequenceIDRef ParentID)
+{
+	FMovieSceneSequenceHierarchyNode* ParentNode = Template.Hierarchy.FindNode(ParentID);
+	checkf(ParentNode, TEXT("Cannot generate a sequence ID for a ParentID that doesn't yet exist"));
+
+	FMovieSceneSequenceID ThisID = SequenceData.DeterministicSequenceID;
+
+	if (const FMovieSceneSubSequenceData* ParentSubSequenceData = Template.Hierarchy.FindSubData(ParentID))
+	{
+#if WITH_EDITORONLY_DATA
+		// Clamp this sequence's valid play range by its parent's valid play range
+		TRange<float> ParentPlayRangeChildSpace = ParentSubSequenceData->ValidPlayRange * (SequenceData.RootToSequenceTransform * ParentSubSequenceData->RootToSequenceTransform.Inverse());
+		SequenceData.ValidPlayRange = TRange<float>::Intersection(ParentPlayRangeChildSpace, SequenceData.ValidPlayRange);
+#endif
+		// Determine its ID from its parent's
+		ThisID = SequenceData.DeterministicSequenceID.AccumulateParentID(ParentSubSequenceData->DeterministicSequenceID);
+	}
+
+	// Ensure we have a unique ID. This should never happen in reality.
+	while(!ensureMsgf(!Template.Hierarchy.FindNode(ThisID), TEXT("CRC collision on deterministic hashes. Manually hashing a random new one.")))
+	{
+		ThisID = ThisID.AccumulateParentID(ThisID);
+	}
+
+	SequenceData.DeterministicSequenceID = ThisID;
+
+	ParentNode->Children.Add(ThisID);
+	Template.Hierarchy.Add(SequenceData, ThisID, ParentID);
+
+	return ThisID;
+}
+
+void FMovieSceneEvaluationTemplateGenerator::FlushGroupImmediately(FName EvaluationGroup)
+{
+	ImmediateFlushGroups.Add(EvaluationGroup);
+}
+
+void FMovieSceneEvaluationTemplateGenerator::Generate(FMovieSceneTrackCompilationParams InParams)
+{
+	Template.Hierarchy = FMovieSceneSequenceHierarchy();
+
+	// Generate templates for each track in the movie scene
+	UMovieScene* MovieScene = SourceSequence.GetMovieScene();
+
+	TransientArgs.Params = InParams;
+
+	if (UMovieSceneTrack* Track = MovieScene->GetCameraCutTrack())
+	{
+		ProcessTrack(*Track);
+	}
+
+	for (UMovieSceneTrack* Track : MovieScene->GetMasterTracks())
+	{
+		ProcessTrack(*Track);
+	}
+
+	for (const FMovieSceneBinding& ObjectBinding : MovieScene->GetBindings())
+	{
+		for (UMovieSceneTrack* Track : ObjectBinding.GetTracks())
+		{
+			ProcessTrack(*Track, ObjectBinding.GetObjectGuid());
+		}
+	}
+
+	// Remove references to old tracks
+	RemoveOldTrackReferences();
+
+	// Add all the tracks in *this* sequence (these exist after any sub section ptrs, not that its important for this algorithm)
+	for (auto& Pair : Template.GetTracks())
+	{
+		TArrayView<const FMovieSceneSegment> Segments = Pair.Value.GetSegments();
+
+		// Add the segment range data to the master collection for overall compilation
+		for (int32 SegmentIndex = 0; SegmentIndex < Segments.Num(); ++SegmentIndex)
+		{
+			SegmentData.Add(FMovieSceneSectionData(Segments[SegmentIndex].Range, TrackLUT.Num()));
+			TrackLUT.Add(FMovieSceneEvaluationFieldSegmentPtr(MovieSceneSequenceID::Root, Pair.Key, SegmentIndex));
+		}
+	}
+
+	TMap<FMovieSceneSequenceID, FMovieSceneEvaluationTemplate*> SequenceIdToTemplate;
+	for (auto& Pair : Template.Hierarchy.AllSubSequenceData())
+	{
+		if (UMovieSceneSequence* Sequence = Pair.Value.Sequence)
+		{
+			SequenceIdToTemplate.Add(Pair.Key, &TransientArgs.SubSequenceStore.GetCompiledTemplate(*Sequence));
+		}
+	}
+
+	// Compile the new evaluation field
+	TArray<FMovieSceneSegment> NewSegments = FMovieSceneSegmentCompiler().Compile(SegmentData);
+	UpdateEvaluationField(NewSegments, TrackLUT, SequenceIdToTemplate);
+}
+
+void FMovieSceneEvaluationTemplateGenerator::ProcessTrack(const UMovieSceneTrack& InTrack, const FGuid& ObjectId)
+{
+	FGuid Signature = InTrack.GetSignature();
+
+	// See if this track signature already exists in the ledger, if it does, we don't need to regenerate it
+	if (Template.FindTracks(Signature).Num())
+	{
+		CompiledSignatures.Add(Signature);
+		return;
+	}
+
+	TransientArgs.ObjectBindingId = ObjectId;
+
+	// Potentially expensive generation is required
+	InTrack.GenerateTemplate(TransientArgs);
+}
+
+void FMovieSceneEvaluationTemplateGenerator::RemoveOldTrackReferences()
+{
+	TArray<FGuid> SignaturesToRemove;
+
+	// Go through the template ledger, and remove anything that is no longer referenced
+	for (auto& Pair : Template.GetLedger().TrackSignatureToTrackIdentifier)
+	{
+		if (!CompiledSignatures.Contains(Pair.Key))
+		{
+			SignaturesToRemove.Add(Pair.Key);
+		}
+	}
+
+	// Remove the signatures, updating entries in the evaluation field as we go
+	for (const FGuid& Signature : SignaturesToRemove)
+	{
+		Template.RemoveTrack(Signature);
+	}
+}
+
+void FMovieSceneEvaluationTemplateGenerator::UpdateEvaluationField(const TArray<FMovieSceneSegment>& Segments, const TArray<FMovieSceneEvaluationFieldSegmentPtr>& Ptrs, const TMap<FMovieSceneSequenceID, FMovieSceneEvaluationTemplate*>& Templates)
+{
+	FMovieSceneEvaluationField& Field = Template.EvaluationField;
+
+	Field = FMovieSceneEvaluationField();
+
+	TArray<FMovieSceneEvaluationFieldSegmentPtr> AllTracksInSegment;
+
+	for (int32 Index = 0; Index < Segments.Num(); ++Index)
+	{
+		const FMovieSceneSegment& Segment = Segments[Index];
+		if (Segment.Impls.Num() == 0)
+		{
+			continue;
+		}
+
+		Field.Ranges.Add(Segment.Range);
+
+		AllTracksInSegment.Reset();
+		for (const FSectionEvaluationData& LUTData : Segment.Impls)
+		{
+			AllTracksInSegment.Add(Ptrs[LUTData.ImplIndex]);
+		}
+
+		// Sort the track ptrs, and define flush ranges
+		AllTracksInSegment.Sort(
+			[&](const FMovieSceneEvaluationFieldSegmentPtr& A, const FMovieSceneEvaluationFieldSegmentPtr& B){
+				return SortPredicate(LookupTrack(A, Templates), LookupTrack(B, Templates));
+			}
+		);
+
+		FMovieSceneEvaluationGroup& Group = Field.Groups[Field.Groups.Emplace()];
+
+		TArray<FMovieSceneEvaluationFieldSegmentPtr> EvalPtrs;
+		TArray<FMovieSceneEvaluationFieldSegmentPtr> InitPtrs;
+
+		// Now iterate the tracks and insert indices for initialization and evaluation
+		FName CurrentEvaluationGroup, LastEvaluationGroup;
+		TOptional<uint32> LastPriority;
+
+		for (const FMovieSceneEvaluationFieldSegmentPtr& Ptr : AllTracksInSegment)
+		{
+			const FMovieSceneEvaluationTrack* Track = LookupTrack(Ptr, Templates);
+			if (!ensure(Track))
+			{
+				continue;
+			}
+
+			// If we're now in a different flush group, add the ptrs
+			CurrentEvaluationGroup = Track->GetEvaluationGroup();
+			if (CurrentEvaluationGroup != LastEvaluationGroup)
+			{
+				AddPtrsToGroup(Group, ImmediateFlushGroups.Contains(LastEvaluationGroup), InitPtrs, EvalPtrs);
+			}
+			LastEvaluationGroup = Track->GetEvaluationGroup();
+
+			const bool bRequiresInitialization = Track->GetSegment(Ptr.SegmentIndex).Impls.ContainsByPredicate(
+				[Track](FSectionEvaluationData EvalData)
+				{
+					return Track->GetChildTemplate(EvalData.ImplIndex).RequiresInitialization();
+				}
+			);
+
+			if (bRequiresInitialization)
+			{
+				InitPtrs.Add(Ptr);
+			}
+
+			EvalPtrs.Add(Ptr);
+		}
+
+		AddPtrsToGroup(Group, ImmediateFlushGroups.Contains(LastEvaluationGroup), InitPtrs, EvalPtrs);
+
+		// Copmpute meta data for this segment
+		FMovieSceneEvaluationMetaData& MetaData = Field.MetaData[Field.MetaData.Emplace()];
+
+		TArray<FMovieSceneSequenceID, TInlineAllocator<16>> ActiveSequenceIDs;
+		for (const FMovieSceneEvaluationFieldSegmentPtr& SegmentPtr : Group.SegmentPtrLUT)
+		{
+			ActiveSequenceIDs.AddUnique(SegmentPtr.SequenceID);
+		}
+		MetaData.ActiveSequences = ActiveSequenceIDs;
+		MetaData.ActiveSequences.Sort();
+	}
+}
+
+const FMovieSceneEvaluationTrack* FMovieSceneEvaluationTemplateGenerator::LookupTrack(const FMovieSceneEvaluationFieldTrackPtr& InPtr, const TMap<FMovieSceneSequenceID, FMovieSceneEvaluationTemplate*>& Templates)
+{
+	if (InPtr.SequenceID == MovieSceneSequenceID::Root)
+	{
+		return Template.FindTrack(InPtr.TrackIdentifier);
+	}
+	else if (const FMovieSceneEvaluationTemplate* SubTemplate = Templates.FindRef(InPtr.SequenceID))
+	{
+		return SubTemplate->FindTrack(InPtr.TrackIdentifier);
+	}
+	ensure(false);
+	return nullptr;
+}
+
+bool FMovieSceneEvaluationTemplateGenerator::SortPredicate(const FMovieSceneEvaluationTrack* A, const FMovieSceneEvaluationTrack* B) const
+{
+	if (!ensure(A && B))
+	{
+		return false;
+	}
+
+	if (A->GetEvaluationPriority() != B->GetEvaluationPriority())
+	{
+		return A->GetEvaluationPriority() > B->GetEvaluationPriority();
+	}
+
+	if (A->GetEvaluationGroup() != B->GetEvaluationGroup())
+	{
+		return A->GetEvaluationGroup() < B->GetEvaluationGroup();
+	}
+
+	return false;
+}
\ No newline at end of file
diff --git a/Engine/Source/Runtime/MovieScene/Private/Compilation/MovieSceneSegmentCompiler.cpp b/Engine/Source/Runtime/MovieScene/Private/Compilation/MovieSceneSegmentCompiler.cpp
new file mode 100644
--- /dev/null
+++ b/Engine/Source/Runtime/MovieScene/Private/Compilation/MovieSceneSegmentCompiler.cpp
@@ -0,0 +1,349 @@
+// Copyright 1998-2016 Epic Games, Inc. All Rights Reserved.
+
+#include "MovieScenePrivatePCH.h"
+#include "MovieSceneSegmentCompiler.h"
+
+FMovieSceneSectionData::FMovieSceneSectionData(const TRange<float>& InBounds, int32 InSourceIndex, int32 InPriority)
+	: Bounds(InBounds)
+	, SourceIndex(InSourceIndex)
+	, Priority(InPriority)
+{}
+
+void FMovieSceneSegmentCompilerRules::ProcessSegments(TArray<FMovieSceneSegment>& Segments, const TArrayView<const FMovieSceneSectionData>& SourceData) const
+{
+	if (!Segments.Num())
+	{
+		return;
+	}
+
+	for (FMovieSceneSegment& Segment : Segments)
+	{
+		BlendSegment(Segment, SourceData);
+	}
+
+	int32 Index = 0;
+
+	if (!Segments[0].Range.GetLowerBound().IsOpen())
+	{
+		InsertSegment(Segments, 0, SourceData);
+		++Index;
+	}
+
+	for (; Index < Segments.Num() - 1; ++Index)
+	{
+		if (InsertSegment(Segments, Index, SourceData))
+		{
+			++Index;
+		}
+	}
+
+	if (!Segments.Last().Range.GetUpperBound().IsOpen())
+	{
+		InsertSegment(Segments, Segments.Num(), SourceData);
+	}
+
+	PostProcessSegments(Segments, SourceData);
+}
+
+bool FMovieSceneSegmentCompilerRules::InsertSegment(TArray<FMovieSceneSegment>& Segments, int32 Index, const TArrayView<const FMovieSceneSectionData>& SourceData) const
+{
+	const FMovieSceneSegment* PreviousSegment = Segments.IsValidIndex(Index-1) ? &Segments[Index-1] : nullptr;
+	const FMovieSceneSegment* NextSegment = Segments.IsValidIndex(Index) ? &Segments[Index] : nullptr;
+
+	TRange<float> EmptyRange(
+		PreviousSegment 	? TRangeBound<float>::FlipInclusion(PreviousSegment->Range.GetUpperBound()) : TRangeBound<float>(),
+		NextSegment 		? TRangeBound<float>::FlipInclusion(NextSegment->Range.GetLowerBound()) 	: TRangeBound<float>()
+		);
+
+	if (EmptyRange.IsEmpty())
+	{
+		return false;
+	}
+
+	TOptional<FMovieSceneSegment> NewSeg = InsertEmptySpace(EmptyRange, PreviousSegment, NextSegment);
+	if (!NewSeg.IsSet())
+	{
+		return false;
+	}
+	else if (!ensureMsgf(EmptyRange.Contains(NewSeg.GetValue().Range), TEXT("Attempting to insert an range that overflows the empty space. Correcting....")))
+	{
+		NewSeg.GetValue().Range = TRange<float>::Intersection(NewSeg.GetValue().Range, EmptyRange);
+	}
+
+	BlendSegment(NewSeg.GetValue(), SourceData);
+	Segments.Insert(MoveTemp(NewSeg.GetValue()), Index);
+
+	return true;
+}
+
+TArray<FMovieSceneSegment> FMovieSceneSegmentCompiler::Compile(const TArrayView<const FMovieSceneSectionData>& Data, const FMovieSceneSegmentCompilerRules* Rules)
+{
+	OverlappingSections.Reset(16);
+	OverlappingRefCounts.Reset(16);
+	LowerBounds.Reset(Data.Num());
+	UpperBounds.Reset(Data.Num());
+	CompiledSegments.Reset();
+
+	LowerReadIndex = 0;
+	UpperReadIndex = 0;
+
+	// Populate the lists of lower/upper bounds
+	for (int32 Index = 0; Index < Data.Num(); ++Index)
+	{
+		const FMovieSceneSectionData& Section = Data[Index];
+		if (!Section.Bounds.IsEmpty())
+		{
+			ensure(Section.SourceIndex != -1);
+			LowerBounds.Add(FBound(Section.SourceIndex, Section.Bounds.GetLowerBound()));
+			UpperBounds.Add(FBound(Section.SourceIndex, Section.Bounds.GetUpperBound()));
+		}
+	}
+
+	LowerBounds.Sort(
+		[](const FBound& A, const FBound& B)
+		{
+			return FFloatRangeBound::MinLower(A.Bound, B.Bound) == A.Bound;
+		}
+	);
+
+	UpperBounds.Sort(
+		[](const FBound& A, const FBound& B)
+		{
+			return FFloatRangeBound::MinUpper(A.Bound, B.Bound) == A.Bound;
+		}
+	);
+	
+	while (LowerReadIndex < LowerBounds.Num())
+	{
+		CloseCompletedSegments();
+
+		TRangeBound<float> OpeningBound = LowerBounds[LowerReadIndex].Bound;
+
+		const bool bHadOverlappingSections = OverlappingSections.Num() != 0;
+
+		// Add the currently overlapping sections for any sections starting at exactly this time
+		do
+		{
+			// Reference count how many times this section is overlapping the current time. This is to support multiple references to the same section.
+			const int32 OverlapIndex = OverlappingSections.IndexOfByKey(LowerBounds[LowerReadIndex].ImplIndex);
+			if (OverlapIndex == INDEX_NONE)
+			{
+				OverlappingSections.Add(LowerBounds[LowerReadIndex].ImplIndex);
+				OverlappingRefCounts.Add(1);
+			}
+			else
+			{
+				++OverlappingRefCounts[OverlapIndex];
+			}
+		}
+		while(++LowerReadIndex < LowerBounds.Num() && LowerBounds[LowerReadIndex].Bound == OpeningBound);
+
+		CompiledSegments.Emplace(
+			TRange<float>(OpeningBound, TRangeBound<float>()),
+			OverlappingSections
+			);
+	}
+
+	CloseCompletedSegments();
+
+	ensure(OverlappingSections.Num() == 0);
+
+	if (Rules)
+	{
+		Rules->ProcessSegments(CompiledSegments, Data);
+	}
+
+	return MoveTemp(CompiledSegments);
+}
+
+void FMovieSceneSegmentCompiler::CloseCompletedSegments()
+{
+	if (!CompiledSegments.Num())
+	{
+		return;
+	}
+
+	while(UpperReadIndex < UpperBounds.Num())
+	{
+		FMovieSceneSegment& LastSegment = CompiledSegments.Last();
+
+		// If there is a non-empty range between the next lower bound and upper bound, we can't close any more
+		const bool bHasOpeningRange = LowerReadIndex < LowerBounds.Num() && !TRange<float>(LowerBounds[LowerReadIndex].Bound, UpperBounds[UpperReadIndex].Bound).IsEmpty();
+		if (bHasOpeningRange)
+		{
+			if (OverlappingSections.Num())
+			{
+				TRangeBound<float> ClosingBound = TRangeBound<float>::FlipInclusion(LowerBounds[LowerReadIndex].Bound);
+				TRange<float> NewRange(LastSegment.Range.GetLowerBound(), ClosingBound);
+				if (!NewRange.IsEmpty())
+				{
+					// Just set the closing bound of the last segment and return
+					LastSegment.Range = NewRange;
+				}
+				else
+				{
+					// If it's empty, there's not point adding a segment, just allow the next segment to include the current OverlappingSegments
+					CompiledSegments.RemoveAt(CompiledSegments.Num() - 1, 1, false);
+				}
+			}
+			return;
+		}
+
+		TRangeBound<float> ClosingBound = UpperBounds[UpperReadIndex].Bound;
+		// Update the last segment's closing range
+		LastSegment.Range = TRange<float>(
+			LastSegment.Range.GetLowerBound(),
+			ClosingBound);
+
+		ensure(!LastSegment.Range.IsEmpty());
+
+		// Remove all sections that finish on this time
+		while (UpperReadIndex < UpperBounds.Num() && UpperBounds[UpperReadIndex].Bound == ClosingBound)
+		{
+			int32 OverlappingIndex = OverlappingSections.IndexOfByKey(UpperBounds[UpperReadIndex].ImplIndex);
+			if (ensure(OverlappingIndex != INDEX_NONE))
+			{
+				if (--OverlappingRefCounts[OverlappingIndex] == 0)
+				{
+					OverlappingSections.RemoveAtSwap(OverlappingIndex, 1, false);
+					OverlappingRefCounts.RemoveAtSwap(OverlappingIndex, 1, false);
+				}
+			}
+			++UpperReadIndex;
+		}
+
+		// If there are any more sections still active, create a new segment for those
+		if (OverlappingSections.Num())
+		{
+			CompiledSegments.Emplace(
+				TRange<float>(TRangeBound<float>::FlipInclusion(ClosingBound), TRangeBound<float>()),
+				OverlappingSections
+				);
+		}
+	}
+}
+
+
+FMovieSceneTrackCompiler::FRows::FRows(const TArray<UMovieSceneSection*>& Sections, const FMovieSceneSegmentCompilerRules* CompileRules)
+{
+	for (int32 Index = 0; Index < Sections.Num(); ++Index)
+	{
+		const UMovieSceneSection* Section = Sections[Index];
+		if (!ensure(Section) || !Section->IsActive())
+		{
+			continue;
+		}
+		
+		const int32 RowIndex = Section->GetRowIndex();
+		if (!Rows.IsValidIndex(RowIndex))
+		{
+			Rows.AddDefaulted(RowIndex - Rows.Num() + 1);
+		}
+
+		const TRange<float> Range = Section->IsInfinite() ? TRange<float>::All() : Section->GetRange();
+		Rows[RowIndex].Sections.Add(
+			FMovieSceneSectionRowData(Index, Range, Rows[RowIndex].Sections.Num(), Section->GetOverlapPriority())
+			);
+	}
+
+	Rows.RemoveAll(
+		[](const FRow& Row)
+		{
+			return Row.Sections.Num() == 0;
+		});
+
+	for (FRow& Row : Rows)
+	{
+		Row.CompileRules = CompileRules;
+	}
+}
+
+FMovieSceneTrackEvaluationField FMovieSceneTrackCompiler::Compile(const TArrayView<FRow>& Rows, const FMovieSceneSegmentCompilerRules* Rules)
+{
+	FMovieSceneTrackEvaluationField Result;
+
+	TArray<FMovieSceneSegment> AllRowSegments;
+	TArray<FMovieSceneSectionData> TrackCompileData;
+	TArray<const FMovieSceneSectionData*> AllRows;
+
+	// Compile each row
+	for (int32 RowIndex = 0; RowIndex < Rows.Num(); ++RowIndex)
+	{
+		const FRow& Row = Rows[RowIndex];
+		if (!Row.Sections.Num())
+		{
+			continue;
+		}
+		
+		FMovieSceneSegmentCompiler Compiler;
+
+		for (const FMovieSceneSectionData& Section : Row.Sections)
+		{
+			AllRows.Add(&Section);
+		}
+		
+		TArray<FMovieSceneSegment> RowSegments = Compiler.Compile(
+			TArray<FMovieSceneSectionData>(Row.Sections),
+			Row.CompileRules);
+
+		const int32 Priority = Rows.Num() - RowIndex;
+		for (FMovieSceneSegment& Segment : RowSegments)
+		{
+			TrackCompileData.Add(FMovieSceneSectionData(Segment.Range, TrackCompileData.Num(), Priority));
+
+			// Remap impl indices to their actual sections
+			for (FSectionEvaluationData& EvalData : Segment.Impls)
+			{
+				EvalData.ImplIndex = Row.Sections[EvalData.ImplIndex].ActualSectionIndex;
+			}
+		}
+
+		AllRowSegments.Append(MoveTemp(RowSegments));
+	}
+
+	// Boil down each row into a single, blended field
+	FMovieSceneSegmentCompiler Compiler;
+	TArray<FMovieSceneSegment> TrackSegments = Compiler.Compile(TrackCompileData, nullptr);
+
+	if (Rules)
+	{
+		Rules->ProcessSegments(TrackSegments, TrackCompileData);
+	}
+
+	// Compile each of the row segments into a single segment comprising all overlapping sections on the track, blended and ordered appropriately
+	// There should be no empty space by this point
+	for (FMovieSceneSegment& Segment : TrackSegments)
+	{
+		FMovieSceneSegment CompiledSegment(Segment.Range);
+
+		// Segment.Impls relates to the segments stored in AllRowSegments
+		for (const FSectionEvaluationData& RowEvalData : Segment.Impls)
+		{
+			for (FSectionEvaluationData SegmentEvalData : AllRowSegments[RowEvalData.ImplIndex].Impls)
+			{
+				// Take a forced time from the track level, if it's not specified on the row level
+				if (SegmentEvalData.ForcedTime == TNumericLimits<float>::Lowest())
+				{
+					SegmentEvalData.ForcedTime = RowEvalData.ForcedTime;
+				}
+				CompiledSegment.Impls.Add(SegmentEvalData);
+			}
+		}
+		
+		if (CompiledSegment.Impls.Num() != 0 || (Rules && Rules->AllowEmptySegments()))
+		{
+			// If this is the same as the previous segment, and it ajoins the previous segment's range, just increase the range of the previous segment
+			if (Result.Segments.Num() != 0 && Result.Segments.Last().Range.Adjoins(CompiledSegment.Range) && Result.Segments.Last().Impls == CompiledSegment.Impls)
+			{
+				Result.Segments.Last().Range = FFloatRange::Hull(Result.Segments.Last().Range, CompiledSegment.Range);
+			}
+			else
+			{
+				Result.Segments.Add(MoveTemp(CompiledSegment));
+			}
+		}
+	}
+
+	return Result;
+}
+
diff --git a/Engine/Source/Runtime/MovieScene/Private/Evaluation/MovieSceneEvalTemplate.cpp b/Engine/Source/Runtime/MovieScene/Private/Evaluation/MovieSceneEvalTemplate.cpp
new file mode 100644
--- /dev/null
+++ b/Engine/Source/Runtime/MovieScene/Private/Evaluation/MovieSceneEvalTemplate.cpp
@@ -0,0 +1,10 @@
+// Copyright 1998-2016 Epic Games, Inc. All Rights Reserved.
+
+#include "MovieScenePrivatePCH.h"
+#include "MovieSceneEvalTemplate.h"
+#include "MovieSceneEvalTemplateSerializer.h"
+
+bool FMovieSceneEvalTemplatePtr::Serialize(FArchive& Ar)
+{
+	return SerializeEvaluationTemplate(*this, Ar);
+}
\ No newline at end of file
diff --git a/Engine/Source/Runtime/MovieScene/Private/Evaluation/MovieSceneEvaluationCustomVersion.cpp b/Engine/Source/Runtime/MovieScene/Private/Evaluation/MovieSceneEvaluationCustomVersion.cpp
new file mode 100644
--- /dev/null
+++ b/Engine/Source/Runtime/MovieScene/Private/Evaluation/MovieSceneEvaluationCustomVersion.cpp
@@ -0,0 +1,7 @@
+// Copyright 1998-2016 Epic Games, Inc. All Rights Reserved.
+
+#include "MovieScenePrivatePCH.h"
+#include "MovieSceneEvaluationCustomVersion.h"
+
+// Register the custom version with core
+FCustomVersionRegistration GRegisterMovieSceneEvaluationCustomVersion(FMovieSceneEvaluationCustomVersion::GUID, FMovieSceneEvaluationCustomVersion::LatestVersion, TEXT("MovieSceneEvaluation"));
\ No newline at end of file
diff --git a/Engine/Source/Runtime/MovieScene/Private/Evaluation/MovieSceneEvaluationField.cpp b/Engine/Source/Runtime/MovieScene/Private/Evaluation/MovieSceneEvaluationField.cpp
new file mode 100644
--- /dev/null
+++ b/Engine/Source/Runtime/MovieScene/Private/Evaluation/MovieSceneEvaluationField.cpp
@@ -0,0 +1,41 @@
+// Copyright 1998-2016 Epic Games, Inc. All Rights Reserved.
+
+#include "MovieScenePrivatePCH.h"
+#include "MovieSceneEvaluationField.h"
+
+int32 FMovieSceneEvaluationField::GetSegmentFromTime(float Time) const
+{
+	// Linear search
+	// @todo: accelerated search based on the last evaluated index?
+	for (int32 Index = 0; Index < Ranges.Num(); ++Index)
+	{
+		if (Ranges[Index].Contains(Time))
+		{
+			return Index;
+		}
+	}
+
+	return INDEX_NONE;
+}
+
+TRange<int32> FMovieSceneEvaluationField::OverlapRange(TRange<float> Range) const
+{
+	int32 StartIndex = 0, Num = 0;
+	for (int32 Index = 0; Index < Ranges.Num(); ++Index)
+	{
+		if (Ranges[Index].Overlaps(Range))
+		{
+			if (Num == 0)
+			{
+				StartIndex = Index;
+			}
+			++Num;
+		}
+		else if (Num != 0)
+		{
+			break;
+		}
+	}
+
+	return Num != 0 ? TRange<int32>(StartIndex, StartIndex + Num) : TRange<int32>::Empty();
+}
\ No newline at end of file
diff --git a/Engine/Source/Runtime/MovieScene/Private/Evaluation/MovieSceneEvaluationState.cpp b/Engine/Source/Runtime/MovieScene/Private/Evaluation/MovieSceneEvaluationState.cpp
new file mode 100644
--- /dev/null
+++ b/Engine/Source/Runtime/MovieScene/Private/Evaluation/MovieSceneEvaluationState.cpp
@@ -0,0 +1,270 @@
+// Copyright 1998-2016 Epic Games, Inc. All Rights Reserved.
+
+#include "MovieScenePrivatePCH.h"
+#include "MovieSceneEvaluationState.h"
+
+
+DECLARE_CYCLE_STAT(TEXT("Find Bound Objects"), MovieSceneEval_FindBoundObjects, STATGROUP_MovieSceneEval);
+
+FMovieSceneSharedDataId FMovieSceneSharedDataId::Allocate()
+{
+	static uint32 Counter = 0;
+
+	FMovieSceneSharedDataId Value;
+	Value.UniqueId = ++Counter;
+	check(Counter != -1);
+	return Value;
+}
+
+TArrayView<TWeakObjectPtr<>> FMovieSceneObjectCache::FindBoundObjects(const FGuid& InBindingID, IMovieScenePlayer& Player)
+{
+	MOVIESCENE_DETAILED_SCOPE_CYCLE_COUNTER(MovieSceneEval_FindBoundObjects)
+	
+	// Fast route - where everything's cached
+	FBoundObjects* Bindings = BoundObjects.Find(InBindingID);
+	if (Bindings && Bindings->bUpToDate)
+	{
+		return TArrayView<TWeakObjectPtr<>>(
+			Bindings->Objects.GetData(),
+			Bindings->Objects.Num()
+			);
+	}
+
+	// Attempt to update the bindings
+	UpdateBindings(InBindingID, Player);
+
+	Bindings = BoundObjects.Find(InBindingID);
+	if (Bindings)
+	{
+		return TArrayView<TWeakObjectPtr<>>(Bindings->Objects.GetData(), Bindings->Objects.Num());
+	}
+
+	// Just return nothing
+	return TArrayView<TWeakObjectPtr<>>();
+}
+
+FGuid FMovieSceneObjectCache::FindObjectId(UObject& InObject, IMovieScenePlayer& Player)
+{
+	UMovieSceneSequence* Sequence = WeakSequence.Get();
+	UMovieScene* MovieScene = Sequence ? Sequence->GetMovieScene() : nullptr;
+	if (!MovieScene)
+	{
+		return FGuid();
+	}
+
+	TWeakObjectPtr<> ObjectToFind(&InObject);
+
+	// Search all possessables
+	for (int32 Index = 0; Index < MovieScene->GetPossessableCount(); ++Index)
+	{
+		FGuid ThisGuid = MovieScene->GetPossessable(Index).GetGuid();
+		if (FindBoundObjects(ThisGuid, Player).Contains(ObjectToFind))
+		{
+			return ThisGuid;
+		}
+	}
+
+	// Search all spawnables
+	for (int32 Index = 0; Index < MovieScene->GetSpawnableCount(); ++Index)
+	{
+		FGuid ThisGuid = MovieScene->GetSpawnable(Index).GetGuid();
+		if (FindBoundObjects(ThisGuid, Player).Contains(ObjectToFind))
+		{
+			return ThisGuid;
+		}
+	}
+
+	return FGuid();
+}
+
+void FMovieSceneObjectCache::InvalidateExpiredObjects()
+{
+	for (auto& Pair : BoundObjects)
+	{
+		if (!Pair.Value.bUpToDate)
+		{
+			continue;
+		}
+
+		for (TWeakObjectPtr<>& Ptr : Pair.Value.Objects)
+		{
+			if (!Ptr.Get())
+			{
+				Invalidate(Pair.Key);
+				break;
+			}
+		}
+	}
+}
+
+void FMovieSceneObjectCache::Invalidate(const FGuid& InGuid)
+{
+	// Don't manipulate the actual map structure, since this can be called from inside an iterator
+	FBoundObjects* Cache = BoundObjects.Find(InGuid);
+	if (Cache)
+	{
+		Cache->bUpToDate = false;
+
+		auto* Children = ChildBindings.Find(InGuid);
+		if (Children)
+		{
+			for (const FGuid& Child : *Children)
+			{
+				Invalidate(Child);
+			}
+		}
+	}
+}
+
+void FMovieSceneObjectCache::SetSequence(UMovieSceneSequence& InSequence, FMovieSceneSequenceIDRef InSequenceID)
+{
+	if (WeakSequence != &InSequence)
+	{
+		Clear();
+	}
+
+	WeakSequence = &InSequence;
+	SequenceID = InSequenceID;
+}
+
+void FMovieSceneObjectCache::UpdateBindings(const FGuid& InGuid, IMovieScenePlayer& Player)
+{
+	FBoundObjects& Bindings = BoundObjects.FindOrAdd(InGuid);
+	Bindings.Objects.Reset();
+
+	if (auto* Children = ChildBindings.Find(InGuid))
+	{
+		for (const FGuid& Child : *Children)
+		{
+			Invalidate(Child);
+		}
+	}
+	ChildBindings.Remove(InGuid);
+
+	// Find the sequence for the current
+	UMovieSceneSequence* Sequence = WeakSequence.Get();
+	if (!ensure(Sequence))
+	{
+		return;
+	}
+
+	const bool bUseParentsAsContext = Sequence->AreParentContextsSignificant();
+	
+	UObject* Context = Player.GetPlaybackContext();
+
+	const FMovieScenePossessable* Possessable = Sequence->GetMovieScene()->FindPossessable(InGuid);
+	if (Possessable)
+	{
+		UObject* ResolutionContext = Context;
+
+		// Because these are ordered parent-first, the parent must have already been bound, if it exists
+		if (Possessable->GetParent().IsValid())
+		{
+			ChildBindings.FindOrAdd(Possessable->GetParent()).AddUnique(InGuid);
+
+			for (TWeakObjectPtr<> Parent : FindBoundObjects(Possessable->GetParent(), Player))
+			{
+				if (bUseParentsAsContext)
+				{
+					ResolutionContext = Parent.Get();
+					if (!ResolutionContext)
+					{
+						continue;
+					}
+				}
+
+				TArray<UObject*, TInlineAllocator<1>> FoundObjects;
+				Sequence->LocateBoundObjects(InGuid, ResolutionContext, FoundObjects);
+				for (UObject* Object : FoundObjects)
+				{
+					Bindings.Objects.Add(Object);
+				}
+			}
+		}
+		else
+		{
+			TArray<UObject*, TInlineAllocator<1>> FoundObjects;
+			Sequence->LocateBoundObjects(InGuid, ResolutionContext, FoundObjects);
+			for (UObject* Object : FoundObjects)
+			{
+				Bindings.Objects.Add(Object);
+			}
+		}
+	}
+	else
+	{
+		// Probably a spawnable then (or an phantom)
+		UObject* SpawnedObject = Player.GetSpawnRegister().FindSpawnedObject(InGuid, SequenceID);
+		if (SpawnedObject)
+		{
+			Bindings.Objects.Add(SpawnedObject);
+		}
+	}
+
+	if (Bindings.Objects.Num())
+	{
+		Bindings.bUpToDate = true;
+	}
+}
+
+bool FOrderedEvaluationCache::Contains(FMovieSceneEvaluationKey InKey) const
+{
+	return KeySet.Contains(InKey);
+}
+
+void FOrderedEvaluationCache::Add(FMovieSceneEvaluationKey InKey)
+{
+	if (!KeySet.Contains(InKey))
+	{
+		KeySet.Add(InKey);
+		OrderedKeys.Add(InKey);
+	}
+}
+
+void FOrderedEvaluationCache::Reset()
+{
+	KeySet.Reset();
+	OrderedKeys.Reset();
+}
+
+void FMovieSceneEvaluationState::InvalidateExpiredObjects()
+{
+	for (auto& Pair : ObjectCaches)
+	{
+		Pair.Value.InvalidateExpiredObjects();
+	}
+}
+
+void FMovieSceneEvaluationState::Invalidate(const FGuid& InGuid, FMovieSceneSequenceIDRef SequenceID)
+{
+	FMovieSceneObjectCache* Cache = ObjectCaches.Find(SequenceID);
+	if (Cache)
+	{
+		Cache->Invalidate(InGuid);
+	}
+}
+
+void FMovieSceneEvaluationState::ClearObjectCaches()
+{
+	for (auto& Pair : ObjectCaches)
+	{
+		Pair.Value.Clear();
+	}
+}
+
+void FMovieSceneEvaluationState::AssignSequence(FMovieSceneSequenceIDRef InSequenceID, UMovieSceneSequence& InSequence)
+{
+	GetObjectCache(InSequenceID).SetSequence(InSequence, InSequenceID);
+}
+
+UMovieSceneSequence* FMovieSceneEvaluationState::FindSequence(FMovieSceneSequenceIDRef InSequenceID) const
+{
+	const FMovieSceneObjectCache* Cache = ObjectCaches.Find(InSequenceID);
+	return Cache ? Cache->GetSequence() : nullptr;
+}
+
+FGuid FMovieSceneEvaluationState::FindObjectId(UObject& Object, FMovieSceneSequenceIDRef InSequenceID, IMovieScenePlayer& Player)
+{
+	FMovieSceneObjectCache* Cache = ObjectCaches.Find(InSequenceID);
+	return Cache ? Cache->FindObjectId(Object, Player) : FGuid();
+}
diff --git a/Engine/Source/Runtime/MovieScene/Private/Evaluation/MovieSceneEvaluationTemplate.cpp b/Engine/Source/Runtime/MovieScene/Private/Evaluation/MovieSceneEvaluationTemplate.cpp
new file mode 100644
--- /dev/null
+++ b/Engine/Source/Runtime/MovieScene/Private/Evaluation/MovieSceneEvaluationTemplate.cpp
@@ -0,0 +1,215 @@
+// Copyright 1998-2016 Epic Games, Inc. All Rights Reserved.
+
+#include "MovieScenePrivatePCH.h"
+#include "MovieSceneEvaluationTemplate.h"
+
+FMovieSceneSequenceCachedSignature::FMovieSceneSequenceCachedSignature(UMovieSceneSequence& InSequence)
+	: Sequence(&InSequence)
+	, CachedSignature(InSequence.GetSignature())
+{
+}
+
+#if WITH_EDITORONLY_DATA
+
+void FCachedMovieSceneEvaluationTemplate::Initialize(UMovieSceneSequence& InSequence, FMovieSceneSequenceTemplateStore* InOrigin)
+{
+	SourceSequence = &InSequence;
+	Origin = InOrigin;
+}
+
+void FCachedMovieSceneEvaluationTemplate::Regenerate()
+{
+	Regenerate(CachedCompilationParams);
+}
+
+void FCachedMovieSceneEvaluationTemplate::Regenerate(const FMovieSceneTrackCompilationParams& Params)
+{
+	if (IsOutOfDate(Params))
+	{
+		RegenerateImpl(Params);
+	}
+}
+
+void FCachedMovieSceneEvaluationTemplate::ForceRegenerate(const FMovieSceneTrackCompilationParams& Params)
+{
+	// Reassign the template back to default
+	static_cast<FMovieSceneEvaluationTemplate&>(*this) = FMovieSceneEvaluationTemplate();
+	RegenerateImpl(Params);
+}
+
+void FCachedMovieSceneEvaluationTemplate::RegenerateImpl(const FMovieSceneTrackCompilationParams& Params)
+{
+	CachedSignatures.Reset();
+	CachedCompilationParams = Params;
+
+	if (UMovieSceneSequence* Sequence = SourceSequence.Get())
+	{
+		FMovieSceneSequenceTemplateStore DefaultStore;
+		Sequence->GenerateEvaluationTemplate(*this, Params, Origin ? *Origin : DefaultStore);
+
+		CachedSignatures.Add(FMovieSceneSequenceCachedSignature(*Sequence));
+		for (auto& Pair : Hierarchy.AllSubSequenceData())
+		{
+			if (UMovieSceneSequence* SubSequence = Pair.Value.Sequence)
+			{
+				CachedSignatures.Add(FMovieSceneSequenceCachedSignature(*SubSequence));
+			}
+		}
+	}
+}
+
+bool FCachedMovieSceneEvaluationTemplate::IsOutOfDate(const FMovieSceneTrackCompilationParams& Params) const
+{
+	if (Params != CachedCompilationParams || CachedSignatures.Num() == 0)
+	{
+		return true;
+	}
+
+	// out of date if the cached signautres contains any out of date sigs
+	return CachedSignatures.ContainsByPredicate(
+		[](const FMovieSceneSequenceCachedSignature& Sig){
+			UMovieSceneSequence* Sequence = Sig.Sequence.Get();
+			return !Sequence || Sequence->GetSignature() != Sig.CachedSignature;
+		}
+	);
+}
+
+#endif // WITH_EDITORONLY_DATA
+
+bool FMovieSceneGenerationLedger::Serialize(FArchive& Ar)
+{
+	if (Ar.IsLoading())
+	{
+		int32 NumReferenceCounts = 0;
+
+		Ar << NumReferenceCounts;
+		for (int32 Index = 0; Index < NumReferenceCounts; ++Index)
+		{
+			FMovieSceneTrackIdentifier Identifier;
+			FMovieSceneTrackIdentifier::StaticStruct()->SerializeItem(Ar, &Identifier, nullptr);
+
+			int32 Count = 0;
+			Ar << Count;
+
+			TrackReferenceCounts.Add(Identifier, Count);
+		}
+
+		int32 SignatureToTrackIDs = 0;
+		Ar << SignatureToTrackIDs;
+
+		for (int32 Index = 0; Index < SignatureToTrackIDs; ++Index)
+		{
+			FGuid Signature;
+			Ar << Signature;
+
+			TArray<FMovieSceneTrackIdentifier, TInlineAllocator<1>>& Identifiers = TrackSignatureToTrackIdentifier.Add(Signature);
+
+			int32 NumIdentifiers = 0;
+			Ar << NumIdentifiers;
+			for (int32 IdentifierIndex = 0; IdentifierIndex < NumIdentifiers; ++IdentifierIndex)
+			{
+				FMovieSceneTrackIdentifier Identifier;
+				FMovieSceneTrackIdentifier::StaticStruct()->SerializeItem(Ar, &Identifier, nullptr);
+
+				Identifiers.Add(Identifier);
+			}
+		}
+
+		return true;
+	}
+	else if (Ar.IsSaving())
+	{
+		int32 NumReferenceCounts = TrackReferenceCounts.Num();
+
+		Ar << NumReferenceCounts;
+		for (auto& Pair : TrackReferenceCounts)
+		{
+			FMovieSceneTrackIdentifier Identifier = Pair.Key;
+			FMovieSceneTrackIdentifier::StaticStruct()->SerializeItem(Ar, &Identifier, nullptr);
+
+			int32 Count = Pair.Value;
+			Ar << Count;
+		}
+
+		int32 SignatureToTrackIDs = TrackSignatureToTrackIdentifier.Num();
+		Ar << SignatureToTrackIDs;
+
+		for (auto& Pair : TrackSignatureToTrackIdentifier)
+		{
+			FGuid Signature = Pair.Key;
+			Ar << Signature;
+
+			int32 NumIdentifiers = Pair.Value.Num();
+			Ar << NumIdentifiers;
+
+			for (FMovieSceneTrackIdentifier Identifier : Pair.Value)
+			{
+				FMovieSceneTrackIdentifier::StaticStruct()->SerializeItem(Ar, &Identifier, nullptr);
+			}
+		}
+
+		return true;
+	}
+
+	return false;
+}
+
+TArrayView<FMovieSceneTrackIdentifier> FMovieSceneGenerationLedger::FindTracks(const FGuid& InSignature)
+{
+	if (auto* Tracks = TrackSignatureToTrackIdentifier.Find(InSignature))
+	{
+		return *Tracks;
+	}
+	return TArrayView<FMovieSceneTrackIdentifier>();
+}
+
+void FMovieSceneGenerationLedger::AddTrack(const FGuid& InSignature, FMovieSceneTrackIdentifier Identifier)
+{
+	TrackSignatureToTrackIdentifier.FindOrAdd(InSignature).Add(Identifier);
+	++TrackReferenceCounts.FindOrAdd(Identifier);
+}
+
+FMovieSceneTrackIdentifier FMovieSceneEvaluationTemplate::AddTrack(const FGuid& InSignature, FMovieSceneEvaluationTrack&& InTrack)
+{
+	FMovieSceneTrackIdentifier NewIdentifier = ++Ledger.LastTrackIdentifier;
+	
+	InTrack.SetupOverrides();
+	Tracks.Add(NewIdentifier.Value, MoveTemp(InTrack));
+	Ledger.AddTrack(InSignature, NewIdentifier);
+
+	return NewIdentifier;
+}
+
+void FMovieSceneEvaluationTemplate::RemoveTrack(const FGuid& InSignature)
+{
+	for (FMovieSceneTrackIdentifier TrackIdentifier : Ledger.FindTracks(InSignature))
+	{
+		int32* RefCount = Ledger.TrackReferenceCounts.Find(TrackIdentifier);
+		if (ensure(RefCount) && --(*RefCount) == 0)
+		{
+			if (bKeepStaleTracks)
+			{
+				if (FMovieSceneEvaluationTrack* Track = Tracks.Find(TrackIdentifier.Value))
+				{
+					StaleTracks.Add(TrackIdentifier.Value, MoveTemp(*Track));
+				}
+			}
+			
+			Tracks.Remove(TrackIdentifier.Value);
+			Ledger.TrackReferenceCounts.Remove(TrackIdentifier);
+		}
+	}
+	Ledger.TrackSignatureToTrackIdentifier.Remove(InSignature);
+}
+
+const TMap<FMovieSceneTrackIdentifier, FMovieSceneEvaluationTrack>& FMovieSceneEvaluationTemplate::GetTracks() const
+{
+	// Reinterpret the uint32 as FMovieSceneTrackIdentifier
+	static_assert(sizeof(FMovieSceneTrackIdentifier) == sizeof(uint32), "FMovieSceneTrackIdentifier is not convertible directly to/from uint32.");
+	return *reinterpret_cast<const TMap<FMovieSceneTrackIdentifier, FMovieSceneEvaluationTrack>*>(&Tracks);
+}
+
+TArrayView<FMovieSceneTrackIdentifier> FMovieSceneEvaluationTemplate::FindTracks(const FGuid& InSignature)
+{
+	return Ledger.FindTracks(InSignature);
+}
\ No newline at end of file
diff --git a/Engine/Source/Runtime/MovieScene/Private/Evaluation/MovieSceneEvaluationTemplateInstance.cpp b/Engine/Source/Runtime/MovieScene/Private/Evaluation/MovieSceneEvaluationTemplateInstance.cpp
new file mode 100644
--- /dev/null
+++ b/Engine/Source/Runtime/MovieScene/Private/Evaluation/MovieSceneEvaluationTemplateInstance.cpp
@@ -0,0 +1,434 @@
+// Copyright 1998-2016 Epic Games, Inc. All Rights Reserved.
+
+#include "MovieScenePrivatePCH.h"
+#include "MovieSceneSequence.h"
+#include "MovieSceneEvaluationTemplateInstance.h"
+
+
+DECLARE_CYCLE_STAT(TEXT("Entire Evaluation Cost"), MovieSceneEval_EntireEvaluationCost, STATGROUP_MovieSceneEval);
+DECLARE_CYCLE_STAT(TEXT("Gather Entries For Frame"), MovieSceneEval_GatherEntries, STATGROUP_MovieSceneEval);
+DECLARE_CYCLE_STAT(TEXT("Call Setup() and TearDown()"), MovieSceneEval_CallSetupTearDown, STATGROUP_MovieSceneEval);
+DECLARE_CYCLE_STAT(TEXT("Evaluate Group"), MovieSceneEval_EvaluateGroup, STATGROUP_MovieSceneEval);
+
+
+FMovieSceneEvaluationTemplateInstance::FMovieSceneEvaluationTemplateInstance(UMovieSceneSequence& InSequence, const FMovieSceneEvaluationTemplate& InTemplate)
+	: Sequence(&InSequence)
+	, Template(&InTemplate)
+{
+	if (Template->bHasLegacyTrackInstances)
+	{
+		// Initialize the legacy sequence instance
+		LegacySequenceInstance = MakeShareable(new FMovieSceneSequenceInstance(InSequence, MovieSceneSequenceID::Root));
+	}
+}
+
+FMovieSceneEvaluationTemplateInstance::FMovieSceneEvaluationTemplateInstance(const FMovieSceneSubSequenceData& InSubData, const FMovieSceneEvaluationTemplate& InTemplate, FMovieSceneSequenceIDRef InSequenceID)
+	: Sequence(InSubData.Sequence)
+	, RootToSequenceTransform(InSubData.RootToSequenceTransform)
+	, Template(&InTemplate)
+{
+	if (Template->bHasLegacyTrackInstances)
+	{
+		// Initialize the legacy sequence instance
+		LegacySequenceInstance = MakeShareable(new FMovieSceneSequenceInstance(*InSubData.Sequence, InSequenceID));
+	}
+}
+
+FMovieSceneRootEvaluationTemplateInstance::FMovieSceneRootEvaluationTemplateInstance()
+	: RootSequence(nullptr)
+	, TemplateStore(MakeShareable(new FMovieSceneSequenceTemplateStore))
+	, bIsDirty(false)
+{
+
+}
+
+FMovieSceneRootEvaluationTemplateInstance::~FMovieSceneRootEvaluationTemplateInstance()
+{
+	Reset();
+}
+
+void FMovieSceneRootEvaluationTemplateInstance::Reset()
+{
+	if (TemplateStore->AreTemplatesVolatile())
+	{
+		if (RootInstance.Sequence)
+		{
+			RootInstance.Sequence->OnSignatureChanged().RemoveAll(this);
+		}
+
+		for (auto& Pair : SubInstances)
+		{
+			if (Pair.Value.Sequence)
+			{
+				Pair.Value.Sequence->OnSignatureChanged().RemoveAll(this);
+			}
+		}
+	}
+
+	SubInstances.Reset();
+}
+
+void FMovieSceneRootEvaluationTemplateInstance::Initialize(UMovieSceneSequence& InRootSequence, IMovieScenePlayer& Player, TSharedRef<FMovieSceneSequenceTemplateStore> InTemplateStore)
+{
+	if (RootSequence != &InRootSequence)
+	{
+		Finish(Player);
+	}
+
+	// Ensure we reset everything before we overwrite the template store (which potentially owns templates we've previously referenced)
+	Reset();
+
+	TemplateStore = MoveTemp(InTemplateStore);
+
+	Initialize(InRootSequence, Player);
+}
+
+void FMovieSceneRootEvaluationTemplateInstance::Initialize(UMovieSceneSequence& InRootSequence, IMovieScenePlayer& Player)
+{
+	Reset();
+
+	const bool bAddEvents = TemplateStore->AreTemplatesVolatile();
+	
+	RootSequence = &InRootSequence;
+
+	const FMovieSceneEvaluationTemplate& RootTemplate = TemplateStore->GetCompiledTemplate(InRootSequence);
+	
+	Player.State.AssignSequence(MovieSceneSequenceID::Root, InRootSequence);
+	RootInstance = FMovieSceneEvaluationTemplateInstance(InRootSequence, RootTemplate);
+
+	if (bAddEvents)
+	{
+		InRootSequence.OnSignatureChanged().AddRaw(this, &FMovieSceneRootEvaluationTemplateInstance::OnSequenceChanged);
+	}
+
+	for (auto& Pair : RootTemplate.Hierarchy.AllSubSequenceData())
+	{
+		const FMovieSceneSubSequenceData& Data = Pair.Value;
+		FMovieSceneSequenceID SequenceID = Pair.Key;
+
+		if (!Data.Sequence)
+		{
+			continue;
+		}
+
+		Player.State.AssignSequence(SequenceID, *Data.Sequence);
+
+		const FMovieSceneEvaluationTemplate& ChildTemplate = TemplateStore->GetCompiledTemplate(*Data.Sequence);
+		FMovieSceneEvaluationTemplateInstance& Instance = SubInstances.Add(SequenceID, FMovieSceneEvaluationTemplateInstance(Data, ChildTemplate, SequenceID));
+
+		if (bAddEvents)
+		{
+			Data.Sequence->OnSignatureChanged().AddRaw(this, &FMovieSceneRootEvaluationTemplateInstance::OnSequenceChanged);
+		}
+	}
+
+	bIsDirty = false;
+
+	OnUpdatedEvent.Broadcast();
+}
+
+void FMovieSceneRootEvaluationTemplateInstance::Finish(IMovieScenePlayer& Player)
+{
+	Swap(EntitiesEvaluatedThisFrame, EntitiesEvaluatedLastFrame);
+	EntitiesEvaluatedThisFrame.Reset();
+
+	CallSetupTearDown(Player);
+}
+
+void FMovieSceneRootEvaluationTemplateInstance::Evaluate(FMovieSceneContext Context, IMovieScenePlayer& Player, FMovieSceneSequenceID OverrideRootID)
+{
+	if (bIsDirty && RootSequence)
+	{
+		Initialize(*RootSequence, Player);
+	}
+
+	SCOPE_CYCLE_COUNTER(MovieSceneEval_EntireEvaluationCost);
+
+	Swap(EntitiesEvaluatedThisFrame, EntitiesEvaluatedLastFrame);
+	EntitiesEvaluatedThisFrame.Reset();
+
+	Swap(ActiveSequencesThisFrame, ActiveSequencesLastFrame);
+	ActiveSequencesThisFrame.Reset();
+
+	const FMovieSceneEvaluationTemplateInstance* Instance = GetInstance(OverrideRootID);
+
+	ensureMsgf(Instance, TEXT("Could not find instance for supplied sequence ID."));
+
+	const int32 FieldIndex = Instance ? Instance->Template->EvaluationField.GetSegmentFromTime(Context.GetTime() * Instance->RootToSequenceTransform) : INDEX_NONE;
+	if (FieldIndex == INDEX_NONE)
+	{
+		CallSetupTearDown(Player);
+		return;
+	}
+
+	const FMovieSceneEvaluationGroup& Group = Instance->Template->EvaluationField.Groups[FieldIndex];
+
+	GatherEntities(Group, Player, OverrideRootID);
+
+	// Gather the active sequences for this frame, remapping them to the root if necessary
+	ActiveSequencesThisFrame = Instance->Template->EvaluationField.MetaData[FieldIndex].ActiveSequences;
+	if (OverrideRootID != MovieSceneSequenceID::Root)
+	{
+		for (FMovieSceneSequenceID& ID : ActiveSequencesThisFrame)
+		{
+			ID = GetSequenceIdForRoot(ID, OverrideRootID);
+		}
+	}
+
+	// Run the post root evaluate steps which invoke tear downs for anything no longer evaluated.
+	// Do this now to ensure they don't undo any of the current frame's execution tokens 
+	CallSetupTearDown(Player);
+
+	// Ensure any null objects are not cached
+	Player.State.InvalidateExpiredObjects();
+
+	// Accumulate execution tokens into this structure
+	FMovieSceneExecutionTokens ExecutionTokens;
+	EvaluateGroup(Group, Context, Player, ExecutionTokens, OverrideRootID);
+
+	// Process execution tokens
+	ExecutionTokens.Apply(Player);
+}
+
+void FMovieSceneRootEvaluationTemplateInstance::GatherEntities(const FMovieSceneEvaluationGroup& Group, IMovieScenePlayer& Player, FMovieSceneSequenceID OverrideRootID)
+{
+	MOVIESCENE_DETAILED_SCOPE_CYCLE_COUNTER(MovieSceneEval_GatherEntries);
+
+	// See what's happening this frame
+	for (const FMovieSceneEvaluationGroupLUTIndex& Index : Group.LUTIndices)
+	{
+		for (int32 PreEvalTrackIndex = Index.LUTOffset; PreEvalTrackIndex < Index.LUTOffset + Index.NumInitPtrs + Index.NumEvalPtrs; ++PreEvalTrackIndex)
+		{
+			FMovieSceneEvaluationFieldSegmentPtr SegmentPtr = Group.SegmentPtrLUT[PreEvalTrackIndex];
+			
+			// Ensure we're able to find the sequence instance in our root if we've overridden
+			SegmentPtr.SequenceID = GetSequenceIdForRoot(SegmentPtr.SequenceID, OverrideRootID);
+
+			const FMovieSceneEvaluationTemplateInstance& Instance = GetInstanceChecked(SegmentPtr.SequenceID);
+
+			if (const FMovieSceneEvaluationTrack* Track = Instance.Template->FindTrack(SegmentPtr.TrackIdentifier))
+			{
+				const FMovieSceneSegment& Segment = Track->GetSegment(SegmentPtr.SegmentIndex);
+
+				FMovieSceneEvaluationKey TrackKey(SegmentPtr.SequenceID, SegmentPtr.TrackIdentifier);
+
+				// Important: Add the track key first. When tearing down unevaluated entites, we do so in a strictly reverse order.
+				// This is necessary as sections may depend on track data, but tracks should have no knowledge of their inner section data
+				EntitiesEvaluatedThisFrame.Add(TrackKey);
+
+				// Add all the sections
+				for (const FSectionEvaluationData& EvalData : Segment.Impls)
+				{
+					EntitiesEvaluatedThisFrame.Add(TrackKey.AsSection(EvalData.ImplIndex));
+				}
+			}
+		}
+	}
+}
+
+void FMovieSceneRootEvaluationTemplateInstance::EvaluateGroup(const FMovieSceneEvaluationGroup& Group, const FMovieSceneContext& RootContext, IMovieScenePlayer& Player, FMovieSceneExecutionTokens& ExecutionTokens, FMovieSceneSequenceID OverrideRootID) const
+{
+	MOVIESCENE_DETAILED_SCOPE_CYCLE_COUNTER(MovieSceneEval_EvaluateGroup);
+
+	FPersistentEvaluationData PersistentDataProxy(Player.State.PersistentEntityData, Player.State.PersistentSharedData);
+	
+	FMovieSceneEvaluationOperand Operand;
+
+	FMovieSceneContext Context = RootContext;
+	FMovieSceneContext SubContext = Context;
+
+	for (const FMovieSceneEvaluationGroupLUTIndex& Index : Group.LUTIndices)
+	{
+		int32 TrackIndex = Index.LUTOffset;
+		
+		// Initialize anything that wants to be initialized first
+		for ( ; TrackIndex < Index.LUTOffset + Index.NumInitPtrs; ++TrackIndex)
+		{
+			FMovieSceneEvaluationFieldSegmentPtr SegmentPtr = Group.SegmentPtrLUT[TrackIndex];
+
+			// Ensure we're able to find the sequence instance in our root if we've overridden
+			SegmentPtr.SequenceID = GetSequenceIdForRoot(SegmentPtr.SequenceID, OverrideRootID);
+
+			const FMovieSceneEvaluationTemplateInstance& Instance = GetInstanceChecked(SegmentPtr.SequenceID);
+			const FMovieSceneEvaluationTrack* Track = Instance.Template->FindTrack(SegmentPtr.TrackIdentifier);
+
+			if (Track)
+			{
+				Operand.ObjectBindingID = Track->GetObjectBindingID();
+				Operand.SequenceID = SegmentPtr.SequenceID;
+				
+				FMovieSceneEvaluationKey TrackKey(SegmentPtr.SequenceID, SegmentPtr.TrackIdentifier);
+
+				PersistentDataProxy.SetTrackKey(TrackKey);
+				Player.PreAnimatedState.SetCaptureEntity(TrackKey, EMovieSceneCompletionMode::KeepState);
+
+				SubContext = Context;
+				if (SegmentPtr.SequenceID != MovieSceneSequenceID::Root)
+				{
+					SubContext = Context.Transform(Instance.RootToSequenceTransform);
+				}
+
+				Track->Initialize(SegmentPtr.SegmentIndex, Operand, SubContext, PersistentDataProxy, Player);
+			}
+		}
+
+		// Then evaluate
+
+		// *Threading candidate*
+		// @todo: if we want to make this threaded, we need to:
+		//  - Make the execution tokens threadsafe, and sortable (one container per thread + append?)
+		//  - Do the above in a lockless manner
+		for (; TrackIndex < Index.LUTOffset + Index.NumInitPtrs + Index.NumEvalPtrs; ++TrackIndex)
+		{
+			FMovieSceneEvaluationFieldSegmentPtr SegmentPtr = Group.SegmentPtrLUT[TrackIndex];
+			
+			// Ensure we're able to find the sequence instance in our root if we've overridden
+			SegmentPtr.SequenceID = GetSequenceIdForRoot(SegmentPtr.SequenceID, OverrideRootID);
+
+			const FMovieSceneEvaluationTemplateInstance& Instance = GetInstanceChecked(SegmentPtr.SequenceID);
+			const FMovieSceneEvaluationTrack* Track = Instance.Template->FindTrack(SegmentPtr.TrackIdentifier);
+
+			if (Track)
+			{
+				Operand.ObjectBindingID = Track->GetObjectBindingID();
+				Operand.SequenceID = SegmentPtr.SequenceID;
+
+				FMovieSceneEvaluationKey TrackKey(SegmentPtr.SequenceID, SegmentPtr.TrackIdentifier);
+
+				PersistentDataProxy.SetTrackKey(TrackKey);
+
+				ExecutionTokens.SetOperand(Operand);
+				ExecutionTokens.SetTrack(TrackKey);
+
+				Player.PreAnimatedState.SetCaptureEntity(TrackKey, EMovieSceneCompletionMode::KeepState);
+
+				SubContext = Context;
+				if (SegmentPtr.SequenceID != MovieSceneSequenceID::Root)
+				{
+					SubContext = Context.Transform(Instance.RootToSequenceTransform);
+				}
+
+				Track->Evaluate(
+					SegmentPtr.SegmentIndex,
+					Operand,
+					SubContext,
+					PersistentDataProxy,
+					ExecutionTokens);
+			}
+		}
+
+		if (Index.bRequiresImmediateFlush)
+		{
+			ExecutionTokens.Apply(Player);
+		}
+	}
+}
+
+void FMovieSceneRootEvaluationTemplateInstance::CallSetupTearDown(IMovieScenePlayer& Player)
+{
+	MOVIESCENE_DETAILED_SCOPE_CYCLE_COUNTER(MovieSceneEval_CallSetupTearDown);
+
+	FPersistentEvaluationData PersistentDataProxy(Player.State.PersistentEntityData, Player.State.PersistentSharedData);
+	
+	// Always tear down entites in reverse order
+	Algo::Reverse(EntitiesEvaluatedLastFrame.OrderedKeys);
+	for (const FMovieSceneEvaluationKey& Key : EntitiesEvaluatedLastFrame.OrderedKeys)
+	{
+		if (EntitiesEvaluatedThisFrame.Contains(Key))
+		{
+			continue;
+		}
+
+		const FMovieSceneEvaluationTemplateInstance* Instance = FindInstance(Key.SequenceID);
+		if (!Instance)
+		{
+			continue;
+		}
+
+		const FMovieSceneEvaluationTrack* Track = Instance->Template->FindTrack(Key.TrackIdentifier);
+
+		// Track data key may be required by both tracks and sections
+		PersistentDataProxy.SetTrackKey(Key.AsTrack());
+
+		if (Key.SectionIdentifier == uint32(-1))
+		{
+			if (Track)
+			{
+				Track->OnEndEvaluation(PersistentDataProxy, Player);
+			}
+			
+			Player.PreAnimatedState.RestorePreAnimatedState(Player, Key);
+			PersistentDataProxy.ResetTrackData();
+		}
+		else
+		{
+			PersistentDataProxy.SetSectionKey(Key);
+			if (Track)
+			{
+				Track->GetChildTemplate(Key.SectionIdentifier).OnEndEvaluation(PersistentDataProxy, Player);
+			}
+
+			Player.PreAnimatedState.RestorePreAnimatedState(Player, Key);
+			PersistentDataProxy.ResetSectionData();
+		}
+	}
+
+	for (const FMovieSceneEvaluationKey& Key : EntitiesEvaluatedThisFrame.OrderedKeys)
+	{
+		if (EntitiesEvaluatedLastFrame.Contains(Key))
+		{
+			continue;
+		}
+
+		const FMovieSceneEvaluationTemplateInstance& Instance = GetInstanceChecked(Key.SequenceID);
+
+		if (const FMovieSceneEvaluationTrack* Track = Instance.Template->FindTrack(Key.TrackIdentifier))
+		{
+			PersistentDataProxy.SetTrackKey(Key.AsTrack());
+
+			if (Key.SectionIdentifier == uint32(-1))
+			{
+				Track->OnBeginEvaluation(PersistentDataProxy, Player);
+			}
+			else
+			{
+				PersistentDataProxy.SetSectionKey(Key);
+				Track->GetChildTemplate(Key.SectionIdentifier).OnBeginEvaluation(PersistentDataProxy, Player);
+			}
+		}
+	}
+
+	// Determine which sub sequences are no longer evaluated this frame.
+	// This algorithm works on the premise that each array is sorted, and each ID can only appear once
+	FMovieSceneSpawnRegister& Register = Player.GetSpawnRegister();
+
+	auto ThisFrameIDs = ActiveSequencesThisFrame.CreateConstIterator();
+	for (FMovieSceneSequenceID LastID : ActiveSequencesLastFrame)
+	{
+		if (ThisFrameIDs)
+		{
+			FMovieSceneSequenceID ThisID = *ThisFrameIDs;
+			if (LastID == ThisID)
+			{
+				continue;
+			}
+			else if (LastID < ThisID)
+			{
+				Register.OnSequenceExpired(LastID, Player);
+			}
+			else
+			{
+				++ThisFrameIDs;
+			}
+		}
+		else
+		{
+			Register.OnSequenceExpired(LastID, Player);
+		}
+	}
+}
+
+void FMovieSceneRootEvaluationTemplateInstance::OnSequenceChanged()
+{
+	bIsDirty = true;
+}
\ No newline at end of file
diff --git a/Engine/Source/Runtime/MovieScene/Private/Evaluation/MovieSceneEvaluationTrack.cpp b/Engine/Source/Runtime/MovieScene/Private/Evaluation/MovieSceneEvaluationTrack.cpp
new file mode 100644
--- /dev/null
+++ b/Engine/Source/Runtime/MovieScene/Private/Evaluation/MovieSceneEvaluationTrack.cpp
@@ -0,0 +1,223 @@
+// Copyright 1998-2016 Epic Games, Inc. All Rights Reserved.
+
+#include "MovieScenePrivatePCH.h"
+#include "MovieSceneEvaluationTrack.h"
+
+FMovieSceneEvaluationTrack::FMovieSceneEvaluationTrack()
+{
+}
+
+FMovieSceneEvaluationTrack::FMovieSceneEvaluationTrack(const FGuid& InObjectBindingID)
+	: ObjectBindingID(InObjectBindingID)
+	, EvaluationPriority(1000)
+	, EvaluationMethod(EEvaluationMethod::Static)
+	, bEvaluateInPreroll(true)
+	, bEvaluateInPostroll(true)
+{
+}
+
+void FMovieSceneEvaluationTrack::DefineAsSingleTemplate(FMovieSceneEvalTemplatePtr&& InTemplate)
+{
+	ChildTemplates.Reset(1);
+	Segments.Reset(1);
+
+	ChildTemplates.Add(MoveTemp(InTemplate));
+	
+	int32 SegmentIndex = 0;
+	Segments.Add(FMovieSceneSegment(TRange<float>::All(), TArrayView<int32>(&SegmentIndex, 1)));
+}
+
+int32 FMovieSceneEvaluationTrack::AddChildTemplate(FMovieSceneEvalTemplatePtr&& InTemplate)
+{
+	return ChildTemplates.Add(MoveTemp(InTemplate));
+}
+
+void FMovieSceneEvaluationTrack::SetSegments(TArray<FMovieSceneSegment>&& InSegments)
+{
+	Segments = MoveTemp(InSegments);
+
+	ValidateSegments();
+}
+
+void FMovieSceneEvaluationTrack::ValidateSegments()
+{
+	// We don't remove segments as this may break ptrs that have been set up in the evaluation field. Instead we just remove invalid track indices.
+
+	for (FMovieSceneSegment& Segment : Segments)
+	{
+		for (int32 Index = Segment.Impls.Num() - 1; Index >= 0; --Index)
+		{
+			int32 TemplateIndex = Segment.Impls[Index].ImplIndex;
+			if (!HasChildTemplate(TemplateIndex))
+			{
+				Segment.Impls.RemoveAt(Index, 1, false);
+			}
+		}
+	}
+}
+
+void FMovieSceneEvaluationTrack::SetupOverrides()
+{
+	for (FMovieSceneEvalTemplatePtr& ChildTemplate : ChildTemplates)
+	{
+		if (ChildTemplate.IsValid())
+		{
+			ChildTemplate->SetupOverrides();
+		}
+	}
+
+	if (TrackTemplate.IsValid())
+	{
+		TrackTemplate->SetupOverrides();
+	}
+}
+
+void FMovieSceneEvaluationTrack::Initialize(int32 SegmentIndex, const FMovieSceneEvaluationOperand& Operand, const FMovieSceneContext& Context, FPersistentEvaluationData& PersistentData, IMovieScenePlayer& Player) const
+{
+	if (!TrackTemplate.IsValid() || !TrackTemplate->HasCustomInitialize())
+	{
+		DefaultInitialize(SegmentIndex, Operand, Context, PersistentData, Player);
+	}
+	else
+	{
+		TrackTemplate->Initialize(*this, SegmentIndex, Operand, Context, PersistentData, Player);
+	}
+}
+
+void FMovieSceneEvaluationTrack::DefaultInitialize(int32 SegmentIndex, const FMovieSceneEvaluationOperand& Operand, FMovieSceneContext Context, FPersistentEvaluationData& PersistentData, IMovieScenePlayer& Player) const
+{
+	const FMovieSceneSegment& Segment = Segments[SegmentIndex];
+	for (const FSectionEvaluationData& EvalData : Segment.Impls)
+	{
+		const FMovieSceneEvalTemplate& Template = GetChildTemplate(EvalData.ImplIndex);
+
+		if (Template.RequiresInitialization())
+		{
+			PersistentData.DeriveSectionKey(EvalData.ImplIndex);
+			
+			Context.OverrideTime(EvalData.ForcedTime);
+
+			Template.Initialize(
+				Operand,
+				Context,
+				PersistentData,
+				Player);
+		}
+	}
+}
+
+void FMovieSceneEvaluationTrack::Evaluate(int32 SegmentIndex, const FMovieSceneEvaluationOperand& Operand, const FMovieSceneContext& Context, const FPersistentEvaluationData& PersistentData, FMovieSceneExecutionTokens& ExecutionTokens) const
+{
+	ExecutionTokens.SetCompletionMode(EMovieSceneCompletionMode::KeepState);
+
+	if (!TrackTemplate.IsValid() || !TrackTemplate->HasCustomEvaluate())
+	{
+		DefaultEvaluate(SegmentIndex, Operand, Context, PersistentData, ExecutionTokens);
+	}
+	else
+	{
+		TrackTemplate->Evaluate(*this, SegmentIndex, Operand, Context, PersistentData, ExecutionTokens);
+	}
+}
+
+void FMovieSceneEvaluationTrack::DefaultEvaluate(int32 SegmentIndex, const FMovieSceneEvaluationOperand& Operand, const FMovieSceneContext& Context, const FPersistentEvaluationData& PersistentData, FMovieSceneExecutionTokens& ExecutionTokens) const
+{
+	switch (EvaluationMethod)
+	{
+	case EEvaluationMethod::Static:
+		EvaluateStatic(SegmentIndex, Operand, Context, PersistentData, ExecutionTokens);
+		break;
+	case EEvaluationMethod::Swept:
+		EvaluateSwept(SegmentIndex, Operand, Context, PersistentData, ExecutionTokens);
+		break;
+	}
+}
+
+void FMovieSceneEvaluationTrack::EvaluateStatic(int32 SegmentIndex, const FMovieSceneEvaluationOperand& Operand, FMovieSceneContext Context, const FPersistentEvaluationData& PersistentData, FMovieSceneExecutionTokens& ExecutionTokens) const
+{
+	for (const FSectionEvaluationData& EvalData : GetSegment(SegmentIndex).Impls)
+	{
+		const FMovieSceneEvalTemplate& Template = GetChildTemplate(EvalData.ImplIndex);
+		
+		Context.OverrideTime(EvalData.ForcedTime);
+
+		PersistentData.DeriveSectionKey(EvalData.ImplIndex);
+		ExecutionTokens.SetSectionIdentifier(EvalData.ImplIndex);
+		ExecutionTokens.SetCompletionMode(Template.GetCompletionMode());
+		ExecutionTokens.SetContext(Context);
+
+		Template.Evaluate(
+			Operand,
+			Context,
+			PersistentData,
+			ExecutionTokens);
+	}
+}
+
+namespace
+{
+	bool IntersectSegmentRanges(const FMovieSceneSegment& Segment, TRange<float> TraversedRange, TMap<int32, TRange<float>>& ImplToAccumulatedRange)
+	{
+		TRange<float> Intersection = TRange<float>::Intersection(Segment.Range, TraversedRange);
+		if (Intersection.IsEmpty())
+		{
+			return false;
+		}
+
+		for (const FSectionEvaluationData& EvalData : Segment.Impls)
+		{
+			TRange<float>* AccumulatedRange = ImplToAccumulatedRange.Find(EvalData.ImplIndex);
+			if (!AccumulatedRange)
+			{
+				ImplToAccumulatedRange.Add(EvalData.ImplIndex, Intersection);
+			}
+			else
+			{
+				*AccumulatedRange = TRange<float>::Hull(*AccumulatedRange, Intersection);
+			}
+		}
+
+		return true;
+	}
+}
+
+void FMovieSceneEvaluationTrack::EvaluateSwept(int32 SegmentIndex, const FMovieSceneEvaluationOperand& Operand, const FMovieSceneContext& Context, const FPersistentEvaluationData& PersistentData, FMovieSceneExecutionTokens& ExecutionTokens) const
+{
+	TRange<float> TraversedRange = Context.GetRange();
+
+	// Accumulate the relevant ranges that each section intersects with the evaluated range
+	TMap<int32, TRange<float>> ImplToAccumulatedRange;
+
+	// Search backwards from the current segment for any segments intersecting the traversed range
+	{
+		int32 PreviousIndex = SegmentIndex;
+		while(--PreviousIndex >= 0 && IntersectSegmentRanges(Segments[PreviousIndex], TraversedRange, ImplToAccumulatedRange));
+	}
+
+	// Obviously the current segment intersects otherwise we wouldn't be in here
+	IntersectSegmentRanges(Segments[SegmentIndex], TraversedRange, ImplToAccumulatedRange);
+
+	// Search forwards from the current segment for any segments intersecting the traversed range
+	{
+		int32 NextIndex = SegmentIndex;
+		while(++NextIndex < Segments.Num() && IntersectSegmentRanges(Segments[NextIndex], TraversedRange, ImplToAccumulatedRange));
+	}
+
+	for (auto& Pair : ImplToAccumulatedRange)
+	{
+		const int32 SectionIndex = Pair.Key;
+		TRange<float> EvaluationRange = Pair.Value;
+		const FMovieSceneEvalTemplate& Template = GetChildTemplate(SectionIndex);
+
+		PersistentData.DeriveSectionKey(SectionIndex);
+		ExecutionTokens.SetSectionIdentifier(SectionIndex);
+		ExecutionTokens.SetCompletionMode(Template.GetCompletionMode());
+		ExecutionTokens.SetContext(Context);
+
+		Template.EvaluateSwept(
+			Operand,
+			Context.Clamp(EvaluationRange),
+			PersistentData,
+			ExecutionTokens);
+	}
+}
\ No newline at end of file
diff --git a/Engine/Source/Runtime/MovieScene/Private/Evaluation/MovieSceneExecutionTokens.cpp b/Engine/Source/Runtime/MovieScene/Private/Evaluation/MovieSceneExecutionTokens.cpp
new file mode 100644
--- /dev/null
+++ b/Engine/Source/Runtime/MovieScene/Private/Evaluation/MovieSceneExecutionTokens.cpp
@@ -0,0 +1,30 @@
+// Copyright 1998-2016 Epic Games, Inc. All Rights Reserved.
+
+#include "MovieScenePrivatePCH.h"
+#include "MovieSceneExecutionTokens.h"
+
+DECLARE_CYCLE_STAT(TEXT("Apply Execution Tokens"), MovieSceneEval_ApplyExecutionTokens, STATGROUP_MovieSceneEval);
+DECLARE_CYCLE_STAT(TEXT("Apply Execution Token"), MovieSceneEval_ApplyExecutionToken, STATGROUP_MovieSceneEval);
+
+void FMovieSceneExecutionTokens::Apply(IMovieScenePlayer& Player)
+{
+	MOVIESCENE_DETAILED_SCOPE_CYCLE_COUNTER(MovieSceneEval_ApplyExecutionTokens);
+
+	FPersistentEvaluationData PersistentDataProxy(Player.State.PersistentEntityData, Player.State.PersistentSharedData);
+	for (FMovieSceneExecutionTokens::FEntry& Entry : Tokens)
+	{
+		PersistentDataProxy.SetTrackKey(Entry.TrackKey);
+
+		FMovieSceneEvaluationKey SectionKey(Entry.TrackKey.AsSection(Entry.SectionIdentifier));
+		PersistentDataProxy.SetSectionKey(SectionKey);
+		
+		Player.PreAnimatedState.SetCaptureEntity(SectionKey, Entry.CompletionMode);
+
+		{
+			MOVIESCENE_DETAILED_SCOPE_CYCLE_COUNTER(MovieSceneEval_ApplyExecutionToken);
+			Entry.Token->Execute(Entry.Context, Entry.Operand, PersistentDataProxy, Player);
+		}
+	}
+
+	Tokens.Reset();
+}
\ No newline at end of file
diff --git a/Engine/Source/Runtime/MovieScene/Private/Evaluation/MovieSceneLegacyTrackInstanceTemplate.cpp b/Engine/Source/Runtime/MovieScene/Private/Evaluation/MovieSceneLegacyTrackInstanceTemplate.cpp
new file mode 100644
--- /dev/null
+++ b/Engine/Source/Runtime/MovieScene/Private/Evaluation/MovieSceneLegacyTrackInstanceTemplate.cpp
@@ -0,0 +1,122 @@
+// Copyright 1998-2016 Epic Games, Inc. All Rights Reserved.
+
+#include "MovieScenePrivatePCH.h"
+#include "MovieSceneLegacyTrackInstanceTemplate.h"
+#include "MovieSceneEvaluationTemplateInstance.h"
+
+/** Legacy tracks just get a new ID each time (so they always save/restore state, regardless of whether a similar track has already animated the object. That was the old behaviour) */
+struct FMovieSceneLegacyAnimTypeID : FMovieSceneAnimTypeID
+{
+	FMovieSceneLegacyAnimTypeID(uint32 InSeed)
+	{
+		ID = GenerateHash(&UnusedStatic, InSeed);
+	}
+	static uint64 UnusedStatic;
+};
+
+uint64 FMovieSceneLegacyAnimTypeID::UnusedStatic = 0;
+
+static uint32 LegacyTrackIdSeed = 0;
+IMovieSceneTrackInstance::IMovieSceneTrackInstance()
+	: AnimTypeID(FMovieSceneLegacyAnimTypeID(LegacyTrackIdSeed++))
+{
+}
+
+struct FLegacyTrackData : IPersistentEvaluationData
+{
+	TSharedPtr<IMovieSceneTrackInstance> TrackInstance;
+};
+
+struct FRestoreStateToken : IMovieScenePreAnimatedGlobalToken
+{
+	TArray<TWeakObjectPtr<>> RuntimeObjects;
+	TSharedPtr<FMovieSceneSequenceInstance> LegacySequence;
+	TSharedPtr<IMovieSceneTrackInstance> LegacyTrackInstance;
+
+	virtual void RestoreState(IMovieScenePlayer& InPlayer) override
+	{
+		LegacyTrackInstance->RestoreState(RuntimeObjects, InPlayer, *LegacySequence);
+	}
+};
+
+struct FLegacyPreAnimatedStateProducer : IMovieScenePreAnimatedGlobalTokenProducer
+{
+	FLegacyPreAnimatedStateProducer(TFunctionRef<IMovieScenePreAnimatedGlobalTokenPtr()> In) : Producer(MoveTemp(In)) {}
+
+	TFunctionRef<IMovieScenePreAnimatedGlobalTokenPtr()> Producer;
+
+	virtual IMovieScenePreAnimatedGlobalTokenPtr CacheExistingState() const override
+	{
+		return Producer();
+	}
+};
+
+struct FLegacyExecutionToken : IMovieSceneExecutionToken
+{
+	virtual void Execute(const FMovieSceneContext& Context, const FMovieSceneEvaluationOperand& Operand, FPersistentEvaluationData& PersistentData, IMovieScenePlayer& Player) override
+	{
+		const FMovieSceneEvaluationTemplateInstance* TemplateInstance = Player.GetEvaluationTemplate().GetInstance(Operand.SequenceID);
+
+		TSharedPtr<FMovieSceneSequenceInstance> LegacySequence = TemplateInstance ? TemplateInstance->LegacySequenceInstance : nullptr;
+		TSharedPtr<IMovieSceneTrackInstance> LegacyTrackInstance = PersistentData.GetSectionData<FLegacyTrackData>().TrackInstance;
+		
+		if (!ensure(LegacyTrackInstance.IsValid() && LegacySequence.IsValid()))
+		{
+			return;
+		}
+
+		EMovieSceneUpdateData UpdateData(Context.GetTime(), Context.GetPreviousTime());
+		UpdateData.bJumpCut = Context.HasJumped();
+
+		TArray<TWeakObjectPtr<>> RuntimeObjects;
+		for (TWeakObjectPtr<> Obj : Player.FindBoundObjects(Operand))
+		{
+			RuntimeObjects.Add(Obj);
+		}
+
+		FLegacyPreAnimatedStateProducer Producer([&]() -> IMovieScenePreAnimatedGlobalTokenPtr{
+			LegacyTrackInstance->SaveState(RuntimeObjects, Player, *LegacySequence);
+
+			FRestoreStateToken Token;
+			Token.RuntimeObjects = RuntimeObjects;
+			Token.LegacySequence = LegacySequence;
+			Token.LegacyTrackInstance = LegacyTrackInstance;
+			return MoveTemp(Token);
+		});
+		
+		Player.SavePreAnimatedState(LegacyTrackInstance->AnimTypeID, Producer);
+
+		UpdateData.UpdatePass = MSUP_PreUpdate;
+		if (LegacyTrackInstance->HasUpdatePasses() & UpdateData.UpdatePass)
+		{
+			LegacyTrackInstance->Update(UpdateData, RuntimeObjects, Player, *LegacySequence);
+		}
+
+		UpdateData.UpdatePass = MSUP_Update;
+		if (LegacyTrackInstance->HasUpdatePasses() & UpdateData.UpdatePass)
+		{
+			LegacyTrackInstance->Update(UpdateData, RuntimeObjects, Player, *LegacySequence);
+		}
+
+		UpdateData.UpdatePass = MSUP_PostUpdate;
+		if (LegacyTrackInstance->HasUpdatePasses() & UpdateData.UpdatePass)
+		{
+			LegacyTrackInstance->Update(UpdateData, RuntimeObjects, Player, *LegacySequence);
+		}
+	}
+};
+
+FMovieSceneLegacyTrackInstanceTemplate::FMovieSceneLegacyTrackInstanceTemplate(const UMovieSceneTrack* InTrack)
+	: Track(InTrack)
+{
+}
+
+void FMovieSceneLegacyTrackInstanceTemplate::Evaluate(const FMovieSceneEvaluationOperand& Operand, const FMovieSceneContext& Context, const FPersistentEvaluationData& PersistentData, FMovieSceneExecutionTokens& ExecutionTokens) const
+{
+	ExecutionTokens.Add(FLegacyExecutionToken());
+}
+
+void FMovieSceneLegacyTrackInstanceTemplate::Setup(FPersistentEvaluationData& PersistentData, IMovieScenePlayer& Player) const
+{
+	PersistentData.AddSectionData<FLegacyTrackData>().TrackInstance = Track->CreateLegacyInstance();
+}
\ No newline at end of file
diff --git a/Engine/Source/Runtime/MovieScene/Private/Evaluation/MovieScenePlayback.cpp b/Engine/Source/Runtime/MovieScene/Private/Evaluation/MovieScenePlayback.cpp
new file mode 100644
--- /dev/null
+++ b/Engine/Source/Runtime/MovieScene/Private/Evaluation/MovieScenePlayback.cpp
@@ -0,0 +1,81 @@
+// Copyright 1998-2016 Epic Games, Inc. All Rights Reserved.
+
+#include "MovieScenePrivatePCH.h"
+#include "MovieScenePlayback.h"
+
+#include "MovieScene.h"
+
+namespace 
+{
+	TRange<float> CalculateEvaluationRange(float CurrentTime, float PreviousTime, bool bInclusivePreviousTime)
+	{
+		if (CurrentTime == PreviousTime)
+		{
+			return TRange<float>(CurrentTime);
+		}
+		else if (CurrentTime < PreviousTime)
+		{
+			return TRange<float>(
+				TRangeBound<float>::Inclusive(CurrentTime),
+				bInclusivePreviousTime
+					? TRangeBound<float>::Inclusive(PreviousTime)
+					: TRangeBound<float>::Exclusive(PreviousTime)
+				);
+		}
+
+		return TRange<float>(
+			bInclusivePreviousTime
+				? TRangeBound<float>::Inclusive(PreviousTime)
+				: TRangeBound<float>::Exclusive(PreviousTime)
+			, TRangeBound<float>::Inclusive(CurrentTime)
+			);
+	}
+}
+
+FMovieSceneEvaluationRange::FMovieSceneEvaluationRange(float InTime)
+	: EvaluationRange(InTime)
+	, Direction(EPlayDirection::Forwards)
+	, TimeOverride(TNumericLimits<float>::Lowest())
+{
+}
+
+FMovieSceneEvaluationRange::FMovieSceneEvaluationRange(TRange<float> InRange, EPlayDirection InDirection)
+	: EvaluationRange(InRange)
+	, Direction(InDirection)
+	, TimeOverride(TNumericLimits<float>::Lowest())
+{
+}
+
+FMovieSceneEvaluationRange::FMovieSceneEvaluationRange(float InCurrentTime, float InPreviousTime, bool bInclusivePreviousTime)
+	: EvaluationRange(CalculateEvaluationRange(InCurrentTime, InPreviousTime, bInclusivePreviousTime))
+	, Direction((InCurrentTime - InPreviousTime >= 0) ? EPlayDirection::Forwards : EPlayDirection::Backwards)
+	, TimeOverride(TNumericLimits<float>::Lowest())
+{
+}
+
+void FMovieScenePlaybackPosition::Reset(float StartPos)
+{
+	PreviousPosition = StartPos;
+	PreviousPlayEvalPosition.Reset();
+}
+
+FMovieSceneEvaluationRange FMovieScenePlaybackPosition::JumpTo(float NewPosition, TOptional<float> FixedInterval)
+{
+	Reset(NewPosition);
+
+	PreviousPosition = NewPosition;
+	float EvalPosition = FixedInterval.IsSet() ? UMovieScene::CalculateFixedFrameTime(NewPosition, FixedInterval.GetValue()) : NewPosition;
+	return FMovieSceneEvaluationRange(TRange<float>(EvalPosition), EPlayDirection::Forwards);
+}
+
+FMovieSceneEvaluationRange FMovieScenePlaybackPosition::PlayTo(float NewPosition, TOptional<float> FixedInterval)
+{
+	float EvalPositionTo = FixedInterval.IsSet() ? UMovieScene::CalculateFixedFrameTime(NewPosition, FixedInterval.GetValue()) : NewPosition;
+	float EvalPositionFrom = PreviousPlayEvalPosition.Get(PreviousPosition.Get(NewPosition));
+
+	FMovieSceneEvaluationRange Range(EvalPositionTo, EvalPositionFrom, !PreviousPlayEvalPosition.IsSet());
+
+	PreviousPosition = NewPosition;
+	PreviousPlayEvalPosition = EvalPositionTo;
+	return Range;
+}
diff --git a/Engine/Source/Runtime/MovieScene/Private/Evaluation/MovieScenePreAnimatedState.cpp b/Engine/Source/Runtime/MovieScene/Private/Evaluation/MovieScenePreAnimatedState.cpp
new file mode 100644
--- /dev/null
+++ b/Engine/Source/Runtime/MovieScene/Private/Evaluation/MovieScenePreAnimatedState.cpp
@@ -0,0 +1,198 @@
+// Copyright 1998-2016 Epic Games, Inc. All Rights Reserved.
+
+#include "MovieScenePrivatePCH.h"
+#include "MovieScenePreAnimatedState.h"
+
+DECLARE_CYCLE_STAT(TEXT("Save Pre Animated State"), MovieSceneEval_SavePreAnimatedState, STATGROUP_MovieSceneEval);
+
+void FMovieSceneSavedObjectTokens::Add(ECapturePreAnimatedState CaptureState, FMovieSceneAnimTypeID InAnimTypeID, FMovieSceneEvaluationKey AssociatedKey, const IMovieScenePreAnimatedTokenProducer& Producer, UObject& InObject, IMovieScenePreAnimatedState& Parent)
+{
+	MOVIESCENE_DETAILED_SCOPE_CYCLE_COUNTER(MovieSceneEval_SavePreAnimatedState)
+
+	switch(CaptureState)
+	{
+	case ECapturePreAnimatedState::Entity:
+		{
+			FMovieSceneEntityAndAnimTypeID EntityAndTypeID{AssociatedKey, InAnimTypeID};
+			if (AnimatedEntityTypeIDs.Contains(EntityAndTypeID))
+			{
+				return;
+			}
+			else
+			{
+				AnimatedEntityTypeIDs.Add(EntityAndTypeID);
+
+				FRefCountedPreAnimatedObjectToken* Token = EntityTokens.FindByPredicate([=](const FRefCountedPreAnimatedObjectToken& InToken){ return InToken.AnimTypeID == InAnimTypeID; });
+				if (!Token)
+				{
+					EntityTokens.Add(FRefCountedPreAnimatedObjectToken(Producer.CacheExistingState(InObject), InAnimTypeID));
+				}
+				else
+				{
+					++Token->RefCount;
+				}
+
+				Parent.EntityHasAnimatedObject(AssociatedKey, FObjectKey(&InObject));
+			}
+		}
+		// fallthrough
+
+	case ECapturePreAnimatedState::Global:
+		if (!AnimatedGlobalTypeIDs.Contains(InAnimTypeID))
+		{
+			AnimatedGlobalTypeIDs.Add(InAnimTypeID);
+			GlobalStateTokens.Add(Producer.CacheExistingState(InObject));
+		}
+		break;
+
+	default:
+		break;
+	}
+}
+
+void FMovieSceneSavedObjectTokens::Restore(IMovieScenePlayer& Player)
+{
+	UObject* ObjectPtr = Object.Get();
+	if (!ObjectPtr)
+	{
+		return;
+	}
+
+	// Restore in reverse
+	for (int32 Index = GlobalStateTokens.Num() - 1; Index >= 0; --Index)
+	{
+		GlobalStateTokens[Index]->RestoreState(*ObjectPtr, Player);
+	}
+
+	Reset();
+}
+
+void FMovieSceneSavedObjectTokens::Restore(IMovieScenePlayer& Player, TFunctionRef<bool(FMovieSceneAnimTypeID)> InFilter)
+{
+	UObject* ObjectPtr = Object.Get();
+	if (!ObjectPtr)
+	{
+		return;
+	}
+
+	TArray<FMovieSceneAnimTypeID> TypesToRemove;
+
+	// Restore in reverse
+	for (int32 Index = AnimatedGlobalTypeIDs.Num() - 1; Index >= 0; --Index)
+	{
+		if (InFilter(AnimatedGlobalTypeIDs[Index]))
+		{
+			TypesToRemove.Add(AnimatedGlobalTypeIDs[Index]);
+
+			GlobalStateTokens[Index]->RestoreState(*ObjectPtr, Player);
+
+			GlobalStateTokens.RemoveAtSwap(Index);
+			AnimatedGlobalTypeIDs.RemoveAtSwap(Index);
+		}
+	}
+
+	EntityTokens.RemoveAll(
+		[&](const FRefCountedPreAnimatedObjectToken& In){
+			return TypesToRemove.Contains(In.AnimTypeID);
+		}
+	);
+
+	AnimatedEntityTypeIDs.RemoveAll(
+		[&](const FMovieSceneEntityAndAnimTypeID& In){
+			return TypesToRemove.Contains(In.AnimTypeID);
+		}
+	);
+}
+
+void FMovieSceneSavedMasterTokens::Add(ECapturePreAnimatedState CaptureState, FMovieSceneAnimTypeID InAnimTypeID, FMovieSceneEvaluationKey AssociatedKey, const IMovieScenePreAnimatedGlobalTokenProducer& Producer, IMovieScenePreAnimatedState& Parent)
+{
+	MOVIESCENE_DETAILED_SCOPE_CYCLE_COUNTER(MovieSceneEval_SavePreAnimatedState)
+
+	switch(CaptureState)
+	{
+	case ECapturePreAnimatedState::Entity:
+		{
+			FMovieSceneEntityAndAnimTypeID EntityAndTypeID{AssociatedKey, InAnimTypeID};
+			if (AnimatedEntityTypeIDs.Contains(EntityAndTypeID))
+			{
+				return;
+			}
+			else
+			{
+				AnimatedEntityTypeIDs.Add(EntityAndTypeID);
+
+				FRefCountedPreAnimatedGlobalToken* Token = EntityTokens.FindByPredicate([=](const FRefCountedPreAnimatedGlobalToken& InToken){ return InToken.AnimTypeID == InAnimTypeID; });
+				if (!Token)
+				{
+					EntityTokens.Add(FRefCountedPreAnimatedGlobalToken(Producer.CacheExistingState(), InAnimTypeID));
+				}
+				else
+				{
+					++Token->RefCount;
+				}
+
+				Parent.EntityHasAnimatedMaster(AssociatedKey);
+			}
+		}
+		// fallthrough
+
+	case ECapturePreAnimatedState::Global:
+		if (!AnimatedGlobalTypeIDs.Contains(InAnimTypeID))
+		{
+			AnimatedGlobalTypeIDs.Add(InAnimTypeID);
+			GlobalStateTokens.Add(Producer.CacheExistingState());
+		}
+		break;
+
+	default:
+		break;
+	}
+}
+
+void FMovieSceneSavedMasterTokens::Restore(IMovieScenePlayer& Player)
+{
+	// Restore in reverse
+	for (int32 Index = GlobalStateTokens.Num() - 1; Index >= 0; --Index)
+	{
+		GlobalStateTokens[Index]->RestoreState(Player);
+	}
+
+	Reset();
+}
+
+void FMovieScenePreAnimatedState::RestorePreAnimatedState(IMovieScenePlayer& Player)
+{
+	for (auto& Pair : ObjectTokens)
+	{
+		Pair.Value.Restore(Player);
+	}
+
+	MasterTokens.Restore(Player);
+
+	EntityToAnimatedObjects.Reset();
+}
+
+void FMovieScenePreAnimatedState::RestorePreAnimatedState(IMovieScenePlayer& Player, UObject& Object)
+{
+	FObjectKey ObjectKey(&Object);
+
+	FMovieSceneSavedObjectTokens* FoundObjectTokens = ObjectTokens.Find(ObjectKey);
+	if (FoundObjectTokens)
+	{
+		FoundObjectTokens->Restore(Player);
+	}
+
+	for (auto& Pair : EntityToAnimatedObjects)
+	{
+		Pair.Value.Remove(ObjectKey);
+	}
+}
+
+void FMovieScenePreAnimatedState::RestorePreAnimatedState(IMovieScenePlayer& Player, UObject& Object, TFunctionRef<bool(FMovieSceneAnimTypeID)> InFilter)
+{
+	FMovieSceneSavedObjectTokens* FoundObjectTokens = ObjectTokens.Find(&Object);
+	if (FoundObjectTokens)
+	{
+		FoundObjectTokens->Restore(Player, InFilter);
+	}
+}
\ No newline at end of file
diff --git a/Engine/Source/Runtime/MovieScene/Private/Evaluation/MovieScenePropertyTemplate.cpp b/Engine/Source/Runtime/MovieScene/Private/Evaluation/MovieScenePropertyTemplate.cpp
new file mode 100644
--- /dev/null
+++ b/Engine/Source/Runtime/MovieScene/Private/Evaluation/MovieScenePropertyTemplate.cpp
@@ -0,0 +1,19 @@
+// Copyright 1998-2016 Epic Games, Inc. All Rights Reserved.
+
+#include "MovieScenePrivatePCH.h"
+#include "MovieScenePropertyTemplate.h"
+
+static TMovieSceneAnimTypeIDContainer<FString> PropertyTypeIDs;
+
+// Default property ID to our own type - this implies an empty property
+PropertyTemplate::FSectionData::FSectionData()
+	: PropertyID(TMovieSceneAnimTypeID<FSectionData>())
+{
+
+}
+
+void PropertyTemplate::FSectionData::Initialize(FName InPropertyName, FString InPropertyPath, FName InFunctionName)
+{
+	PropertyBindings = MakeShareable(new FTrackInstancePropertyBindings(InPropertyName, MoveTemp(InPropertyPath), InFunctionName));
+	PropertyID = PropertyTypeIDs.GetAnimTypeID(InPropertyPath);
+}
\ No newline at end of file
diff --git a/Engine/Source/Runtime/MovieScene/Private/Evaluation/MovieSceneSequenceTemplateStore.cpp b/Engine/Source/Runtime/MovieScene/Private/Evaluation/MovieSceneSequenceTemplateStore.cpp
new file mode 100644
--- /dev/null
+++ b/Engine/Source/Runtime/MovieScene/Private/Evaluation/MovieSceneSequenceTemplateStore.cpp
@@ -0,0 +1,17 @@
+// Copyright 1998-2016 Epic Games, Inc. All Rights Reserved.
+
+#include "MovieScenePrivatePCH.h"
+#include "MovieSceneSequenceTemplateStore.h"
+#include "MovieSceneSequence.h"
+#include "MovieSceneTrack.h"
+
+FMovieSceneEvaluationTemplate& FMovieSceneSequenceTemplateStore::GetCompiledTemplate(UMovieSceneSequence& Sequence)
+{
+#if WITH_EDITORONLY_DATA
+	if (AreTemplatesVolatile())
+	{
+		Sequence.EvaluationTemplate.Regenerate(Sequence.TemplateParameters);
+	}
+#endif
+	return Sequence.EvaluationTemplate;
+}
\ No newline at end of file
diff --git a/Engine/Source/Runtime/MovieScene/Private/Evaluation/MovieSceneTrackImplementation.cpp b/Engine/Source/Runtime/MovieScene/Private/Evaluation/MovieSceneTrackImplementation.cpp
new file mode 100644
--- /dev/null
+++ b/Engine/Source/Runtime/MovieScene/Private/Evaluation/MovieSceneTrackImplementation.cpp
@@ -0,0 +1,10 @@
+// Copyright 1998-2016 Epic Games, Inc. All Rights Reserved.
+
+#include "MovieScenePrivatePCH.h"
+#include "MovieSceneEvalTemplate.h"
+#include "MovieSceneEvalTemplateSerializer.h"
+
+bool FMovieSceneTrackImplementationPtr::Serialize(FArchive& Ar)
+{
+	return SerializeEvaluationTemplate(*this, Ar);
+}
\ No newline at end of file
diff --git a/Engine/Source/Runtime/MovieScene/Private/MovieScene.cpp b/Engine/Source/Runtime/MovieScene/Private/MovieScene.cpp
--- a/Engine/Source/Runtime/MovieScene/Private/MovieScene.cpp
+++ b/Engine/Source/Runtime/MovieScene/Private/MovieScene.cpp
@@ -3,6 +3,7 @@
 #include "MovieScenePrivatePCH.h"
 #include "MovieScene.h"
 #include "MovieSceneFolder.h"
+#include "MovieSceneEvaluationCustomVersion.h"
 
 /* UMovieScene interface
  *****************************************************************************/
@@ -25,6 +26,8 @@ UMovieScene::UMovieScene(const FObjectInitializer& ObjectInitializer)
 
 void UMovieScene::Serialize( FArchive& Ar )
 {
+	Ar.UsingCustomVersion(FMovieSceneEvaluationCustomVersion::GUID);
+
 #if WITH_EDITOR
 
 	// Perform optimizations for cooking
@@ -194,7 +197,7 @@ bool UMovieScene::RemovePossessable( const FGuid& PossessableGuid )
 }
 
 
-bool UMovieScene::ReplacePossessable( const FGuid& OldGuid, const FGuid& NewGuid, const FString& NewName )
+bool UMovieScene::ReplacePossessable( const FGuid& OldGuid, const FMovieScenePossessable& InNewPosessable )
 {
 	bool bAnythingReplaced = false;
 
@@ -205,10 +208,20 @@ bool UMovieScene::ReplacePossessable( const FGuid& OldGuid, const FGuid& NewGuid
 			Modify();
 
 			// Found it!
-			Possessable.SetGuid(NewGuid);
-			Possessable.SetName(NewName);
+			if (InNewPosessable.GetPossessedObjectClass() == nullptr)
+			{
+				// @todo: delete this when
+				// bool ReplacePossessable(const FGuid& OldGuid, const FGuid& NewGuid, const FString& Name)
+				// is removed
+				Possessable.SetGuid(InNewPosessable.GetGuid());
+				Possessable.SetName(InNewPosessable.GetName());
+			}
+			else
+			{
+				Possessable = InNewPosessable;
+			}
 
-			ReplaceBinding( OldGuid, NewGuid, NewName );
+			ReplaceBinding( OldGuid, InNewPosessable.GetGuid(), InNewPosessable.GetName() );
 			bAnythingReplaced = true;
 
 			break;
@@ -318,35 +331,49 @@ void UMovieScene::SetPlaybackRange(float Start, float End, bool bAlwaysMarkDirty
 		PlaybackRange = NewRange;
 
 #if WITH_EDITORONLY_DATA
+		// Initialize the working and view range with a little bit more space
+		const float OutputViewSize = PlaybackRange.GetUpperBoundValue() - PlaybackRange.GetLowerBoundValue();
+		const float OutputChange = OutputViewSize * 0.1f;
+
+		TRange<float> ExpandedPlaybackRange = TRange<float>(PlaybackRange.GetLowerBoundValue() - OutputChange, PlaybackRange.GetUpperBoundValue() + OutputChange);
+
 		if (EditorData.WorkingRange.IsEmpty())
 		{
-			EditorData.WorkingRange = PlaybackRange;
+			EditorData.WorkingRange = ExpandedPlaybackRange;
 		}
 
 		if (EditorData.ViewRange.IsEmpty())
 		{
-			EditorData.ViewRange = PlaybackRange;
+			EditorData.ViewRange = ExpandedPlaybackRange;
 		}
 #endif
 	}
 }
 
+void UMovieScene::SetWorkingRange(float Start, float End)
+{
+#if WITH_EDITORONLY_DATA
+	EditorData.WorkingRange = TRange<float>(Start, End);
+#endif
+}
 
-bool UMovieScene::GetForceFixedFrameIntervalPlayback() const
+void UMovieScene::SetViewRange(float Start, float End)
 {
-	return bForceFixedFrameIntervalPlayback;
+#if WITH_EDITORONLY_DATA
+	EditorData.ViewRange = TRange<float>(Start, End);
+#endif
 }
 
 
-void UMovieScene::SetForceFixedFrameIntervalPlayback( bool bInForceFixedFrameIntervalPlayback )
+bool UMovieScene::GetForceFixedFrameIntervalPlayback() const
 {
-	bForceFixedFrameIntervalPlayback = bInForceFixedFrameIntervalPlayback;
+	return bForceFixedFrameIntervalPlayback;
 }
 
 
-float UMovieScene::GetFixedFrameInterval() const
+void UMovieScene::SetForceFixedFrameIntervalPlayback( bool bInForceFixedFrameIntervalPlayback )
 {
-	return FixedFrameInterval;
+	bForceFixedFrameIntervalPlayback = bInForceFixedFrameIntervalPlayback;
 }
 
 
@@ -554,7 +581,6 @@ void UMovieScene::RemoveCameraCutTrack()
 	}
 }
 
-
 bool UMovieScene::RemoveMasterTrack(UMovieSceneTrack& Track) 
 {
 	Modify();
@@ -723,6 +749,12 @@ void UMovieScene::ReplaceBinding(const FGuid& OldGuid, const FGuid& NewGuid, con
 		{
 			Binding.SetObjectGuid(NewGuid);
 			Binding.SetName(Name);
+
+			// Changing a binding guid invalidates any tracks contained within the binding
+			for (UMovieSceneTrack* Track : Binding.GetTracks())
+			{
+				Track->MarkAsChanged();
+			}
 			break;
 		}
 	}
diff --git a/Engine/Source/Runtime/MovieScene/Private/MovieSceneCommonHelpers.cpp b/Engine/Source/Runtime/MovieScene/Private/MovieSceneCommonHelpers.cpp
--- a/Engine/Source/Runtime/MovieScene/Private/MovieSceneCommonHelpers.cpp
+++ b/Engine/Source/Runtime/MovieScene/Private/MovieSceneCommonHelpers.cpp
@@ -188,6 +188,7 @@ UCameraComponent* MovieSceneHelpers::CameraComponentFromActor(const AActor* InAc
 	TArray<UCameraComponent*> CameraComponents;
 	InActor->GetComponents<UCameraComponent>(CameraComponents);
 
+	// If there's a camera component that's active, return that one
 	for (UCameraComponent* CameraComponent : CameraComponents)
 	{
 		if (CameraComponent->bIsActive)
@@ -196,6 +197,12 @@ UCameraComponent* MovieSceneHelpers::CameraComponentFromActor(const AActor* InAc
 		}
 	}
 
+	// Otherwise, return the first camera component
+	for (UCameraComponent* CameraComponent : CameraComponents)
+	{
+		return CameraComponent;
+	}
+
 	// now see if any actors are attached to us, directly or indirectly, that have an active camera component we might want to use
 	// we will just return the first one.
 	// #note: assumption here that attachment cannot be circular
@@ -290,7 +297,7 @@ FTrackInstancePropertyBindings::FTrackInstancePropertyBindings( FName InProperty
 	}
 }
 
-FTrackInstancePropertyBindings::FPropertyAddress FTrackInstancePropertyBindings::FindPropertyRecursive( const UObject* Object, void* BasePointer, UStruct* InStruct, TArray<FString>& InPropertyNames, uint32 Index ) const
+FTrackInstancePropertyBindings::FPropertyAddress FTrackInstancePropertyBindings::FindPropertyRecursive( void* BasePointer, UStruct* InStruct, TArray<FString>& InPropertyNames, uint32 Index )
 {
 	UProperty* Property = FindField<UProperty>(InStruct, *InPropertyNames[Index]);
 	
@@ -305,7 +312,7 @@ FTrackInstancePropertyBindings::FPropertyAddress FTrackInstancePropertyBindings:
 		if( InPropertyNames.IsValidIndex(Index+1) )
 		{
 			void* StructContainer = StructProp->ContainerPtrToValuePtr<void>(BasePointer);
-			return FindPropertyRecursive( Object, StructContainer, StructProp->Struct, InPropertyNames, Index+1 );
+			return FindPropertyRecursive( StructContainer, StructProp->Struct, InPropertyNames, Index+1 );
 		}
 		else
 		{
@@ -323,7 +330,7 @@ FTrackInstancePropertyBindings::FPropertyAddress FTrackInstancePropertyBindings:
 }
 
 
-FTrackInstancePropertyBindings::FPropertyAddress FTrackInstancePropertyBindings::FindProperty( const UObject* InObject, const FString& InPropertyPath ) const
+FTrackInstancePropertyBindings::FPropertyAddress FTrackInstancePropertyBindings::FindProperty( const UObject& InObject, const FString& InPropertyPath )
 {
 	TArray<FString> PropertyNames;
 
@@ -331,45 +338,30 @@ FTrackInstancePropertyBindings::FPropertyAddress FTrackInstancePropertyBindings:
 
 	if( PropertyNames.Num() > 0 )
 	{
-		return FindPropertyRecursive( InObject, (void*)InObject, InObject->GetClass(), PropertyNames, 0 );
+		return FindPropertyRecursive( (void*)&InObject, InObject.GetClass(), PropertyNames, 0 );
 	}
 	else
 	{
 		return FTrackInstancePropertyBindings::FPropertyAddress();
 	}
-
 }
 
-
-void FTrackInstancePropertyBindings::UpdateBindings( const TArray<TWeakObjectPtr<UObject>>& InRuntimeObjects )
+void FTrackInstancePropertyBindings::CacheBinding(const UObject& Object)
 {
-	for (auto ObjectPtr : InRuntimeObjects)
+	FPropertyAndFunction PropAndFunction;
 	{
-		UpdateBinding(ObjectPtr);
+		PropAndFunction.Function = Object.FindFunction(FunctionName);
+		PropAndFunction.PropertyAddress = FindProperty(Object, PropertyPath);
 	}
-}
 
-void FTrackInstancePropertyBindings::UpdateBinding(const TWeakObjectPtr<UObject>& InRuntimeObject)
-{
-	UObject* Object = InRuntimeObject.Get();
-
-	if (Object != nullptr)
-	{
-		FPropertyAndFunction PropAndFunction;
-		{
-			PropAndFunction.Function = Object->FindFunction(FunctionName);
-			PropAndFunction.PropertyAddress = FindProperty(Object, PropertyPath);
-		}
-
-		RuntimeObjectToFunctionMap.Add(Object, PropAndFunction);
-	}
+	RuntimeObjectToFunctionMap.Add(FObjectKey(&Object), PropAndFunction);
 }
 
-UProperty* FTrackInstancePropertyBindings::GetProperty(const UObject* Object) const
+UProperty* FTrackInstancePropertyBindings::GetProperty(const UObject& Object) const
 {
 	UProperty* Prop = nullptr;
 
-	FPropertyAndFunction PropAndFunction = RuntimeObjectToFunctionMap.FindRef(Object);
+	FPropertyAndFunction PropAndFunction = RuntimeObjectToFunctionMap.FindRef(&Object);
 	if( PropAndFunction.PropertyAddress.Property ) 
 	{
 		Prop = PropAndFunction.PropertyAddress.Property;
@@ -382,3 +374,49 @@ UProperty* FTrackInstancePropertyBindings::GetProperty(const UObject* Object) co
 
 	return Prop;
 }
+
+template<> void FTrackInstancePropertyBindings::CallFunction<bool>(UObject& InRuntimeObject, TCallTraits<bool>::ParamType PropertyValue)
+{
+	FPropertyAndFunction PropAndFunction = FindOrAdd(InRuntimeObject);
+	if (PropAndFunction.Function)
+	{
+		// ProcessEvent should really be taking const void*
+		InRuntimeObject.ProcessEvent(PropAndFunction.Function, (void*)&PropertyValue);
+	}
+	else if (PropAndFunction.PropertyAddress.Property != nullptr)
+	{
+		if (UBoolProperty* BoolProperty = CastChecked<UBoolProperty>(PropAndFunction.PropertyAddress.Property))
+		{
+			uint8* ValuePtr = BoolProperty->ContainerPtrToValuePtr<uint8>(&InRuntimeObject);
+			BoolProperty->SetPropertyValue(ValuePtr, PropertyValue);
+		}
+	}
+}
+
+template<> bool FTrackInstancePropertyBindings::GetCurrentValue<bool>(const UObject& Object)
+{
+	FPropertyAndFunction PropAndFunction = FindOrAdd(Object);
+	if (PropAndFunction.PropertyAddress.Property != nullptr)
+	{
+		if (UBoolProperty* BoolProperty = Cast<UBoolProperty>(PropAndFunction.PropertyAddress.Property))
+		{
+			const uint8* ValuePtr = BoolProperty->ContainerPtrToValuePtr<uint8>(&Object);
+			return BoolProperty->GetPropertyValue(ValuePtr);
+		}
+	}
+
+	return false;
+}
+
+template<> void FTrackInstancePropertyBindings::SetCurrentValue<bool>(UObject& Object, TCallTraits<bool>::ParamType InValue)
+{
+	FPropertyAndFunction PropAndFunction = FindOrAdd(Object);
+	if (PropAndFunction.PropertyAddress.Property != nullptr)
+	{
+		if (UBoolProperty* BoolProperty = Cast<UBoolProperty>(PropAndFunction.PropertyAddress.Property))
+		{
+			uint8* ValuePtr = BoolProperty->ContainerPtrToValuePtr<uint8>(&Object);
+			BoolProperty->SetPropertyValue(ValuePtr, InValue);
+		}
+	}
+}
\ No newline at end of file
diff --git a/Engine/Source/Runtime/MovieScene/Private/MovieScenePlayer.cpp b/Engine/Source/Runtime/MovieScene/Private/MovieScenePlayer.cpp
new file mode 100644
--- /dev/null
+++ b/Engine/Source/Runtime/MovieScene/Private/MovieScenePlayer.cpp
@@ -0,0 +1,19 @@
+// Copyright 1998-2016 Epic Games, Inc. All Rights Reserved.
+
+#include "MovieScenePrivatePCH.h"
+#include "IMovieScenePlayer.h"
+#include "MovieSceneSequenceInstance.h"
+#include "Evaluation/MovieSceneEvaluationTemplateInstance.h"
+
+TSharedRef<FMovieSceneSequenceInstance> IMovieScenePlayer::GetRootMovieSceneSequenceInstance()
+{
+	return GetEvaluationTemplate().GetInstance(MovieSceneSequenceID::Root)->LegacySequenceInstance.ToSharedRef();
+}
+
+void IMovieScenePlayer::GetRuntimeObjects(TSharedRef<FMovieSceneSequenceInstance> MovieSceneInstance, const FGuid& Guid, TArray<TWeakObjectPtr<>>& OutRuntimeObjects)
+{
+	for (TWeakObjectPtr<> Obj : FindBoundObjects(Guid, MovieSceneInstance->GetSequenceID()))
+	{
+		OutRuntimeObjects.Add(Obj);
+	}
+}
\ No newline at end of file
diff --git a/Engine/Source/Runtime/MovieScene/Private/MovieSceneSection.cpp b/Engine/Source/Runtime/MovieScene/Private/MovieSceneSection.cpp
--- a/Engine/Source/Runtime/MovieScene/Private/MovieSceneSection.cpp
+++ b/Engine/Source/Runtime/MovieScene/Private/MovieSceneSection.cpp
@@ -159,8 +159,14 @@ void UMovieSceneSection::AddKeyToCurve(FRichCurve& InCurve, float Time, float Va
 
 void UMovieSceneSection::SetCurveDefault(FRichCurve& InCurve, float Value)
 {
-	if (TryModify())
+	if (InCurve.GetDefaultValue() != Value && TryModify())
 	{
 		InCurve.SetDefaultValue(Value);
 	}
 }
+
+
+FMovieSceneEvalTemplatePtr UMovieSceneSection::GenerateTemplate() const
+{
+	return FMovieSceneEvalTemplatePtr();
+}
\ No newline at end of file
diff --git a/Engine/Source/Runtime/MovieScene/Private/MovieSceneSequence.cpp b/Engine/Source/Runtime/MovieScene/Private/MovieSceneSequence.cpp
new file mode 100644
--- /dev/null
+++ b/Engine/Source/Runtime/MovieScene/Private/MovieSceneSequence.cpp
@@ -0,0 +1,74 @@
+// Copyright 1998-2016 Epic Games, Inc. All Rights Reserved.
+
+#include "MovieScenePrivatePCH.h"
+#include "MovieSceneSequence.h"
+#include "MovieSceneEvaluationTemplateGenerator.h"
+#include "MovieSceneEvaluationCustomVersion.h"
+
+
+UMovieSceneSequence::UMovieSceneSequence(const FObjectInitializer& Init)
+	: Super(Init)
+{
+	bParentContextsAreSignificant = false;
+
+	TemplateParameters.bForEditorPreview = false;
+#if WITH_EDITORONLY_DATA
+	EvaluationTemplate.Initialize(*this);
+#endif
+}
+
+void UMovieSceneSequence::PreSave(const ITargetPlatform* TargetPlatform)
+{
+	Super::PreSave(TargetPlatform);
+}
+
+#if WITH_EDITORONLY_DATA
+void UMovieSceneSequence::PostDuplicate(bool bDuplicateForPIE)
+{
+	if (bDuplicateForPIE)
+	{
+		EvaluationTemplate.Regenerate(TemplateParameters);
+	}
+
+	Super::PostDuplicate(bDuplicateForPIE);
+}
+#endif // WITH_EDITORONLY_DATA
+
+void UMovieSceneSequence::Serialize(FArchive& Ar)
+{
+	Ar.UsingCustomVersion(FMovieSceneEvaluationCustomVersion::GUID);
+
+#if WITH_EDITORONLY_DATA
+	if (Ar.IsCooking() && !HasAnyFlags(RF_ClassDefaultObject))
+	{
+		EvaluationTemplate.Regenerate(TemplateParameters);
+	}
+#endif
+	
+	Super::Serialize(Ar);
+}
+
+void UMovieSceneSequence::GenerateEvaluationTemplate(FMovieSceneEvaluationTemplate& Template, const FMovieSceneTrackCompilationParams& Params, FMovieSceneSequenceTemplateStore& Store)
+{
+	FMovieSceneEvaluationTemplateGenerator(*this, Template, Store).Generate(Params);
+}
+
+FGuid UMovieSceneSequence::FindPossessableObjectId(UObject& Object, UObject* Context) const
+{
+	UMovieScene* MovieScene = GetMovieScene();
+	if (!MovieScene)
+	{
+		return FGuid();
+	}
+
+	// Search all possessables
+	for (int32 Index = 0; Index < MovieScene->GetPossessableCount(); ++Index)
+	{
+		FGuid ThisGuid = MovieScene->GetPossessable(Index).GetGuid();
+		if (LocateBoundObjects(ThisGuid, Context).Contains(&Object))
+		{
+			return ThisGuid;
+		}
+	}
+	return FGuid();
+}
\ No newline at end of file
