Copying //UE4/Dev-Sequencer to //UE4/Dev-Main (Source: //UE4/Dev-Sequencer @ 3178529)

diff --git a/Engine/Source/Runtime/MovieScene/Private/MovieSceneSequenceInstance.cpp b/Engine/Source/Runtime/MovieScene/Private/MovieSceneSequenceInstance.cpp
--- a/Engine/Source/Runtime/MovieScene/Private/MovieSceneSequenceInstance.cpp
+++ b/Engine/Source/Runtime/MovieScene/Private/MovieSceneSequenceInstance.cpp
@@ -1,520 +1,15 @@
 // Copyright 1998-2016 Epic Games, Inc. All Rights Reserved.
 
 #include "MovieScenePrivatePCH.h"
-#include "MovieSceneSequenceInstance.h"
 #include "MovieSceneSequence.h"
-#include "MovieSceneTrack.h"
 
-FMovieSceneSequenceInstance::FMovieSceneSequenceInstance(const UMovieSceneSequence& InMovieSceneSequence)
-	: MovieSceneSequence( &InMovieSceneSequence )
-#if WITH_EDITOR
-	, InstanceId(FGuid::NewGuid())
-#endif
+FMovieSceneSequenceInstance::FMovieSceneSequenceInstance(const UMovieSceneSequence& InMovieSceneSequence, FMovieSceneSequenceIDRef InSequenceID)
 {
-	TimeRange = MovieSceneSequence->GetMovieScene()->GetPlaybackRange();
+	MovieSceneSequence = &InMovieSceneSequence;
+	SequenceID = InSequenceID;
 }
 
-FMovieSceneSequenceInstance::FMovieSceneSequenceInstance(const UMovieSceneTrack& InMovieSceneTrack)
-#if WITH_EDITOR
-	: InstanceId(FGuid::NewGuid())
-#endif
+TRange<float> FMovieSceneSequenceInstance::GetTimeRange() const
 {
-	TimeRange = CastChecked<UMovieScene>(InMovieSceneTrack.GetOuter())->GetPlaybackRange();
+	return MovieSceneSequence->GetMovieScene()->GetPlaybackRange();
 }
-
-FMovieSceneSequenceInstance::~FMovieSceneSequenceInstance()
-{
-	for (auto& Pair : SpawnedObjects)
-	{
-		ensureMsgf(!Pair.Value.Get(), TEXT("Sequence instance is being destroyed but still has spawned objects"));
-	}
-}
-
-FGuid FMovieSceneSequenceInstance::FindObjectId(UObject& Object) const
-{
-	if(MovieSceneSequence.IsValid())
-	{
-		for (auto& Pair : ObjectBindingInstances)
-		{
-			if (Pair.Value.RuntimeObjects.Contains(&Object))
-			{
-				return Pair.Key;
-			}
-		}
-
-		// At this point the only possibility left is that we have not cached the object
-		// in ObjectBindingInstances, so we to see if the sequence itself can tell us the GUID
-		return MovieSceneSequence->FindPossessableObjectId(Object);
-	}
-
-	return FGuid();
-}
-
-FGuid FMovieSceneSequenceInstance::FindParentObjectId(UObject& Object) const
-{
-	if(MovieSceneSequence.IsValid())
-	{
-		UObject* ParentObject = MovieSceneSequence->GetParentObject(&Object);
-		if (ParentObject)
-		{
-			return FindObjectId(*ParentObject);
-		}
-	}
-	return FGuid();
-}
-
-UObject* FMovieSceneSequenceInstance::FindObject(const FGuid& ObjectId, const IMovieScenePlayer& Player) const
-{
-	if(MovieSceneSequence.IsValid())
-	{
-		// Attempt to find a possessable first
-		UMovieScene* MovieScene = MovieSceneSequence->GetMovieScene();
-		if (MovieScene != nullptr)
-		{
-			FMovieScenePossessable* Possessable = MovieScene->FindPossessable(ObjectId);
-			if (Possessable)
-			{
-				UObject* ParentObject = Player.GetPlaybackContext();
-				if (Possessable->GetParent().IsValid())
-				{
-					ParentObject = FindObject(Possessable->GetParent(), Player);
-				}
-
-				return MovieSceneSequence->FindPossessableObject(ObjectId, ParentObject);
-			}
-			else
-			{
-				return FindSpawnedObject(ObjectId);
-			}
-		}
-	}
-
-	return nullptr;
-}
-
-UObject* FMovieSceneSequenceInstance::FindSpawnedObject(const FGuid& ObjectId) const
-{
-	const TWeakObjectPtr<>* SpawnedObject = SpawnedObjects.Find(ObjectId);
-	return SpawnedObject ? SpawnedObject->Get() : nullptr;
-}
-
-void FMovieSceneSequenceInstance::SaveState(class IMovieScenePlayer& Player)
-{
-	TArray<TWeakObjectPtr<UObject>> NoObjects;
-	if (CameraCutTrackInstance.IsValid())
-	{
-		CameraCutTrackInstance->SaveState(NoObjects, Player, *this);
-	}
-
-	for (auto& Pair : MasterTrackInstances)
-	{
-		Pair.Value->SaveState(NoObjects, Player, *this);
-	}
-
-	TMap<FGuid, FMovieSceneObjectBindingInstance>::TIterator ObjectIt = ObjectBindingInstances.CreateIterator();
-	for (; ObjectIt; ++ObjectIt)
-	{
-		FMovieSceneObjectBindingInstance& ObjectBindingInstance = ObjectIt.Value();
-
-		for (FMovieSceneInstanceMap::TIterator It = ObjectBindingInstance.TrackInstances.CreateIterator(); It; ++It)
-		{
-			It.Value()->SaveState(ObjectBindingInstance.RuntimeObjects, Player, *this);
-		}
-
-		for( FMovieSceneInstanceMap::TIterator It( MasterTrackInstances ); It; ++It )
-		{
-			It.Value()->SaveState( ObjectBindingInstance.RuntimeObjects, Player, *this );
-		}
-	}
-}
-
-
-void FMovieSceneSequenceInstance::RestoreState(class IMovieScenePlayer& Player)
-{
-	Player.GetSpawnRegister().CleanUpSequence(*this, Player);
-
-	TMap<FGuid, FMovieSceneObjectBindingInstance>::TIterator ObjectIt = ObjectBindingInstances.CreateIterator();
-	for (; ObjectIt; ++ObjectIt)
-	{
-		FMovieSceneObjectBindingInstance& ObjectBindingInstance = ObjectIt.Value();
-
-		for (FMovieSceneInstanceMap::TIterator It = ObjectBindingInstance.TrackInstances.CreateIterator(); It; ++It)
-		{
-			It.Value()->RestoreState(ObjectBindingInstance.RuntimeObjects, Player, *this);
-		}
-
-		for( FMovieSceneInstanceMap::TIterator It( MasterTrackInstances ); It; ++It )
-		{
-			It.Value()->RestoreState( ObjectBindingInstance.RuntimeObjects, Player, *this );
-		}
-	}
-
-	TArray<TWeakObjectPtr<UObject>> NoObjects;
-	if (CameraCutTrackInstance.IsValid())
-	{
-		CameraCutTrackInstance->RestoreState(NoObjects, Player, *this);
-	}
-
-	for (auto& Pair : MasterTrackInstances)
-	{
-		Pair.Value->RestoreState(NoObjects, Player, *this);
-	}
-}
-
-void FMovieSceneSequenceInstance::RestoreSpecificState(const FGuid& ObjectGuid, IMovieScenePlayer& Player)
-{
-	FMovieSceneObjectBindingInstance* ObjectBindingInstance = ObjectBindingInstances.Find(ObjectGuid);
-	if (!ObjectBindingInstance)
-	{
-		return;
-	}
-
-	for (auto& Pair : ObjectBindingInstance->TrackInstances)
-	{
-		Pair.Value->RestoreState(ObjectBindingInstance->RuntimeObjects, Player, *this);
-	}
-}
-
-void FMovieSceneSequenceInstance::Update( EMovieSceneUpdateData& UpdateData, class IMovieScenePlayer& Player )
-{
-	if(MovieSceneSequence.IsValid())
-	{
-		PreUpdate(Player);
-
-		UpdatePasses(UpdateData, Player);
-
-		PostUpdate(Player);
-	}
-}
-
-void FMovieSceneSequenceInstance::PreUpdate(class IMovieScenePlayer& Player)
-{
-	// Remove any stale runtime objects
-	TMap<FGuid, FMovieSceneObjectBindingInstance>::TIterator ObjectIt = ObjectBindingInstances.CreateIterator();
-	for(; ObjectIt; ++ObjectIt )
-	{
-		FMovieSceneObjectBindingInstance& ObjectBindingInstance = ObjectIt.Value();
-		for (int32 ObjectIndex = 0; ObjectIndex < ObjectBindingInstance.RuntimeObjects.Num(); )
-		{
-			UObject* RuntimeObject = ObjectBindingInstance.RuntimeObjects[ObjectIndex].Get();
-			if (RuntimeObject == nullptr || RuntimeObject->HasAnyFlags(RF_BeginDestroyed|RF_FinishDestroyed) || RuntimeObject->IsPendingKill())
-			{
-				ObjectBindingInstance.RuntimeObjects.RemoveAt(ObjectIndex);
-			}
-			else
-			{
-				++ObjectIndex;
-			}
-		}
-	}
-
-	Player.GetSpawnRegister().PreUpdateSequenceInstance(*this, Player);
-}
-
-void FMovieSceneSequenceInstance::UpdatePasses(EMovieSceneUpdateData& UpdateData, class IMovieScenePlayer& Player)
-{
-	UpdateData.UpdatePass = MSUP_PreUpdate;
-	UpdatePassSingle(UpdateData, Player);
-	UpdateData.UpdatePass = MSUP_Update;
-	UpdatePassSingle(UpdateData, Player);
-	UpdateData.UpdatePass = MSUP_PostUpdate;
-	UpdatePassSingle(UpdateData, Player);
-}
-
-void FMovieSceneSequenceInstance::PostUpdate(class IMovieScenePlayer& Player)
-{
-	Player.GetSpawnRegister().PostUpdateSequenceInstance(*this, Player);
-}
-
-void FMovieSceneSequenceInstance::UpdatePassSingle( EMovieSceneUpdateData& UpdateData, class IMovieScenePlayer& Player )
-{
-	if(MovieSceneSequence.IsValid())
-	{
-		// Refresh time range so that spawnables can be created if they fall within the playback range, or destroyed if not
-		UMovieScene* MovieScene = MovieSceneSequence->GetMovieScene();
-		TimeRange = MovieScene->GetPlaybackRange();
-
-		TArray<TWeakObjectPtr<UObject>> NoObjects;
-
-		// update each master track
-		for( FMovieSceneInstanceMap::TIterator It( MasterTrackInstances ); It; ++It )
-		{		
-			if (It.Value()->HasUpdatePasses() & UpdateData.UpdatePass && It.Value() != CameraCutTrackInstance &&
-				( UpdateData.bSubSceneDeactivate == false || It.Value()->RequiresUpdateForSubSceneDeactivate () ))
-			{
-				It.Value()->Update( UpdateData, NoObjects, Player, *this);
-			}
-		}
-
-		// update tracks bound to objects
-		TMap<FGuid, FMovieSceneObjectBindingInstance>::TIterator ObjectIt = ObjectBindingInstances.CreateIterator();
-		for(; ObjectIt; ++ObjectIt )
-		{
-			FMovieSceneObjectBindingInstance& ObjectBindingInstance = ObjectIt.Value();
-		
-			for( FMovieSceneInstanceMap::TIterator It = ObjectBindingInstance.TrackInstances.CreateIterator(); It; ++It )
-			{
-				if (It.Value()->HasUpdatePasses() & UpdateData.UpdatePass && It.Value() != CameraCutTrackInstance &&
-					( UpdateData.bSubSceneDeactivate == false || It.Value()->RequiresUpdateForSubSceneDeactivate() ) )
-				{
-					It.Value()->Update( UpdateData, ObjectBindingInstance.RuntimeObjects, Player, *this );
-				}
-			}
-		}
-
-		// update camera cut track last to make sure spawnable cameras are there, and to override sub-shots
-		if (CameraCutTrackInstance.IsValid())
-		{
-			if (CameraCutTrackInstance->HasUpdatePasses() & UpdateData.UpdatePass &&
-				( UpdateData.bSubSceneDeactivate == false || CameraCutTrackInstance->RequiresUpdateForSubSceneDeactivate() ) )
-			{
-				CameraCutTrackInstance->Update(UpdateData, NoObjects, Player, *this );
-			}
-		}
-	}
-}
-
-
-void FMovieSceneSequenceInstance::RefreshInstance( IMovieScenePlayer& Player )
-{
-	if(MovieSceneSequence.IsValid())
-	{
-		UMovieScene* MovieScene = MovieSceneSequence->GetMovieScene();
-		TimeRange = MovieScene->GetPlaybackRange();
-
-		UMovieSceneTrack* CameraCutTrack = MovieScene->GetCameraCutTrack();
-
-		if (CameraCutTrack != nullptr)
-		{
-			FMovieSceneInstanceMap CameraCutTrackInstanceMap;
-
-			if (CameraCutTrackInstance.IsValid())
-			{
-				CameraCutTrackInstanceMap.Add(CameraCutTrack, CameraCutTrackInstance);
-			}
-
-			TArray<TWeakObjectPtr<UObject>> Objects;
-			TArray<UMovieSceneTrack*> Tracks;
-			Tracks.Add(CameraCutTrack);
-			RefreshInstanceMap(Tracks, Objects, CameraCutTrackInstanceMap, Player);
-
-			CameraCutTrackInstance = CameraCutTrackInstanceMap.FindRef(CameraCutTrack);
-		}
-		else if(CameraCutTrackInstance.IsValid())
-		{
-			CameraCutTrackInstance->ClearInstance(Player, *this);
-			CameraCutTrackInstance.Reset();
-		}
-
-		// Get all the master tracks and create instances for them if needed
-		const TArray<UMovieSceneTrack*>& MasterTracks = MovieScene->GetMasterTracks();
-		TArray<TWeakObjectPtr<UObject>> Objects;
-		RefreshInstanceMap( MasterTracks, Objects, MasterTrackInstances, Player );
-
-		TSet< FGuid > FoundObjectBindings;
-		// Get all tracks for each object binding and create instances for them if needed
-		const TArray<FMovieSceneBinding>& ObjectBindings = MovieScene->GetBindings();
-		for( int32 BindingIndex = 0; BindingIndex < ObjectBindings.Num(); ++BindingIndex )
-		{
-			const FMovieSceneBinding& ObjectBinding = ObjectBindings[BindingIndex];
-
-			// Create an instance for this object binding
-			FMovieSceneObjectBindingInstance& BindingInstance = ObjectBindingInstances.FindOrAdd( ObjectBinding.GetObjectGuid() );
-			BindingInstance.ObjectGuid = ObjectBinding.GetObjectGuid();
-
-			FoundObjectBindings.Add( ObjectBinding.GetObjectGuid() );
-
-			// Populate the runtime objects for this instance of the binding.
-			// @todo sequencer: SubSequences: We need to know which actors were removed and which actors were added so we know which saved actor state to restore/create
-			BindingInstance.RuntimeObjects.Empty();
-			Player.GetRuntimeObjects( SharedThis( this ), BindingInstance.ObjectGuid, BindingInstance.RuntimeObjects );
-
-			// Refresh the instance's tracks
-			const TArray<UMovieSceneTrack*>& Tracks = ObjectBinding.GetTracks();
-			RefreshInstanceMap( Tracks, BindingInstance.RuntimeObjects, BindingInstance.TrackInstances, Player );
-		}
-
-		IMovieSceneSpawnRegister& SpawnRegister = Player.GetSpawnRegister();
-
-		// Remove object binding instances which are no longer bound
-		TMap<FGuid, FMovieSceneObjectBindingInstance>::TIterator It = ObjectBindingInstances.CreateIterator();
-		for( ; It; ++It )
-		{
-			if( !FoundObjectBindings.Contains( It.Key() ) )
-			{
-				SpawnRegister.DestroySpawnedObject(It.Key(), *this, Player);
-
-				// The instance no longer is bound to an existing guid
-				It.RemoveCurrent();
-			}
-		}
-	}
-}
-
-
-struct FTrackInstanceEvalSorter
-{
-	bool operator()( const TSharedPtr<IMovieSceneTrackInstance> A, const TSharedPtr<IMovieSceneTrackInstance> B ) const
-	{
-		return A->EvalOrder() < B->EvalOrder();
-	}
-};
-
-
-void FMovieSceneSequenceInstance::RefreshInstanceMap( const TArray<UMovieSceneTrack*>& Tracks, const TArray<TWeakObjectPtr<UObject>>& RuntimeObjects, FMovieSceneInstanceMap& TrackInstances, IMovieScenePlayer& Player  )
-{
-	// All the tracks we found during this pass
-	TSet< UMovieSceneTrack* > FoundTracks;
-
-	// For every track, check if it has an instance, if not create one, otherwise refresh that instance
-	for( int32 TrackIndex = 0; TrackIndex < Tracks.Num(); ++TrackIndex )
-	{
-		UMovieSceneTrack* Track = Tracks[TrackIndex];
-
-		// A new track has been encountered
-		FoundTracks.Add( Track );
-
-		// See if the track has an instance
-		TSharedPtr<IMovieSceneTrackInstance> Instance = TrackInstances.FindRef( Track );
-		if ( !Instance.IsValid() )
-		{
-			// The track does not have an instance, create one
-			Instance = Track->CreateInstance();
-			Instance->RefreshInstance( RuntimeObjects, Player, *this );
-			Instance->SaveState(RuntimeObjects, Player, *this);
-
-			TrackInstances.Add( Track, Instance );
-		}
-		else
-		{
-			// The track has an instance, refresh it
-			Instance->RefreshInstance( RuntimeObjects, Player, *this );
-			Instance->SaveState(RuntimeObjects, Player, *this);
-		}
-
-	}
-
-	// Remove entries which no longer have a track associated with them
-	FMovieSceneInstanceMap::TIterator It = TrackInstances.CreateIterator();
-	for( ; It; ++It )
-	{
-		if( !FoundTracks.Contains( It.Key().Get() ) )
-		{
-			It.Value()->ClearInstance( Player, *this );
-
-			// This track was not found in the moviescene's track list so it was removed.
-			It.RemoveCurrent();
-		}
-	}
-
-	// Sort based on evaluation order
-	TrackInstances.ValueSort(FTrackInstanceEvalSorter());
-}
-
-void FMovieSceneSequenceInstance::UpdateObjectBinding(const FGuid& ObjectId, IMovieScenePlayer& Player)
-{
-	if(MovieSceneSequence.IsValid())
-	{
-		UMovieSceneSequence* Sequence = MovieSceneSequence.Get();
-		auto* BindingInstance = ObjectBindingInstances.Find(ObjectId);
-
-		if (!BindingInstance || !Sequence)
-		{
-			return;
-		}
-
-		// Update the runtime objects
-		BindingInstance->RuntimeObjects.Reset();
-
-		TWeakObjectPtr<UObject>* WeakSpawnedObject = SpawnedObjects.Find(ObjectId);
-		if (WeakSpawnedObject)
-		{
-			UObject* SpawnedObject = WeakSpawnedObject->Get();
-			if (SpawnedObject)
-			{
-				BindingInstance->RuntimeObjects.Add(SpawnedObject);
-			}
-		}
-		else
-		{
-			Player.GetRuntimeObjects(SharedThis(this), BindingInstance->ObjectGuid, BindingInstance->RuntimeObjects);
-		}
-
-		const FMovieSceneBinding* ObjectBinding = Sequence->GetMovieScene()->GetBindings().FindByPredicate([&](const FMovieSceneBinding& In){
-			return In.GetObjectGuid() == ObjectId;
-		});
-
-		// Refresh the instance map, if we found the binding itself
-		if (ObjectBinding)
-		{
-			RefreshInstanceMap(ObjectBinding->GetTracks(), BindingInstance->RuntimeObjects, BindingInstance->TrackInstances, Player);
-		}
-	}
-}
-
-void FMovieSceneSequenceInstance::OnObjectSpawned(const FGuid& ObjectId, UObject& SpawnedObject, IMovieScenePlayer& Player)
-{
-	if(MovieSceneSequence.IsValid())
-	{
-		auto* BindingInstance = ObjectBindingInstances.Find(ObjectId);
-
-		if (!BindingInstance)
-		{
-			return;
-		}
-
-		SpawnedObjects.Add(ObjectId, &SpawnedObject);
-
-		// Add it to the instance's runtime objects array, and update any child possessable binding instances
-		BindingInstance->RuntimeObjects.Reset();
-		BindingInstance->RuntimeObjects.Emplace(&SpawnedObject);
-
-		UMovieSceneSequence* Sequence = GetSequence();
-		FMovieSceneSpawnable* Spawnable = Sequence ? Sequence->GetMovieScene()->FindSpawnable(ObjectId) : nullptr;
-		if (Spawnable)
-		{
-			UpdateObjectBinding(ObjectId, Player);
-			for (const FGuid& Child : Spawnable->GetChildPossessables())
-			{
-				UpdateObjectBinding(Child, Player);
-			}
-		}
-	}
-}
-
-void FMovieSceneSequenceInstance::OnSpawnedObjectDestroyed(const FGuid& ObjectId, IMovieScenePlayer& Player)
-{
-	if(MovieSceneSequence.IsValid())
-	{
-		auto* BindingInstance = ObjectBindingInstances.Find(ObjectId);
-		if (!BindingInstance)
-		{
-			return;
-		}
-
-		SpawnedObjects.Remove(ObjectId);
-
-		// Destroy the object
-		BindingInstance->RuntimeObjects.Reset();
-
-		// Update any child possessable object bindings
-		UMovieSceneSequence* Sequence = GetSequence();
-		FMovieSceneSpawnable* Spawnable = Sequence ? Sequence->GetMovieScene()->FindSpawnable(ObjectId) : nullptr;
-		if (Spawnable)
-		{
-			for (const FGuid& Child : Spawnable->GetChildPossessables())
-			{
-				UpdateObjectBinding(Child, Player);
-			}
-		}
-	}
-}
-
-void FMovieSceneSequenceInstance::HandleSequenceSectionChanged(UMovieSceneSequence* Sequence)
-{
-	MovieSceneSequence = Sequence;
-	if(MovieSceneSequence.IsValid())
-	{
-		TimeRange = MovieSceneSequence->GetMovieScene()->GetPlaybackRange();
-	}
-}
\ No newline at end of file
diff --git a/Engine/Source/Runtime/MovieScene/Private/MovieSceneSpawnRegister.cpp b/Engine/Source/Runtime/MovieScene/Private/MovieSceneSpawnRegister.cpp
new file mode 100644
--- /dev/null
+++ b/Engine/Source/Runtime/MovieScene/Private/MovieSceneSpawnRegister.cpp
@@ -0,0 +1,108 @@
+// Copyright 1998-2016 Epic Games, Inc. All Rights Reserved.
+
+#include "MovieScenePrivatePCH.h"
+#include "MovieSceneSpawnRegister.h"
+
+UObject* FMovieSceneSpawnRegister::FindSpawnedObject(const FGuid& BindingId, FMovieSceneSequenceIDRef TemplateID) const
+{
+	FMovieSceneSpawnRegisterKey Key(TemplateID, BindingId);
+
+	const FSpawnedObject* Existing = Register.Find(Key);
+	return Existing ? Existing->Object.Get() : nullptr;
+}
+
+UObject* FMovieSceneSpawnRegister::SpawnObject(const FGuid& BindingId, UMovieScene& MovieScene, FMovieSceneSequenceIDRef TemplateID, IMovieScenePlayer& Player)
+{
+	UObject* ObjectInstance = FindSpawnedObject(BindingId, TemplateID);
+	if (ObjectInstance)
+	{
+		return ObjectInstance;
+	}
+
+	// Find the spawnable definition
+	FMovieSceneSpawnable* Spawnable = MovieScene.FindSpawnable(BindingId);
+	if (!Spawnable)
+	{
+		return nullptr;
+	}
+
+	UObject* SpawnedActor = SpawnObject(*Spawnable, TemplateID, Player);
+	
+	if (ensure(SpawnedActor))
+	{
+		FMovieSceneSpawnRegisterKey Key(TemplateID, BindingId);
+		Register.Add(Key, FSpawnedObject(BindingId, *SpawnedActor, Spawnable->GetSpawnOwnership()));
+
+		Player.State.Invalidate(BindingId, TemplateID);
+	}
+
+	return SpawnedActor;
+};
+
+void FMovieSceneSpawnRegister::DestroySpawnedObject(const FGuid& BindingId, FMovieSceneSequenceIDRef TemplateID, IMovieScenePlayer& Player)
+{
+	FMovieSceneSpawnRegisterKey Key(TemplateID, BindingId);
+	
+	FSpawnedObject* Existing = Register.Find(Key);
+	UObject* SpawnedObject = Existing ? Existing->Object.Get() : nullptr;
+	if (SpawnedObject)
+	{
+		PreDestroyObject(*SpawnedObject, BindingId, TemplateID);
+		DestroySpawnedObject(*SpawnedObject);
+	}
+
+	Register.Remove(Key);
+
+	Player.State.Invalidate(BindingId, TemplateID);
+}
+
+void FMovieSceneSpawnRegister::DestroyObjectsByPredicate(IMovieScenePlayer& Player, const TFunctionRef<bool(const FGuid&, ESpawnOwnership, FMovieSceneSequenceIDRef)>& Predicate)
+{
+	for (auto It = Register.CreateIterator(); It; ++It)
+	{
+		if (Predicate(It.Value().Guid, It.Value().Ownership, It.Key().TemplateID))
+		{
+			UObject* SpawnedObject = It.Value().Object.Get();
+			if (SpawnedObject)
+			{
+				PreDestroyObject(*SpawnedObject, It.Key().BindingId, It.Key().TemplateID);
+				DestroySpawnedObject(*SpawnedObject);
+			}
+
+			It.RemoveCurrent();
+		}
+	}
+}
+
+void FMovieSceneSpawnRegister::ForgetExternallyOwnedSpawnedObjects(FMovieSceneEvaluationState& State, IMovieScenePlayer& Player)
+{
+	for (auto It = Register.CreateIterator(); It; ++It)
+	{
+		if (It.Value().Ownership == ESpawnOwnership::External)
+		{
+			Player.State.Invalidate(It.Key().BindingId, It.Key().TemplateID);
+			It.RemoveCurrent();
+		}
+	}
+}
+
+void FMovieSceneSpawnRegister::CleanUp(IMovieScenePlayer& Player)
+{
+	DestroyObjectsByPredicate(Player, [&](const FGuid&, ESpawnOwnership, FMovieSceneSequenceIDRef){
+		return true;
+	});
+}
+
+void FMovieSceneSpawnRegister::CleanUpSequence(FMovieSceneSequenceIDRef TemplateID, IMovieScenePlayer& Player)
+{
+	DestroyObjectsByPredicate(Player, [&](const FGuid&, ESpawnOwnership, FMovieSceneSequenceIDRef ThisTemplateID){
+		return ThisTemplateID == TemplateID;
+	});
+}
+
+void FMovieSceneSpawnRegister::OnSequenceExpired(FMovieSceneSequenceIDRef TemplateID, IMovieScenePlayer& Player)
+{
+	DestroyObjectsByPredicate(Player, [&](const FGuid& ObjectId, ESpawnOwnership Ownership, FMovieSceneSequenceIDRef ThisTemplateID){
+		return (Ownership == ESpawnOwnership::InnerSequence) && (TemplateID == ThisTemplateID);
+	});
+}
diff --git a/Engine/Source/Runtime/MovieScene/Private/MovieSceneTrack.cpp b/Engine/Source/Runtime/MovieScene/Private/MovieSceneTrack.cpp
new file mode 100644
--- /dev/null
+++ b/Engine/Source/Runtime/MovieScene/Private/MovieSceneTrack.cpp
@@ -0,0 +1,199 @@
+// Copyright 1998-2016 Epic Games, Inc. All Rights Reserved.
+
+#include "MovieScenePrivatePCH.h"
+#include "MovieSceneTrack.h"
+#include "MovieSceneCompilerRules.h"
+
+#include "MovieSceneEvaluationTemplate.h"
+#include "Compilation/IMovieSceneTemplateGenerator.h"
+#include "Evaluation/MovieSceneLegacyTrackInstanceTemplate.h"
+
+UMovieSceneTrack::UMovieSceneTrack(const FObjectInitializer& InInitializer)
+	: Super(InInitializer)
+{
+#if WITH_EDITORONLY_DATA
+	TrackTint = FColor(127, 127, 127, 0);
+#endif
+}
+
+PRAGMA_DISABLE_DEPRECATION_WARNINGS
+TSharedPtr<IMovieSceneTrackInstance> UMovieSceneTrack::CreateLegacyInstance() const
+{
+	// Ugly const cast required due to lack of const-correctness in old system
+	return const_cast<UMovieSceneTrack*>(this)->CreateInstance();
+}
+PRAGMA_ENABLE_DEPRECATION_WARNINGS
+
+TInlineValue<FMovieSceneSegmentCompilerRules> UMovieSceneTrack::GetRowCompilerRules() const
+{
+	// By default we only evaluate the section with the highest Z-Order if they overlap on the same row
+	struct FDefaultCompilerRules : FMovieSceneSegmentCompilerRules
+	{
+		virtual void BlendSegment(FMovieSceneSegment& Segment, const TArrayView<const FMovieSceneSectionData>& SourceData) const
+		{
+			MovieSceneSegmentCompiler::BlendSegmentHighPass(Segment, SourceData);
+		}
+	};
+	return FDefaultCompilerRules();
+}
+
+TInlineValue<FMovieSceneSegmentCompilerRules> UMovieSceneTrack::GetTrackCompilerRules() const
+{
+	struct FRoundToNearstSectionRules : FMovieSceneSegmentCompilerRules
+	{
+		virtual TOptional<FMovieSceneSegment> InsertEmptySpace(const TRange<float>& Range, const FMovieSceneSegment* PreviousSegment, const FMovieSceneSegment* NextSegment) const
+		{
+			return MovieSceneSegmentCompiler::EvaluateNearestSegment(Range, PreviousSegment, NextSegment);
+		}
+	};
+
+	// Evaluate according to bEvaluateNearestSection preference
+	if (EvalOptions.bCanEvaluateNearestSection && EvalOptions.bEvaluateNearestSection)
+	{
+		return FRoundToNearstSectionRules();
+	}
+	else
+	{
+		return FMovieSceneSegmentCompilerRules();
+	}
+}
+
+void UMovieSceneTrack::GenerateTemplate(const FMovieSceneTrackCompilerArgs& Args) const
+{
+	FMovieSceneEvaluationTrack NewTrackTemplate(Args.ObjectBindingId);
+
+	// Legacy path
+	if (CreateLegacyInstance().IsValid())
+	{
+		NewTrackTemplate.DefineAsSingleTemplate(FMovieSceneLegacyTrackInstanceTemplate(this));
+		Args.Generator.AddLegacyTrack(MoveTemp(NewTrackTemplate), *this);
+		return;
+	}
+
+	if (Compile(NewTrackTemplate, Args) != EMovieSceneCompileResult::Success)
+	{
+		return;
+	}
+
+	Args.Generator.AddOwnedTrack(MoveTemp(NewTrackTemplate), *this);
+}
+
+struct FSegmentRemapper
+{
+	void ProcessSegments(const TArray<FMovieSceneSegment>& SourceSegments, FMovieSceneEvaluationTrack& OutTrack, const TFunctionRef<FMovieSceneEvalTemplatePtr(int32)>& TemplateFactory)
+	{
+		NewSegments.Reset(SourceSegments.Num());
+		NewIndices.Reset();
+
+		for (const FMovieSceneSegment& Segment : SourceSegments)
+		{
+			AddSegment(Segment, OutTrack, TemplateFactory);
+		}
+
+		OutTrack.SetSegments(MoveTemp(NewSegments));
+	}
+
+	void AddSegment(const FMovieSceneSegment& SourceSegment, FMovieSceneEvaluationTrack& OutTrack, const TFunctionRef<FMovieSceneEvalTemplatePtr(int32)>& TemplateFactory)
+	{
+		FMovieSceneSegment NewSegment;
+		NewSegment.Range = SourceSegment.Range;
+
+		for (const FSectionEvaluationData& EvalData : SourceSegment.Impls)
+		{
+			EnsureIndexIsValid(EvalData.ImplIndex);
+
+			// Ensure all our templates have been added to the track
+			if (NewIndices[EvalData.ImplIndex] == NotCreatedYet)
+			{
+				FMovieSceneEvalTemplatePtr NewTemplate = TemplateFactory(EvalData.ImplIndex);
+				if (NewTemplate.IsValid())
+				{
+					NewIndices[EvalData.ImplIndex] = OutTrack.AddChildTemplate(MoveTemp(NewTemplate));
+				}
+				else
+				{
+					NewIndices[EvalData.ImplIndex] = NullTemplate;
+				}
+			}
+
+			if (NewIndices[EvalData.ImplIndex] == NullTemplate)
+			{
+				continue;
+			}
+
+			FSectionEvaluationData NewData = EvalData;
+			NewData.ImplIndex = NewIndices[EvalData.ImplIndex];
+			NewSegment.Impls.Add(NewData);
+		}
+
+		if (NewSegment.Impls.Num())
+		{
+			NewSegments.Add(NewSegment);
+		}
+	}
+
+private:
+
+	static const int32 NotCreatedYet;
+	static const int32 NullTemplate;
+
+	TArray<FMovieSceneSegment> NewSegments;
+	TArray<int32> NewIndices;
+
+	void EnsureIndexIsValid(int32 InSourceIndex)
+	{
+		NewIndices.Reserve(InSourceIndex + 1);
+		for (int32 Index = NewIndices.Num(); Index < InSourceIndex + 1; ++Index)
+		{
+			NewIndices.Add(NotCreatedYet);
+		}
+	}
+};
+
+const int32 FSegmentRemapper::NotCreatedYet = -1;
+const int32 FSegmentRemapper::NullTemplate = -2;
+
+EMovieSceneCompileResult UMovieSceneTrack::Compile(FMovieSceneEvaluationTrack& OutTrack, const FMovieSceneTrackCompilerArgs& Args) const
+{
+	OutTrack.SetPreAndPostrollConditions(EvalOptions.bEvaluateInPreroll, EvalOptions.bEvaluateInPostroll);
+
+	EMovieSceneCompileResult Result = CustomCompile(OutTrack, Args);
+
+	if (Result == EMovieSceneCompileResult::Unimplemented)
+	{
+		// Default implementation
+		const TArray<UMovieSceneSection*>& AllSections = GetAllSections();
+		
+		TInlineValue<FMovieSceneSegmentCompilerRules> RowCompilerRules = GetRowCompilerRules();
+		FMovieSceneTrackCompiler::FRows TrackRows(AllSections, RowCompilerRules.GetPtr(nullptr));
+
+		FMovieSceneTrackCompiler Compiler;
+		FMovieSceneTrackEvaluationField EvaluationField = Compiler.Compile(TrackRows.Rows, GetTrackCompilerRules().GetPtr(nullptr));
+
+		FSegmentRemapper Remapper;
+		Remapper.ProcessSegments(EvaluationField.Segments, OutTrack,
+			[&](int32 SectionIndex){
+				FMovieSceneEvalTemplatePtr NewTemplate = CreateTemplateForSection(*AllSections[SectionIndex]);
+				if (NewTemplate.IsValid())
+				{
+					NewTemplate->SetCompletionMode(AllSections[SectionIndex]->EvalOptions.CompletionMode);
+				}
+				return NewTemplate;
+			}
+		);
+
+		Result = EMovieSceneCompileResult::Success;
+	}
+
+	if (Result == EMovieSceneCompileResult::Success)
+	{
+		PostCompile(OutTrack, Args);
+	}
+
+	return Result;
+}
+
+FMovieSceneEvalTemplatePtr UMovieSceneTrack::CreateTemplateForSection(const UMovieSceneSection& InSection) const
+{
+	return InSection.GenerateTemplate();
+}
diff --git a/Engine/Source/Runtime/MovieScene/Private/Tests/MovieSceneTransformTests.cpp b/Engine/Source/Runtime/MovieScene/Private/Tests/MovieSceneTransformTests.cpp
new file mode 100644
--- /dev/null
+++ b/Engine/Source/Runtime/MovieScene/Private/Tests/MovieSceneTransformTests.cpp
@@ -0,0 +1,169 @@
+// Copyright 1998-2016 Epic Games, Inc. All Rights Reserved.
+
+#include "MovieScenePrivatePCH.h"
+#include "MovieSceneSectionParameters.h"
+
+#define LOCTEXT_NAMESPACE "MovieSceneSubTrack"
+
+bool IsNearly(TRangeBound<float> A, TRangeBound<float> B)
+{
+	if (A.IsOpen() || B.IsOpen())
+	{
+		return A.IsOpen() == B.IsOpen();
+	}
+	else if (A.IsInclusive() != B.IsInclusive())
+	{
+		return false;
+	}
+	
+	return FMath::IsNearlyEqual(A.GetValue(), B.GetValue());
+}
+
+bool IsNearly(TRange<float> A, TRange<float> B)
+{
+	return IsNearly(A.GetLowerBound(), B.GetLowerBound()) && IsNearly(A.GetUpperBound(), B.GetUpperBound());
+}
+
+namespace LexicalConversion
+{
+	FString ToString(const TRange<float>& InRange)
+	{
+		TRangeBound<float> SourceLower = InRange.GetLowerBound();
+		TRangeBound<float> SourceUpper = InRange.GetUpperBound();
+
+		return *FString::Printf(TEXT("%s, %s"),
+			SourceLower.IsOpen() ? 
+				TEXT("...") : 
+				SourceLower.IsInclusive() ?
+					*FString::Printf(TEXT("[%.5f"), SourceLower.GetValue()) :
+					*FString::Printf(TEXT("(%.5f"), SourceLower.GetValue()),
+
+			SourceUpper.IsOpen() ? 
+				TEXT("...") : 
+				SourceUpper.IsInclusive() ?
+					*FString::Printf(TEXT("%.5f]"), SourceUpper.GetValue()) :
+					*FString::Printf(TEXT("%.5f)"), SourceUpper.GetValue())
+			);
+	}
+}
+
+bool TestTransform(FAutomationTestBase& Test, FMovieSceneSequenceTransform Transform, TArrayView<TRange<float>> InSource, TArrayView<TRange<float>> InExpected, const TCHAR* TestName)
+{
+	using namespace LexicalConversion;
+
+	check(InSource.Num() == InExpected.Num());
+
+	bool bSuccess = true;
+	for (int32 Index = 0; Index < InSource.Num(); ++Index)
+	{
+		TRange<float> Result = InSource[Index] * Transform;
+		if (!IsNearly(Result, InExpected[Index]))
+		{
+			Test.AddError(FString::Printf(TEXT("Test '%s' failed (Index %d). Transform (Scale %.3f, Offset %.3f) did not apply correctly (%s != %s)"),
+				TestName,
+				Index,
+				Transform.TimeScale,
+				Transform.Offset,
+				*ToString(Result),
+				*ToString(InExpected[Index])));
+
+			bSuccess = false;
+		}
+	}
+
+	return bSuccess;
+}
+
+// Calculate the transform that transforms from range A to range B
+FMovieSceneSequenceTransform TransformRange(float StartA, float EndA, float StartB, float EndB)
+{
+	return FMovieSceneSequenceTransform(StartB, (EndB - StartB) / (EndA - StartA)) * FMovieSceneSequenceTransform(-StartA);
+}
+
+IMPLEMENT_SIMPLE_AUTOMATION_TEST(FMovieSceneSubSectionCoreTransformsTest, "System.Engine.Sequencer.Core.Transforms", EAutomationTestFlags::EditorContext | EAutomationTestFlags::EngineFilter)
+bool FMovieSceneSubSectionCoreTransformsTest::RunTest(const FString& Parameters)
+{
+	// We test using ranges since that implicitly tests float transformation
+	static const TRangeBound<float> OpenBound;
+
+	TRange<float> InfiniteRange(OpenBound, OpenBound);
+	TRange<float> OpenLowerRange(OpenBound, 200.f);
+	TRange<float> OpenUpperRange(100.f, OpenBound);
+	TRange<float> ClosedRange(100.f, 200.f);
+
+	TRange<float> SourceRanges[] = {
+		InfiniteRange, OpenLowerRange, OpenUpperRange, ClosedRange
+	};
+
+	bool bSuccess = true;
+
+	{
+		// Test Multiplication with an identity transform
+		FMovieSceneSequenceTransform IdentityTransform;
+
+		TRange<float> Expected[] = {
+			InfiniteRange, OpenLowerRange, OpenUpperRange, ClosedRange
+		};
+		
+		bSuccess = TestTransform(*this, IdentityTransform, SourceRanges, Expected, TEXT("IdentityTransform")) && bSuccess;
+	}
+
+	{
+		// Test a simple translation
+		FMovieSceneSequenceTransform Transform(100.f, 1);
+
+		TRange<float> Expected[] = {
+			InfiniteRange, TRange<float>(OpenBound, 300.f), TRange<float>(200.f, OpenBound), TRange<float>(200.f, 300.f)
+		};
+
+		bSuccess = TestTransform(*this, Transform, SourceRanges, Expected, TEXT("Simple Translation")) && bSuccess;
+	}
+
+	{
+		// Test a simple translation + time scale
+
+		// Transform 100 - 200 to -200 - 1000
+		FMovieSceneSequenceTransform Transform = TransformRange(100, 200, -200, 1000);
+
+		TRange<float> Expected[] = {
+			InfiniteRange, TRange<float>(OpenBound, 1000.f), TRange<float>(-200.f, OpenBound), TRange<float>(-200.f, 1000.f)
+		};
+
+		bSuccess = TestTransform(*this, Transform, SourceRanges, Expected, TEXT("Simple Translation + half speed")) && bSuccess;
+	}
+
+	{
+		// Test that transforming a float by the same transform multiple times, does the same as the equivalent accumulated transform
+
+		// scales by 2, then offsets by 100
+		FMovieSceneSequenceTransform SeedTransform = FMovieSceneSequenceTransform(100.f, 0.5f);
+		FMovieSceneSequenceTransform AccumulatedTransform;
+
+		float SeedValue = 10.f;
+		for (int32 i = 0; i < 5; ++i)
+		{
+			AccumulatedTransform = SeedTransform * AccumulatedTransform;
+
+			SeedValue = SeedValue * SeedTransform;
+		}
+
+		float AccumValue = 10.f * AccumulatedTransform;
+		if (!FMath::IsNearlyEqual(AccumValue, SeedValue))
+		{
+			AddError(FString::Printf(TEXT("Accumulated transform does not have the same effect as separate transformations (%.7f != %.7f)"), AccumValue, SeedValue));
+		}
+
+		FMovieSceneSequenceTransform InverseTransform = AccumulatedTransform.Inverse();
+
+		float InverseValue = AccumValue * InverseTransform;
+		if (!FMath::IsNearlyEqual(InverseValue, 10.f))
+		{
+			AddError(FString::Printf(TEXT("Inverse accumulated transform does not return value back to its original value (%.7f != 10.f)"), InverseValue));
+		}
+	}
+
+	return true;
+}
+
+
+#undef LOCTEXT_NAMESPACE
\ No newline at end of file
diff --git a/Engine/Source/Runtime/MovieScene/Public/Compilation/IMovieSceneTemplateGenerator.h b/Engine/Source/Runtime/MovieScene/Public/Compilation/IMovieSceneTemplateGenerator.h
new file mode 100644
--- /dev/null
+++ b/Engine/Source/Runtime/MovieScene/Public/Compilation/IMovieSceneTemplateGenerator.h
@@ -0,0 +1,65 @@
+// Copyright 1998-2016 Epic Games, Inc. All Rights Reserved.
+
+#pragma once
+
+#include "MovieSceneFwd.h"
+
+/** Abstract base class used to generate evaluation templates */
+struct IMovieSceneTemplateGenerator
+{
+	/**
+	 * Add a new track that is to be owned by this template
+	 *
+	 * @param InTrackTemplate			The track template to add
+	 * @param SourceTrack				The originating track
+	 */
+	virtual void AddOwnedTrack(FMovieSceneEvaluationTrack&& InTrackTemplate, const UMovieSceneTrack& SourceTrack) = 0;
+
+	/**
+	 * Add a new track that is potentially shared between multiple tracks. Only one instance of SharedId can exist.
+	 *
+	 * @param InTrackTemplate			The track template to add
+	 * @param SharedId					The identifier of the shared track (for cross referencing between tracks)
+	 * @param SourceTrack				The originating track
+	 */
+	virtual void AddSharedTrack(FMovieSceneEvaluationTrack&& InTrackTemplate, FMovieSceneSharedDataId SharedId, const UMovieSceneTrack& SourceTrack) = 0;
+
+	/**
+	 * Add a legacy track to the template
+	 *
+	 * @param InTrackTemplate			The legacy track template
+	 * @param SourceTrack				The originating track
+	 */
+	virtual void AddLegacyTrack(FMovieSceneEvaluationTrack&& InTrackTemplate, const UMovieSceneTrack& SourceTrack) = 0;
+
+	/**
+	 * Add a set of remapped segments from a sub sequence to this template
+	 *
+	 * @param RootRange					The range in which to add the specified segments (at the root level)
+	 * @param SegmentPtrs				The remapped segment ptrs
+	 */
+	virtual void AddExternalSegments(TRange<float> RootRange, TArrayView<const FMovieSceneEvaluationFieldSegmentPtr> SegmentPtrs) = 0;
+
+	/**
+	 * Get a sequence's transform from its ID
+	 *
+	 * @param SequenceID				The ID of the sequence to get a transform for
+	 * @return The sequence's transform
+	 */
+	virtual FMovieSceneSequenceTransform GetSequenceTransform(FMovieSceneSequenceIDRef InSequenceID) const = 0;
+
+	/**
+	 * Generate a new sequence ID for the specified sub sequence
+	 *
+	 * @param SequenceData				Data pertaining to the sequence to add
+	 * @param ParentID					ID of the parent sequence
+	 */
+	virtual FMovieSceneSequenceID GenerateSequenceID(FMovieSceneSubSequenceData SequenceData, FMovieSceneSequenceIDRef ParentID) = 0;
+
+	/**
+	 * Indicate that the specified evaluation group should flush the execution stack immediately
+	 *
+	 * @param EvaluationGroup			The name of the evaluation group to flush
+	 */
+	virtual void FlushGroupImmediately(FName EvaluationGroup) = 0;
+};
diff --git a/Engine/Source/Runtime/MovieScene/Public/Compilation/MovieSceneCompilerRules.h b/Engine/Source/Runtime/MovieScene/Public/Compilation/MovieSceneCompilerRules.h
new file mode 100644
--- /dev/null
+++ b/Engine/Source/Runtime/MovieScene/Public/Compilation/MovieSceneCompilerRules.h
@@ -0,0 +1,84 @@
+
+
+#pragma once
+
+#include "MovieSceneSegmentCompiler.h"
+
+namespace MovieSceneSegmentCompiler
+{
+	static TOptional<FMovieSceneSegment> EvaluateNearestSegment(const TRange<float>& Range, const FMovieSceneSegment* PreviousSegment, const FMovieSceneSegment* NextSegment)
+	{
+		if (PreviousSegment)
+		{
+			// There is a preceeding segment
+			float PreviousSegmentRangeBound = PreviousSegment->Range.GetUpperBoundValue();
+
+			FMovieSceneSegment EmptySpace(Range);
+			for (FSectionEvaluationData Data : PreviousSegment->Impls)
+			{
+				EmptySpace.Impls.Add(FSectionEvaluationData(Data.ImplIndex, PreviousSegmentRangeBound));
+			}
+			return EmptySpace;
+		}
+		else if (NextSegment)
+		{
+			// Before any sections
+			float NextSegmentRangeBound = NextSegment->Range.GetLowerBoundValue();
+
+			FMovieSceneSegment EmptySpace(Range);
+			for (FSectionEvaluationData Data : NextSegment->Impls)
+			{
+				EmptySpace.Impls.Add(FSectionEvaluationData(Data.ImplIndex, NextSegmentRangeBound));
+			}
+			return EmptySpace;
+		}
+
+		return TOptional<FMovieSceneSegment>();
+	}
+
+	static void BlendSegmentHighPass(FMovieSceneSegment& Segment, const TArrayView<const FMovieSceneSectionData>& SourceData)
+	{
+		if (!Segment.Impls.Num())
+		{
+			return;
+		}
+
+		Segment.Impls.Sort(
+			[&](const FSectionEvaluationData& A, const FSectionEvaluationData& B)
+			{
+				return SourceData[A.ImplIndex].Priority > SourceData[B.ImplIndex].Priority;
+			}
+		);
+
+		const int32 HighestPriority = SourceData[Segment.Impls[0].ImplIndex].Priority;
+		for (int32 RemoveAtIndex = 1; RemoveAtIndex < Segment.Impls.Num(); ++RemoveAtIndex)
+		{
+			if (SourceData[Segment.Impls[RemoveAtIndex].ImplIndex].Priority != HighestPriority)
+			{
+				Segment.Impls.RemoveAt(RemoveAtIndex, Segment.Impls.Num() - RemoveAtIndex, false);
+				break;
+			}
+		}
+	}
+}
+
+class FMovieSceneAdditiveCameraRules : public FMovieSceneSegmentCompilerRules
+{
+public:
+	FMovieSceneAdditiveCameraRules(const UMovieSceneTrack* InTrack)
+		: Sections(InTrack->GetAllSections())
+	{
+	}
+
+	virtual void BlendSegment(FMovieSceneSegment& Segment, const TArrayView<const FMovieSceneSectionData>& SourceData) const
+	{
+		// sort by start time to match application order of player camera
+		Segment.Impls.Sort(
+			[&](FSectionEvaluationData A, FSectionEvaluationData B){
+				return Sections[SourceData[A.ImplIndex].SourceIndex]->GetStartTime() < Sections[SourceData[B.ImplIndex].SourceIndex]->GetStartTime();
+			}
+		);
+	}
+
+	const TArray<UMovieSceneSection*>& Sections;
+};
\ No newline at end of file
diff --git a/Engine/Source/Runtime/MovieScene/Public/Compilation/MovieSceneEvaluationTemplateGenerator.h b/Engine/Source/Runtime/MovieScene/Public/Compilation/MovieSceneEvaluationTemplateGenerator.h
new file mode 100644
--- /dev/null
+++ b/Engine/Source/Runtime/MovieScene/Public/Compilation/MovieSceneEvaluationTemplateGenerator.h
@@ -0,0 +1,105 @@
+// Copyright 1998-2016 Epic Games, Inc. All Rights Reserved.
+
+#pragma once
+
+#include "MovieSceneFwd.h"
+#include "IMovieSceneTemplateGenerator.h"
+#include "MovieSceneEvaluationTemplate.h"
+#include "MovieSceneSegmentCompiler.h"
+#include "ObjectKey.h"
+
+/**
+ * An immutable class used for efficient runtime evaluation of UMovieScene data.
+ */
+class FMovieSceneEvaluationTemplateGenerator : IMovieSceneTemplateGenerator
+{
+public:
+
+	/**
+	 * Construction from a sequence and a destination template
+	 */
+	FMovieSceneEvaluationTemplateGenerator(UMovieSceneSequence& InSequence, FMovieSceneEvaluationTemplate& OutTemplate, FMovieSceneSequenceTemplateStore& InStore);
+
+	/**
+	 * Generate the template with the specified parameters
+	 *
+	 * @param InParams 			Compilation parameters
+	 */
+	void Generate(FMovieSceneTrackCompilationParams InParams);
+
+private:
+
+	/**
+	 * Process a track, potentially causing it to be compiled if it's out of date
+	 *
+	 * @param InTrack 			The track to process
+	 * @param ObjectId 			Object binding ID of the track
+	 */
+	void ProcessTrack(const UMovieSceneTrack& InTrack, const FGuid& ObjectId = FGuid());
+
+
+	/**
+	 * Remove any references to old track templates that are no longer required to exist
+	 */
+	void RemoveOldTrackReferences();
+
+	/**
+	 * Update the evaluation field with the specified segments and segment ptrs
+	 *
+	 * @param Segments 			The segments to populate the evaluation field with
+	 * @param Impls 			Array of implementation pointers that each segment refers to
+	 */
+	void UpdateEvaluationField(const TArray<FMovieSceneSegment>& Segments, const TArray<FMovieSceneEvaluationFieldSegmentPtr>& Impls, const TMap<FMovieSceneSequenceID, FMovieSceneEvaluationTemplate*>& Templates);
+
+private:
+
+	/** IMovieSceneTemplateGenerator overrides */
+	virtual void AddLegacyTrack(FMovieSceneEvaluationTrack&& InTrackTemplate, const UMovieSceneTrack& SourceTrack) override;
+	virtual void AddOwnedTrack(FMovieSceneEvaluationTrack&& InTrackTemplate, const UMovieSceneTrack& SourceTrack) override;
+	virtual void AddSharedTrack(FMovieSceneEvaluationTrack&& InTrackTemplate, FMovieSceneSharedDataId SharedId, const UMovieSceneTrack& SourceTrack) override;
+	virtual void AddExternalSegments(TRange<float> RootRange, TArrayView<const FMovieSceneEvaluationFieldSegmentPtr> SegmentPtrs) override;
+	virtual FMovieSceneSequenceTransform GetSequenceTransform(FMovieSceneSequenceIDRef InSequenceID) const override;
+	virtual FMovieSceneSequenceID GenerateSequenceID(FMovieSceneSubSequenceData SequenceData, FMovieSceneSequenceIDRef ParentID) override;
+	virtual void FlushGroupImmediately(FName EvaluationGroup) override;
+
+private:
+
+	/**
+	 * Sort predicate used to sort evaluation tracks in a given group
+	 */
+	bool SortPredicate(const FMovieSceneEvaluationTrack* A, const FMovieSceneEvaluationTrack* B) const;
+
+	/**
+	 * Attempt to find a track from its evaluation field pointer
+	 *
+	 * @param InPtr 			Ptr to the track to lookup
+	 * @return The track, or nullptr
+	 */
+	const FMovieSceneEvaluationTrack* LookupTrack(const FMovieSceneEvaluationFieldTrackPtr& InPtr, const TMap<FMovieSceneSequenceID, FMovieSceneEvaluationTemplate*>& Templates);
+
+private:
+
+	/** The source sequence we are generating a template for */
+	UMovieSceneSequence& SourceSequence;
+	/** The destination template we're generating into */
+	FMovieSceneEvaluationTemplate& Template;
+
+	/** Transient compiler args that are passed to each track that is compiled */
+	FMovieSceneTrackCompilerArgs TransientArgs;
+
+	/** Temporary set of signatures we've come across in the sequence we're compiling */
+	TSet<FGuid> CompiledSignatures;
+
+	/** Segment data used for evaluation field generation */
+	TArray<FMovieSceneSectionData> SegmentData;
+	/** Look up table whose indices match implementation indices in the above segments */
+	TArray<FMovieSceneEvaluationFieldSegmentPtr> TrackLUT;
+
+	/** Map of external segments added from sub sequences */
+	TMap<FMovieSceneEvaluationFieldSegmentPtr, int32> ExternalSegmentLookup;
+
+	/** Set of shared track keys that have been added (used to ensure we only add a shared track ID once) */
+	TSet<FSharedPersistentDataKey> AddedSharedTracks;
+	/** Set that defines whether an evaluation group should immediately flush the execution token stack or not */
+	TSet<FName> ImmediateFlushGroups;
+};
\ No newline at end of file
diff --git a/Engine/Source/Runtime/MovieScene/Public/Compilation/MovieSceneSegmentCompiler.h b/Engine/Source/Runtime/MovieScene/Public/Compilation/MovieSceneSegmentCompiler.h
new file mode 100644
--- /dev/null
+++ b/Engine/Source/Runtime/MovieScene/Public/Compilation/MovieSceneSegmentCompiler.h
@@ -0,0 +1,110 @@
+// Copyright 1998-2016 Epic Games, Inc. All Rights Reserved.
+
+#pragma once
+
+#include "Evaluation/MovieSceneSegment.h"
+
+struct FMovieSceneSectionData
+{
+	MOVIESCENE_API FMovieSceneSectionData(const TRange<float>& InBounds, int32 InSourceIndex, int32 InPriority = 0);
+
+	TRange<float> Bounds;
+	int32 SourceIndex;
+	int32 Priority;
+};
+
+struct FMovieSceneSegmentCompilerRules
+{
+	FMovieSceneSegmentCompilerRules()
+	{
+		bAllowEmptySegments = false;
+	}
+
+	bool AllowEmptySegments() const { return bAllowEmptySegments; }
+
+	virtual TOptional<FMovieSceneSegment> InsertEmptySpace(const TRange<float>& Range, const FMovieSceneSegment* PreviousSegment, const FMovieSceneSegment* NextSegment) const
+	{
+		return TOptional<FMovieSceneSegment>();
+	}
+
+	virtual void BlendSegment(FMovieSceneSegment& Segment, const TArrayView<const FMovieSceneSectionData>& SourceData) const
+	{
+	}
+
+	virtual void PostProcessSegments(TArray<FMovieSceneSegment>& Segments, const TArrayView<const FMovieSceneSectionData>& SourceData) const
+	{
+	}
+
+	MOVIESCENE_API void ProcessSegments(TArray<FMovieSceneSegment>& Segments, const TArrayView<const FMovieSceneSectionData>& SourceData) const;
+
+protected:
+
+	bool InsertSegment(TArray<FMovieSceneSegment>& Segments, int32 Index, const TArrayView<const FMovieSceneSectionData>& SourceData) const;
+
+	bool bAllowEmptySegments;
+};
+
+struct FMovieSceneSegmentCompilerResult
+{
+	TArray<FMovieSceneSegment> Segments;
+};
+
+struct FMovieSceneSegmentCompiler
+{
+	MOVIESCENE_API TArray<FMovieSceneSegment> Compile(const TArrayView<const FMovieSceneSectionData>& Data, const FMovieSceneSegmentCompilerRules* Rules = nullptr);
+
+private:
+
+	void CloseCompletedSegments();
+
+	struct FBound
+	{
+		FBound(int32 InImplIndex, TRangeBound<float> InBound) : ImplIndex(InImplIndex), Bound(InBound) {}
+
+		int32 ImplIndex;
+		TRangeBound<float> Bound;
+	};
+
+	int32 LowerReadIndex, UpperReadIndex;
+	TArray<FBound> LowerBounds, UpperBounds;
+	TArray<FMovieSceneSegment> CompiledSegments;
+	TArray<int32, TInlineAllocator<16>> OverlappingSections;
+	TArray<uint32, TInlineAllocator<16>> OverlappingRefCounts;
+};
+
+struct FMovieSceneTrackEvaluationField
+{
+	TArray<FMovieSceneSegment> Segments;
+};
+
+struct FMovieSceneTrackCompiler
+{
+	struct FMovieSceneSectionRowData : FMovieSceneSectionData
+	{
+		FMovieSceneSectionRowData(int32 InActualSectionIndex, const TRange<float>& InBounds, int32 InIndexInRow, int32 InPriority = 0)
+			: FMovieSceneSectionData(InBounds, InIndexInRow, InPriority)
+			, ActualSectionIndex(InActualSectionIndex)
+		{}
+
+		int32 ActualSectionIndex;
+	};
+
+	struct FRow
+	{
+		TArray<FMovieSceneSectionRowData, TInlineAllocator<8>> Sections;
+
+		const FMovieSceneSegmentCompilerRules* CompileRules;
+	};
+
+	struct FRows
+	{
+		MOVIESCENE_API FRows(const TArray<UMovieSceneSection*>& Sections, const FMovieSceneSegmentCompilerRules* CompileRules = nullptr);
+		
+		TArray<FRow, TInlineAllocator<2>> Rows;
+	};
+
+	MOVIESCENE_API FMovieSceneTrackEvaluationField Compile(const TArrayView<FRow>& Rows, const FMovieSceneSegmentCompilerRules* Rules = nullptr);
+};
+
+
+
diff --git a/Engine/Source/Runtime/MovieScene/Public/Evaluation/MovieSceneEvaluationKey.h b/Engine/Source/Runtime/MovieScene/Public/Evaluation/MovieSceneEvaluationKey.h
new file mode 100644
--- /dev/null
+++ b/Engine/Source/Runtime/MovieScene/Public/Evaluation/MovieSceneEvaluationKey.h
@@ -0,0 +1,72 @@
+// Copyright 1998-2016 Epic Games, Inc. All Rights Reserved.
+
+#pragma once
+
+#include "MovieSceneTrackIdentifier.h"
+
+/** Keyable struct that represents a particular entity within an evaluation template (either a section/template or a track) */
+struct FMovieSceneEvaluationKey
+{
+	/**
+	 * Default construction to an invalid key
+	 */
+	FMovieSceneEvaluationKey()
+		: SequenceID(MovieSceneSequenceID::Invalid)
+		, TrackIdentifier(FMovieSceneTrackIdentifier::Invalid())
+		, SectionIdentifier(-1)
+	{}
+
+	/**
+	 * User construction
+	 */
+	FMovieSceneEvaluationKey(FMovieSceneSequenceIDRef InSequenceID, FMovieSceneTrackIdentifier InTrackIdentifier, uint32 InSectionIdentifier = uint32(-1))
+		: SequenceID(InSequenceID)
+		, TrackIdentifier(InTrackIdentifier)
+		, SectionIdentifier(InSectionIdentifier)
+	{}
+
+	/**
+	 * Check whether this key is valid
+	 */
+	bool IsValid() const
+	{
+		return SequenceID != MovieSceneSequenceID::Invalid && TrackIdentifier != FMovieSceneTrackIdentifier::Invalid();
+	}
+
+	/**
+	 * Derive a new key from this one using the specified section identifier
+	 */
+	FMovieSceneEvaluationKey AsSection(uint32 InSectionIdentifier) const
+	{
+		FMovieSceneEvaluationKey NewKey = *this;
+		NewKey.SectionIdentifier = InSectionIdentifier;
+		return NewKey;
+	}
+
+	/**
+	 * Convert this key into a track key
+	 */
+	FMovieSceneEvaluationKey AsTrack() const
+	{
+		FMovieSceneEvaluationKey NewKey = *this;
+		NewKey.SectionIdentifier = uint32(-1);
+		return NewKey;
+	}
+
+	friend bool operator==(const FMovieSceneEvaluationKey& A, const FMovieSceneEvaluationKey& B)
+	{
+		return A.TrackIdentifier == B.TrackIdentifier && A.SequenceID == B.SequenceID && A.SectionIdentifier == B.SectionIdentifier;
+	}
+
+	friend uint32 GetTypeHash(const FMovieSceneEvaluationKey& In)
+	{
+		return GetTypeHash(In.SequenceID) ^ (~GetTypeHash(In.TrackIdentifier)) ^ In.SectionIdentifier;
+	}
+
+	/** ID of the sequence that the entity is contained within */
+	FMovieSceneSequenceID SequenceID;
+	/** ID of the track this key relates to */
+	FMovieSceneTrackIdentifier TrackIdentifier;
+	/** ID of the section this key relates to (or -1 where this key relates to a track) */
+	uint32 SectionIdentifier;
+};
\ No newline at end of file
diff --git a/Engine/Source/Runtime/MovieScene/Public/Evaluation/MovieSceneEvaluationState.h b/Engine/Source/Runtime/MovieScene/Public/Evaluation/MovieSceneEvaluationState.h
new file mode 100644
--- /dev/null
+++ b/Engine/Source/Runtime/MovieScene/Public/Evaluation/MovieSceneEvaluationState.h
@@ -0,0 +1,253 @@
+// Copyright 1998-2016 Epic Games, Inc. All Rights Reserved.
+
+#pragma once
+
+#include "MovieSceneFwd.h"
+#include "PersistentEvaluationData.h"
+
+class UMovieSceneSequence;
+class IMovieScenePlayer;
+
+/**
+ * Object cache that looks up, resolves, and caches object bindings for a specific sequence
+ */
+struct FMovieSceneObjectCache
+{
+	/**
+	 * Find all objects that are bound to the specified binding ID
+	 * @note Will look up, and cache any objects if the cache has been invalidated
+	 *
+	 * @param InBindingID		The object binding GUID for a spawnable or posessable in a UMovieScene
+	 * @param Player			The movie scene player that is playing back the sequence
+	 * @return An iterable type of all objects bound to the specified ID.
+	 */
+	MOVIESCENE_API TArrayView<TWeakObjectPtr<>> FindBoundObjects(const FGuid& InBindingID, IMovieScenePlayer& Player);
+
+	/**
+	 * Set the sequence that this cache applies to
+	 *
+	 * @param InSequence		The sequence that this cache applies to
+	 * @param InSequenceID		The ID of the sequence within the master sequence
+	 */
+	MOVIESCENE_API void SetSequence(UMovieSceneSequence& InSequence, FMovieSceneSequenceIDRef InSequenceID);
+
+	/**
+	 * Attempt deduce the posessable or spawnable that relates to the specified object
+	 * @note Will forcably resolve any out of date bindings in the entire sequence
+	 *
+	 * @param InObject			The object whose binding ID is to be find
+	 * @param Player			The movie scene player that is playing back the sequence
+	 * @return The object's spawnable or possessable GUID, or a zero GUID if it was not found
+	 */
+	MOVIESCENE_API FGuid FindObjectId(UObject& InObject, IMovieScenePlayer& Player);
+
+	/**
+	 * Invalidate any object bindings for objects that have been destroyed
+	 */
+	MOVIESCENE_API void InvalidateExpiredObjects();
+
+	/**
+	 * Invalidate the object bindings for a specific object binding ID
+	 *
+	 * @param InGuid			The object binding ID to invalidate bindings for
+	 */
+	MOVIESCENE_API void Invalidate(const FGuid& InGuid);
+
+	/**
+	 * Completely erase all knowledge of, anc caches for all object bindings
+	 */
+	void Clear()
+	{
+		BoundObjects.Reset();
+		ChildBindings.Reset();
+	}
+
+	/**
+	 * Get the sequence that this cache relates to
+	 */
+	UMovieSceneSequence* GetSequence() const { return WeakSequence.Get(); }
+
+private:
+	/**
+	 * Update the bindings for the specified GUID
+	 *
+	 * @param InGuid			The object binding ID to update bindings for
+	 * @param Player			The movie scene player that is playing back the sequence
+	 */
+	void UpdateBindings(const FGuid& InGuid, IMovieScenePlayer& Player);
+
+	struct FBoundObjects
+	{
+		bool bUpToDate;
+		TArray<TWeakObjectPtr<>, TInlineAllocator<1>> Objects;
+	};
+
+private:
+
+	/** The sequence that we're caching objects for */
+	TWeakObjectPtr<UMovieSceneSequence> WeakSequence;
+
+	/** The sequence ID of the sequence within the master sequence */
+	FMovieSceneSequenceID SequenceID;
+
+	template<typename ValueType>
+	struct TFastGuidKeyFuncs : BaseKeyFuncs<TPair<FGuid,ValueType>,FGuid,false>
+	{
+		typedef typename TTypeTraits<FGuid>::ConstPointerType KeyInitType;
+		typedef const TPairInitializer<typename TTypeTraits<FGuid>::ConstInitType, typename TTypeTraits<ValueType>::ConstInitType>& ElementInitType;
+
+		static FORCEINLINE KeyInitType GetSetKey(ElementInitType Element)
+		{
+			return Element.Key;
+		}
+		static FORCEINLINE bool Matches(KeyInitType A,KeyInitType B)
+		{
+			return A == B;
+		}
+		static FORCEINLINE uint32 GetKeyHash(KeyInitType Key)
+		{
+			return Key.A ^ Key.B ^ Key.C ^ Key.D;
+		}
+	};
+
+	/** A map of bound objects */
+	TMap<FGuid, FBoundObjects, FDefaultSetAllocator, TFastGuidKeyFuncs<FBoundObjects>> BoundObjects;
+
+	/** Map of child bindings for any given object binding */
+	typedef TArray<FGuid, TInlineAllocator<4>> FGuidArray;
+	TMap<FGuid, FGuidArray, FDefaultSetAllocator, TFastGuidKeyFuncs<FGuidArray>> ChildBindings;
+};
+
+/**
+ * Container that stores which entities have been evalauted in any given frame, and in what order
+ */
+struct FOrderedEvaluationCache
+{
+	FOrderedEvaluationCache() = default;
+
+#if PLATFORM_COMPILER_HAS_DEFAULTED_FUNCTIONS
+	FOrderedEvaluationCache(FOrderedEvaluationCache&&) = default;
+	FOrderedEvaluationCache& operator=(FOrderedEvaluationCache&&) = default;
+#else
+	FOrderedEvaluationCache(FOrderedEvaluationCache&& RHS) : KeySet(MoveTemp(RHS.KeySet)), OrderedKeys(MoveTemp(RHS.OrderedKeys)) {}
+	FOrderedEvaluationCache& operator=(FOrderedEvaluationCache&& RHS) { KeySet = MoveTemp(RHS.KeySet); OrderedKeys = MoveTemp(RHS.OrderedKeys); return *this; }
+#endif
+
+	/**
+	 * Check whether the specified entity was evaluated
+	 */
+	bool Contains(FMovieSceneEvaluationKey InKey) const;
+
+	/**
+	 * Report that the specified entity was evaluated
+	 */
+	void Add(FMovieSceneEvaluationKey InKey);
+
+	/**
+	 * Reset the container
+	 */
+	void Reset();
+
+public:
+	
+	/** Unordered set of evaluated items */
+	TSet<FMovieSceneEvaluationKey> KeySet;
+
+	/** Ordered set of evaluated items */
+	TArray<FMovieSceneEvaluationKey> OrderedKeys;
+};
+
+/**
+ * Provides runtime evaluation functions with the ability to look up state from the main game environment
+ */
+struct FMovieSceneEvaluationState
+{
+	/**
+	 * Assign a sequence to a specific ID
+	 *
+	 * @param InSequenceID		The sequence ID to assign to
+	 * @param InSequence		The sequence to assign
+	 */
+	void AssignSequence(FMovieSceneSequenceIDRef InSequenceID, UMovieSceneSequence& InSequence);
+
+	/**
+	 * Attempt to locate a sequence from its ID
+	 *
+	 * @param InSequenceID		The sequence ID to lookup
+	 */
+	MOVIESCENE_API UMovieSceneSequence* FindSequence(FMovieSceneSequenceIDRef InSequenceID) const;
+
+	/**
+	 * Attempt deduce the posessable or spawnable that relates to the specified object
+	 * @note Will forcably resolve any out of date bindings in the entire sequence
+	 *
+	 * @param InObject			The object whose binding ID is to be find
+	 * @param Player			The movie scene player that is playing back the sequence
+	 * @return The object's spawnable or possessable GUID, or a zero GUID if it was not found
+	 */
+	MOVIESCENE_API FGuid FindObjectId(UObject& Object, FMovieSceneSequenceIDRef InSequenceID, IMovieScenePlayer& Player);
+
+	/**
+	 * Find an object cache pertaining to the specified sequence
+	 *
+	 * @param InSequenceID		The sequence ID to lookup
+	 */
+	FORCEINLINE FMovieSceneObjectCache* FindObjectCache(FMovieSceneSequenceIDRef SequenceID)
+	{
+		return ObjectCaches.Find(SequenceID);
+	}
+
+	/**
+	 * Get an object cache pertaining to the specified sequence
+	 *
+	 * @param InSequenceID		The sequence ID to lookup
+	 */
+	FORCEINLINE FMovieSceneObjectCache& GetObjectCache(FMovieSceneSequenceIDRef SequenceID)
+	{
+		FMovieSceneObjectCache* Cache = ObjectCaches.Find(SequenceID);
+		if (!Cache)
+		{
+			Cache = &ObjectCaches.Add(SequenceID, FMovieSceneObjectCache());
+		}
+		return *Cache;
+	}
+
+	/**
+	 * Remove the object cache pertaining to the specified sequence
+	 *
+	 * @param InSequenceID		The sequence ID to remove
+	 */
+	FORCEINLINE void RemoveObjectCache(FMovieSceneSequenceIDRef SequenceID)
+	{
+		ObjectCaches.Remove(SequenceID);
+	}
+
+	/**
+	 * Invalidate any object caches that may now contain expired objects
+	 */
+	MOVIESCENE_API void InvalidateExpiredObjects();
+
+	/**
+	 * Forcably invalidate the specified object binding in the specified sequence
+	 *
+	 * @param InGuid			The object binding ID to invalidate
+	 * @param InSequenceID		The sequence ID to which the object binding belongs
+	 */
+	MOVIESCENE_API void Invalidate(const FGuid& InGuid, FMovieSceneSequenceIDRef InSequenceID);
+
+	/**
+	 * Forcably clear all object caches
+	 */
+	MOVIESCENE_API void ClearObjectCaches();
+
+	/** A map of persistent evaluation data mapped by movie scene evaluation entity (i.e, a given track or section) */
+	TMap<FMovieSceneEvaluationKey, TUniquePtr<IPersistentEvaluationData>> PersistentEntityData;
+
+	/** A map of persistent evaluation data mapped by shared evaluation key. Such data can be accessed from anywhere given an operand and a unique identifier. */
+	TMap<FSharedPersistentDataKey, TUniquePtr<IPersistentEvaluationData>> PersistentSharedData;
+
+private:
+
+	/** Maps of bound objects, arranged by template ID */
+	TMap<FMovieSceneSequenceID, FMovieSceneObjectCache> ObjectCaches;
+};
\ No newline at end of file
diff --git a/Engine/Source/Runtime/MovieScene/Public/Evaluation/MovieSceneEvaluationTemplate.h b/Engine/Source/Runtime/MovieScene/Public/Evaluation/MovieSceneEvaluationTemplate.h
new file mode 100644
--- /dev/null
+++ b/Engine/Source/Runtime/MovieScene/Public/Evaluation/MovieSceneEvaluationTemplate.h
@@ -0,0 +1,244 @@
+// Copyright 1998-2016 Epic Games, Inc. All Rights Reserved.
+
+#pragma once
+
+#include "MovieSceneFwd.h"
+#include "MovieSceneTrack.h"
+#include "MovieSceneEvaluationField.h"
+#include "MovieSceneEvaluationTrack.h"
+#include "MovieSceneSequenceHierarchy.h"
+#include "Compilation/IMovieSceneTemplateGenerator.h"
+#include "MovieSceneSequenceTemplateStore.h"
+#include "MovieSceneEvaluationTemplate.generated.h"
+
+USTRUCT()
+struct FMovieSceneGenerationLedger
+{
+	GENERATED_BODY()
+
+public:
+
+	TArrayView<FMovieSceneTrackIdentifier> FindTracks(const FGuid& InSignature);
+
+	void AddTrack(const FGuid& InSignature, FMovieSceneTrackIdentifier Identifier);
+
+	bool Serialize(FArchive& Ar);
+
+public:
+
+	UPROPERTY()
+	FMovieSceneTrackIdentifier LastTrackIdentifier;
+
+	/** Custom serialized - map of track identifiers to number of references within th template (generally 1, maybe >1 for shared tracks) */
+	TMap<FMovieSceneTrackIdentifier, int32> TrackReferenceCounts;
+
+	/** Custom serialized - map of track signature to array of track identifiers that it created */
+	TMap<FGuid, TArray<FMovieSceneTrackIdentifier, TInlineAllocator<1>>> TrackSignatureToTrackIdentifier;
+};
+template<> struct TStructOpsTypeTraits<FMovieSceneGenerationLedger> : public TStructOpsTypeTraitsBase { enum { WithSerializer = true }; };
+
+/**
+ * Template that is used for efficient runtime evaluation of a movie scene sequence. Potentially serialized into the asset.
+ */
+USTRUCT()
+struct FMovieSceneEvaluationTemplate
+{
+	GENERATED_BODY()
+
+	/**
+	 * Default construction
+	 */
+	FMovieSceneEvaluationTemplate()
+	{
+		bHasLegacyTrackInstances = 0;
+		bKeepStaleTracks = 0;
+	}
+
+	/**
+	 * Move construction/assignment
+	 */
+#if PLATFORM_COMPILER_HAS_DEFAULTED_FUNCTIONS
+	FMovieSceneEvaluationTemplate(FMovieSceneEvaluationTemplate&&) = default;
+	FMovieSceneEvaluationTemplate& operator=(FMovieSceneEvaluationTemplate&&) = default;
+#else
+	FMovieSceneEvaluationTemplate(FMovieSceneEvaluationTemplate&& RHS)
+	{
+		*this = MoveTemp(RHS);
+	}
+	FMovieSceneEvaluationTemplate& operator=(FMovieSceneEvaluationTemplate&& RHS)
+	{
+		Tracks = MoveTemp(RHS.Tracks);
+		StaleTracks = MoveTemp(RHS.StaleTracks);
+		EvaluationField = MoveTemp(RHS.EvaluationField);
+		Hierarchy = MoveTemp(RHS.Hierarchy);
+		Ledger = MoveTemp(RHS.Ledger);
+		bHasLegacyTrackInstances = RHS.bHasLegacyTrackInstances;
+		bKeepStaleTracks = RHS.bKeepStaleTracks;
+		return *this;
+	}
+#endif
+
+	/**
+	 * Non-copyable
+	 */
+	FMovieSceneEvaluationTemplate(const FMovieSceneEvaluationTemplate&) = default;
+	FMovieSceneEvaluationTemplate& operator=(const FMovieSceneEvaluationTemplate&) = default;
+
+public:
+
+	/**
+	 * Attempt to locate a track with the specified identifier
+	 */
+	FMovieSceneEvaluationTrack* FindTrack(FMovieSceneTrackIdentifier Identifier)
+	{
+		// Fast, most common path
+		if (FMovieSceneEvaluationTrack* Track = Tracks.Find(Identifier.Value))
+		{
+			return Track;
+		}
+
+		return bKeepStaleTracks ? StaleTracks.Find(Identifier.Value) : nullptr;
+	}
+
+	/**
+	 * Attempt to locate a track with the specified identifier
+	 */
+	const FMovieSceneEvaluationTrack* FindTrack(FMovieSceneTrackIdentifier Identifier) const
+	{
+		// Fast, most common path
+		if (const FMovieSceneEvaluationTrack* Track = Tracks.Find(Identifier.Value))
+		{
+			return Track;
+		}
+
+		return bKeepStaleTracks ? StaleTracks.Find(Identifier.Value) : nullptr;
+	}
+
+	/**
+	 * Add a new track for the specified identifier
+	 */
+	MOVIESCENE_API FMovieSceneTrackIdentifier AddTrack(const FGuid& InSignature, FMovieSceneEvaluationTrack&& InTrack);
+
+	/**
+	 * Remove any tracks that correspond to the specified signature
+	 */
+	MOVIESCENE_API void RemoveTrack(const FGuid& InSignature);
+
+	/**
+	 * Iterate this template's tracks.
+	 */
+	MOVIESCENE_API const TMap<FMovieSceneTrackIdentifier, FMovieSceneEvaluationTrack>& GetTracks() const;
+
+	/**
+	 * Find tracks within this template that relate to the specified signature
+	 */
+	MOVIESCENE_API TArrayView<FMovieSceneTrackIdentifier> FindTracks(const FGuid& InSignature);
+
+	/**
+	 * Purge any stale tracks we may have
+	 */
+	void PurgeStaleTracks()
+	{
+		StaleTracks.Reset();
+	}
+
+public:
+
+	/**
+	 * Get this template's generation ledger
+	 */
+	const FMovieSceneGenerationLedger& GetLedger() const
+	{
+		return Ledger;
+	}
+
+private:
+
+	/** Map of evaluation tracks from identifier to track */
+	UPROPERTY()
+	TMap<uint32, FMovieSceneEvaluationTrack> Tracks;
+
+	/** Transient map of stale tracks. Only populated dureing regeneration where bKeepStaleTracks is true */
+	TMap<uint32, FMovieSceneEvaluationTrack> StaleTracks;
+
+public:
+
+	/** Evaluation field for efficient runtime evaluation */
+	UPROPERTY()
+	FMovieSceneEvaluationField EvaluationField;
+
+	/** Map of all sequences found in this template (recursively) */
+	UPROPERTY()
+	FMovieSceneSequenceHierarchy Hierarchy;
+
+private:
+
+	UPROPERTY()
+	FMovieSceneGenerationLedger Ledger;
+
+public:
+
+	/** When set, this template contains legacy track instances that require the initialization of a legacy sequence instance */
+	UPROPERTY()
+	uint32 bHasLegacyTrackInstances : 1;
+
+	/** Primarily used in editor to keep stale tracks around during template regeneration to ensure we can call OnEndEvaluation on them. */
+	UPROPERTY()
+	uint32 bKeepStaleTracks : 1;
+};
+
+USTRUCT()
+struct FMovieSceneSequenceCachedSignature
+{
+	GENERATED_BODY()
+
+	FMovieSceneSequenceCachedSignature() {}
+	MOVIESCENE_API FMovieSceneSequenceCachedSignature(UMovieSceneSequence& InSequence);
+
+	UPROPERTY()
+	TWeakObjectPtr<UMovieSceneSequence> Sequence;
+
+	UPROPERTY()
+	FGuid CachedSignature;
+};
+
+USTRUCT()
+struct FCachedMovieSceneEvaluationTemplate : public FMovieSceneEvaluationTemplate
+{
+	GENERATED_BODY()
+
+#if WITH_EDITORONLY_DATA
+
+	FCachedMovieSceneEvaluationTemplate()
+		: Origin(nullptr)
+	{
+		bKeepStaleTracks = 1;
+	}
+
+	MOVIESCENE_API void Initialize(UMovieSceneSequence& InSequence, FMovieSceneSequenceTemplateStore* InOrigin = nullptr);
+
+	MOVIESCENE_API void Regenerate();
+
+	MOVIESCENE_API void Regenerate(const FMovieSceneTrackCompilationParams& NewParams);
+
+	MOVIESCENE_API void ForceRegenerate(const FMovieSceneTrackCompilationParams& NewParams);
+
+	MOVIESCENE_API bool IsOutOfDate(const FMovieSceneTrackCompilationParams& NewParams) const;
+
+private:
+
+	void RegenerateImpl(const FMovieSceneTrackCompilationParams& Params);
+
+	/** Transient data */
+	TWeakObjectPtr<UMovieSceneSequence> SourceSequence;
+	FMovieSceneSequenceTemplateStore* Origin;
+
+	UPROPERTY()
+	FMovieSceneTrackCompilationParams CachedCompilationParams;
+
+	UPROPERTY()
+	TArray<FMovieSceneSequenceCachedSignature> CachedSignatures;
+#endif
+};
+
+template<> struct TStructOpsTypeTraits<FCachedMovieSceneEvaluationTemplate> : public TStructOpsTypeTraitsBase { enum { WithCopy = false }; };
\ No newline at end of file
diff --git a/Engine/Source/Runtime/MovieScene/Public/Evaluation/MovieSceneEvaluationTemplateInstance.h b/Engine/Source/Runtime/MovieScene/Public/Evaluation/MovieSceneEvaluationTemplateInstance.h
new file mode 100644
--- /dev/null
+++ b/Engine/Source/Runtime/MovieScene/Public/Evaluation/MovieSceneEvaluationTemplateInstance.h
@@ -0,0 +1,286 @@
+// Copyright 1998-2016 Epic Games, Inc. All Rights Reserved.
+
+#pragma once
+
+#include "MovieSceneFwd.h"
+#include "MovieSceneEvaluationTemplate.h"
+#include "MovieSceneEvaluationTemplateInstance.generated.h"
+
+class FMovieSceneSequenceInstance;
+
+/**
+ * An instance of an evaluation template. Fast to initialize and evaluate.
+ */
+struct FMovieSceneEvaluationTemplateInstance
+{
+	FMovieSceneEvaluationTemplateInstance() : Sequence(nullptr), Template(nullptr), SequenceID(MovieSceneSequenceID::Invalid) {}
+
+	/**
+	 * Constructor for root instances
+	 *
+	 * @param InSequence			The source master sequence we're playing back at the root level
+	 * @param InTemplate			The template to use when playing back the specified sequence. A pointer to which will be held within this instance.
+	 */
+	FMovieSceneEvaluationTemplateInstance(UMovieSceneSequence& InSequence, const FMovieSceneEvaluationTemplate& InTemplate);
+
+	/**
+	 * Constructor for sub sequence instances. SubData must point to a valid sequence
+	 *
+	 * @param SubData				Information about the sub sequence we're instantiating
+	 * @param InTemplate			The template to use when playing back the sub sequence. A pointer to which will be held within this instance.
+	 * @param InSequenceID			The unique identifier within the master sequence
+	 */
+	FMovieSceneEvaluationTemplateInstance(const FMovieSceneSubSequenceData& SubData, const FMovieSceneEvaluationTemplate& InTemplate, FMovieSceneSequenceIDRef InSequenceID);
+
+	FMovieSceneEvaluationTemplateInstance(const FMovieSceneEvaluationTemplateInstance&) = delete;
+	FMovieSceneEvaluationTemplateInstance& operator=(const FMovieSceneEvaluationTemplateInstance&) = delete;
+
+#if PLATFORM_COMPILER_HAS_DEFAULTED_FUNCTIONS
+	FMovieSceneEvaluationTemplateInstance(FMovieSceneEvaluationTemplateInstance&&) = default;
+	FMovieSceneEvaluationTemplateInstance& operator=(FMovieSceneEvaluationTemplateInstance&&) = default;
+#else
+	FMovieSceneEvaluationTemplateInstance(FMovieSceneEvaluationTemplateInstance&& RHS)
+	{
+		*this = MoveTemp(RHS);
+	}
+	FMovieSceneEvaluationTemplateInstance& operator=(FMovieSceneEvaluationTemplateInstance&& RHS)
+	{
+		Sequence = MoveTemp(RHS.Sequence);
+		RootToSequenceTransform = MoveTemp(RHS.RootToSequenceTransform);
+		Template = MoveTemp(RHS.Template);
+		SequenceID = MoveTemp(RHS.SequenceID);
+		LegacySequenceInstance = MoveTemp(RHS.LegacySequenceInstance);
+		return *this;
+	}
+#endif
+
+public:
+
+	/** The sequence that the sub section references */
+	UMovieSceneSequence* Sequence;
+
+	/** Transform that transforms a given time from the sequences outer space, to its authored space. */
+	FMovieSceneSequenceTransform RootToSequenceTransform;
+
+	/** Pointer to the evaluation template we're evaluating */
+	const FMovieSceneEvaluationTemplate* Template;
+
+	/** ID of the sequence within the master sequence */
+	FMovieSceneSequenceID SequenceID;
+
+	/** Legacy sequence instance required for old track instances */
+	TSharedPtr<FMovieSceneSequenceInstance> LegacySequenceInstance;
+};
+
+/**
+ * Root evaluation template instance used to play back any sequence
+ */
+USTRUCT()
+struct FMovieSceneRootEvaluationTemplateInstance
+{
+public:
+	GENERATED_BODY()
+
+	MOVIESCENE_API FMovieSceneRootEvaluationTemplateInstance();
+	MOVIESCENE_API ~FMovieSceneRootEvaluationTemplateInstance();
+
+	FMovieSceneRootEvaluationTemplateInstance(const FMovieSceneRootEvaluationTemplateInstance&) = delete;
+	FMovieSceneRootEvaluationTemplateInstance& operator=(const FMovieSceneRootEvaluationTemplateInstance&) = delete;
+
+	FMovieSceneRootEvaluationTemplateInstance(FMovieSceneRootEvaluationTemplateInstance&&) = delete;
+	FMovieSceneRootEvaluationTemplateInstance& operator=(FMovieSceneRootEvaluationTemplateInstance&&) = delete;
+
+	/**
+	 * Check if this instance has been initialized correctly
+	 */
+	bool IsValid() const
+	{
+		return RootSequence && RootInstance.Template;
+	}
+
+	/**
+	 * Initialize this template instance with the specified sequence
+	 *
+	 * @param RootSequence				The sequence play back
+	 * @param Player					The player responsible for playback
+	 */
+	MOVIESCENE_API void Initialize(UMovieSceneSequence& RootSequence, IMovieScenePlayer& Player);
+
+	/**
+	 * Initialize this template instance with the specified sequence
+	 *
+	 * @param RootSequence				The sequence play back
+	 * @param Player					The player responsible for playback
+	 * @param TemplateStore				Template store responsible for supplying templates for a given sequence
+	 */
+	MOVIESCENE_API void Initialize(UMovieSceneSequence& RootSequence, IMovieScenePlayer& Player, TSharedRef<FMovieSceneSequenceTemplateStore> TemplateStore);
+
+	/**
+	 * Evaluate this sequence
+	 *
+	 * @param Context				Evaluation context containing the time (or range) to evaluate
+	 * @param Player				The player responsible for playback
+	 * @param OverrideRootID		The ID of the sequence from which to evaluate.
+	 */
+	MOVIESCENE_API void Evaluate(FMovieSceneContext Context, IMovieScenePlayer& Player, FMovieSceneSequenceID OverrideRootID = MovieSceneSequenceID::Root);
+
+	/**
+	 * Indicate that we're not going to evaluate this instance again, and that we should tear down any current state
+	 *
+	 * @param Player				The player responsible for playback
+	 */
+	MOVIESCENE_API void Finish(IMovieScenePlayer& Player);
+
+public:
+
+	class FOnUpdated : public TBaseMulticastDelegate<void>
+	{
+		friend FMovieSceneRootEvaluationTemplateInstance;
+	};
+
+	/**
+	 * Event that is triggered when this sequence has been updated
+	 */
+	FOnUpdated& OnUpdated() { return OnUpdatedEvent; }
+
+public:
+
+	/**
+	 * Retrieve a specific template instance given a sequence ID
+	 *
+	 * @param SequenceID			The sequence ID to find a template for
+	 * @return Pointer to a template instance, or nullptr if the ID was not found
+	 */
+	FORCEINLINE const FMovieSceneEvaluationTemplateInstance* GetInstance(FMovieSceneSequenceIDRef SequenceID) const
+	{
+		return SequenceID == MovieSceneSequenceID::Root ? &RootInstance : SubInstances.Find(SequenceID);
+	}
+
+	/**
+	 * Access the underlying map of sub template instances
+	 */
+	const TMap<FMovieSceneSequenceID, FMovieSceneEvaluationTemplateInstance>& GetSubInstances() const
+	{
+		return SubInstances;
+	}
+
+	/**
+	 * Attempt to locate the underlying sequence given a sequence ID
+	 *
+	 * @param SequenceID 			ID of the sequence to locate
+	 * @return The sequence, or nullptr if the ID was not found
+	 */
+	UMovieSceneSequence* GetSequence(FMovieSceneSequenceIDRef SequenceID) const
+	{
+		if (const FMovieSceneEvaluationTemplateInstance* Instance = GetInstance(SequenceID))
+		{
+			return Instance->Sequence;
+		}
+		return nullptr;
+	}
+
+
+	/**
+	 * Access the master sequence's hierarchy data
+	 */
+	const FMovieSceneSequenceHierarchy& GetHierarchy() const
+	{
+		return RootInstance.Template->Hierarchy;
+	}
+
+	/**
+	 * Check whether this instance is dirty (one of its referenced templates has changed)
+	 */
+	bool IsDirty() const
+	{
+		return bIsDirty;
+	}
+
+private:
+
+	/**
+	 * Gather entities that are to be evaluated this frame
+	 */
+	void GatherEntities(const FMovieSceneEvaluationGroup& Group, IMovieScenePlayer& Player, FMovieSceneSequenceID OverrideRootID);
+
+	/**
+	 * Process entities that are newly evaluated, and those that are no longer being evaluated
+	 */
+	void CallSetupTearDown(IMovieScenePlayer& Player);
+
+	/**
+	 * Evaluate a particular group of a segment
+	 */
+	void EvaluateGroup(const FMovieSceneEvaluationGroup& Group, const FMovieSceneContext& Context, IMovieScenePlayer& Player, FMovieSceneExecutionTokens& ExecutionTokens, FMovieSceneSequenceID OverrideRootID) const;
+
+	/**
+	 * Remap the specified sequence ID from relative to OverrideRootID, to the Root
+	 *
+	 * @param SequenceID			The sequence ID to find a template for
+	 * @param OverrideRootID		Overriden ID that we're treating as the root
+	 * @return Pointer to a template instance, or nullptr if the ID was not found
+	 */
+	FORCEINLINE FMovieSceneSequenceID GetSequenceIdForRoot(FMovieSceneSequenceID SequenceID, FMovieSceneSequenceID OverrideRootID) const
+	{
+		return OverrideRootID == MovieSceneSequenceID::Root ? SequenceID : SequenceID.AccumulateParentID(OverrideRootID);
+	}
+
+	/**
+	 * Retrieve a specific template instance given a sequence ID. Assumes it exists.
+	 *
+	 * @param SequenceID			The sequence ID to find a template for
+	 * @return The template instance
+	 */
+	FORCEINLINE const FMovieSceneEvaluationTemplateInstance& GetInstanceChecked(FMovieSceneSequenceIDRef SequenceID) const
+	{
+		return SequenceID == MovieSceneSequenceID::Root ? RootInstance : *SubInstances.Find(SequenceID);
+	}
+
+	/**
+	 * Attempt to retrieve a specific template instance given a sequence ID
+	 *
+	 * @param SequenceID			The sequence ID to find a template for
+	 * @return Pointer to a template instance, or nullptr if the ID was not found
+	 */
+	FORCEINLINE const FMovieSceneEvaluationTemplateInstance* FindInstance(FMovieSceneSequenceIDRef SequenceID) const
+	{
+		return SequenceID == MovieSceneSequenceID::Root ? &RootInstance : SubInstances.Find(SequenceID);
+	}
+
+private:
+
+	/** Called whenever any of the sequences contained within this instance is changed */
+	void OnSequenceChanged();
+
+	/** Reset this template instance */
+	void Reset();
+
+private:
+
+	UPROPERTY()
+	UMovieSceneSequence* RootSequence;
+
+	FMovieSceneEvaluationTemplateInstance RootInstance;
+
+	/** Cache of everything that was evaluated last frame */
+	FOrderedEvaluationCache EntitiesEvaluatedLastFrame;
+
+	/** Cache of everything that was just evaluated */
+	FOrderedEvaluationCache EntitiesEvaluatedThisFrame;
+
+	TArray<FMovieSceneSequenceID> ActiveSequencesThisFrame;
+	TArray<FMovieSceneSequenceID> ActiveSequencesLastFrame;
+
+	/** Map of all sub instances, arranged by sequence ID */
+	TMap<FMovieSceneSequenceID, FMovieSceneEvaluationTemplateInstance> SubInstances;
+
+	/** Template store responsible for supplying templates for a given sequence */
+	TSharedPtr<FMovieSceneSequenceTemplateStore> TemplateStore;
+
+	/** True when any of our templates are out of date, and need reinitializing */
+	bool bIsDirty;
+
+	/** Event that is triggered on update */
+	FOnUpdated OnUpdatedEvent;
+};
+template<> struct TStructOpsTypeTraits<FMovieSceneRootEvaluationTemplateInstance> : public TStructOpsTypeTraitsBase { enum { WithCopy = false }; };
\ No newline at end of file
diff --git a/Engine/Source/Runtime/MovieScene/Public/Evaluation/MovieSceneEvaluationTrack.h b/Engine/Source/Runtime/MovieScene/Public/Evaluation/MovieSceneEvaluationTrack.h
new file mode 100644
--- /dev/null
+++ b/Engine/Source/Runtime/MovieScene/Public/Evaluation/MovieSceneEvaluationTrack.h
@@ -0,0 +1,392 @@
+// Copyright 1998-2016 Epic Games, Inc. All Rights Reserved.
+
+#pragma once
+
+#include "MovieSceneSegment.h"
+#include "MovieSceneEvalTemplate.h"
+#include "MovieSceneTrackImplementation.h"
+#include "MovieSceneEvaluationTrack.generated.h"
+
+/** Enumeration to determine how a track should be evaluated */
+UENUM()
+enum class EEvaluationMethod : uint8
+{
+	/** Evaluation only ever occurs at a single time. Delta is irrelevant. Example: Basic curve eval, animation */
+	Static,
+
+	/** Evaluation from one frame to the next must consider the entire swept delta range on the track. Example: Events */
+	Swept,
+};
+
+/**
+ * Evaluation track that is stored within an evaluation template for a sequence.
+ * Contains user-defined evaluation templates, and an optional track implementation
+ */
+USTRUCT()
+struct FMovieSceneEvaluationTrack
+{
+	GENERATED_BODY()
+
+	/** Default construction (only for serialization) */
+	FMovieSceneEvaluationTrack();
+
+	/**
+	 * User construction, for initialization during compilation
+	 */
+	MOVIESCENE_API FMovieSceneEvaluationTrack(const FGuid& InObjectBindingID);
+
+	/** Copy construction/assignment */
+	FMovieSceneEvaluationTrack(const FMovieSceneEvaluationTrack&) = default;
+	FMovieSceneEvaluationTrack& operator=(const FMovieSceneEvaluationTrack&) = default;
+
+	/** Move construction/assignment */
+#if PLATFORM_COMPILER_HAS_DEFAULTED_FUNCTIONS
+	FMovieSceneEvaluationTrack(FMovieSceneEvaluationTrack&&) = default;
+	FMovieSceneEvaluationTrack& operator=(FMovieSceneEvaluationTrack&&) = default;
+#else
+	FMovieSceneEvaluationTrack(FMovieSceneEvaluationTrack&& RHS)
+	{
+		*this = MoveTemp(RHS);
+	}
+	FMovieSceneEvaluationTrack& operator=(FMovieSceneEvaluationTrack&& RHS)
+	{
+		ObjectBindingID = RHS.ObjectBindingID;
+		EvaluationPriority = RHS.EvaluationPriority;
+		EvaluationMethod = RHS.EvaluationMethod;
+		Segments = MoveTemp(RHS.Segments);
+		ChildTemplates = MoveTemp(RHS.ChildTemplates);
+		TrackTemplate = MoveTemp(RHS.TrackTemplate);
+		EvaluationGroup = MoveTemp(RHS.EvaluationGroup);
+		bEvaluateInPreroll = RHS.bEvaluateInPreroll;
+		bEvaluateInPostroll = RHS.bEvaluateInPostroll;
+		return *this;
+	}
+#endif
+
+public:
+
+	/**
+	 * Get the object binding ID that this track belongs to
+	 */
+	FORCEINLINE const FGuid& GetObjectBindingID() const
+	{
+		return ObjectBindingID;
+	}
+
+	/**
+	 * Const iteration of this track's segments
+	 */
+	FORCEINLINE TArrayView<const FMovieSceneSegment> GetSegments() const
+	{
+		return TArrayView<const FMovieSceneSegment>(Segments.GetData(), Segments.Num());
+	}
+
+	/**
+	 * Get the segment from the given segment index
+	 */
+	FORCEINLINE const FMovieSceneSegment& GetSegment(int32 SegmentIndex) const
+	{
+		return Segments[SegmentIndex];
+	}
+
+	/**
+	 * Get the template from the given template index
+	 */
+	FORCEINLINE const FMovieSceneEvalTemplate& GetChildTemplate(int32 TemplateIndex) const
+	{
+		return *ChildTemplates[TemplateIndex];
+	}
+
+	/**
+	 * Check whether we have a valid child template for the specified index
+	 */
+	FORCEINLINE bool HasChildTemplate(int32 TemplateIndex) const
+	{
+		return ChildTemplates.IsValidIndex(TemplateIndex) && ChildTemplates[TemplateIndex].IsValid();
+	}
+	
+public:
+
+	/**
+	 * Get this track's evaluation group name. Only used during compilation.
+	 *
+	 * @return The evaluation group
+	 */
+	FName GetEvaluationGroup() const
+	{
+		return EvaluationGroup;
+	}
+
+	/**
+	 * Set this track's flush group name.
+	 *
+	 * @note 	When not 'None', setting an evaluation group indicates that all tracks with similar
+	 *			groups and priorities should be grouped together at runtime. Named groups can
+	 *			be optionally flushed immediately at runtime by calling IMovieSceneTemplateGenerator::FlushGroupImmediately
+	 *			with the appropriate group.
+	 *
+	 * @param InEvaluationGroup 		The evaluation group to assign this track to
+	 */
+	void SetEvaluationGroup(FName InEvaluationGroup)
+	{
+		EvaluationGroup = InEvaluationGroup;
+	}
+
+public:
+
+	/**
+	 * Get the evaluation bias to apply to this track. Higher priority tracks will be evaluated first.
+	 *
+	 * @return The evaluation priority
+	 */
+	uint16 GetEvaluationPriority() const
+	{
+		return EvaluationPriority;
+	}
+
+	/**
+	 * Get the evaluation bias to apply to this track. Higher priority tracks will be evaluated first.
+	 *
+	 * @param InEvaluationPriority			The new priority
+	 */
+	void SetEvaluationPriority(uint16 InEvaluationPriority = 1000)
+	{
+		EvaluationPriority = InEvaluationPriority;
+	}
+
+	/**
+	 * Get the method we should use to evaluate this track
+	 *
+	 * @return The method to use when evaluating this track
+	 */
+	EEvaluationMethod GetEvaluationMethod() const
+	{
+		return EvaluationMethod;
+	}
+
+	/**
+	 * Set the method we should use to evaluate this track
+	 *
+	 * @param InMethod				The method to use when evaluating this track
+	 */
+	void SetEvaluationMethod(EEvaluationMethod InMethod)
+	{
+		EvaluationMethod = InMethod;
+	}
+
+	/**
+	 * Define how this track evaluates in pre and postroll
+	 *
+	 * @param bInEvaluateInPreroll	Whether this track should evaluate in preroll
+	 * @param bInEvaluateInPostroll	Whether this track should evaluate in postroll
+	 */
+	void SetPreAndPostrollConditions(bool bInEvaluateInPreroll, bool bInEvaluateInPostroll)
+	{
+		bEvaluateInPreroll = bInEvaluateInPreroll;
+		bEvaluateInPostroll = bInEvaluateInPostroll;
+	}
+
+	/**
+	 * @return Whether this track should evaluate in preroll
+	 */
+	bool ShouldEvaluateInPreroll() const
+	{
+		return bEvaluateInPreroll;
+	}
+
+	/**
+	 * @return Whether this track should evaluate in postroll
+	 */
+	bool ShouldEvaluateInPostroll() const
+	{
+		return bEvaluateInPostroll;
+	}
+
+public:
+
+	/**
+	 * Called to initialize the specified segment index
+	 * 
+	 * @param SegmentIndex			The segment we are evaluating
+	 * @param Operand				Operand that relates to the thing we will animate
+	 * @param Context				Current sequence context
+	 * @param PersistentData		Persistent data store
+	 * @param Player				The player that is responsible for playing back this template
+	 */
+	void Initialize(int32 SegmentIndex, const FMovieSceneEvaluationOperand& Operand, const FMovieSceneContext& Context, FPersistentEvaluationData& PersistentData, IMovieScenePlayer& Player) const;
+
+	/**
+	 * Called to evaluate the specified segment index
+	 * 
+	 * @param SegmentIndex			The segment we are evaluating
+	 * @param Operand				Operand that relates to the thing we will animate
+	 * @param Context				Current sequence context
+	 * @param PersistentData		Persistent data store
+	 * @param ExecutionTokens		Token stack on which to add tokens that will be executed later
+	 */
+	void Evaluate(int32 SegmentIndex, const FMovieSceneEvaluationOperand& Operand, const FMovieSceneContext& Context, const FPersistentEvaluationData& PersistentData, FMovieSceneExecutionTokens& ExecutionTokens) const;
+
+	/**
+	 * Default implementation of initialization of child templates for the specified segment
+	 * 
+	 * @param SegmentIndex			The segment we are evaluating
+	 * @param Operand				Operand that relates to the thing we will animate
+	 * @param Context				Current sequence context
+	 * @param PersistentData		Persistent data store
+	 * @param Player				The player that is responsible for playing back this template
+	 */
+	MOVIESCENE_API void DefaultInitialize(int32 SegmentIndex, const FMovieSceneEvaluationOperand& Operand, FMovieSceneContext Context, FPersistentEvaluationData& PersistentData, IMovieScenePlayer& Player) const;
+
+	/**
+	 * Default implementation of evaluation of child templates for the specified segment
+	 * 
+	 * @param SegmentIndex			The segment we are evaluating
+	 * @param Operand				Operand that relates to the thing we will animate
+	 * @param Context				Current sequence context
+	 * @param PersistentData		Persistent data store
+	 * @param ExecutionTokens		Token stack on which to add tokens that will be executed later
+	 */
+	MOVIESCENE_API void DefaultEvaluate(int32 SegmentIndex, const FMovieSceneEvaluationOperand& Operand, const FMovieSceneContext& Context, const FPersistentEvaluationData& PersistentData, FMovieSceneExecutionTokens& ExecutionTokens) const;
+
+private:
+
+	/**
+	 * Implementation function for static evaluation
+	 */
+	void EvaluateStatic(int32 SegmentIndex, const FMovieSceneEvaluationOperand& Operand, FMovieSceneContext Context, const FPersistentEvaluationData& PersistentData, FMovieSceneExecutionTokens& ExecutionTokens) const;
+
+	/**
+	 * Implementation function for swept evaluation
+	 */
+	void EvaluateSwept(int32 SegmentIndex, const FMovieSceneEvaluationOperand& Operand, const FMovieSceneContext& Context, const FPersistentEvaluationData& PersistentData, FMovieSceneExecutionTokens& ExecutionTokens) const;
+
+public:
+
+	/**
+	 * Assign a single eval template to this track, spanning the entire sequence
+	 * 
+	 * @param InTemplate		The template to insert
+	 */
+	MOVIESCENE_API void DefineAsSingleTemplate(FMovieSceneEvalTemplatePtr&& InTemplate);
+
+	/**
+	 * Add an evaluation template to this track with the given track index
+	 * 
+	 * @param InTemplate		The template to insert
+	 * @return The index of the new template
+	 */
+	MOVIESCENE_API int32 AddChildTemplate(FMovieSceneEvalTemplatePtr&& InTemplate);
+
+	/**
+	 * Assign the specified segments to this track.
+	 *
+	 * @param InSegments		The segments that define this track
+	 */
+	MOVIESCENE_API void SetSegments(TArray<FMovieSceneSegment>&& InSegments);
+	
+	/**
+	 * Assign a track implementation template to this track
+	 * @note Track implementations are evaluated once per frame before any segments.
+	 *
+	 * @param InImplementation	The track implementation to use
+	 */
+	template<typename T>
+	typename TEnableIf<TPointerIsConvertibleFromTo<T, FMovieSceneTrackImplementation>::Value>::Type 
+		SetTrackImplementation(T&& InImpl)
+	{
+		TrackTemplate = FMovieSceneTrackImplementationPtr(MoveTemp(InImpl));
+		TrackTemplate->SetupOverrides();
+	}
+
+	/**
+	 * Setup overrides for any contained templates
+	 */
+	MOVIESCENE_API void SetupOverrides();
+
+	/**
+	 * Post serialize function
+	 */
+	void PostSerialize(const FArchive& Ar)
+	{
+		SetupOverrides();
+	}
+
+private:
+
+	/**
+	 * Validate the segment array and remove any invalid ptrs
+	 */
+	void ValidateSegments();
+
+public:
+
+	/**
+	 * Called before this track is evaluated for the first time, or since OnEndEvaluation has been called
+	 *
+	 * @param PersistentData		Persistent data proxy that may contain data pertaining to this entity
+	 * @param Player				The player that is responsible for playing back this template
+	 */
+	FORCEINLINE void OnBeginEvaluation(FPersistentEvaluationData& PersistentData, IMovieScenePlayer& Player) const
+	{
+		if (!TrackTemplate.IsValid())
+		{
+			return;
+		}
+		TrackTemplate->OnBeginEvaluation(PersistentData, Player);
+	}
+
+	/**
+	 * Called after this track is no longer being evaluated
+	 *
+	 * @param PersistentData		Persistent data proxy that may contain data pertaining to this entity
+	 * @param Player				The player that is responsible for playing back this template
+	 */
+	void OnEndEvaluation(FPersistentEvaluationData& PersistentData, IMovieScenePlayer& Player) const
+	{
+		if (!TrackTemplate.IsValid())
+		{
+			return;
+		}
+		TrackTemplate->OnEndEvaluation(PersistentData, Player);
+	}
+
+private:
+
+	/** ID of the possessable or spawnable within the UMovieScene this track belongs to, if any. Zero guid where this relates to a master track. */
+	UPROPERTY()
+	FGuid ObjectBindingID;
+
+	/** Evaluation priority. Highest is evaluated first */
+	UPROPERTY()
+	uint16 EvaluationPriority;
+
+	/** Evaluation method - static or swept */
+	UPROPERTY()
+	EEvaluationMethod EvaluationMethod;
+
+	/** Array of segmented ranges contained within the track. */
+	UPROPERTY()
+	TArray<FMovieSceneSegment> Segments;
+	
+	/** Domain-specific evaluation templates (normally 1 per section) */
+	UPROPERTY()
+	TArray<FMovieSceneEvalTemplatePtr> ChildTemplates;
+
+	/** Domain-specific track implementation override. */
+	UPROPERTY()
+	FMovieSceneTrackImplementationPtr TrackTemplate;
+
+	/** Flush group that determines whether this track belongs to a group of tracks */
+	UPROPERTY()
+	FName EvaluationGroup;
+
+	/** Whether this track is evaluated in preroll */
+	UPROPERTY()
+	uint32 bEvaluateInPreroll : 1;
+
+	/** Whether this track is evaluated in postroll */
+	UPROPERTY()
+	uint32 bEvaluateInPostroll : 1;
+};
+
+template<> struct TStructOpsTypeTraits<FMovieSceneEvaluationTrack> : public TStructOpsTypeTraitsBase { enum { WithPostSerialize = true, WithCopy = false }; };
\ No newline at end of file
diff --git a/Engine/Source/Runtime/MovieScene/Public/Evaluation/MovieSceneExecutionTokens.h b/Engine/Source/Runtime/MovieScene/Public/Evaluation/MovieSceneExecutionTokens.h
new file mode 100644
--- /dev/null
+++ b/Engine/Source/Runtime/MovieScene/Public/Evaluation/MovieSceneExecutionTokens.h
@@ -0,0 +1,139 @@
+// Copyright 1998-2016 Epic Games, Inc. All Rights Reserved.
+
+#pragma once
+
+#include "MovieSceneEvaluationKey.h"
+#include "MovieSceneEvalTemplate.h"
+
+/**
+ * Ordered execution token stack that accumulates tokens that will apply animated state to the sequence environment at a later time
+ */
+struct FMovieSceneExecutionTokens
+{
+	FMovieSceneExecutionTokens()
+		: Context(FMovieSceneEvaluationRange(0.f), EMovieScenePlayerStatus::Stopped)
+	{
+	}
+
+	FMovieSceneExecutionTokens(const FMovieSceneExecutionTokens&) = delete;
+	FMovieSceneExecutionTokens& operator=(const FMovieSceneExecutionTokens&) = delete;
+	
+	/**
+	 * Add a new IMovieSceneExecutionToken derived token to the stack
+	 */
+	template<typename T>
+	FORCEINLINE typename TEnableIf<TPointerIsConvertibleFromTo<typename TRemoveReference<T>::Type, const IMovieSceneExecutionToken>::Value>::Type Add(T&& InToken)
+	{
+		check(TrackKey.IsValid() && Operand.IsValid());
+		Tokens.Add(FEntry(Operand, TrackKey, SectionIdentifier, CompletionMode, Context, Forward<T>(InToken)));
+	}
+
+	/**
+	 * Internal: Set the operand we're currently operating on
+	 */
+	FORCEINLINE void SetOperand(const FMovieSceneEvaluationOperand& InOperand)
+	{
+		Operand = InOperand;
+	}
+
+	/**
+	 * Internal: Set the track we're currently evaluating
+	 */
+	FORCEINLINE void SetTrack(const FMovieSceneEvaluationKey& InTrackKey)
+	{
+		TrackKey = InTrackKey;
+	}
+
+	/**
+	 * Internal: Set the section we're currently evaluating
+	 */
+	FORCEINLINE void SetSectionIdentifier(uint32 InSectionIdentifier)
+	{
+		SectionIdentifier = InSectionIdentifier;
+	}
+
+	/**
+	 * Internal: Set the current context
+	 */
+	FORCEINLINE void SetContext(const FMovieSceneContext& InContext)
+	{
+		Context = InContext;
+	}
+
+	/**
+	 * Internal: Set what should happen when the section is no longer evaluated
+	 */
+	FORCEINLINE void SetCompletionMode(EMovieSceneCompletionMode InCompletionMode)
+	{
+		CompletionMode = InCompletionMode;
+	}
+
+	struct FEntry
+	{
+		FEntry(const FMovieSceneEvaluationOperand& InOperand, const FMovieSceneEvaluationKey& InTrackKey, uint32 InSectionIdentifier, EMovieSceneCompletionMode InCompletionMode, const FMovieSceneContext& InContext, TInlineValue<IMovieSceneExecutionToken, 32>&& InToken)
+			: Operand(InOperand)
+			, TrackKey(InTrackKey)
+			, SectionIdentifier(InSectionIdentifier)
+			, CompletionMode(InCompletionMode)
+			, Context(InContext)
+			, Token(MoveTemp(InToken))
+		{
+		}
+
+#if PLATFORM_COMPILER_HAS_DEFAULTED_FUNCTIONS
+		FEntry(FEntry&&) = default;
+		FEntry& operator=(FEntry&&) = default;
+#else
+		FEntry(FEntry&& RHS)
+			: Operand(MoveTemp(RHS.Operand))
+			, TrackKey(MoveTemp(RHS.TrackKey))
+			, SectionIdentifier(MoveTemp(RHS.SectionIdentifier))
+			, CompletionMode(MoveTemp(RHS.CompletionMode))
+			, Context(MoveTemp(RHS.Context))
+			, Token(MoveTemp(RHS.Token))
+		{
+		}
+		FEntry& operator=(FEntry&& RHS)
+		{
+			Operand = MoveTemp(RHS.Operand);
+			TrackKey = MoveTemp(RHS.TrackKey);
+			SectionIdentifier = MoveTemp(RHS.SectionIdentifier);
+			CompletionMode = MoveTemp(RHS.CompletionMode);
+			Context = MoveTemp(RHS.Context);
+			Token = MoveTemp(RHS.Token);
+			return *this;
+		}
+#endif
+
+		/** The operand we were operating on when this token was added */
+		FMovieSceneEvaluationOperand Operand;
+		/** The track we were evaluating when this token was added */
+		FMovieSceneEvaluationKey TrackKey;
+		/** The section identifier we were evaluating when this token was added */
+		uint32 SectionIdentifier;
+		/** What should happen when the entity is no longer evaluated */
+		EMovieSceneCompletionMode CompletionMode;
+		/** The context from when this token was added */
+		FMovieSceneContext Context;
+		/** The user-provided token */
+		TInlineValue<IMovieSceneExecutionToken, 32> Token;
+	};
+
+	MOVIESCENE_API void Apply(IMovieScenePlayer& Player);
+
+	/** Ordered array of tokens */
+	TArray<FEntry> Tokens;
+
+private:
+
+	/** The operand we're currently operating on */
+	FMovieSceneEvaluationOperand Operand;
+	/** The track we're currently evaluating */
+	FMovieSceneEvaluationKey TrackKey;
+	/** The section we're currently evaluating */
+	uint32 SectionIdentifier;
+	/* What should happen when the section is no longer evaluated */
+	EMovieSceneCompletionMode CompletionMode;
+	/** The current context */
+	FMovieSceneContext Context;
+};
\ No newline at end of file
diff --git a/Engine/Source/Runtime/MovieScene/Public/Evaluation/MovieSceneLegacyTrackInstanceTemplate.h b/Engine/Source/Runtime/MovieScene/Public/Evaluation/MovieSceneLegacyTrackInstanceTemplate.h
new file mode 100644
--- /dev/null
+++ b/Engine/Source/Runtime/MovieScene/Public/Evaluation/MovieSceneLegacyTrackInstanceTemplate.h
@@ -0,0 +1,30 @@
+// Copyright 1998-2016 Epic Games, Inc. All Rights Reserved.
+
+#pragma once
+
+#include "MovieSceneFwd.h"
+#include "IMovieSceneTrackInstance.h"
+
+#include "MovieSceneLegacyTrackInstanceTemplate.generated.h"
+
+USTRUCT()
+struct MOVIESCENE_API FMovieSceneLegacyTrackInstanceTemplate : public FMovieSceneEvalTemplate
+{
+	GENERATED_BODY()
+
+	FMovieSceneLegacyTrackInstanceTemplate() {}
+	FMovieSceneLegacyTrackInstanceTemplate(const UMovieSceneTrack* InTrack);
+
+	
+	virtual void Evaluate(const FMovieSceneEvaluationOperand& Operand, const FMovieSceneContext& Context, const FPersistentEvaluationData& PersistentData, FMovieSceneExecutionTokens& ExecutionTokens) const override;
+
+	virtual UScriptStruct& GetScriptStructImpl() const override { return *StaticStruct(); }
+	virtual void Setup(FPersistentEvaluationData& PersistentData, IMovieScenePlayer& Player) const;
+	virtual void SetupOverrides() override
+	{
+		EnableOverrides(RequiresSetupFlag);
+	}
+
+	UPROPERTY()
+	const UMovieSceneTrack* Track;
+};
\ No newline at end of file
