Copying //UE4/Dev-Sequencer to //UE4/Dev-Main (Source: //UE4/Dev-Sequencer @ 3178529)

diff --git a/Engine/Source/Runtime/MovieScene/Public/Evaluation/MovieScenePlayback.h b/Engine/Source/Runtime/MovieScene/Public/Evaluation/MovieScenePlayback.h
new file mode 100644
--- /dev/null
+++ b/Engine/Source/Runtime/MovieScene/Public/Evaluation/MovieScenePlayback.h
@@ -0,0 +1,238 @@
+// Copyright 1998-2016 Epic Games, Inc. All Rights Reserved.
+
+#pragma once
+
+#include "MovieSceneFwd.h"
+#include "MovieSceneSequenceTransform.h"
+#include "IMovieScenePlayer.h"
+
+/** Enumeration specifying whether we're playing forwards or backwards */
+enum class EPlayDirection
+{
+	Forwards, Backwards
+};
+
+/** MovieScene evaluation context. Should remain bitwise copyable, and contain no external state since this has the potential to be used on a thread */
+struct MOVIESCENE_API FMovieSceneEvaluationRange
+{
+	/**
+	 * Construct this range from a single fixed time
+	 */
+	FMovieSceneEvaluationRange(float InTime);
+
+	/**
+	 * Construct this range from a raw range and a direction
+	 */
+	FMovieSceneEvaluationRange(TRange<float> InRange, EPlayDirection InDirection);
+
+	/**
+	 * Construct this range from 2 times, and whether the range should include the previous time or not
+	 */
+	FMovieSceneEvaluationRange(float InCurrentTime, float InPreviousTime, bool bInclusivePreviousTime = false);
+
+	/**
+	 * Get the range that we should be evaluating
+	 */
+	FORCEINLINE TRange<float> GetRange() const
+	{
+		return EvaluationRange;
+	}
+
+	/**
+	 * Get the direction to evaluate our range
+	 */
+	FORCEINLINE EPlayDirection GetDirection() const 
+	{
+		return Direction;
+	}
+
+	/**
+	 * Get the current time of evaluation.
+	 */
+	FORCEINLINE float GetTime() const
+	{
+		if (TimeOverride != TNumericLimits<float>::Lowest())
+		{
+			return TimeOverride;
+		}
+
+		return Direction == EPlayDirection::Forwards ? EvaluationRange.GetUpperBoundValue() : EvaluationRange.GetLowerBoundValue();
+	}
+
+	/**
+	 * Get the absolute amount of time that has passed since the last update (will always be >= 0)
+	 */
+	FORCEINLINE float GetDelta() const
+	{
+		return EvaluationRange.Size<float>();
+	}
+
+	/**
+	 * Get the previous time of evaluation. Should not generally be used. Prefer GetRange instead.
+	 */
+	FORCEINLINE float GetPreviousTime() const
+	{
+		return Direction == EPlayDirection::Forwards ? EvaluationRange.GetLowerBoundValue() : EvaluationRange.GetUpperBoundValue();
+	}
+	
+	/**
+	 * Override the time that we're actually evaluating at
+	 */
+	FORCEINLINE void OverrideTime(float InTimeOverride)
+	{
+		TimeOverride = InTimeOverride;
+	}
+
+protected:
+
+	/** The range to evaluate */
+	TRange<float> EvaluationRange;
+
+	/** Whether to evaluate the range forwards, or backwards */
+	EPlayDirection Direction;
+
+	/** Overridden current time (doesn't manipulate the actual evaluated range) */
+	float TimeOverride;
+};
+
+/** MovieScene evaluation context. Should remain bitwise copyable, and contain no external state since this has the potential to be used on a thread */
+struct FMovieSceneContext : FMovieSceneEvaluationRange
+{
+	/**
+	 * Construction from an evaluation range, and a current status
+	 */
+	FMovieSceneContext(FMovieSceneEvaluationRange InRange, EMovieScenePlayerStatus::Type InStatus)
+		: FMovieSceneEvaluationRange(InRange)
+		, Status(InStatus)
+		, bHasJumped(false)
+		, bSilent(false)
+	{}
+
+	/**
+	 * Get the playback status
+	 */
+	FORCEINLINE EMovieScenePlayerStatus::Type GetStatus() const
+	{
+		return Status;
+	}
+
+	/**
+	 * Check whether we've just jumped to a different time
+	 */
+	FORCEINLINE bool HasJumped() const
+	{
+		return bHasJumped;
+	}
+
+	/**
+	 * Check whether we're evaluating in silent mode (no audio or mutating eval)
+	 */
+	FORCEINLINE bool IsSilent() const
+	{
+		return bSilent;
+	}
+
+	/**
+	 * Get the current root to sequence transform for the current sub sequence
+	 */
+	FORCEINLINE const FMovieSceneSequenceTransform& GetRootToSequenceTransform() const
+	{
+		return RootToSequenceTransform;
+	}
+
+public:
+
+	/**
+	 * Indicate that we've just jumped to a different time
+	 */
+	FMovieSceneContext& SetHasJumped(bool bInHasJumped)
+	{
+		bHasJumped = bInHasJumped;
+		return *this;
+	}
+
+	/**
+	 * Set the context to silent mode
+	 */
+	FMovieSceneContext& SetIsSilent(bool bInIsSilent)
+	{
+		bSilent = bInIsSilent;
+		return *this;
+	}
+
+	/**
+	 * Clamp the current evaluation range to the specified range (in the current transform space)
+	 */
+	FMovieSceneContext Clamp(TRange<float> NewRange) const
+	{
+		FMovieSceneContext NewContext = *this;
+		NewContext.EvaluationRange = NewRange;
+		return NewContext;
+	}
+
+	/**
+	 * Transform this context to a different sub sequence space
+	 */
+	FMovieSceneContext Transform(const FMovieSceneSequenceTransform& InTransform) const
+	{
+		FMovieSceneContext NewContext = *this;
+		NewContext.EvaluationRange = EvaluationRange * InTransform;
+		NewContext.RootToSequenceTransform = NewContext.RootToSequenceTransform * InTransform;
+		return NewContext;
+	}
+
+protected:
+
+	/** The transform from the root sequence to the current sequence space */
+	FMovieSceneSequenceTransform RootToSequenceTransform;
+
+	/** The current playback status */
+	EMovieScenePlayerStatus::Type Status;
+
+	/** Whether this evaluation frame is happening as part of a large jump */
+	bool bHasJumped : 1;
+
+	/** Whether this evaluation should happen silently */
+	bool bSilent : 1;
+};
+
+/** Helper class designed to abstract the complexity of calculating evaluation ranges for previous times and fixed time intervals */
+struct MOVIESCENE_API FMovieScenePlaybackPosition
+{
+	/**
+	 * Reset this position to the specified time.
+	 * @note Future calls to 'PlayTo' will include this time in its resulting evaluation range
+	 */
+	void Reset(float StartPos);
+
+	/**
+	 * Jump to the specified time.
+	 * @note Will reset previous play position. Any subsequent call to 'PlayTo' will include NewPosition.
+	 * @return A range encompassing only the specified time.
+	 */
+	FMovieSceneEvaluationRange JumpTo(float NewPosition, TOptional<float> FixedInterval);
+
+	/**
+	 * Play from the previously evaluated play time, to the specified time
+	 * @return An evaluation range from the previously evaluated time, to the specified time
+	 */
+	FMovieSceneEvaluationRange PlayTo(float NewPosition, TOptional<float> FixedInterval);
+
+	/**
+	 * Get the last position that was set
+	 */
+	TOptional<float> GetPreviousPosition() const { return PreviousPosition; }
+
+	/**
+	 * Get the last actual time that was evaluated during playback
+	 */
+	TOptional<float> GetLastPlayEvalPostition() const { return PreviousPlayEvalPosition; }
+
+
+private:
+	/** The previous *actual* time position set. Never rounded to a fixed interval. */
+	TOptional<float> PreviousPosition;
+
+	/** The previous evaluated position when playing, potentially rounded to a frame interval. */
+	TOptional<float> PreviousPlayEvalPosition;
+};
\ No newline at end of file
diff --git a/Engine/Source/Runtime/MovieScene/Public/Evaluation/MovieScenePreAnimatedState.h b/Engine/Source/Runtime/MovieScene/Public/Evaluation/MovieScenePreAnimatedState.h
new file mode 100644
--- /dev/null
+++ b/Engine/Source/Runtime/MovieScene/Public/Evaluation/MovieScenePreAnimatedState.h
@@ -0,0 +1,401 @@
+// Copyright 1998-2016 Epic Games, Inc. All Rights Reserved.
+
+#pragma once
+
+#include "MovieSceneExecutionToken.h"
+#include "MovieSceneSection.h"
+#include "ObjectKey.h"
+
+enum class ECapturePreAnimatedState : uint8
+{
+	None,
+	Global,
+	Entity,
+};
+
+struct IMovieScenePreAnimatedState
+{
+	virtual void EntityHasAnimatedObject(FMovieSceneEvaluationKey EntityKey, FObjectKey Object) = 0;
+	virtual void EntityHasAnimatedMaster(FMovieSceneEvaluationKey EntityKey) = 0;
+};
+
+struct FMovieSceneEntityAndAnimTypeID
+{
+	FMovieSceneEvaluationKey EntityKey;
+	FMovieSceneAnimTypeID AnimTypeID;
+
+	friend bool operator==(const FMovieSceneEntityAndAnimTypeID& A, const FMovieSceneEntityAndAnimTypeID& B)
+	{
+		return A.AnimTypeID == B.AnimTypeID && A.EntityKey == B.EntityKey;
+	}
+};
+
+/** Saved state for animation bound to particular animated objects */
+struct FMovieSceneSavedObjectTokens
+{
+	FMovieSceneSavedObjectTokens(UObject& InObject)
+		: Object(&InObject)
+	{}
+	
+#if PLATFORM_COMPILER_HAS_DEFAULTED_FUNCTIONS
+	FMovieSceneSavedObjectTokens(FMovieSceneSavedObjectTokens&&) = default;
+	FMovieSceneSavedObjectTokens& operator=(FMovieSceneSavedObjectTokens&&) = default;
+#else
+	FMovieSceneSavedObjectTokens(FMovieSceneSavedObjectTokens&& RHS)
+	{
+		*this = MoveTemp(RHS);
+	}
+	FMovieSceneSavedObjectTokens& operator=(FMovieSceneSavedObjectTokens&& RHS)
+	{
+		AnimatedGlobalTypeIDs = MoveTemp(RHS.AnimatedGlobalTypeIDs);
+		AnimatedEntityTypeIDs = MoveTemp(RHS.AnimatedEntityTypeIDs);
+		GlobalStateTokens = MoveTemp(RHS.GlobalStateTokens);
+		EntityTokens = MoveTemp(RHS.EntityTokens);
+		Object = MoveTemp(RHS.Object);
+		return *this;
+	}
+#endif
+
+	MOVIESCENE_API void Add(ECapturePreAnimatedState CaptureMode, FMovieSceneAnimTypeID InAnimTypeID, FMovieSceneEvaluationKey AssociatedKey, const IMovieScenePreAnimatedTokenProducer& Producer, UObject& InObject, IMovieScenePreAnimatedState& Parent);
+
+	MOVIESCENE_API void Restore(IMovieScenePlayer& Player);
+
+	MOVIESCENE_API void Restore(IMovieScenePlayer& Player, TFunctionRef<bool(FMovieSceneAnimTypeID)> InFilter);
+
+	void RestoreEntity(IMovieScenePlayer& Player, FMovieSceneEvaluationKey EntityKey)
+	{
+		TArray<FMovieSceneAnimTypeID, TInlineAllocator<8>> AnimTypesToRestore;
+
+		for (int32 LUTIndex = AnimatedEntityTypeIDs.Num() - 1; LUTIndex >= 0; --LUTIndex)
+		{
+			FMovieSceneEntityAndAnimTypeID EntityAndAnimType = AnimatedEntityTypeIDs[LUTIndex];
+			if (EntityAndAnimType.EntityKey == EntityKey)
+			{
+				AnimTypesToRestore.Add(EntityAndAnimType.AnimTypeID);
+				AnimatedEntityTypeIDs.RemoveAtSwap(LUTIndex);
+			}
+		}
+
+		UObject* ObjectPtr = Object.Get();
+		for (int32 TokenIndex = EntityTokens.Num() - 1; TokenIndex >= 0; --TokenIndex)
+		{
+			FRefCountedPreAnimatedObjectToken& Token = EntityTokens[TokenIndex];
+			if (AnimTypesToRestore.Contains(Token.AnimTypeID) && --Token.RefCount == 0)
+			{
+				if (ObjectPtr)
+				{
+					Token.Token->RestoreState(*ObjectPtr, Player);
+				}
+				EntityTokens.RemoveAtSwap(TokenIndex);
+			}
+		}
+	}
+
+	void Reset()
+	{
+		AnimatedGlobalTypeIDs.Reset();
+		AnimatedEntityTypeIDs.Reset();
+		GlobalStateTokens.Reset();
+		EntityTokens.Reset();
+	}
+
+private:
+
+	/** Array of type IDs for cache efficient searching. Must have 1-1 index mapping into GlobalStateTokens. */
+	TArray<FMovieSceneAnimTypeID, TInlineAllocator<8>> AnimatedGlobalTypeIDs;
+	TArray<FMovieSceneEntityAndAnimTypeID, TInlineAllocator<8>> AnimatedEntityTypeIDs;
+
+	/** Tokens mapped to specific anim type IDs. Can only ever be one at a time. */
+	TArray<IMovieScenePreAnimatedTokenPtr> GlobalStateTokens;
+
+	/** Tokens mapped to specific anim type IDs. Can only ever be one at a time. */
+	struct FRefCountedPreAnimatedObjectToken
+	{
+		FRefCountedPreAnimatedObjectToken(IMovieScenePreAnimatedTokenPtr&& InToken, FMovieSceneAnimTypeID InAnimTypeID) : AnimTypeID(InAnimTypeID), RefCount(1), Token(MoveTemp(InToken)) {}
+
+#if PLATFORM_COMPILER_HAS_DEFAULTED_FUNCTIONS
+		FRefCountedPreAnimatedObjectToken(FRefCountedPreAnimatedObjectToken&&) = default;
+		FRefCountedPreAnimatedObjectToken& operator=(FRefCountedPreAnimatedObjectToken&&) = default;
+#else
+		FRefCountedPreAnimatedObjectToken(FRefCountedPreAnimatedObjectToken&& RHS)
+			: AnimTypeID(RHS.AnimTypeID)
+			, RefCount(RHS.RefCount)
+			, Token(MoveTemp(RHS.Token))
+		{
+		}
+
+		FRefCountedPreAnimatedObjectToken& operator=(FRefCountedPreAnimatedObjectToken&& RHS)
+		{
+			AnimTypeID = MoveTemp(RHS.AnimTypeID);
+			RefCount = RHS.RefCount;
+			Token = MoveTemp(RHS.Token);
+			return *this;
+		}
+#endif
+
+		/** The type of the animation */
+		FMovieSceneAnimTypeID AnimTypeID;
+
+		/** The number of entities that are referencing this data */
+		uint32 RefCount;
+
+		/** The token that defines how to restore this object's state */
+		IMovieScenePreAnimatedTokenPtr Token;
+	};
+	TArray<FRefCountedPreAnimatedObjectToken> EntityTokens;
+
+	/** The object that was animated */
+	TWeakObjectPtr<> Object;
+};
+
+
+/** Saved state for animation bound to particular animated */
+struct FMovieSceneSavedMasterTokens
+{
+	FMovieSceneSavedMasterTokens(){}
+
+#if PLATFORM_COMPILER_HAS_DEFAULTED_FUNCTIONS
+	FMovieSceneSavedMasterTokens(FMovieSceneSavedMasterTokens&&) = default;
+	FMovieSceneSavedMasterTokens& operator=(FMovieSceneSavedMasterTokens&&) = default;
+#else
+	FMovieSceneSavedMasterTokens(FMovieSceneSavedMasterTokens&& RHS)
+	{
+		*this = MoveTemp(RHS);
+	}
+	FMovieSceneSavedMasterTokens& operator=(FMovieSceneSavedMasterTokens&& RHS)
+	{
+		AnimatedGlobalTypeIDs = MoveTemp(RHS.AnimatedGlobalTypeIDs);
+		AnimatedEntityTypeIDs = MoveTemp(RHS.AnimatedEntityTypeIDs);
+		GlobalStateTokens = MoveTemp(RHS.GlobalStateTokens);
+		EntityTokens = MoveTemp(RHS.EntityTokens);
+		return *this;
+	}
+#endif
+
+	MOVIESCENE_API void Add(ECapturePreAnimatedState CaptureMode, FMovieSceneAnimTypeID InAnimTypeID, FMovieSceneEvaluationKey AssociatedKey, const IMovieScenePreAnimatedGlobalTokenProducer& Producer, IMovieScenePreAnimatedState& Parent);
+
+	MOVIESCENE_API void Restore(IMovieScenePlayer& Player);
+
+	void RestoreEntity(IMovieScenePlayer& Player, FMovieSceneEvaluationKey EntityKey)
+	{
+		TArray<FMovieSceneAnimTypeID, TInlineAllocator<8>> AnimTypesToRestore;
+
+		for (int32 LUTIndex = AnimatedEntityTypeIDs.Num() - 1; LUTIndex >= 0; --LUTIndex)
+		{
+			FMovieSceneEntityAndAnimTypeID EntityAndAnimType = AnimatedEntityTypeIDs[LUTIndex];
+			if (EntityAndAnimType.EntityKey == EntityKey)
+			{
+				AnimTypesToRestore.Add(EntityAndAnimType.AnimTypeID);
+				AnimatedEntityTypeIDs.RemoveAtSwap(LUTIndex);
+			}
+		}
+
+		for (int32 TokenIndex = EntityTokens.Num() - 1; TokenIndex >= 0; --TokenIndex)
+		{
+			FRefCountedPreAnimatedGlobalToken& Token = EntityTokens[TokenIndex];
+			if (AnimTypesToRestore.Contains(Token.AnimTypeID) && --Token.RefCount == 0)
+			{
+				Token.Token->RestoreState(Player);
+				EntityTokens.RemoveAtSwap(TokenIndex);
+			}
+		}
+	}
+
+	void Reset()
+	{
+		AnimatedGlobalTypeIDs.Reset();
+		AnimatedEntityTypeIDs.Reset();
+		GlobalStateTokens.Reset();
+		EntityTokens.Reset();
+	}
+
+private:
+
+	TArray<FMovieSceneAnimTypeID, TInlineAllocator<8>> AnimatedGlobalTypeIDs;
+	TArray<FMovieSceneEntityAndAnimTypeID, TInlineAllocator<8>> AnimatedEntityTypeIDs;
+
+	/** Tokens mapped to specific anim type IDs. Can only ever be one at a time. */
+	TArray<IMovieScenePreAnimatedGlobalTokenPtr> GlobalStateTokens;
+
+	/** Tokens mapped to specific anim type IDs. Can only ever be one at a time. */
+	struct FRefCountedPreAnimatedGlobalToken
+	{
+		FRefCountedPreAnimatedGlobalToken(IMovieScenePreAnimatedGlobalTokenPtr&& InToken, FMovieSceneAnimTypeID InAnimTypeID) : AnimTypeID(InAnimTypeID), RefCount(1), Token(MoveTemp(InToken)) {}
+
+#if PLATFORM_COMPILER_HAS_DEFAULTED_FUNCTIONS
+		FRefCountedPreAnimatedGlobalToken(FRefCountedPreAnimatedGlobalToken&&) = default;
+		FRefCountedPreAnimatedGlobalToken& operator=(FRefCountedPreAnimatedGlobalToken&&) = default;
+#else
+		FRefCountedPreAnimatedGlobalToken(FRefCountedPreAnimatedGlobalToken&& RHS)
+			: AnimTypeID(RHS.AnimTypeID)
+			, RefCount(RHS.RefCount)
+			, Token(MoveTemp(RHS.Token))
+		{
+		}
+		FRefCountedPreAnimatedGlobalToken& operator=(FRefCountedPreAnimatedGlobalToken&& RHS)
+		{
+			AnimTypeID = MoveTemp(RHS.AnimTypeID);
+			RefCount = RHS.RefCount;
+			Token = MoveTemp(RHS.Token);
+			return *this;
+		}
+#endif
+
+		/** The type of the animation */
+		FMovieSceneAnimTypeID AnimTypeID;
+
+		/** The number of entities that are referencing this data */
+		uint32 RefCount;
+
+		/** The token that defines how to restore the state */
+		IMovieScenePreAnimatedGlobalTokenPtr Token;
+	};
+	TArray<FRefCountedPreAnimatedGlobalToken> EntityTokens;
+};
+
+/**
+ * Class that caches pre-animated state for objects that were manipulated by sequencer
+ */
+class FMovieScenePreAnimatedState : IMovieScenePreAnimatedState
+{
+public:
+
+	/**
+	 * Default construction
+	 */
+	FMovieScenePreAnimatedState()
+	{
+		DefaultGlobalCaptureMode = ECapturePreAnimatedState::None;
+	}
+
+	FMovieScenePreAnimatedState(const FMovieScenePreAnimatedState&) = delete;
+	FMovieScenePreAnimatedState& operator=(const FMovieScenePreAnimatedState&) = delete;
+
+	/**
+	 * Check whether we're currently caching pre-animated state at a global level
+	 */
+	bool IsGlobalCaptureEnabled() const
+	{
+		return DefaultGlobalCaptureMode == ECapturePreAnimatedState::Global;
+	}
+
+	/**
+	 * Enable this cache and allow it to start caching state
+	 */
+	void EnableGlobalCapture()
+	{
+		DefaultGlobalCaptureMode = ECapturePreAnimatedState::Global;
+	}
+	
+	/**
+	 * Disable this cache, preventing it from caching state
+	 */
+	void DisableGlobalCapture()
+	{
+		DefaultGlobalCaptureMode = ECapturePreAnimatedState::None;
+	}
+
+public:
+
+	FORCEINLINE void SavePreAnimatedState(FMovieSceneAnimTypeID InTokenType, const IMovieScenePreAnimatedTokenProducer& Producer, UObject& InObject)
+	{
+		SavePreAnimatedState(InTokenType, Producer, InObject, CurrentCaptureState, CapturingStateFor);
+	}
+
+	FORCEINLINE void SavePreAnimatedState(FMovieSceneAnimTypeID InTokenType, const IMovieScenePreAnimatedGlobalTokenProducer& Producer)
+	{
+		SavePreAnimatedState(InTokenType, Producer, CurrentCaptureState, CapturingStateFor);
+	}
+
+	void SavePreAnimatedState(FMovieSceneAnimTypeID InTokenType, const IMovieScenePreAnimatedTokenProducer& Producer, UObject& InObject, ECapturePreAnimatedState CaptureState, FMovieSceneEvaluationKey CaptureEntity)
+	{
+		if (CaptureState == ECapturePreAnimatedState::None)
+		{
+			return;
+		}
+
+		FObjectKey ObjectKey(&InObject);
+
+		FMovieSceneSavedObjectTokens* Container = ObjectTokens.Find(ObjectKey);
+		if (!Container)
+		{
+			Container = &ObjectTokens.Add(ObjectKey, FMovieSceneSavedObjectTokens(InObject));
+		}
+
+		Container->Add(CaptureState, InTokenType, CaptureEntity, Producer, InObject, *this);
+	}
+
+	void SavePreAnimatedState(FMovieSceneAnimTypeID InTokenType, const IMovieScenePreAnimatedGlobalTokenProducer& Producer, ECapturePreAnimatedState CaptureState, FMovieSceneEvaluationKey CaptureEntity)
+	{
+		if (CaptureState == ECapturePreAnimatedState::None)
+		{
+			return;
+		}
+
+		MasterTokens.Add(CaptureState, InTokenType, CaptureEntity, Producer, *this);
+	}
+
+	void RestorePreAnimatedState(IMovieScenePlayer& Player, const FMovieSceneEvaluationKey& Key)
+	{
+		auto* AnimatedObjects = EntityToAnimatedObjects.Find(Key);
+		if (!AnimatedObjects)
+		{
+			return;
+		}
+
+		for (FObjectKey ObjectKey : *AnimatedObjects)
+		{
+			if (ObjectKey == FObjectKey())
+			{
+				MasterTokens.RestoreEntity(Player, Key);
+			}
+			else if (FMovieSceneSavedObjectTokens* FoundState = ObjectTokens.Find(ObjectKey))
+			{
+				FoundState->RestoreEntity(Player, Key);
+			}
+		}
+
+		EntityToAnimatedObjects.Remove(Key);
+	}
+
+	MOVIESCENE_API void RestorePreAnimatedState(IMovieScenePlayer& Player);
+
+	MOVIESCENE_API void RestorePreAnimatedState(IMovieScenePlayer& Player, UObject& Object);
+
+	MOVIESCENE_API void RestorePreAnimatedState(IMovieScenePlayer& Player, UObject& Object, TFunctionRef<bool(FMovieSceneAnimTypeID)> InFilter);
+
+public:
+
+	void SetCaptureEntity(FMovieSceneEvaluationKey InEntity, EMovieSceneCompletionMode InCompletionMode)
+	{
+		CapturingStateFor = InEntity;
+
+		switch(InCompletionMode)
+		{
+		case EMovieSceneCompletionMode::RestoreState: 	CurrentCaptureState = ECapturePreAnimatedState::Entity;	break;
+		case EMovieSceneCompletionMode::KeepState: 		CurrentCaptureState = DefaultGlobalCaptureMode;	break;
+		}
+	}
+
+private:
+
+	virtual void EntityHasAnimatedObject(FMovieSceneEvaluationKey EntityKey, FObjectKey ObjectKey) override
+	{
+		EntityToAnimatedObjects.FindOrAdd(EntityKey).Add(ObjectKey);
+	}
+
+	virtual void EntityHasAnimatedMaster(FMovieSceneEvaluationKey EntityKey) override
+	{
+		EntityToAnimatedObjects.FindOrAdd(EntityKey).Add(FObjectKey());
+	}
+
+	TMap<FObjectKey, FMovieSceneSavedObjectTokens> ObjectTokens;
+	FMovieSceneSavedMasterTokens MasterTokens;
+
+	TMap<FMovieSceneEvaluationKey, TArray<FObjectKey, TInlineAllocator<4>>> EntityToAnimatedObjects;
+
+	FMovieSceneEvaluationKey CapturingStateFor;
+
+	ECapturePreAnimatedState CurrentCaptureState;
+	ECapturePreAnimatedState DefaultGlobalCaptureMode;
+};
\ No newline at end of file
diff --git a/Engine/Source/Runtime/MovieScene/Public/Evaluation/MovieScenePropertyTemplate.h b/Engine/Source/Runtime/MovieScene/Public/Evaluation/MovieScenePropertyTemplate.h
new file mode 100644
--- /dev/null
+++ b/Engine/Source/Runtime/MovieScene/Public/Evaluation/MovieScenePropertyTemplate.h
@@ -0,0 +1,245 @@
+// Copyright 1998-2016 Epic Games, Inc. All Rights Reserved.
+
+#pragma once
+
+#include "MovieSceneFwd.h"
+#include "MovieSceneCommonHelpers.h"
+#include "MovieSceneEvalTemplate.h"
+#include "MovieScenePropertyTemplate.generated.h"
+
+DECLARE_CYCLE_STAT(TEXT("Property Track Token Execute"), MovieSceneEval_PropertyTrack_TokenExecute, STATGROUP_MovieSceneEval);
+
+namespace PropertyTemplate
+{
+	/** Persistent section data for a property section */
+	struct FSectionData : IPersistentEvaluationData
+	{
+		MOVIESCENE_API FSectionData();
+
+		/** Initialize track data with the specified property name, path, and optional setter function */
+		MOVIESCENE_API void Initialize(FName InPropertyName, FString InPropertyPath, FName InFunctionName = NAME_None);
+
+		/** Property bindings used to get and set the property */
+		TSharedPtr<FTrackInstancePropertyBindings> PropertyBindings;
+
+		/** Cached identifier of the property we're editing */
+		FMovieSceneAnimTypeID PropertyID;
+	};
+
+	/** The value of the object as it existed before this frame's evaluation */
+	template<typename PropertyValueType>
+	struct TCachedValue
+	{
+	/** The object */
+		TWeakObjectPtr<UObject> WeakObject;
+
+		/** The property value */
+		PropertyValueType Value;
+	};
+
+	/** Typed persistent section data for a property section that caches object values at the start of the frame. Use in conjunction with TPropertyTrackExecutionToken. */
+	template<typename PropertyValueType>
+	struct TCachedSectionData : FSectionData
+	{
+		/** Setup the data for the current frame */
+		void SetupFrame(const FMovieSceneEvaluationOperand& Operand, IMovieScenePlayer& Player)
+		{
+			ObjectsAndValues.Reset();
+			
+			for (TWeakObjectPtr<> Object : Player.FindBoundObjects(Operand))
+			{
+				if (UObject* ObjectPtr = Object.Get())
+				{
+					ObjectsAndValues.Add(TCachedValue<PropertyValueType>{ ObjectPtr, PropertyBindings->GetCurrentValue<PropertyValueType>(*ObjectPtr) });
+				}
+			}
+		}
+
+		TArray<TCachedValue<PropertyValueType>, TInlineAllocator<1>> ObjectsAndValues;
+	};
+
+	template<typename PropertyValueType, typename IntermediateType = PropertyValueType>
+	struct TTemporarySetterType { typedef PropertyValueType Type; };
+
+	/** Convert from an intermediate type to the type used for setting a property value. Called when resetting pre animated state */
+	template<typename PropertyValueType, typename IntermediateType = PropertyValueType>
+	typename TTemporarySetterType<PropertyValueType, IntermediateType>::Type
+		ConvertFromIntermediateType(const IntermediateType& InIntermediateType, IMovieScenePlayer& Player)
+	{
+		return InIntermediateType;
+	}
+
+	/** Convert from an intermediate type to the type used for setting a property value. Called during token execution. */
+	template<typename PropertyValueType, typename IntermediateType = PropertyValueType>
+	typename TTemporarySetterType<PropertyValueType, IntermediateType>::Type
+		ConvertFromIntermediateType(const IntermediateType& InIntermediateType, const FMovieSceneEvaluationOperand& Operand, FPersistentEvaluationData& PersistentData, IMovieScenePlayer& Player)
+	{
+		return InIntermediateType;
+	}
+	
+	template<typename PropertyValueType, typename IntermediateType = PropertyValueType>
+	IntermediateType ConvertToIntermediateType(PropertyValueType&& NewValue)
+	{
+		return MoveTemp(NewValue);
+	}
+
+	template<typename T>
+	bool IsValueValid(const T& InValue)
+	{
+		return true;
+	}
+
+	/** Cached preanimated state for a given property */
+	template<typename PropertyValueType, typename IntermediateType = PropertyValueType>
+	struct TCachedState : IMovieScenePreAnimatedToken
+	{
+		TCachedState(typename TCallTraits<IntermediateType>::ParamType InValue, const FTrackInstancePropertyBindings& InBindings)
+			: Value(MoveTemp(InValue))
+			, Bindings(InBindings)
+		{
+		}
+
+		virtual void RestoreState(UObject& Object, IMovieScenePlayer& Player) override
+		{
+			auto NewValue = PropertyTemplate::ConvertFromIntermediateType<PropertyValueType, IntermediateType>(Value, Player);
+			if (IsValueValid(NewValue))
+			{
+				Bindings.CallFunction<PropertyValueType>(Object, NewValue);
+			}
+		}
+
+		IntermediateType Value;
+		FTrackInstancePropertyBindings Bindings;
+	};
+
+	template<typename PropertyValueType, typename IntermediateType = PropertyValueType>
+	static IMovieScenePreAnimatedTokenPtr CacheExistingState(UObject& Object, FTrackInstancePropertyBindings& PropertyBindings)
+	{
+		return TCachedState<PropertyValueType, IntermediateType>(PropertyTemplate::ConvertToIntermediateType(PropertyBindings.GetCurrentValue<PropertyValueType>(Object)), PropertyBindings);
+	}
+
+
+	template<typename PropertyValueType>
+	struct FTokenProducer : IMovieScenePreAnimatedTokenProducer
+	{
+		FTokenProducer(FTrackInstancePropertyBindings& InPropertyBindings) : PropertyBindings(InPropertyBindings) {}
+
+		virtual IMovieScenePreAnimatedTokenPtr CacheExistingState(UObject& Object) const
+		{
+			return PropertyTemplate::CacheExistingState<PropertyValueType>(Object, PropertyBindings);
+		}
+
+		FTrackInstancePropertyBindings& PropertyBindings;
+	};
+}
+
+/** Execution token for a given property */
+template<typename PropertyValueType>
+struct TCachedPropertyTrackExecutionToken : IMovieSceneExecutionToken
+{
+	/** Execute this token, operating on all objects referenced by 'Operand' */
+	virtual void Execute(const FMovieSceneContext& Context, const FMovieSceneEvaluationOperand& Operand, FPersistentEvaluationData& PersistentData, IMovieScenePlayer& Player) override
+	{
+		using namespace PropertyTemplate;
+
+		MOVIESCENE_DETAILED_SCOPE_CYCLE_COUNTER(MovieSceneEval_PropertyTrack_TokenExecute)
+		
+		TCachedSectionData<PropertyValueType>& PropertyTrackData = PersistentData.GetSectionData<TCachedSectionData<PropertyValueType>>();
+		FTrackInstancePropertyBindings* PropertyBindings = PropertyTrackData.PropertyBindings.Get();
+
+		for (TCachedValue<PropertyValueType>& ObjectAndValue : PropertyTrackData.ObjectsAndValues)
+		{
+			if (UObject* ObjectPtr = ObjectAndValue.WeakObject.Get())
+			{
+				Player.SavePreAnimatedState(*ObjectPtr, PropertyTrackData.PropertyID, FTokenProducer<PropertyValueType>(*PropertyBindings));
+				
+				PropertyBindings->CallFunction<PropertyValueType>(*ObjectPtr, ObjectAndValue.Value);
+			}
+		}
+	}
+};
+
+/** Execution token that simple stores a value, and sets it when executed */
+template<typename PropertyValueType, typename IntermediateType = PropertyValueType>
+struct TPropertyTrackExecutionToken : IMovieSceneExecutionToken
+{
+	TPropertyTrackExecutionToken(IntermediateType InValue)
+		: Value(MoveTemp(InValue))
+	{}
+
+	/** Execute this token, operating on all objects referenced by 'Operand' */
+	virtual void Execute(const FMovieSceneContext& Context, const FMovieSceneEvaluationOperand& Operand, FPersistentEvaluationData& PersistentData, IMovieScenePlayer& Player) override
+	{
+		using namespace PropertyTemplate;
+
+		MOVIESCENE_DETAILED_SCOPE_CYCLE_COUNTER(MovieSceneEval_PropertyTrack_TokenExecute)
+		
+		FSectionData& PropertyTrackData = PersistentData.GetSectionData<FSectionData>();
+		FTrackInstancePropertyBindings* PropertyBindings = PropertyTrackData.PropertyBindings.Get();
+
+		auto NewValue = PropertyTemplate::ConvertFromIntermediateType<PropertyValueType, IntermediateType>(Value, Operand, PersistentData, Player);
+		if (!IsValueValid(NewValue))
+		{
+			return;
+		}
+
+		for (TWeakObjectPtr<> WeakObject : Player.FindBoundObjects(Operand))
+		{
+			if (UObject* ObjectPtr = WeakObject.Get())
+			{
+				Player.SavePreAnimatedState(*ObjectPtr, PropertyTrackData.PropertyID, FTokenProducer<PropertyValueType>(*PropertyBindings));
+				
+				PropertyBindings->CallFunction<PropertyValueType>(*ObjectPtr, NewValue);
+			}
+		}
+	}
+
+	IntermediateType Value;
+};
+
+USTRUCT()
+struct FMovieScenePropertySectionData
+{
+	GENERATED_BODY()
+	
+	FMovieScenePropertySectionData()
+	{}
+
+	FMovieScenePropertySectionData(FName InPropertyName, FString InPropertyPath, FName InFunctionName = NAME_None)
+		: PropertyName(InPropertyName), PropertyPath(MoveTemp(InPropertyPath)), FunctionName(InFunctionName) 
+	{
+	}
+
+	/** Helper function to create FSectionData for this property section */
+	void SetupTrack(FPersistentEvaluationData& PersistentData) const
+	{
+		PersistentData.AddSectionData<PropertyTemplate::FSectionData>().Initialize(PropertyName, PropertyPath, FunctionName);
+	}
+
+	/** Helper function to create TCachedSectionData<T> for this property section */
+	template<typename T>
+	void SetupCachedTrack(FPersistentEvaluationData& PersistentData) const
+	{
+		typedef PropertyTemplate::TCachedSectionData<T> FSectionData;
+		PersistentData.AddSectionData<FSectionData>().Initialize(PropertyName, PropertyPath, FunctionName);
+	}
+
+	/** Helper function to initialize the TCachedSectionData<T> for this property section for the current frame */
+	template<typename T>
+	void SetupCachedFrame(const FMovieSceneEvaluationOperand& Operand, FPersistentEvaluationData& PersistentData, IMovieScenePlayer& Player) const
+	{
+		typedef PropertyTemplate::TCachedSectionData<T> FSectionData;
+		PersistentData.GetSectionData<FSectionData>().SetupFrame(Operand, Player);
+	}
+
+private:
+
+	UPROPERTY()
+	FName PropertyName;
+
+	UPROPERTY()
+	FString PropertyPath;
+
+	UPROPERTY()
+	FName FunctionName;
+};
\ No newline at end of file
diff --git a/Engine/Source/Runtime/MovieScene/Public/Evaluation/MovieSceneSectionParameters.h b/Engine/Source/Runtime/MovieScene/Public/Evaluation/MovieSceneSectionParameters.h
new file mode 100644
--- /dev/null
+++ b/Engine/Source/Runtime/MovieScene/Public/Evaluation/MovieSceneSectionParameters.h
@@ -0,0 +1,36 @@
+// Copyright 1998-2016 Epic Games, Inc. All Rights Reserved.
+
+#pragma once
+
+#include "MovieSceneSequenceTransform.h"
+#include "MovieSceneSectionParameters.generated.h"
+
+USTRUCT()
+struct FMovieSceneSectionParameters
+{
+	GENERATED_BODY()
+
+	/** Default constructor */
+	FMovieSceneSectionParameters()
+		: StartOffset(0.0f)
+		, TimeScale(1.0f)
+		, PrerollTime(0.0f)
+		, PostrollTime(0.0f)
+	{}
+
+	/** Number of seconds to skip at the beginning of the sub-sequence. */
+	UPROPERTY(EditAnywhere, Category="Clipping")
+	float StartOffset;
+
+	/** Playback time scaling factor. */
+	UPROPERTY(EditAnywhere, Category="Timing")
+	float TimeScale;
+
+	/** Amount of time to evaluate the section before its actual physical start time. */
+	UPROPERTY(EditAnywhere, Category="Timing")
+	float PrerollTime;
+
+	/** Amount of time to evaluate the section after its actual physical end time. */
+	UPROPERTY(EditAnywhere, Category="Timing")
+	float PostrollTime;
+};
\ No newline at end of file
diff --git a/Engine/Source/Runtime/MovieScene/Public/Evaluation/MovieSceneSegment.h b/Engine/Source/Runtime/MovieScene/Public/Evaluation/MovieSceneSegment.h
new file mode 100644
--- /dev/null
+++ b/Engine/Source/Runtime/MovieScene/Public/Evaluation/MovieSceneSegment.h
@@ -0,0 +1,101 @@
+// Copyright 1998-2016 Epic Games, Inc. All Rights Reserved.
+
+#pragma once
+
+#include "Range.h"
+#include "ArrayView.h"
+#include "MovieSceneSegment.generated.h"
+
+/**
+ * Evaluation data that specifies information about what to evaluate for a given template
+ */
+USTRUCT()
+struct FSectionEvaluationData
+{
+	GENERATED_BODY()
+
+	FSectionEvaluationData() : ForcedTime(TNumericLimits<float>::Lowest()) {}
+	explicit FSectionEvaluationData(int32 InImplIndex, float InForcedTime = TNumericLimits<float>::Lowest()) : ImplIndex(InImplIndex), ForcedTime(InForcedTime) {}
+
+	friend bool operator==(FSectionEvaluationData A, FSectionEvaluationData B)
+	{
+		return A.ImplIndex == B.ImplIndex && A.ForcedTime == B.ForcedTime;
+	}
+
+	float GetTime(float InActualTime)
+	{
+		return ForcedTime == TNumericLimits<float>::Lowest() ? InActualTime : ForcedTime;
+	}
+
+	/** The implementation index we should evaluation (index into FMovieSceneEvaluationTrack::ChildTemplates) */
+	UPROPERTY()
+	int32 ImplIndex;
+
+	/** A forced time to evaluate this section at */
+	UPROPERTY()
+	float ForcedTime;
+};
+
+/**
+ * Information about a singe segment of an evaluation track
+ */
+USTRUCT()
+struct FMovieSceneSegment
+{
+	GENERATED_BODY()
+
+	FMovieSceneSegment()
+	{}
+
+	FMovieSceneSegment(const TRange<float>& InRange)
+		: Range(InRange)
+	{}
+
+	FMovieSceneSegment(const TRange<float>& InRange, TArrayView<int32> InApplicationImpls)
+		: Range(InRange)
+	{
+		Impls.Reserve(InApplicationImpls.Num());
+		for (int32 Impl : InApplicationImpls)
+		{
+			Impls.Add(FSectionEvaluationData(Impl));
+		}
+	}
+
+	/** Custom serializer to accomodate the inline allocator on our array */
+	bool Serialize(FArchive& Ar)
+	{
+		Ar << Range;
+
+		int32 NumStructs = Impls.Num();
+		Ar << NumStructs;
+
+		if (Ar.IsLoading())
+		{
+			for (int32 Index = 0; Index < NumStructs; ++Index)
+			{
+				FSectionEvaluationData Data;
+				FSectionEvaluationData::StaticStruct()->SerializeItem(Ar, &Data, nullptr);
+				Impls.Add(Data);
+			}
+		}
+		else if (Ar.IsSaving())
+		{
+			for (FSectionEvaluationData& Data : Impls)
+			{
+				FSectionEvaluationData::StaticStruct()->SerializeItem(Ar, &Data, nullptr);
+			}
+		}
+		return true;
+	}
+
+	/** The segment's range */
+	FFloatRange Range;
+
+	/** Array of implementations that reside at the segment's range */
+	TArray<FSectionEvaluationData, TInlineAllocator<4>> Impls;
+};
+
+template<> struct TStructOpsTypeTraits<FMovieSceneSegment> : public TStructOpsTypeTraitsBase
+{
+	enum { WithSerializer = true };
+};
\ No newline at end of file
diff --git a/Engine/Source/Runtime/MovieScene/Public/Evaluation/MovieSceneSequenceHierarchy.h b/Engine/Source/Runtime/MovieScene/Public/Evaluation/MovieSceneSequenceHierarchy.h
new file mode 100644
--- /dev/null
+++ b/Engine/Source/Runtime/MovieScene/Public/Evaluation/MovieSceneSequenceHierarchy.h
@@ -0,0 +1,192 @@
+// Copyright 1998-2016 Epic Games, Inc. All Rights Reserved.
+
+#pragma once
+
+#include "MovieSceneSectionParameters.h"
+#include "MovieSceneSequenceHierarchy.generated.h"
+
+/**
+ * Sub sequence data that is stored within an evaluation template as a backreference to the originating sequence, and section
+ */
+USTRUCT()
+struct FMovieSceneSubSequenceData
+{
+	GENERATED_BODY()
+
+	/**
+	 * Default constructor for serialization
+	 */
+	FMovieSceneSubSequenceData()
+		: Sequence(nullptr)
+	{}
+
+	/**
+	 * Construction from a movie scene sequence, and a sub section name, and its valid play range
+	 */
+	FMovieSceneSubSequenceData(UMovieSceneSequence& InSequence, FMovieSceneSequenceID InDeterministicSequenceID
+#if WITH_EDITORONLY_DATA
+		, FName InSectionPath
+		, TRange<float> InValidPlayRange
+#endif
+		)
+		: Sequence(&InSequence)
+		, DeterministicSequenceID(InDeterministicSequenceID)
+#if WITH_EDITORONLY_DATA
+		, SectionPath(InSectionPath)
+		, ValidPlayRange(InValidPlayRange)
+#endif
+	{}
+
+	/** The sequence that the sub section references */
+	UPROPERTY()
+	UMovieSceneSequence* Sequence;
+
+	/** Transform that transforms a given time from the sequences outer space, to its authored space. */
+	UPROPERTY()
+	FMovieSceneSequenceTransform RootToSequenceTransform;
+
+	/** Cached signature of the evaluation template */
+	UPROPERTY()
+	FGuid SourceSequenceSignature;
+
+	/** This sequence's deterministic sequence ID. Used in editor to reduce the risk of collisions on recompilation */ 
+	UPROPERTY()
+	FMovieSceneSequenceID DeterministicSequenceID;
+
+#if WITH_EDITORONLY_DATA
+
+	/** This sequence's path within its movie scene */
+	UPROPERTY()
+	FName SectionPath;
+
+	/** This sub sequence's valid bounds according to its parent sub section. Clamped recursively during template generation */
+	UPROPERTY()
+	FFloatRange ValidPlayRange;
+#endif
+};
+
+/**
+ * Simple structure specifying parent and child sequence IDs for any given sequences
+ */
+USTRUCT()
+struct FMovieSceneSequenceHierarchyNode
+{
+	GENERATED_BODY()
+
+	/**
+	 * Default construction used by serialization
+	 */
+	FMovieSceneSequenceHierarchyNode()
+	{}
+
+	/**
+	 * Construct this hierarchy node from the sequence's parent ID
+	 */
+	FMovieSceneSequenceHierarchyNode(FMovieSceneSequenceIDRef InParentID)
+		: ParentID(InParentID)
+	{}
+
+	/** Movie scene sequence ID of this node's parent sequence */
+	UPROPERTY()
+	FMovieSceneSequenceID ParentID;
+
+	/** Array of child sequences contained within this sequence */
+	UPROPERTY()
+	TArray<FMovieSceneSequenceID> Children;
+};
+
+/**
+ * Structure that stores hierarchical information pertaining to all sequences contained within a master sequence
+ */
+USTRUCT()
+struct FMovieSceneSequenceHierarchy
+{
+	GENERATED_BODY()
+
+	FMovieSceneSequenceHierarchy()
+	{
+		Hierarchy.Add(MovieSceneSequenceID::Root.GetInternalValue(), FMovieSceneSequenceHierarchyNode(MovieSceneSequenceID::Invalid));
+	}
+
+	/**
+	 * Find the structural information for the specified sequence ID
+	 *
+	 * @param SequenceID 			The ID of the sequence to lookup
+	 * @return pointer to the structural information, or nullptr if the sequence ID does not exist in this hierarchy 
+	 */
+	const FMovieSceneSequenceHierarchyNode* FindNode(FMovieSceneSequenceIDRef SequenceID) const
+	{
+		return Hierarchy.Find(SequenceID.GetInternalValue());
+	}
+
+	/**
+	 * Find the structural information for the specified sequence ID
+	 *
+	 * @param SequenceID 			The ID of the sequence to lookup
+	 * @return pointer to the structural information, or nullptr if the sequence ID does not exist in this hierarchy 
+	 */
+	FMovieSceneSequenceHierarchyNode* FindNode(FMovieSceneSequenceIDRef SequenceID)
+	{
+		return Hierarchy.Find(SequenceID.GetInternalValue());
+	}
+
+	/**
+	 * Find the sub sequence and section information for the specified sequence ID
+	 *
+	 * @param SequenceID 			The ID of the sequence to lookup
+	 * @return pointer to the sequence/section information, or nullptr if the sequence ID does not exist in this hierarchy 
+	 */
+	const FMovieSceneSubSequenceData* FindSubData(FMovieSceneSequenceIDRef SequenceID) const
+	{
+		return SubSequences.Find(SequenceID.GetInternalValue());
+	}
+
+	/**
+	 * Find the sub sequence and section information for the specified sequence ID
+	 *
+	 * @param SequenceID 			The ID of the sequence to lookup
+	 * @return pointer to the sequence/section information, or nullptr if the sequence ID does not exist in this hierarchy 
+	 */
+	FMovieSceneSubSequenceData* FindSubData(FMovieSceneSequenceIDRef SequenceID)
+	{
+		return SubSequences.Find(SequenceID.GetInternalValue());
+	}
+
+	/**
+	 * Add the specified sub sequence data to the hierarchy
+	 *
+	 * @param Data 					The data to add
+	 * @param ThisSequenceID 		The sequence ID of the sequence the data relates to
+	 * @param ParentID 				The parent ID of this sequence data
+	 */
+	void Add(const FMovieSceneSubSequenceData& Data, FMovieSceneSequenceIDRef ThisSequenceID, FMovieSceneSequenceIDRef ParentID)
+	{
+		SubSequences.Add(ThisSequenceID.GetInternalValue(), Data);
+
+		FMovieSceneSequenceHierarchyNode Node(ParentID);
+		Hierarchy.Add(ThisSequenceID.GetInternalValue(), Node);
+
+		if (ParentID != MovieSceneSequenceID::Invalid)
+		{
+			FMovieSceneSequenceHierarchyNode* Parent = Hierarchy.Find(ParentID.GetInternalValue());
+			check(Parent);
+			Parent->Children.Add(ThisSequenceID);
+		}
+	}
+
+	/** Access to all the subsequence data */
+	const TMap<FMovieSceneSequenceID, FMovieSceneSubSequenceData>& AllSubSequenceData() const
+	{
+		return (const TMap<FMovieSceneSequenceID, FMovieSceneSubSequenceData>&)SubSequences;
+	}
+
+private:
+
+	/** Map of all (recursive) sub sequences found in this template, keyed on sequence ID */
+	UPROPERTY()
+	TMap<uint32, FMovieSceneSubSequenceData> SubSequences;
+
+	/** Structural information describing the structure of the sequence */
+	UPROPERTY()
+	TMap<uint32, FMovieSceneSequenceHierarchyNode> Hierarchy;
+};
diff --git a/Engine/Source/Runtime/MovieScene/Public/Evaluation/MovieSceneSequenceTemplateStore.h b/Engine/Source/Runtime/MovieScene/Public/Evaluation/MovieSceneSequenceTemplateStore.h
new file mode 100644
--- /dev/null
+++ b/Engine/Source/Runtime/MovieScene/Public/Evaluation/MovieSceneSequenceTemplateStore.h
@@ -0,0 +1,33 @@
+// Copyright 1998-2016 Epic Games, Inc. All Rights Reserved.
+
+#pragma once
+
+struct FMovieSceneEvaluationTemplate;
+class UMovieSceneSequence;
+
+struct FMovieSceneSequenceTemplateStore
+{
+	FMovieSceneSequenceTemplateStore()
+#if WITH_EDITORONLY_DATA
+		: bTemplatesAreVolatile(true)
+#else
+		// In cooked builds we assume the evaluation template can never change, and is thus always up-to-date
+		: bTemplatesAreVolatile(false)
+#endif
+	{
+	}
+
+	virtual ~FMovieSceneSequenceTemplateStore() {}
+	
+	bool AreTemplatesVolatile() const
+	{
+		return bTemplatesAreVolatile;
+	}
+
+	MOVIESCENE_API virtual FMovieSceneEvaluationTemplate& GetCompiledTemplate(UMovieSceneSequence& Sequence);
+
+protected:
+
+	/** True where templates are liable to change, and we should listen for changes */
+	bool bTemplatesAreVolatile;
+};
\ No newline at end of file
diff --git a/Engine/Source/Runtime/MovieScene/Public/Evaluation/MovieSceneSequenceTransform.h b/Engine/Source/Runtime/MovieScene/Public/Evaluation/MovieSceneSequenceTransform.h
new file mode 100644
--- /dev/null
+++ b/Engine/Source/Runtime/MovieScene/Public/Evaluation/MovieSceneSequenceTransform.h
@@ -0,0 +1,137 @@
+// Copyright 1998-2016 Epic Games, Inc. All Rights Reserved.
+
+#pragma once
+
+#include "MovieSceneSequenceTransform.generated.h"
+
+/**
+ * Movie scene sequence transform class that transforms from one time-space to another.
+ *
+ * @note The transform can be thought of as the top row of a 2x2 matrix, where the bottom row is the identity:
+ * 			| TimeScale	Offset	|
+ *			| 0			1		|
+ *
+ * As such, traditional matrix mathematics can be applied to transform between different sequence's time-spaces.
+ * Transforms apply offset first, then time scale.
+ */
+USTRUCT()
+struct FMovieSceneSequenceTransform
+{
+	GENERATED_BODY()
+
+	/**
+	 * Default construction to the identity transform
+	 */
+	FMovieSceneSequenceTransform()
+		: TimeScale(1.f)
+		, Offset(0.f)
+	{}
+
+	/**
+	 * Construction from an offset, and a scale
+	 *
+	 * @param InOffset 			The offset to translate by
+	 * @param InTimeScale 		The timescale. For instance, if a sequence is playing twice as fast, pass 2.f
+	 */
+	FMovieSceneSequenceTransform(float InOffset, float InTimeScale = 1.f)
+		: TimeScale(InTimeScale)
+		, Offset(InOffset)
+	{}
+
+	/**
+	 * Retrieve the inverse of this transform
+	 */
+	FMovieSceneSequenceTransform Inverse() const
+	{
+		return FMovieSceneSequenceTransform(-Offset/TimeScale, 1.f/TimeScale);
+	}
+
+	/** The sequence's time scale (or play rate) */
+	UPROPERTY()
+	float TimeScale;
+
+	/** Scalar time offset applied before the scale */
+	UPROPERTY()
+	float Offset;
+};
+
+/**
+ * Transform a time by a sequence transform
+ *
+ * @param InTime 			The time to transform
+ * @param RHS 				The transform
+ */
+inline float operator*(float InTime, const FMovieSceneSequenceTransform& RHS)
+{
+	return RHS.Offset + InTime * RHS.TimeScale;
+}
+
+/**
+ * Transform a time by a sequence transform
+ *
+ * @param InTime 			The time to transform
+ * @param RHS 				The transform
+ */
+inline float& operator*=(float& InTime, const FMovieSceneSequenceTransform& RHS)
+{
+	InTime = InTime * RHS;
+	return InTime;
+}
+
+/**
+ * Transform a time range by a sequence transform
+ *
+ * @param LHS 				The time range to transform
+ * @param RHS 				The transform
+ */
+template<typename T>
+TRange<T> operator*(const TRange<T>& LHS, const FMovieSceneSequenceTransform& RHS)
+{
+	TRangeBound<float> SourceLower = LHS.GetLowerBound();
+	TRangeBound<float> TransformedLower =
+		SourceLower.IsOpen() ? 
+			TRangeBound<float>() : 
+			SourceLower.IsInclusive() ?
+				TRangeBound<float>::Inclusive(SourceLower.GetValue() * RHS) :
+				TRangeBound<float>::Exclusive(SourceLower.GetValue() * RHS);
+
+	TRangeBound<float> SourceUpper = LHS.GetUpperBound();
+	TRangeBound<float> TransformedUpper =
+		SourceUpper.IsOpen() ? 
+			TRangeBound<float>() : 
+			SourceUpper.IsInclusive() ?
+				TRangeBound<float>::Inclusive(SourceUpper.GetValue() * RHS) :
+				TRangeBound<float>::Exclusive(SourceUpper.GetValue() * RHS);
+
+	return TRange<float>(TransformedLower, TransformedUpper);
+}
+
+/**
+ * Transform a time range by a sequence transform
+ *
+ * @param InTime 			The time range to transform
+ * @param RHS 				The transform
+ */
+template<typename T>
+TRange<T>& operator*=(TRange<T>& LHS, const FMovieSceneSequenceTransform& RHS)
+{
+	LHS = LHS * RHS;
+	return LHS;
+}
+
+/**
+ * Multiply 2 transforms together, resulting in a single transform that gets from RHS parent to LHS space
+ * @note Transforms apply from right to left
+ */
+inline FMovieSceneSequenceTransform operator*(const FMovieSceneSequenceTransform& LHS, const FMovieSceneSequenceTransform& RHS)
+{
+	// The matrix multiplication occurs as follows:
+	//
+	// | TimeScaleA	, OffsetA	|	.	| TimeScaleB, OffsetB	|
+	// | 0			, 1			|		| 0			, 1			|
+
+	return FMovieSceneSequenceTransform(
+		RHS.Offset*LHS.TimeScale + LHS.Offset,		// New Offset
+		LHS.TimeScale * RHS.TimeScale				// New TimeScale
+		);
+}
\ No newline at end of file
diff --git a/Engine/Source/Runtime/MovieScene/Public/Evaluation/MovieSceneTrackIdentifier.h b/Engine/Source/Runtime/MovieScene/Public/Evaluation/MovieSceneTrackIdentifier.h
new file mode 100644
--- /dev/null
+++ b/Engine/Source/Runtime/MovieScene/Public/Evaluation/MovieSceneTrackIdentifier.h
@@ -0,0 +1,45 @@
+// Copyright 1998-2016 Epic Games, Inc. All Rights Reserved.
+
+#pragma once
+
+#include "MovieSceneTrackIdentifier.generated.h"
+
+USTRUCT()
+struct FMovieSceneTrackIdentifier
+{
+	GENERATED_BODY()
+
+	FMovieSceneTrackIdentifier()
+		: Value(-1)
+	{}
+
+	static FMovieSceneTrackIdentifier Invalid() { return FMovieSceneTrackIdentifier(); }
+
+	FMovieSceneTrackIdentifier& operator++()
+	{
+		++Value;
+		return *this;
+	}
+
+	friend bool operator==(FMovieSceneTrackIdentifier A, FMovieSceneTrackIdentifier B)
+	{
+		return A.Value == B.Value;
+	}
+
+	friend bool operator!=(FMovieSceneTrackIdentifier A, FMovieSceneTrackIdentifier B)
+	{
+		return A.Value != B.Value;
+	}
+
+	friend uint32 GetTypeHash(FMovieSceneTrackIdentifier In)
+	{
+		return In.Value;
+	}
+
+private:
+
+	friend struct FMovieSceneEvaluationTemplate;
+	
+	UPROPERTY()
+	uint32 Value;
+};
\ No newline at end of file
diff --git a/Engine/Source/Runtime/MovieScene/Public/IMovieScenePlayer.h b/Engine/Source/Runtime/MovieScene/Public/IMovieScenePlayer.h
--- a/Engine/Source/Runtime/MovieScene/Public/IMovieScenePlayer.h
+++ b/Engine/Source/Runtime/MovieScene/Public/IMovieScenePlayer.h
@@ -2,12 +2,13 @@
 
 #pragma once
 
-#include "IMovieSceneSpawnRegister.h"
+#include "MovieSceneSpawnRegister.h"
 #include "UnrealClient.h"
+#include "MovieSceneEvaluationState.h"
+#include "MovieScenePreAnimatedState.h"
+#include "MovieSceneSequenceInstance.h"
 
-
-class FMovieSceneSequenceInstance;
-
+struct FMovieSceneRootEvaluationTemplateInstance;
 
 namespace EMovieScenePlayerStatus
 {
@@ -57,14 +58,11 @@ struct EMovieSceneViewportParams
 class IMovieScenePlayer
 {
 public:
-	
+
 	/**
-	 * Gets runtime objects associated with an object handle for manipulation during playback
-	 *
-	 * @param ObjectHandle The handle to runtime objects
-	 * @Param The list of runtime objects that will be populated
+	 * Access the evaluation template that we are playing back
 	 */
-	virtual void GetRuntimeObjects(TSharedRef<FMovieSceneSequenceInstance> MovieSceneInstance, const FGuid& ObjectHandle, TArray<TWeakObjectPtr<UObject>>& OutObjects) const = 0;
+	virtual FMovieSceneRootEvaluationTemplateInstance& GetEvaluationTemplate() = 0;
 
 	/**
 	 * Updates the perspective viewports with the actor to view through
@@ -89,27 +87,6 @@ public:
 	 */
 	virtual void GetViewportSettings(TMap<FViewportClient*, EMovieSceneViewportParams>& ViewportParamsMap) const = 0;
 
-	/**
-	 * Adds a MovieScene instance to the player.  MovieScene players need to know about each instance for actor spawning
-	 *
-	 * @param MovieSceneSection	The section owning the MovieScene being instanced. 
-	 * @param InstanceToAdd		The instance being added
-	 */
-	virtual void AddOrUpdateMovieSceneInstance( class UMovieSceneSection& MovieSceneSection, TSharedRef<FMovieSceneSequenceInstance> InstanceToAdd ) = 0;
-
-	/**
-	 * Removes a MovieScene instance from the player. 
-	 *
-	 * @param MovieSceneSection	The section owning the MovieScene that was instanced. 
-	 * @param InstanceToAdd		The instance being removed
-	 */
-	virtual void RemoveMovieSceneInstance( class UMovieSceneSection& MovieSceneSection, TSharedRef<FMovieSceneSequenceInstance> InstanceToRemove ) = 0;
-
-	/**
-	 * @return The root movie scene instance being played
-	 */
-	virtual TSharedRef<FMovieSceneSequenceInstance> GetRootMovieSceneSequenceInstance() const = 0;
-
 	/** @return whether the player is currently playing, scrubbing, etc. */
 	virtual EMovieScenePlayerStatus::Type GetPlaybackStatus() const = 0;
 
@@ -121,7 +98,7 @@ public:
 	/**
 	 * Obtain an object responsible for managing movie scene spawnables
 	 */
-	virtual IMovieSceneSpawnRegister& GetSpawnRegister() { return NullRegister; }
+	virtual FMovieSceneSpawnRegister& GetSpawnRegister() { return NullRegister; }
 
 	/**
 	 * Access the playback context for this movie scene player
@@ -138,7 +115,164 @@ public:
 	 */
 	virtual bool IsPreview() const { return false; }
 
+public:
+
+	/**
+	 * Locate objects bound to the specified object guid, in the specified sequence
+	 * @note: Objects lists are cached internally until they are invalidate.
+	 *
+	 * @param ObjectBindingID 		The object to resolve
+	 * @param SequenceID 			ID of the sequence to resolve for
+	 *
+	 * @return Iterable list of weak object pointers pertaining to the specified GUID
+	 */
+	TArrayView<TWeakObjectPtr<>> FindBoundObjects(const FGuid& ObjectBindingID, FMovieSceneSequenceIDRef SequenceID)
+	{
+		FMovieSceneObjectCache* Cache = State.FindObjectCache(SequenceID);
+		if (Cache)
+		{
+			return Cache->FindBoundObjects(ObjectBindingID, *this);
+		}
+
+		return TArrayView<TWeakObjectPtr<>>();
+	}
+
+	/**
+	 * Locate objects bound to the specified sequence operand
+	 * @note: Objects lists are cached internally until they are invalidate.
+	 *
+	 * @param Operand 			The movie scene operand to resolve
+	 *
+	 * @return Iterable list of weak object pointers pertaining to the specified GUID
+	 */
+	TArrayView<TWeakObjectPtr<>> FindBoundObjects(const FMovieSceneEvaluationOperand& Operand)
+	{
+		return FindBoundObjects(Operand.ObjectBindingID, Operand.SequenceID);
+	}
+
+	/**
+	 * Attempt to find the object binding ID for the specified object, in the specified sequence
+	 * @note: Will forcably resolve all out-of-date object mappings in the sequence
+	 *
+	 * @param InObject 			The object to find a GUID for
+	 * @param SequenceID 		The sequence ID to search within
+	 *
+	 * @return The guid of the object's binding, or zero guid if it was not found
+	 */
+	FGuid FindObjectId(UObject& InObject, FMovieSceneSequenceIDRef SequenceID)
+	{
+		return State.FindObjectId(InObject, SequenceID, *this);
+	}
+
+	/**
+	 * Attempt to save specific state for the specified token state before it animates an object.
+	 * @note: Will only call IMovieSceneExecutionToken::CacheExistingState if no state has been previously cached for the specified token type
+	 *
+	 * @param InObject			The object to cache state for
+	 * @param InTokenType		Unique marker that identifies the originating token type
+	 * @param InProducer		Producer implementation that defines how to create the preanimated token, if it doesn't already exist
+	 */
+	FORCEINLINE void SavePreAnimatedState(UObject& InObject, FMovieSceneAnimTypeID InTokenType, const IMovieScenePreAnimatedTokenProducer& InProducer)
+	{
+		PreAnimatedState.SavePreAnimatedState(InTokenType, InProducer, InObject);
+	}
+
+	/**
+	 * Attempt to save specific state for the specified token state before it mutates state.
+	 * @note: Will only call IMovieSceneExecutionToken::CacheExistingState if no state has been previously cached for the specified token type
+	 *
+	 * @param InTokenType		Unique marker that identifies the originating token type
+	 * @param InProducer		Producer implementation that defines how to create the preanimated token, if it doesn't already exist
+	 */
+	FORCEINLINE void SavePreAnimatedState(FMovieSceneAnimTypeID InTokenType, const IMovieScenePreAnimatedGlobalTokenProducer& InProducer)
+	{
+		PreAnimatedState.SavePreAnimatedState(InTokenType, InProducer);
+	}
+
+	/**
+	 * Attempt to save specific state for the specified token state before it animates an object.
+	 * @note: Will only call IMovieSceneExecutionToken::CacheExistingState if no state has been previously cached for the specified token type
+	 *
+	 * @param InObject			The object to cache state for
+	 * @param InTokenType		Unique marker that identifies the originating token type
+	 * @param InProducer		Producer implementation that defines how to create the preanimated token, if it doesn't already exist
+	 * @param CaptureEntity		The entity key to associate this animated state with
+	 */
+	FORCEINLINE void SavePreAnimatedState(UObject& InObject, FMovieSceneAnimTypeID InTokenType, const IMovieScenePreAnimatedTokenProducer& InProducer, FMovieSceneEvaluationKey CaptureEntity)
+	{
+		PreAnimatedState.SavePreAnimatedState(InTokenType, InProducer, InObject, ECapturePreAnimatedState::Entity, CaptureEntity);
+	}
+
+	/**
+	 * Attempt to save specific state for the specified token state before it mutates state.
+	 * @note: Will only call IMovieSceneExecutionToken::CacheExistingState if no state has been previously cached for the specified token type
+	 *
+	 * @param InTokenType		Unique marker that identifies the originating token type
+	 * @param InProducer		Producer implementation that defines how to create the preanimated token, if it doesn't already exist
+	 * @param CaptureEntity		The entity key to associate this animated state with
+	 */
+	FORCEINLINE void SavePreAnimatedState(FMovieSceneAnimTypeID InTokenType, const IMovieScenePreAnimatedGlobalTokenProducer& InProducer, FMovieSceneEvaluationKey CaptureEntity)
+	{
+		PreAnimatedState.SavePreAnimatedState(InTokenType, InProducer, ECapturePreAnimatedState::Entity, CaptureEntity);
+	}
+
+	/**
+	 * Restore all pre-animated state
+	 */
+	void RestorePreAnimatedState()
+	{
+		PreAnimatedState.RestorePreAnimatedState(*this);
+	}
+
+	/**
+	 * Restore any pre-animated state that has been cached for the specified object
+	 *
+	 * @param Object			The object to restore
+	 */
+	void RestorePreAnimatedState(UObject& Object)
+	{
+		PreAnimatedState.RestorePreAnimatedState(*this, Object);
+	}
+
+	/**
+	 * Restore any pre-animated state that has been cached for the specified object
+	 *
+	 * @param Object			The object to restore
+	 * @param InFilter			Filter that defines whether specific anim types should be restored
+	 */
+	void RestorePreAnimatedState(UObject& Object, TFunctionRef<bool(FMovieSceneAnimTypeID)> InFilter)
+	{
+		PreAnimatedState.RestorePreAnimatedState(*this, Object, InFilter);
+	}
+
+	/**
+	 * Restore any pre-animated state that has been cached from the specified entity (a section or, less commonly, a track)
+	 *
+	 * @param EntityKey			The key to the entity that we want to restore state for (typically retrieved from PersistentData.GetSectionKey())
+	 */
+	FORCEINLINE void RestorePreAnimatedState(const FMovieSceneEvaluationKey& EntityKey)
+	{
+		PreAnimatedState.RestorePreAnimatedState(*this, EntityKey);
+	}
+
+public:
+
+	/** Evaluation state that stores global state to do with the playback operation */
+	FMovieSceneEvaluationState State;
+
+	/** Container that stores any per-animated state tokens  */
+	FMovieScenePreAnimatedState PreAnimatedState;
+
+public:
+
+	DEPRECATED(4.15, "Sequence Instances have been deprecated in favor of a template approach (see GetEvaluationTemplate).")
+	MOVIESCENE_API TSharedRef<FMovieSceneSequenceInstance> GetRootMovieSceneSequenceInstance();
+
+	DEPRECATED(4.15, "Please use FindBoundObjects directly.")
+	MOVIESCENE_API void GetRuntimeObjects(TSharedRef<FMovieSceneSequenceInstance> MovieSceneInstance, const FGuid& Guid, TArray<TWeakObjectPtr<>>& OutRuntimeObjects);
+	
 private:
+
 	/** Null register that asserts on use */
 	FNullMovieSceneSpawnRegister NullRegister;
 };
diff --git a/Engine/Source/Runtime/MovieScene/Public/IMovieSceneTrackInstance.h b/Engine/Source/Runtime/MovieScene/Public/IMovieSceneTrackInstance.h
--- a/Engine/Source/Runtime/MovieScene/Public/IMovieSceneTrackInstance.h
+++ b/Engine/Source/Runtime/MovieScene/Public/IMovieSceneTrackInstance.h
@@ -2,6 +2,7 @@
 
 #pragma once
 
+#include "MovieSceneAnimTypeID.h"
 
 class IMovieScenePlayer;
 class FMovieSceneSequenceInstance;
@@ -56,6 +57,8 @@ class IMovieSceneTrackInstance
 {
 public:
 
+	MOVIESCENE_API IMovieSceneTrackInstance();
+	
 	/** Virtual destructor. */
 	virtual ~IMovieSceneTrackInstance() { }
 
@@ -106,4 +109,6 @@ public:
 	* Evaluation order
 	*/
 	virtual float EvalOrder() { return 0.f; }
+
+	const FMovieSceneAnimTypeID AnimTypeID;
 };
diff --git a/Engine/Source/Runtime/MovieScene/Public/MovieScene.h b/Engine/Source/Runtime/MovieScene/Public/MovieScene.h
--- a/Engine/Source/Runtime/MovieScene/Public/MovieScene.h
+++ b/Engine/Source/Runtime/MovieScene/Public/MovieScene.h
@@ -5,6 +5,7 @@
 #include "MovieSceneBinding.h"
 #include "MovieScenePossessable.h"
 #include "MovieSceneSpawnable.h"
+#include "MovieSceneSignedObject.h"
 #include "MovieScene.generated.h"
 
 
@@ -13,7 +14,6 @@ class UMovieSceneSection;
 class UMovieSceneTrack;
 class UMovieSceneFolder;
 
-
 MOVIESCENE_API DECLARE_LOG_CATEGORY_EXTERN(LogMovieScene, Log, All);
 
 
@@ -80,7 +80,7 @@ struct FMovieSceneTrackLabels
  */
 UCLASS()
 class MOVIESCENE_API UMovieScene
-	: public UObject
+	: public UMovieSceneSignedObject
 {
 	GENERATED_UCLASS_BODY()
 
@@ -165,7 +165,15 @@ public:
 	/*
 	* Replace an existing possessable with another 
 	*/
-	bool ReplacePossessable(const FGuid& OldGuid, const FGuid& NewGuid, const FString& Name);
+	bool ReplacePossessable(const FGuid& OldGuid, const FMovieScenePossessable& InNewPosessable);
+
+	DEPRECATED(4.15, "Please use ReplacePossessable(const FGuid&, const FMovieScenePossessable&) so that the possessable class gets updated correctly.")
+	bool ReplacePossessable(const FGuid& OldGuid, const FGuid& NewGuid, const FString& Name)
+	{
+		FMovieScenePossessable NewPossessable(Name, nullptr);
+		NewPossessable.SetGuid(NewGuid);
+		return ReplacePossessable(OldGuid, NewPossessable);
+	}
 
 	/**
 	 * Tries to locate a possessable in this MovieScene for the specified possessable GUID.
@@ -462,6 +470,22 @@ public:
 	 */
 	void SetPlaybackRange(float Start, float End, bool bAlwaysMarkDirty = true);
 
+	/**
+	 * Set the start and end working range (outer) for this movie scene
+	 *
+	 * @param Start The offset from 0-time to view this movie scene.
+	 * @param End The offset from 0-time to view this movie scene
+	 */
+	void SetWorkingRange(float Start, float End);
+
+	/**
+	 * Set the start and end view range (inner) for this movie scene
+	 *
+	 * @param Start The offset from 0-time to view this movie scene
+	 * @param End The offset from 0-time to view this movie scene
+	 */
+	void SetViewRange(float Start, float End);
+
 	/**
 	 * Gets whether or not playback should be forced to match the fixed frame interval.  When true all time values will be rounded to a fixed
 	 * frame value which will force editor and runtime playback to match exactly, but will result in duplicate frames if the runtime and editor
@@ -479,13 +503,18 @@ public:
 	/**
 	* Gets the fixed frame interval to be used when "force fixed frame interval playback" is set.
 	*/
-	float GetFixedFrameInterval() const;
+	float GetFixedFrameInterval() const { return FixedFrameInterval; }
 
 	/**
 	* Gets the fixed frame interval to be used when "force fixed frame interval playback" is set.
 	*/
 	void SetFixedFrameInterval( float InFixedFrameInterval );
 
+	/**
+	 * Gets the fixed frame interval to be used when "force fixed frame interval playback" is set. Only returns a valid result when GetForceFixedFrameIntervalPlayback() is true, and the interval is > 0.
+	 */
+	TOptional<float> GetOptionalFixedFrameInterval() const { return (!GetForceFixedFrameIntervalPlayback() || GetFixedFrameInterval() <= 0) ? TOptional<float>() : GetFixedFrameInterval(); }
+
 	/**
 	 * Calculates a fixed frame time based on a current time, a fixed frame interval, and an internal epsilon to account
 	 * for floating point consistency.
diff --git a/Engine/Source/Runtime/MovieScene/Public/MovieSceneCommonHelpers.h b/Engine/Source/Runtime/MovieScene/Public/MovieSceneCommonHelpers.h
--- a/Engine/Source/Runtime/MovieScene/Public/MovieSceneCommonHelpers.h
+++ b/Engine/Source/Runtime/MovieScene/Public/MovieSceneCommonHelpers.h
@@ -103,46 +103,33 @@ public:
 	 * Calls the setter function for a specific runtime object or if the setter function does not exist, the property is set directly
 	 *
 	 * @param InRuntimeObject The runtime object whose function to call
-	 * @param PropertyValue The new value to assign to the property (if bound to property)
-	 * @param FunctionParams Parameters to pass to the function (if bound to function)
+	 * @param PropertyValue The new value to assign to the property
 	 */
 	template <typename ValueType>
-	void 
-	CallFunction( UObject* InRuntimeObject, void* PropertyValue, void* FunctionParams )
+	void CallFunction( UObject& InRuntimeObject, typename TCallTraits<ValueType>::ParamType PropertyValue )
 	{
-		FPropertyAndFunction PropAndFunction = RuntimeObjectToFunctionMap.FindRef(InRuntimeObject);
-		if(PropAndFunction.Function)
+		FPropertyAndFunction PropAndFunction = FindOrAdd(InRuntimeObject);
+		if (PropAndFunction.Function)
 		{
-			InRuntimeObject->ProcessEvent(PropAndFunction.Function, FunctionParams);
+			// ProcessEvent should really be taking const void*
+			InRuntimeObject.ProcessEvent(PropAndFunction.Function, (void*)&PropertyValue);
 		}
-		else
+		else if (PropAndFunction.PropertyAddress.Address)
 		{
-			ValueType* NewValue = (ValueType*)(PropertyValue);
-			SetCurrentValue<ValueType>(InRuntimeObject, *NewValue);
+			ValueType* Val = PropAndFunction.PropertyAddress.Property->ContainerPtrToValuePtr<ValueType>(PropAndFunction.PropertyAddress.Address);
+			if(Val)
+			{
+				*Val = MoveTemp(PropertyValue);
+			}
 		}
 	}
 
-	/** For backwards compatibility. */
-	template <typename ValueType>
-	void 
-	CallFunction( UObject* InRuntimeObject, void* PropertyValue )
-	{
-		CallFunction<ValueType>(InRuntimeObject, PropertyValue, PropertyValue);
-	}
-
-	/**
-	 * Rebuilds the property and function mappings for a set of runtime objects
-	 *
-	 * @param InRuntimeObjects	The objects to rebuild mappings for
-	 */
-	void UpdateBindings( const TArray<TWeakObjectPtr<UObject>>& InRuntimeObjects );
-
 	/**
-	 * Rebuilds the property and function mappings for a single runtime object
+	 * Rebuilds the property and function mappings for a single runtime object, and adds them to the cache
 	 *
-	 * @param InRuntimeObject	The object to rebuild mappings for
+	 * @param InRuntimeObject	The object to cache mappings for
 	 */
-	void UpdateBinding( const TWeakObjectPtr<UObject>& InRuntimeObject );
+	void CacheBinding( const UObject& InRuntimeObject );
 
 	/**
 	 * Gets the UProperty that is bound to the track instance
@@ -150,7 +137,7 @@ public:
 	 * @param Object	The Object that owns the property
 	 * @return			The property on the object if it exists
 	 */
-	UProperty* GetProperty(const UObject* Object) const;
+	UProperty* GetProperty(const UObject& Object) const;
 
 	/**
 	 * Gets the current value of a property on an object
@@ -159,9 +146,9 @@ public:
 	 * @return ValueType	The current value
 	 */
 	template <typename ValueType>
-	ValueType GetCurrentValue(const UObject* Object) const
+	ValueType GetCurrentValue(const UObject& Object)
 	{
-		FPropertyAndFunction PropAndFunction = RuntimeObjectToFunctionMap.FindRef(Object);
+		FPropertyAndFunction PropAndFunction = FindOrAdd(Object);
 
 		if(PropAndFunction.PropertyAddress.Address)
 		{
@@ -182,9 +169,9 @@ public:
 	 * @param InValue   The value to set
 	 */
 	template <typename ValueType>
-	void SetCurrentValue(const UObject* Object, const ValueType& InValue)
+	void SetCurrentValue(UObject& Object, typename TCallTraits<ValueType>::ParamType InValue)
 	{
-		FPropertyAndFunction PropAndFunction = RuntimeObjectToFunctionMap.FindRef(Object);
+		FPropertyAndFunction PropAndFunction = FindOrAdd(Object);
 
 		if(PropAndFunction.PropertyAddress.Address)
 		{
@@ -208,6 +195,17 @@ public:
 		return PropertyName;
 	}
 
+	template<typename ValueType>
+	DEPRECATED(4.15, "Please use GetCurrentValue(const UObject&)")
+	ValueType GetCurrentValue(const UObject* Object) { check(Object) return GetCurrentValue<ValueType>(*Object); }
+	template <typename ValueType>
+	DEPRECATED(4.15, "Please use CallFunction(UObject&)")
+	void CallFunction( UObject* InRuntimeObject, ValueType* PropertyValue ) { CallFunction<ValueType>(*InRuntimeObject, *PropertyValue); }
+	DEPRECATED(4.15, "UpdateBindings is no longer necessary")
+	void UpdateBindings( const TArray<TWeakObjectPtr<UObject>>& InRuntimeObjects ){}
+	DEPRECATED(4.15, "UpdateBindings is no longer necessary")
+	void UpdateBinding( const TWeakObjectPtr<UObject>& InRuntimeObject ) {}
+
 private:
 
 	struct FPropertyAddress
@@ -232,12 +230,27 @@ private:
 		{}
 	};
 
-	FPropertyAddress FindPropertyRecursive(const UObject* Object, void* BasePointer, UStruct* InStruct, TArray<FString>& InPropertyNames, uint32 Index) const;
-	FPropertyAddress FindProperty(const UObject* Object, const FString& InPropertyPath) const;
+	static FPropertyAddress FindPropertyRecursive(void* BasePointer, UStruct* InStruct, TArray<FString>& InPropertyNames, uint32 Index);
+	static FPropertyAddress FindProperty(const UObject& Object, const FString& InPropertyPath);
+
+	/** Find or add the FPropertyAndFunction for the specified object */
+	FPropertyAndFunction FindOrAdd(const UObject& InObject)
+	{
+		FObjectKey ObjectKey(&InObject);
+
+		const FPropertyAndFunction* PropAndFunction = RuntimeObjectToFunctionMap.Find(ObjectKey);
+		if (PropAndFunction)
+		{
+			return *PropAndFunction;
+		}
+
+		CacheBinding(InObject);
+		return RuntimeObjectToFunctionMap.FindRef(ObjectKey);
+	}
 
 private:
 	/** Mapping of objects to bound functions that will be called to update data on the track */
-	TMap< TWeakObjectPtr<UObject>, FPropertyAndFunction > RuntimeObjectToFunctionMap;
+	TMap< FObjectKey, FPropertyAndFunction > RuntimeObjectToFunctionMap;
 
 	/** Path to the property we are bound to */
 	FString PropertyPath;
@@ -248,4 +261,9 @@ private:
 	/** Actual name of the property we are bound to */
 	FName PropertyName;
 
-};
\ No newline at end of file
+};
+
+/** Explicit specializations for bools */
+template<> MOVIESCENE_API void FTrackInstancePropertyBindings::CallFunction<bool>(UObject& InRuntimeObject, TCallTraits<bool>::ParamType PropertyValue);
+template<> MOVIESCENE_API bool FTrackInstancePropertyBindings::GetCurrentValue<bool>(const UObject& Object);
+template<> MOVIESCENE_API void FTrackInstancePropertyBindings::SetCurrentValue<bool>(UObject& Object, TCallTraits<bool>::ParamType InValue);
\ No newline at end of file
diff --git a/Engine/Source/Runtime/MovieScene/Public/MovieSceneExecutionToken.h b/Engine/Source/Runtime/MovieScene/Public/MovieSceneExecutionToken.h
new file mode 100644
--- /dev/null
+++ b/Engine/Source/Runtime/MovieScene/Public/MovieSceneExecutionToken.h
@@ -0,0 +1,41 @@
+// Copyright 1998-2016 Epic Games, Inc. All Rights Reserved.
+
+#pragma once
+
+#include "MovieSceneFwd.h"
+#include "InlineValue.h"
+#include "MovieSceneAnimTypeID.h"
+
+struct IMovieScenePreAnimatedToken
+{
+	virtual ~IMovieScenePreAnimatedToken() {}
+	/** Restore state for the specified object, only called when this token was created with a bound object */
+	virtual void RestoreState(UObject& Object, IMovieScenePlayer& Player) = 0;
+};
+typedef TInlineValue<IMovieScenePreAnimatedToken, 32> IMovieScenePreAnimatedTokenPtr;
+
+/** Type required for production of pre-animated state tokens. Implemented as an interface rather than a callback to ensure efficient construction (these types are often constructed, but rarely called) */
+struct IMovieScenePreAnimatedTokenProducer
+{
+	virtual IMovieScenePreAnimatedTokenPtr CacheExistingState(UObject& Object) const = 0;
+};
+
+struct IMovieScenePreAnimatedGlobalToken
+{
+	virtual ~IMovieScenePreAnimatedGlobalToken() {}
+	/** Restore global state for a master track only called when this token was created with no bound object */
+	virtual void RestoreState(IMovieScenePlayer& Player) = 0;
+};
+typedef TInlineValue<IMovieScenePreAnimatedGlobalToken, 32> IMovieScenePreAnimatedGlobalTokenPtr;
+
+/** Type required for production of pre-animated state tokens. Implemented as an interface rather than a callback to ensure efficient construction (these types are often constructed, but rarely called) */
+struct IMovieScenePreAnimatedGlobalTokenProducer
+{
+	virtual IMovieScenePreAnimatedGlobalTokenPtr CacheExistingState() const = 0;
+};
+
+struct IMovieSceneExecutionToken
+{
+	virtual ~IMovieSceneExecutionToken() {}
+	virtual void Execute(const FMovieSceneContext& Context, const FMovieSceneEvaluationOperand& Operand, FPersistentEvaluationData& PersistentData, IMovieScenePlayer& Player) = 0;
+};
diff --git a/Engine/Source/Runtime/MovieScene/Public/MovieSceneFwd.h b/Engine/Source/Runtime/MovieScene/Public/MovieSceneFwd.h
new file mode 100644
--- /dev/null
+++ b/Engine/Source/Runtime/MovieScene/Public/MovieSceneFwd.h
@@ -0,0 +1,43 @@
+// Copyright 1998-2016 Epic Games, Inc. All Rights Reserved.
+
+#pragma once
+
+#include "MovieSceneSequenceID.h"
+
+enum class ESpawnOwnership : uint8;
+
+class IMovieScenePlayer;
+class UMovieScene;
+class UMovieSceneTrack;
+class UMovieSceneSection;
+class UMovieSceneSequence;
+
+struct FMovieSceneContext;
+struct FMovieSceneSpawnable;
+struct FMovieSceneEvaluationState;
+struct FMovieSceneEvaluationFieldTrackPtr;
+struct FMovieSceneEvaluationFieldSegmentPtr;
+struct FMovieSceneEvalTemplate;
+struct FMovieSceneEvaluationTrack;
+struct FMovieSceneSharedDataId;
+struct IMovieSceneTemplateGenerator;
+struct FMovieSceneSubSequenceData;
+struct FMovieSceneSequenceTransform;
+struct IMovieSceneTemplateGenerator;
+struct FPersistentEvaluationData;
+struct FMovieSceneExecutionTokens;
+struct FMovieSceneEvaluationOperand;
+struct IMovieScenePreAnimatedTokenProducer;
+struct IMovieScenePreAnimatedGlobalTokenProducer;
+
+DECLARE_STATS_GROUP(TEXT("Movie Scene Evaluation"), STATGROUP_MovieSceneEval, STATCAT_Advanced);
+
+#ifndef MOVIESCENE_DETAILED_STATS
+	#define MOVIESCENE_DETAILED_STATS 0
+#endif
+
+#if MOVIESCENE_DETAILED_STATS
+	#define MOVIESCENE_DETAILED_SCOPE_CYCLE_COUNTER SCOPE_CYCLE_COUNTER
+#else
+	#define MOVIESCENE_DETAILED_SCOPE_CYCLE_COUNTER(...)
+#endif
\ No newline at end of file
diff --git a/Engine/Source/Runtime/MovieScene/Public/MovieSceneSection.h b/Engine/Source/Runtime/MovieScene/Public/MovieSceneSection.h
--- a/Engine/Source/Runtime/MovieScene/Public/MovieSceneSection.h
+++ b/Engine/Source/Runtime/MovieScene/Public/MovieSceneSection.h
@@ -2,19 +2,63 @@
 
 #pragma once
 
+#include "MovieSceneFwd.h"
 #include "KeyParams.h"
+#include "Curves/KeyHandle.h"
+#include "Curves/RichCurve.h"
+#include "MovieSceneSignedObject.h"
 #include "MovieSceneSection.generated.h"
 
+struct FMovieSceneEvalTemplatePtr;
+
+/** Enumeration specifying how to handle state when this section is no longer evaluated */
+UENUM()
+enum class EMovieSceneCompletionMode : uint8
+{
+	KeepState,
+
+	RestoreState,
+};
+
+
+USTRUCT()
+struct FMovieSceneSectionEvalOptions
+{
+	GENERATED_BODY()
+	
+	FMovieSceneSectionEvalOptions()
+		: bCanEditCompletionMode(false)
+		, CompletionMode(EMovieSceneCompletionMode::KeepState)
+	{}
+
+	void EnableAndSetCompletionMode(EMovieSceneCompletionMode NewCompletionMode)
+	{
+		bCanEditCompletionMode = true;
+		CompletionMode = NewCompletionMode;
+	}
+
+	UPROPERTY()
+	bool bCanEditCompletionMode;
+
+	/** When set to "RestoreState", this section will restore any animation back to its previous state  */
+	UPROPERTY(EditAnywhere, DisplayName="When Finished", Category="Section")
+	EMovieSceneCompletionMode CompletionMode;
+};
 
 /**
  * Base class for movie scene sections
  */
 UCLASS(abstract, MinimalAPI)
 class UMovieSceneSection
-	: public UObject
+	: public UMovieSceneSignedObject
 {
 	GENERATED_UCLASS_BODY()
 
+public:
+
+	UPROPERTY(EditAnywhere, Category="Section", meta=(ShowOnlyInnerProperties))
+	FMovieSceneSectionEvalOptions EvalOptions;
+
 public:
 
 	/**
@@ -180,6 +224,12 @@ public:
 		return nullptr;
 	}
 
+	/**
+	 * Generate an evaluation template for this section
+	 * @return a valid evaluation template ptr, or nullptr
+	 */
+	MOVIESCENE_API virtual FMovieSceneEvalTemplatePtr GenerateTemplate() const;
+
 	/**
 	 * Gets all snap times for this section
 	 *
@@ -272,7 +322,7 @@ public:
 
 	/** Sets the time for the key referenced by the supplied key handle. */
 	virtual void SetKeyTime( FKeyHandle KeyHandle, float Time ) PURE_VIRTUAL( UAISenseEvent::SetKeyTime, );
-
+	
 private:
 
 	/** The start time of the section */
diff --git a/Engine/Source/Runtime/MovieScene/Public/MovieSceneSequence.h b/Engine/Source/Runtime/MovieScene/Public/MovieSceneSequence.h
--- a/Engine/Source/Runtime/MovieScene/Public/MovieSceneSequence.h
+++ b/Engine/Source/Runtime/MovieScene/Public/MovieSceneSequence.h
@@ -2,23 +2,27 @@
 
 #pragma once
 
+#include "MovieSceneFwd.h"
+#include "MovieSceneEvaluationTemplate.h"
+#include "MovieSceneSignedObject.h"
 #include "MovieSceneSequence.generated.h"
 
-
-class FMovieSceneSequenceInstance;
 class UMovieScene;
-class UObject;
-
+struct FMovieSceneSequenceTemplateStore;
 
 /**
  * Abstract base class for movie scene animations (C++ version).
  */
 UCLASS(MinimalAPI)
 class UMovieSceneSequence
-	: public UObject
+	: public UMovieSceneSignedObject
 {
+public:
+
 	GENERATED_BODY()
 
+	MOVIESCENE_API UMovieSceneSequence(const FObjectInitializer& Init);
+
 public:
 
 	/**
@@ -35,26 +39,53 @@ public:
 	 * Check whether the given object can be possessed by this animation.
 	 *
 	 * @param Object The object to check.
+	 * @param InPlaybackContext The current playback context
 	 * @return true if the object can be possessed, false otherwise.
 	 */
-	virtual bool CanPossessObject(UObject& Object) const PURE_VIRTUAL(UMovieSceneSequence::CanPossessObject, return false;);
+	virtual bool CanPossessObject(UObject& Object, UObject* InPlaybackContext) const PURE_VIRTUAL(UMovieSceneSequence::CanPossessObject, return false;);
+
+	DEPRECATED(4.15, "Please implement LocateBoundObjects instead")
+	virtual UObject* FindPossessableObject(const FGuid& ObjectId, UObject* Context) const
+	{
+		TArray<UObject*, TInlineAllocator<1>> OutObjects;
+		LocateBoundObjects(ObjectId, Context, OutObjects);
+		return OutObjects.Num() ? OutObjects[0] : nullptr;
+	}
 
 	/**
-	 * Finds the possessed for the specified identifier.
+	 * Locate all the objects that correspond to the specified object ID, using the specified context
 	 *
-	 * @param ObjectId The unique identifier of the object.
-	 * @param Context Optional context to use to find the required object (for instance, a parent spawnable object)
-	 * @return The object, or nullptr if not found.
+	 * @param ObjectId				The unique identifier of the object.
+	 * @param Context				Optional context to use to find the required object (for instance, a parent spawnable object)
+	 * @param OutObjects			Destination array to add found objects to
 	 */
-	virtual UObject* FindPossessableObject(const FGuid& ObjectId, UObject* Context) const PURE_VIRTUAL(UMovieSceneSequence::FindPossessedObject, return nullptr;);
+	virtual void LocateBoundObjects(const FGuid& ObjectId, UObject* Context, TArray<UObject*, TInlineAllocator<1>>& OutObjects) const PURE_VIRTUAL(UMovieSceneSequence::LocateBoundObjects, );
 
 	/**
-	 * Finds an ID for the specified the object
+	 * Locate all the objects that correspond to the specified object ID, using the specified context
 	 *
-	 * @param Object The object to find an id for
-	 * @return The object's ID, or an invalid guid if not found
+	 * @param ObjectId				The unique identifier of the object.
+	 * @param Context				Optional context to use to find the required object (for instance, a parent spawnable object)
+	 * @return An array of all bound objects
 	 */
-	virtual FGuid FindPossessableObjectId(UObject& Object) const PURE_VIRTUAL(UMovieSceneSequence::FindPossessableObjectId, return FGuid(););
+	TArray<UObject*, TInlineAllocator<1>> LocateBoundObjects(const FGuid& ObjectId, UObject* Context) const
+	{
+		TArray<UObject*, TInlineAllocator<1>> OutObjects;
+		LocateBoundObjects(ObjectId, Context, OutObjects);
+		return OutObjects;
+	}
+
+	DEPRECATED(4.15, "Please use IMovieScenePlayer::FindObjectId or FindPossessableObjectId(UObject&, UObject*) instead.")
+	virtual FGuid FindPossessableObjectId(UObject& Object) const PURE_VIRTUAL(UMovieSceneSequence::FindPossessableObjectId, return FGuid(); );
+
+	/**
+	 * Attempt to find the guid relating to the specified object
+	 *
+	 * @param ObjectId				The unique identifier of the object.
+	 * @param Context				Optional context to use to find the required object (for instance, a parent spawnable object or its world)
+	 * @return The object's guid, or zero guid if the object is not a valid possessable in the current context
+	 */
+	MOVIESCENE_API FGuid FindPossessableObjectId(UObject& Object, UObject* Context) const;
 
 	/**
 	 * Get the movie scene that controls this animation.
@@ -98,6 +129,44 @@ public:
 	 */
 	virtual UObject* MakeSpawnableTemplateFromInstance(UObject& InSourceObject, FName ObjectName) { return nullptr; }
 
+public:
+
+	MOVIESCENE_API virtual void Serialize(FArchive& Ar) override;
+
+#if WITH_EDITORONLY_DATA
+	MOVIESCENE_API virtual void PostDuplicate(bool bDuplicateForPIE) override;
+#endif
+
+	MOVIESCENE_API virtual void PreSave(const ITargetPlatform* TargetPlatform) override;
+
+	MOVIESCENE_API virtual void GenerateEvaluationTemplate(FMovieSceneEvaluationTemplate& Template, const FMovieSceneTrackCompilationParams& Params, FMovieSceneSequenceTemplateStore& Store);
+
+	UPROPERTY()
+	FCachedMovieSceneEvaluationTemplate EvaluationTemplate;
+
+	UPROPERTY()
+	FMovieSceneTrackCompilationParams TemplateParameters;
+
+public:
+
+	/**
+	 * true if the result of GetParentObject is significant in object resolution for LocateBoundObjects.
+	 */
+	bool AreParentContextsSignificant() const
+	{
+		return bParentContextsAreSignificant;
+	}
+
+protected:
+
+	/**
+	 * true if the result of GetParentObject is significant in object resolution for LocateBoundObjects.
+	 * When true, if GetParentObject returns nullptr, the PlaybackContext will be used for LocateBoundObjects, other wise the object's parent will be used
+	 * When false, the PlaybackContext will always be used for LocateBoundObjects
+	 */
+	UPROPERTY()
+	bool bParentContextsAreSignificant;
+
 public:
 
 #if WITH_EDITOR
@@ -107,4 +176,4 @@ public:
 	 */
 	virtual FText GetDisplayName() const { return FText::FromName(GetFName()); }
 #endif
-};
+};
\ No newline at end of file
diff --git a/Engine/Source/Runtime/MovieScene/Public/MovieSceneSequenceID.h b/Engine/Source/Runtime/MovieScene/Public/MovieSceneSequenceID.h
new file mode 100644
--- /dev/null
+++ b/Engine/Source/Runtime/MovieScene/Public/MovieSceneSequenceID.h
@@ -0,0 +1,82 @@
+// Copyright 1998-2016 Epic Games, Inc. All Rights Reserved.
+
+#pragma once
+
+#include "MovieSceneSequenceID.generated.h"
+
+USTRUCT()
+struct FMovieSceneSequenceID
+{
+	GENERATED_BODY()
+
+	FMovieSceneSequenceID()
+		: Value(-1)
+	{}
+
+	explicit FMovieSceneSequenceID(uint32 InValue)
+		: Value(InValue)
+	{}
+
+	FORCEINLINE friend bool operator==(FMovieSceneSequenceID LHS, FMovieSceneSequenceID RHS)
+	{
+		return LHS.Value == RHS.Value;
+	}
+	
+	FORCEINLINE friend bool operator!=(FMovieSceneSequenceID LHS, FMovieSceneSequenceID RHS)
+	{
+		return LHS.Value != RHS.Value;
+	}
+
+	FORCEINLINE friend bool operator<(FMovieSceneSequenceID LHS, FMovieSceneSequenceID RHS)
+	{
+		return LHS.Value < RHS.Value;
+	}
+
+	FORCEINLINE friend uint32 GetTypeHash(FMovieSceneSequenceID In)
+	{
+		return GetTypeHash(In.Value);
+	}
+
+	FORCEINLINE FMovieSceneSequenceID AccumulateParentID(FMovieSceneSequenceID InParentID)
+	{
+		return FMovieSceneSequenceID(HashCombine(Value, InParentID.Value));
+	}
+
+	FORCEINLINE bool Serialize(FArchive& Ar)
+	{
+		Ar << Value;
+		return true;
+	}
+
+	FORCEINLINE friend FArchive& operator<<(FArchive& Ar, FMovieSceneSequenceID SequenceID)
+	{
+		SequenceID.Serialize(Ar);
+		return Ar;
+	}
+
+	FORCEINLINE uint32 GetInternalValue() const
+	{
+		return Value;
+	}
+private:
+
+	UPROPERTY()
+	uint32 Value;
+};
+
+template<>
+struct TStructOpsTypeTraits<FMovieSceneSequenceID> : public TStructOpsTypeTraitsBase
+{
+	enum
+	{
+		WithSerializer = true
+	};
+};
+
+typedef TCallTraits<FMovieSceneSequenceID>::ParamType FMovieSceneSequenceIDRef;
+
+namespace MovieSceneSequenceID
+{
+	static const FMovieSceneSequenceID Invalid(-1);
+	static const FMovieSceneSequenceID Root(0);
+}
\ No newline at end of file
diff --git a/Engine/Source/Runtime/MovieScene/Public/MovieSceneSequenceInstance.h b/Engine/Source/Runtime/MovieScene/Public/MovieSceneSequenceInstance.h
--- a/Engine/Source/Runtime/MovieScene/Public/MovieSceneSequenceInstance.h
+++ b/Engine/Source/Runtime/MovieScene/Public/MovieSceneSequenceInstance.h
@@ -2,213 +2,44 @@
 
 #pragma once
 
-#include "IMovieSceneTrackInstance.h"
+#include "MovieSceneFwd.h"
 
-class IMovieScenePlayer;
-class IMovieSceneTrackInstance;
-
-class UMovieSceneTrack;
 class UMovieSceneSequence;
-
-typedef TMap<TWeakObjectPtr<UMovieSceneTrack>, TSharedPtr<class IMovieSceneTrackInstance>> FMovieSceneInstanceMap;
+class IMovieScenePlayer;
 
 /**
  * A sequence instance holds the live objects bound to the tracks in a sequence.  This is completely transient
  * There is one instance per sequence.  If a sequence holds multiple sub-sequences each sub-sequence will have it's own instance even if they are the same sequence asset
  * A sequence instance also creates and manages all track instances for the tracks in a sequence
  */
-class FMovieSceneSequenceInstance
+class MOVIESCENE_API FMovieSceneSequenceInstance
 	: public TSharedFromThis<FMovieSceneSequenceInstance>
 {
 public:
-
-	/**
-	 * Creates and initializes a new instance from a Sequence
-	 *
-	 * @param InMovieSceneSequence The sequence that this instance represents.
-	 */
-	MOVIESCENE_API FMovieSceneSequenceInstance(const UMovieSceneSequence& InMovieSceneSequence);
-
-	/** 
-	 * Creates and initializes a new instance that is primed for recording 
-	 * 
-	 * @param InMovieSceneSubTrack The subtrack we are primed for recording in
-	 */
-	MOVIESCENE_API FMovieSceneSequenceInstance(const UMovieSceneTrack& InMovieSceneSubTrack);
-
-	MOVIESCENE_API ~FMovieSceneSequenceInstance();
-
-	/**
-	 * Find the identifier for a possessed or spawned object.
-	 *
-	 * @param Object The object to get the identifier for.
-	 * @return The object identifier, or an invalid GUID if not found.
-	 */
-	MOVIESCENE_API FGuid FindObjectId(UObject& Object) const;
-
-	/**
-	 * Find the identifier for the parent of the specified object
-	 *
-	 * @param Object The object to get the parent identifier for.
-	 * @return The parent's identifier, or an invalid GUID if not applicable.
-	 */
-	MOVIESCENE_API FGuid FindParentObjectId(UObject& Object) const;
 	
-	/**
-	 * Find the object relating to the specified GUID
-	 *
-	 * @param ObjectId The object identifier
-	 * @return The object, or nullptr
-	 */
-	MOVIESCENE_API UObject* FindObject(const FGuid& ObjectId, const IMovieScenePlayer& Player) const;
+	FMovieSceneSequenceInstance(const UMovieSceneSequence& InMovieSceneSequence, FMovieSceneSequenceIDRef InSequenceID);
 
-	/**
-	 * Find the spawned object relating to the specified GUID
-	 *
-	 * @param ObjectId The object identifier
-	 * @return The object, or nullptr
-	 */
-	MOVIESCENE_API UObject* FindSpawnedObject(const FGuid& ObjectId) const;
+	DEPRECATED(4.15, "Please use IMovieScenePlayer::State.FindObjectId")
+	FGuid FindObjectId(UObject& Object) const { return FGuid();}
 
-	/** Save state of the objects that this movie scene controls. */
-	MOVIESCENE_API void SaveState(IMovieScenePlayer& Player);
-
-	/** Restore state of the objects that this movie scene controls. */
-	MOVIESCENE_API void RestoreState(IMovieScenePlayer& Player);
-
-	/** Restore state of a specific object track that this movie scene controls. */
-	MOVIESCENE_API void RestoreSpecificState(const FGuid& ObjectGuid, IMovieScenePlayer& Player);
-
-	/**
-	 * Updates this movie scene.
-	 *
-	 * @param UpdateData The current and previous position of the moviescene that is playing.
-	 * @param Player Movie scene player interface for interaction with runtime data.
-	 */
-	MOVIESCENE_API void Update(EMovieSceneUpdateData& UpdateData, IMovieScenePlayer& Player);
-
-	/*
-	 * Pre update - spawnables, any stale runtime objects.
-	 */
-	MOVIESCENE_API void PreUpdate(IMovieScenePlayer& Player);
+	DEPRECATED(4.15, "Please use IMovieScenePlayer::State.FindObjectId")
+	FGuid FindParentObjectId(UObject& Object) const { return FGuid(); }
 	
-	/*
-	 * Update all the passes
-	 */
-	MOVIESCENE_API void UpdatePasses(EMovieSceneUpdateData& UpdateData, IMovieScenePlayer& Player);
-
-	/*
-	 * Post update - clean up spawnables.
-	 */
-	MOVIESCENE_API void PostUpdate(IMovieScenePlayer& Player);
-
-	/*
-	 * Update a single pass
-	 */
-	MOVIESCENE_API void UpdatePassSingle(EMovieSceneUpdateData& UpdateData, IMovieScenePlayer& Player);
-
-	MOVIESCENE_API void UpdateFromSubSceneDeactivate();
-
-	/**
-	 * Refreshes the existing instance.
-	 *
-	 * Called when something significant about movie scene data changes (like adding or removing a track).
-	 * Instantiates all new tracks found and removes instances for tracks that no longer exist.
-	 */
-	MOVIESCENE_API void RefreshInstance(IMovieScenePlayer& Player);
+	DEPRECATED(4.15, "Please use IMovieScenePlayer::State.FindBoundObjects")
+	UObject* FindObject(const FGuid& ObjectId, const IMovieScenePlayer& Player) const {return nullptr;}
 
-	/**
-	 * Spawn an object relating to the specified object ID
-	 * 
-	 * @param ObjectId The ID of the object to spawn
-	 * @param Player The movie scene player responsible for this instance
-	 * @return The newly spawned or previously-spawned object, or nullptr
-	 */
-	MOVIESCENE_API void OnObjectSpawned(const FGuid& ObjectId, UObject& SpawnedObject, IMovieScenePlayer& Player);
+	DEPRECATED(4.15, "Please use IMovieScenePlayer::SpawnRegister.FindSpawnedObject")
+	UObject* FindSpawnedObject(const FGuid& ObjectId) const{return nullptr;}
 
-	/**
-	 * Destroy a previously spawned object relating to the specified object ID
-	 * 
-	 * @param ObjectId The ID of the object to destroy
-	 * @param Player The movie scene player responsible for this instance
-	 */
-	MOVIESCENE_API void OnSpawnedObjectDestroyed(const FGuid& ObjectId, IMovieScenePlayer& Player);
+	DEPRECATED(4.15, "Direct access to the sequence at runtime should no longer be necessary.")
+	UMovieSceneSequence* GetSequence() const { return MovieSceneSequence.Get(); }
 
-	/**
-	 * Get the sequence associated with this instance.
-	 *
-	 * @return The movie scene sequence object.
-	 */
-	UMovieSceneSequence* GetSequence()
-	{
-		return MovieSceneSequence.Get();
-	}
+	DEPRECATED(4.15, "Direct access to the sequence's time range should no longer be necessary.")
+	TRange<float> GetTimeRange() const;
 
-	/** 
-	 * Get the time range for the movie scene.
-	 *
-	 * Note: This is not necessarily how long the instance will be playing.
-	 * The parent could have changed the length.
-	 *
-	 * @return The time range.
-	 */
-	TRange<float> GetTimeRange() const
-	{
-		return TimeRange;
-	}
-
-	/** Handle the sequence this instance refers to changing */
-	MOVIESCENE_API void HandleSequenceSectionChanged(UMovieSceneSequence* Sequence);
-
-#if WITH_EDITOR
-	/**
-	 * @return a transient but, unique identifier for this instance
-	 */
-	const FGuid& GetInstanceId() const { return InstanceId; }
-#endif
-
-protected:
-
-	void RefreshInstanceMap(const TArray<UMovieSceneTrack*>& Tracks, const TArray<TWeakObjectPtr<UObject>>& RuntimeObjects, FMovieSceneInstanceMap& TrackInstances, IMovieScenePlayer& Player);
-
-	/** Update the object binding instance for the specified object */
-	void UpdateObjectBinding(const FGuid& ObjectId, IMovieScenePlayer& Player);
+	FMovieSceneSequenceIDRef GetSequenceID() const { return SequenceID; }
 
 private:
-
-	/** A paring of a runtime object guid to a runtime data and track instances animating the runtime objects */
-	struct FMovieSceneObjectBindingInstance
-	{
-		/** Runtime object guid for looking up runtime UObjects */
-		FGuid ObjectGuid;
-
-		/** Actual runtime objects */	
-		TArray<TWeakObjectPtr<UObject>> RuntimeObjects;
-
-		/** Instances that animate the runtime objects */
-		FMovieSceneInstanceMap TrackInstances;
-	};
-
-
-	/** MovieScene that is instanced */
 	TWeakObjectPtr<UMovieSceneSequence> MovieSceneSequence;
-
-	/** The camera cut track instance map */
-	TSharedPtr<IMovieSceneTrackInstance> CameraCutTrackInstance;
-
-	/** All Master track instances */
-	FMovieSceneInstanceMap MasterTrackInstances;
-
-	/** All object binding instances */
-	TMap<FGuid, FMovieSceneObjectBindingInstance> ObjectBindingInstances;
-
-	/** A map of object ID -> spawned object */
-	TMap<FGuid, TWeakObjectPtr<UObject>> SpawnedObjects;
-
-	/** Cached time range for the movie scene */
-	TRange<float> TimeRange;
-
-#if WITH_EDITOR
-	FGuid InstanceId;
-#endif
+	FMovieSceneSequenceID SequenceID;
 };
diff --git a/Engine/Source/Runtime/MovieScene/Public/MovieSceneSpawnRegister.h b/Engine/Source/Runtime/MovieScene/Public/MovieSceneSpawnRegister.h
new file mode 100644
--- /dev/null
+++ b/Engine/Source/Runtime/MovieScene/Public/MovieSceneSpawnRegister.h
@@ -0,0 +1,209 @@
+// Copyright 1998-2016 Epic Games, Inc. All Rights Reserved.
+
+#pragma once
+
+#include "MovieSceneFwd.h"
+#include "MovieSceneSpawnable.h"
+#include "ValueOrError.h"
+
+class UMovieSceneSequence;
+
+/** Struct used for defining a new spawnable type */
+struct FNewSpawnable
+{
+	FNewSpawnable() : ObjectTemplate(nullptr) {}
+	FNewSpawnable(UObject* InObjectTemplate, FString InName) : ObjectTemplate(InObjectTemplate), Name(MoveTemp(InName)) {}
+
+	/** The archetype object template that defines the spawnable */
+	UObject* ObjectTemplate;
+
+	/** The desired name of the new spawnable */
+	FString Name;
+};
+
+/**
+ * Class responsible for managing spawnables in a movie scene
+ */
+class FMovieSceneSpawnRegister : public TSharedFromThis<FMovieSceneSpawnRegister>
+{
+public:
+
+	/**
+	 * Virtual destructor
+	 */
+	virtual ~FMovieSceneSpawnRegister()
+	{
+	}
+
+	/**
+	 * Attempt to find a previously spawned object represented by the specified object and template IDs
+	 * @param BindingId		ID of the object to find
+	 * @param TemplateID	Unique ID of the template to look within
+	 * @return The spawned object if found; nullptr otherwise.
+	 */
+	MOVIESCENE_API UObject* FindSpawnedObject(const FGuid& BindingId, FMovieSceneSequenceIDRef TemplateID) const;
+
+	/**
+	 * Spawn an object for the specified GUID, from the specified sequence instance.
+	 *
+	 * @param BindingId 	ID of the object to spawn
+	 * @param TemplateID 	Identifier for the current template we're evaluating
+	 * @param Player 		Movie scene player that is ultimately responsible for spawning the object
+	 * @return the spawned object, or nullptr on failure
+	 */
+	MOVIESCENE_API UObject* SpawnObject(const FGuid& BindingId, UMovieScene& MovieScene, FMovieSceneSequenceIDRef Template, IMovieScenePlayer& Player);
+
+	/**
+	 * Destroy a specific previously spawned object
+	 *
+	 * @param BindingId		ID of the object to destroy
+	 * @param TemplateID 	Identifier for the current template we're evaluating
+	 * @param Player 		Movie scene player that is ultimately responsible for destroying the object
+	 */
+	MOVIESCENE_API void DestroySpawnedObject(const FGuid& BindingId, FMovieSceneSequenceIDRef TemplateID, IMovieScenePlayer& Player);
+
+	/**
+	 * Destroy a specific previously spawned object, where its binding ID and sequence ID is not known.
+	 * @note Should only be used for restoring pre-animated state, or where it is otherwise impossible to call DestroySpawnedObject
+	 *
+	 * @param InObject		the object to destroy
+	 */
+	void DestroyObjectDirectly(UObject& InObject) { DestroySpawnedObject(InObject); }
+
+	/**
+	 * Destroy spawned objects using a custom predicate
+	 *
+	 * @param Player 		Movie scene player that is ultimately responsible for destroying the objects
+	 * @param Predicate		Predicate used for testing whether an object should be destroyed. Returns true for destruction, false to skip.
+	 */
+	MOVIESCENE_API void DestroyObjectsByPredicate(IMovieScenePlayer& Player, const TFunctionRef<bool(const FGuid&, ESpawnOwnership, FMovieSceneSequenceIDRef)>& Predicate);
+
+	/**
+	 * Purge any memory of any objects that are considered externally owned
+
+	 * @param Player 		Movie scene player that is ultimately responsible for destroying the objects
+	 */
+	MOVIESCENE_API void ForgetExternallyOwnedSpawnedObjects(FMovieSceneEvaluationState& State, IMovieScenePlayer& Player);
+
+public:
+
+	/**
+	 * Called to indiscriminately clean up any spawned objects
+	 */
+	MOVIESCENE_API void CleanUp(IMovieScenePlayer& Player);
+
+	/**
+	 * Called to clean up any non-externally owned spawnables that were spawned from the specified instance
+	 */
+	MOVIESCENE_API void CleanUpSequence(FMovieSceneSequenceIDRef TemplateID, IMovieScenePlayer& Player);
+
+	/**
+	 * Called when the current time has moved beyond the specified sequence's play range
+	 */
+	MOVIESCENE_API void OnSequenceExpired(FMovieSceneSequenceIDRef TemplateID, IMovieScenePlayer& Player);
+
+#if WITH_EDITOR
+	
+public:
+
+	/**
+	 * Create a new spawnable type from the given source object
+	 *
+	 * @param SourceObject		The source object to create the spawnable from
+	 * @param OwnerMovieScene	The owner movie scene that this spawnable type should reside in
+	 * @return the new spawnable type, or error
+	 */
+	virtual TValueOrError<FNewSpawnable, FText> CreateNewSpawnableType(UObject& SourceObject, UMovieScene& OwnerMovieScene) { return MakeError(NSLOCTEXT("SpawnRegister", "NotSupported", "Not supported")); }
+
+	/**
+	 * Called to save the default state of the specified spawnable
+	 */
+	virtual void SaveDefaultSpawnableState(FMovieSceneSpawnable& Spawnable, FMovieSceneSequenceIDRef TemplateID, IMovieScenePlayer& Player) {}
+
+#endif
+
+protected:
+
+	/**
+	 * Spawn an object for the specified GUID, from the specified sequence instance.
+	 *
+	 * @param Object 		ID of the object to spawn
+	 * @param TemplateID 	Identifier for the current template we're evaluating
+	 * @param Player 		Movie scene player that is ultimately responsible for spawning the object
+	 * @return the spawned object, or nullptr on failure
+	 */
+	virtual UObject* SpawnObject(FMovieSceneSpawnable& Spawnable, FMovieSceneSequenceIDRef TemplateID, IMovieScenePlayer& Player) = 0;
+
+	/**
+	 * Called right before a spawned object with the specified ID and template ID is destroyed
+	 */
+	virtual void PreDestroyObject(UObject& Object, const FGuid& BindingId, FMovieSceneSequenceIDRef TemplateID) {}
+
+	/**
+	 * Destroy a specific previously spawned object
+	 *
+	 * @param Object 		The object to destroy
+	 */
+	virtual void DestroySpawnedObject(UObject& Object) = 0;
+
+protected:
+
+	/** Structure holding information pertaining to a spawned object */
+	struct FSpawnedObject
+	{
+		FSpawnedObject(const FGuid& InGuid, UObject& InObject, ESpawnOwnership InOwnership)
+			: Guid(InGuid)
+			, Object(&InObject)
+			, Ownership(InOwnership)
+		{}
+
+		/** The ID of the sequencer object binding that this object relates to */
+		FGuid Guid;
+
+		/** The object that has been spawned */
+		TWeakObjectPtr<UObject> Object;
+
+		/** What level of ownership this object was spawned with */
+		ESpawnOwnership Ownership;
+	};
+
+	/**
+	 *	Helper key type for mapping a guid and sequence instance to a specific value
+	 */
+	struct FMovieSceneSpawnRegisterKey
+	{
+		FMovieSceneSpawnRegisterKey(FMovieSceneSequenceIDRef InTemplateID, const FGuid& InBindingId)
+			: BindingId(InBindingId)
+			, TemplateID(InTemplateID)
+		{
+		}
+
+		bool operator==(const FMovieSceneSpawnRegisterKey& Other) const
+		{
+			return BindingId == Other.BindingId && TemplateID == Other.TemplateID;
+		}
+		
+		friend uint32 GetTypeHash(const FMovieSceneSpawnRegisterKey& Key)
+		{
+			return HashCombine(GetTypeHash(Key.BindingId), GetTypeHash(Key.TemplateID));
+		}
+
+		/** BindingId of the object binding */
+		FGuid BindingId;
+
+		/** Movie Scene template identifier that spawned the object */
+		FMovieSceneSequenceID TemplateID;
+	};
+
+protected:
+
+	/** Register of spawned objects */
+	TMap<FMovieSceneSpawnRegisterKey, FSpawnedObject> Register;
+};
+
+class FNullMovieSceneSpawnRegister : public FMovieSceneSpawnRegister
+{
+public:
+	virtual UObject* SpawnObject(FMovieSceneSpawnable&, FMovieSceneSequenceIDRef, IMovieScenePlayer&) override { check(false); return nullptr; }
+	virtual void DestroySpawnedObject(UObject&) override { }
+};
\ No newline at end of file
