Sequencer: Made sequence pointers stored in sequence template instances weak object ptrs

  - We can't guarantee the lifetime of the objects here

diff --git a/Engine/Plugins/MovieScene/LevelSequenceEditor/Source/LevelSequenceEditor/Private/Misc/LevelSequenceEditorSpawnRegister.cpp b/Engine/Plugins/MovieScene/LevelSequenceEditor/Source/LevelSequenceEditor/Private/Misc/LevelSequenceEditorSpawnRegister.cpp
--- a/Engine/Plugins/MovieScene/LevelSequenceEditor/Source/LevelSequenceEditor/Private/Misc/LevelSequenceEditorSpawnRegister.cpp
+++ b/Engine/Plugins/MovieScene/LevelSequenceEditor/Source/LevelSequenceEditor/Private/Misc/LevelSequenceEditorSpawnRegister.cpp
@@ -113,7 +113,7 @@ void FLevelSequenceEditorSpawnRegister::PreDestroyObject(UObject& Object, const
 void FLevelSequenceEditorSpawnRegister::SaveDefaultSpawnableState(FMovieSceneSpawnable& Spawnable, FMovieSceneSequenceIDRef TemplateID, IMovieScenePlayer& Player)
 {
 	const FMovieSceneEvaluationTemplateInstance* Template = Player.GetEvaluationTemplate().GetInstance(TemplateID);
-	UMovieSceneSequence* Sequence = Template ? Template->Sequence : nullptr;
+	UMovieSceneSequence* Sequence = Template ? Template->Sequence.Get() : nullptr;
 
 	UObject* Object = FindSpawnedObject(Spawnable.GetGuid(), TemplateID);
 	if (!Object || !Sequence)
diff --git a/Engine/Source/Editor/Sequencer/Private/Sequencer.cpp b/Engine/Source/Editor/Sequencer/Private/Sequencer.cpp
--- a/Engine/Source/Editor/Sequencer/Private/Sequencer.cpp
+++ b/Engine/Source/Editor/Sequencer/Private/Sequencer.cpp
@@ -2112,8 +2112,10 @@ void FSequencer::AddReferencedObjects( FReferenceCollector& Collector )
 		// Sequencer references all sub movie scene sequences contained within the root
 		for (auto& Pair : RootTemplateInstance.GetSubInstances())
 		{
-			UMovieSceneSequence* Sequence = Pair.Value.Sequence;
-			Collector.AddReferencedObject(Sequence);
+			if (UMovieSceneSequence* Sequence = Pair.Value.Sequence.Get())
+			{
+				Collector.AddReferencedObject(Sequence);
+			}
 		}
 	}
 }
diff --git a/Engine/Source/Runtime/LevelSequence/Public/LevelSequencePlayer.h b/Engine/Source/Runtime/LevelSequence/Public/LevelSequencePlayer.h
--- a/Engine/Source/Runtime/LevelSequence/Public/LevelSequencePlayer.h
+++ b/Engine/Source/Runtime/LevelSequence/Public/LevelSequencePlayer.h
@@ -319,7 +319,6 @@ private:
 	/** The last view target to reset to when updating camera cuts to null */
 	TWeakObjectPtr<AActor> LastViewTarget;
 
-	UPROPERTY(transient)
 	FMovieSceneRootEvaluationTemplateInstance RootTemplateInstance;
 
 	FMovieScenePlaybackPosition PlayPosition;
diff --git a/Engine/Source/Runtime/MovieScene/Private/Evaluation/MovieSceneEvaluationTemplateInstance.cpp b/Engine/Source/Runtime/MovieScene/Private/Evaluation/MovieSceneEvaluationTemplateInstance.cpp
--- a/Engine/Source/Runtime/MovieScene/Private/Evaluation/MovieSceneEvaluationTemplateInstance.cpp
+++ b/Engine/Source/Runtime/MovieScene/Private/Evaluation/MovieSceneEvaluationTemplateInstance.cpp
@@ -51,16 +51,16 @@ void FMovieSceneRootEvaluationTemplateInstance::Reset()
 {
 	if (TemplateStore->AreTemplatesVolatile())
 	{
-		if (RootInstance.Sequence)
+		if (UMovieSceneSequence* Sequence = RootInstance.Sequence.Get())
 		{
-			RootInstance.Sequence->OnSignatureChanged().RemoveAll(this);
+			Sequence->OnSignatureChanged().RemoveAll(this);
 		}
 
 		for (auto& Pair : SubInstances)
 		{
-			if (Pair.Value.Sequence)
+			if (UMovieSceneSequence* Sequence = Pair.Value.Sequence.Get())
 			{
-				Pair.Value.Sequence->OnSignatureChanged().RemoveAll(this);
+				Sequence->OnSignatureChanged().RemoveAll(this);
 			}
 		}
 	}
@@ -137,9 +137,9 @@ void FMovieSceneRootEvaluationTemplateInstance::Finish(IMovieScenePlayer& Player
 
 void FMovieSceneRootEvaluationTemplateInstance::Evaluate(FMovieSceneContext Context, IMovieScenePlayer& Player, FMovieSceneSequenceID OverrideRootID)
 {
-	if (bIsDirty && RootSequence)
+	if (bIsDirty && RootSequence.IsValid())
 	{
-		Initialize(*RootSequence, Player);
+		Initialize(*RootSequence.Get(), Player);
 	}
 
 	SCOPE_CYCLE_COUNTER(MovieSceneEval_EntireEvaluationCost);
diff --git a/Engine/Source/Runtime/MovieScene/Public/Evaluation/MovieSceneEvaluationTemplateInstance.h b/Engine/Source/Runtime/MovieScene/Public/Evaluation/MovieSceneEvaluationTemplateInstance.h
--- a/Engine/Source/Runtime/MovieScene/Public/Evaluation/MovieSceneEvaluationTemplateInstance.h
+++ b/Engine/Source/Runtime/MovieScene/Public/Evaluation/MovieSceneEvaluationTemplateInstance.h
@@ -4,7 +4,6 @@
 
 #include "MovieSceneFwd.h"
 #include "MovieSceneEvaluationTemplate.h"
-#include "MovieSceneEvaluationTemplateInstance.generated.h"
 
 class FMovieSceneSequenceInstance;
 
@@ -13,7 +12,7 @@ class FMovieSceneSequenceInstance;
  */
 struct FMovieSceneEvaluationTemplateInstance
 {
-	FMovieSceneEvaluationTemplateInstance() : Sequence(nullptr), Template(nullptr), SequenceID(MovieSceneSequenceID::Invalid) {}
+	FMovieSceneEvaluationTemplateInstance() : Template(nullptr), SequenceID(MovieSceneSequenceID::Invalid) {}
 
 	/**
 	 * Constructor for root instances
@@ -57,7 +56,7 @@ struct FMovieSceneEvaluationTemplateInstance
 public:
 
 	/** The sequence that the sub section references */
-	UMovieSceneSequence* Sequence;
+	TWeakObjectPtr<UMovieSceneSequence> Sequence;
 
 	/** Transform that transforms a given time from the sequences outer space, to its authored space. */
 	FMovieSceneSequenceTransform RootToSequenceTransform;
@@ -75,12 +74,9 @@ public:
 /**
  * Root evaluation template instance used to play back any sequence
  */
-USTRUCT()
 struct FMovieSceneRootEvaluationTemplateInstance
 {
 public:
-	GENERATED_BODY()
-
 	MOVIESCENE_API FMovieSceneRootEvaluationTemplateInstance();
 	MOVIESCENE_API ~FMovieSceneRootEvaluationTemplateInstance();
 
@@ -95,7 +91,7 @@ public:
 	 */
 	bool IsValid() const
 	{
-		return RootSequence && RootInstance.Template;
+		return RootSequence.Get() && RootInstance.Template;
 	}
 
 	/**
@@ -174,7 +170,7 @@ public:
 	{
 		if (const FMovieSceneEvaluationTemplateInstance* Instance = GetInstance(SequenceID))
 		{
-			return Instance->Sequence;
+			return Instance->Sequence.Get();
 		}
 		return nullptr;
 	}
@@ -257,8 +253,7 @@ private:
 
 private:
 
-	UPROPERTY()
-	UMovieSceneSequence* RootSequence;
+	TWeakObjectPtr<UMovieSceneSequence> RootSequence;
 
 	FMovieSceneEvaluationTemplateInstance RootInstance;
 
@@ -282,5 +277,4 @@ private:
 
 	/** Event that is triggered on update */
 	FOnUpdated OnUpdatedEvent;
-};
-template<> struct TStructOpsTypeTraits<FMovieSceneRootEvaluationTemplateInstance> : public TStructOpsTypeTraitsBase { enum { WithCopy = false }; };
\ No newline at end of file
+};
\ No newline at end of file
diff --git a/Engine/Source/Runtime/UMG/Public/Animation/UMGSequencePlayer.h b/Engine/Source/Runtime/UMG/Public/Animation/UMGSequencePlayer.h
--- a/Engine/Source/Runtime/UMG/Public/Animation/UMGSequencePlayer.h
+++ b/Engine/Source/Runtime/UMG/Public/Animation/UMGSequencePlayer.h
@@ -71,7 +71,6 @@ private:
 	/** The user widget this sequence is animating */
 	TWeakObjectPtr<UUserWidget> UserWidget;
 
-	UPROPERTY(transient)
 	FMovieSceneRootEvaluationTemplateInstance RootTemplateInstance;
 
 	/** Time range of the animation */
