Copying //UE4/Dev-Sequencer to //UE4/Dev-Main (Source: //UE4/Dev-Sequencer @ 3237992)

diff --git a/Engine/Source/Runtime/MediaUtils/Public/MediaPlayerBase.h b/Engine/Source/Runtime/MediaUtils/Public/MediaPlayerBase.h
new file mode 100644
--- /dev/null
+++ b/Engine/Source/Runtime/MediaUtils/Public/MediaPlayerBase.h
@@ -0,0 +1,392 @@
+// Copyright 1998-2016 Epic Games, Inc. All Rights Reserved.
+
+#pragma once
+
+#include "CoreMinimal.h"
+
+
+enum class EMediaEvent;
+enum class EMediaTrackType;
+class IMediaAudioSink;
+class IMediaBinarySink;
+class IMediaOptions;
+class IMediaOverlaySink;
+class IMediaPlayer;
+class IMediaTextureSink;
+
+
+/**
+ * Base class for common shared media player functionality.
+ */
+class MEDIAUTILS_API FMediaPlayerBase
+{
+public:
+
+	/** Virtual destructor. */
+	virtual ~FMediaPlayerBase();
+
+public:
+
+	/**
+	 * Check whether media playback can be paused right now.
+	 *
+	 * Playback can be paused if the media supports pausing and if it is currently playing.
+	 *
+	 * @return true if pausing playback can be paused, false otherwise.
+	 * @see CanPlay, Pause
+	 */
+	bool CanPause() const;
+
+	/**
+	 * Check whether the specified URL can be played by this player.
+	 *
+	 * If a desired player name is set for this player, it will only check
+	 * whether that particular player type can play the specified URL.
+	 *
+	 * @param Url The URL to check.
+	 * @param Options Optional media parameters.
+	 * @see CanPlaySource, SetDesiredPlayerName
+	 */
+	bool CanPlayUrl(const FString& Url, const IMediaOptions& Options);
+
+	/**
+	 * Close the currently open media, if any.
+	 */
+	void Close();
+
+	/**
+	 * Get the media's duration.
+	 *
+	 * @return A time span representing the duration.
+	 * @see GetTime, Seek
+	 */
+	FTimespan GetDuration() const;
+
+	/**
+	 * Get the supported forward playback rates.
+	 *
+	 * @param Unthinned Whether the rates are for unthinned playback (default = true).
+	 * @return The range of supported rates.
+	 * @see GetReverseRates
+	 */
+	TRange<float> GetForwardRates(bool Unthinned = true);
+
+	/**
+	 * Get debug information about the player and currently opened media.
+	 *
+	 * @return Information string.
+	 * @see GetStats
+	 */
+	FString GetInfo() const;
+
+	/**
+	 * Get the number of tracks of the given type.
+	 *
+	 * @param TrackType The type of media tracks.
+	 * @return Number of tracks.
+	 * @see GetSelectedTrack, SelectTrack
+	 */
+	int32 GetNumTracks(EMediaTrackType TrackType) const;
+
+	/**
+	 * Get the low-level player associated with this object.
+	 *
+	 * @return The player, or nullptr if no player was created.
+	 */
+	TSharedPtr<IMediaPlayer> GetNativePlayer() const
+	{
+		return NativePlayer;
+	}
+
+	/**
+	 * Get the name of the current native media player.
+	 *
+	 * @return Player name, or NAME_None if not available.
+	 */
+	FName GetPlayerName() const;
+
+	/**
+	 * Get the media's current playback rate.
+	 *
+	 * @return The playback rate.
+	 * @see SetRate, SupportsRate
+	 */
+	float GetRate() const;
+
+	/**
+	 * Get the supported reverse playback rates.
+	 *
+	 * @param Unthinned Whether the rates are for unthinned playback (default = true).
+	 * @return The range of supported rates.
+	 * @see GetForwardRates
+	 */
+	TRange<float> GetReverseRates(bool Unthinned = true);
+
+	/**
+	 * Get the index of the currently selected track of the given type.
+	 *
+	 * @param TrackType The type of track to get.
+	 * @return The index of the selected track, or INDEX_NONE if no track is active.
+	 * @see GetNumTracks, SelectTrack
+	 */
+	int32 GetSelectedTrack(EMediaTrackType TrackType) const;
+
+	/**
+	 * Get playback statistics information.
+	 *
+	 * @return Information string.
+	 * @see GetInfo
+	 */
+	FString GetStats() const;
+
+	/**
+	 * Get the media's current playback time.
+	 *
+	 * @return Playback time.
+	 * @see GetDuration, Seek
+	 */
+	FTimespan GetTime() const;
+
+	/**
+	 * Get the human readable name of the specified track.
+	 *
+	 * @param TrackType The type of track.
+	 * @param TrackIndex The index of the track.
+	 * @return Display name.
+	 * @see GetNumTracks, GetTrackLanguage
+	 */
+	FText GetTrackDisplayName(EMediaTrackType TrackType, int32 TrackIndex) const;
+
+	/**
+	 * Get the language tag of the specified track.
+	 *
+	 * @param TrackType The type of track.
+	 * @param TrackIndex The index of the track.
+	 * @return Language tag, i.e. "en-US" for English, or "und" for undefined.
+	 * @see GetNumTracks, GetTrackDisplayName
+	 */
+	FString GetTrackLanguage(EMediaTrackType TrackType, int32 TrackIndex) const;
+
+	/**
+	 * Get the URL of the currently loaded media, if any.
+	 *
+	 * @return Media URL, or empty string if no media was loaded.
+	 */
+	const FString& GetUrl() const
+	{
+		return CurrentUrl;
+	}
+
+	/**
+	 * Checks whether playback is looping.
+	 *
+	 * @return true if looping, false otherwise.
+	 * @see SetLooping
+	 */
+	bool IsLooping() const;
+
+	/**
+	 * Checks whether playback is currently paused.
+	 *
+	 * @return true if playback is paused, false otherwise.
+	 * @see CanPause, IsPlaying, IsReady, Pause
+	 */
+	bool IsPaused() const;
+
+	/**
+	 * Checks whether playback has started.
+	 *
+	 * @return true if playback has started, false otherwise.
+	 * @see CanPlay, IsPaused, IsReady, Play
+	 */
+	bool IsPlaying() const;
+
+	/**
+	 * Checks whether the media is currently opening or buffering.
+	 *
+	 * @return true if playback is being prepared, false otherwise.
+	 * @see CanPlay, IsPaused, IsReady, Play
+	 */
+	bool IsPreparing() const;
+
+	/**
+	 * Checks whether media is ready for playback.
+	 *
+	 * A player is ready for playback if it has a media source opened that
+	 * finished preparing and is not in an error state.
+	 *
+	 * @return true if media is ready, false otherwise.
+	 * @see IsPaused, IsPlaying, Stop
+	 */
+	bool IsReady() const;
+
+	/**
+	 * Open a media source from a URL with optional parameters.
+	 *
+	 * @param Url The URL of the media to open (file name or web address).
+	 * @param Options Optional media parameters.
+	 * @return true if the media is being opened, false otherwise.
+	 */
+	bool Open(const FString& Url, const IMediaOptions& Options);
+
+	/**
+	 * Seeks to the specified playback time.
+	 *
+	 * @param Time The playback time to set.
+	 * @return true on success, false otherwise.
+	 * @see GetTime, Rewind
+	 */
+	bool Seek(const FTimespan& Time);
+
+	/**
+	 * Select the active track of the given type.
+	 *
+	 * Only one track of a given type can be active at any time.
+	 *
+	 * @param TrackType The type of track to select.
+	 * @param TrackIndex The index of the track to select, or INDEX_NONE to deselect.
+	 * @return true if the track was selected, false otherwise.
+	 * @see GetNumTracks, GetSelectedTrack
+	 */
+	bool SelectTrack(EMediaTrackType TrackType, int32 TrackIndex);
+
+	/**
+	 * Set the sink to receive data from audio tracks.
+	 *
+	 * @param Sink The sink to set, or nullptr to clear.
+	 * @see SetMetadataSink, SetOverlaySink, SetVideoSink
+	 */
+	void SetAudioSink(IMediaAudioSink* Sink);
+
+	/**
+	 * Enables or disables playback looping.
+	 *
+	 * @param Looping Whether playback should be looped.
+	 * @return true on success, false otherwise.
+	 * @see IsLooping
+	 */
+	bool SetLooping(bool Looping);
+
+	/**
+	 * Set the sink to receive data from binary tracks.
+	 *
+	 * @param Sink The sink to set, or nullptr to clear.
+	 * @see SetAudioSink, SetOverlaySink, SetVideoSink
+	 */
+	void SetMetadataSink(IMediaBinarySink* Sink);
+
+	/**
+	 * Set the sink to receive data from caption and subtitle tracks.
+	 *
+	 * @param Sink The sink to set, or nullptr to clear.
+	 * @see SetAudioSink, SetMetadataSink, SetVideoSink
+	 */
+	void SetOverlaySink(IMediaOverlaySink* Sink);
+
+	/**
+	 * Changes the media's playback rate.
+	 *
+	 * @param Rate The playback rate to set.
+	 * @return true on success, false otherwise.
+	 * @see GetRate, SupportsRate
+	 */
+	bool SetRate(float Rate);
+
+	/**
+	 * Set the sink to receive data from the active video track.
+	 *
+	 * @param Sink The sink to set, or nullptr to clear.
+	 * @see SetAudioSink
+	 */
+	void SetVideoSink(IMediaTextureSink* Sink);
+
+	/**
+	 * Checks whether the specified playback rate is supported.
+	 *
+	 * @param Rate The playback rate to check.
+	 * @param Unthinned Whether no frames should be dropped at the given rate.
+	 * @see SupportsScrubbing, SupportsSeeking
+	 */
+	bool SupportsRate(float Rate, bool Unthinned) const;
+
+	/**
+	 * Checks whether the currently loaded media supports scrubbing.
+	 *
+	 * @return true if scrubbing is supported, false otherwise.
+	 * @see SupportsRate, SupportsSeeking
+	 */
+	bool SupportsScrubbing() const;
+
+	/**
+	 * Checks whether the currently loaded media can jump to a certain position.
+	 *
+	 * @return true if seeking is supported, false otherwise.
+	 * @see SupportsRate, SupportsScrubbing
+	 */
+	bool SupportsSeeking() const;
+
+	/**
+	 * Tick the media player logic.
+	 *
+	 * @param DeltaTime Time since last tick.
+	 * @see TickVideo
+	 */
+	void TickPlayer(float DeltaTime);
+
+	/**
+	 * Tick the media player's video code.
+	 *
+	 * @param DeltaTime Time since last tick.
+	 * @see TickPlayer
+	 */
+	void TickVideo(float DeltaTime);
+
+public:
+
+	/** Get an event delegate that is invoked when a media event occurred. */
+	DECLARE_EVENT_OneParam(FMediaPlayerBase, FOnMediaEvent, EMediaEvent /*Event*/)
+	FOnMediaEvent& OnMediaEvent()
+	{
+		return MediaEvent;
+	}
+
+public:
+
+	/** Name of the desired native player, if any. */
+	FName DesiredPlayerName;
+
+protected:
+
+	/**
+	 * Find a player that can play the specified media URL.
+	 *
+	 * @param Url The URL to play.
+	 * @param Options The media options for the URL.
+	 * @return The player if found, or nullptr otherwise.
+	 */
+	TSharedPtr<IMediaPlayer> FindPlayerForUrl(const FString& Url, const IMediaOptions& Options);
+
+	/** Select the default media tracks. */
+	void SelectDefaultTracks();
+
+protected:
+
+	/** Holds the URL of the currently loaded media source. */
+	FString CurrentUrl;
+
+	/** The low-level player used to play the media source. */
+	TSharedPtr<IMediaPlayer> NativePlayer;
+
+private:
+
+	/** Callback for when a media event occurred in the player. */
+	void HandlePlayerMediaEvent(EMediaEvent Event);
+
+private:
+
+	/** Synchronizes access to Player. */
+	FCriticalSection CriticalSection;
+
+	/** An event delegate that is invoked when a media event occurred. */
+	FOnMediaEvent MediaEvent;
+};
diff --git a/Engine/Source/Runtime/Messaging/Private/Bus/MessageRouter.h b/Engine/Source/Runtime/Messaging/Private/Bus/MessageRouter.h
--- a/Engine/Source/Runtime/Messaging/Private/Bus/MessageRouter.h
+++ b/Engine/Source/Runtime/Messaging/Private/Bus/MessageRouter.h
@@ -201,7 +201,7 @@ private:
 		{
 			const FTimespan Difference = Other.Context->GetTimeSent() - Context->GetTimeSent();
 
-			if (Difference == FTimespan::Zero())
+			if (Difference.IsZero())
 			{
 				return (Sequence < Other.Sequence);
 			}
diff --git a/Engine/Source/Runtime/MovieScene/Private/Compilation/MovieSceneSegmentCompiler.cpp b/Engine/Source/Runtime/MovieScene/Private/Compilation/MovieSceneSegmentCompiler.cpp
--- a/Engine/Source/Runtime/MovieScene/Private/Compilation/MovieSceneSegmentCompiler.cpp
+++ b/Engine/Source/Runtime/MovieScene/Private/Compilation/MovieSceneSegmentCompiler.cpp
@@ -22,16 +22,19 @@ void FMovieSceneSegmentCompilerRules::ProcessSegments(TArray<FMovieSceneSegment>
 
 	int32 Index = 0;
 
+	// First off, add an open segment before the first segment if necessary/possible
 	if (!Segments[0].Range.GetLowerBound().IsOpen())
 	{
 		InsertSegment(Segments, 0, SourceData);
-		++Index;
 	}
+	++Index;
 
-	for (; Index < Segments.Num() - 1; ++Index)
+	// Then add any segments inbetween valid segments
+	for (; Index < Segments.Num(); ++Index)
 	{
 		if (InsertSegment(Segments, Index, SourceData))
 		{
+			// Increment the counter to account for the added segment
 			++Index;
 		}
 	}
diff --git a/Engine/Source/Runtime/MovieScene/Private/Evaluation/MovieSceneEvaluationTemplate.cpp b/Engine/Source/Runtime/MovieScene/Private/Evaluation/MovieSceneEvaluationTemplate.cpp
--- a/Engine/Source/Runtime/MovieScene/Private/Evaluation/MovieSceneEvaluationTemplate.cpp
+++ b/Engine/Source/Runtime/MovieScene/Private/Evaluation/MovieSceneEvaluationTemplate.cpp
@@ -26,7 +26,7 @@ void FCachedMovieSceneEvaluationTemplate::Regenerate(const FMovieSceneTrackCompi
 {
 	if (IsOutOfDate(Params))
 	{
-		RegenerateImpl(Params);
+		ForceRegenerate(Params);
 	}
 }
 
@@ -34,11 +34,7 @@ void FCachedMovieSceneEvaluationTemplate::ForceRegenerate(const FMovieSceneTrack
 {
 	// Reassign the template back to default
 	static_cast<FMovieSceneEvaluationTemplate&>(*this) = FMovieSceneEvaluationTemplate();
-	RegenerateImpl(Params);
-}
 
-void FCachedMovieSceneEvaluationTemplate::RegenerateImpl(const FMovieSceneTrackCompilationParams& Params)
-{
 	CachedSignatures.Reset();
 	CachedCompilationParams = Params;
 
diff --git a/Engine/Source/Runtime/MovieScene/Private/Evaluation/MovieSceneEvaluationTemplateInstance.cpp b/Engine/Source/Runtime/MovieScene/Private/Evaluation/MovieSceneEvaluationTemplateInstance.cpp
--- a/Engine/Source/Runtime/MovieScene/Private/Evaluation/MovieSceneEvaluationTemplateInstance.cpp
+++ b/Engine/Source/Runtime/MovieScene/Private/Evaluation/MovieSceneEvaluationTemplateInstance.cpp
@@ -1,8 +1,11 @@
 // Copyright 1998-2017 Epic Games, Inc. All Rights Reserved.
 
 #include "Evaluation/MovieSceneEvaluationTemplateInstance.h"
+#include "IMovieScenePlayer.h"
 #include "MovieSceneSequence.h"
-
+#include "MovieSceneSequence.h"
+#include "MovieSceneSequenceInstance.h"
+#include "MovieSceneExecutionTokens.h"
 
 DECLARE_CYCLE_STAT(TEXT("Entire Evaluation Cost"), MovieSceneEval_EntireEvaluationCost, STATGROUP_MovieSceneEval);
 DECLARE_CYCLE_STAT(TEXT("Gather Entries For Frame"), MovieSceneEval_GatherEntries, STATGROUP_MovieSceneEval);
diff --git a/Engine/Source/Runtime/MovieScene/Private/Evaluation/MovieSceneEvaluationTrack.cpp b/Engine/Source/Runtime/MovieScene/Private/Evaluation/MovieSceneEvaluationTrack.cpp
--- a/Engine/Source/Runtime/MovieScene/Private/Evaluation/MovieSceneEvaluationTrack.cpp
+++ b/Engine/Source/Runtime/MovieScene/Private/Evaluation/MovieSceneEvaluationTrack.cpp
@@ -1,6 +1,8 @@
 // Copyright 1998-2017 Epic Games, Inc. All Rights Reserved.
 
 #include "Evaluation/MovieSceneEvaluationTrack.h"
+#include "MovieSceneExecutionTokens.h"
+
 
 FMovieSceneEvaluationTrack::FMovieSceneEvaluationTrack()
 {
diff --git a/Engine/Source/Runtime/MovieScene/Private/Evaluation/MovieSceneExecutionTokens.cpp b/Engine/Source/Runtime/MovieScene/Private/Evaluation/MovieSceneExecutionTokens.cpp
--- a/Engine/Source/Runtime/MovieScene/Private/Evaluation/MovieSceneExecutionTokens.cpp
+++ b/Engine/Source/Runtime/MovieScene/Private/Evaluation/MovieSceneExecutionTokens.cpp
@@ -1,10 +1,13 @@
 // Copyright 1998-2017 Epic Games, Inc. All Rights Reserved.
 
 #include "Evaluation/MovieSceneExecutionTokens.h"
+#include "IMovieScenePlayer.h"
+
 
 DECLARE_CYCLE_STAT(TEXT("Apply Execution Tokens"), MovieSceneEval_ApplyExecutionTokens, STATGROUP_MovieSceneEval);
 DECLARE_CYCLE_STAT(TEXT("Apply Execution Token"), MovieSceneEval_ApplyExecutionToken, STATGROUP_MovieSceneEval);
 
+
 void FMovieSceneExecutionTokens::Apply(IMovieScenePlayer& Player)
 {
 	MOVIESCENE_DETAILED_SCOPE_CYCLE_COUNTER(MovieSceneEval_ApplyExecutionTokens);
diff --git a/Engine/Source/Runtime/MovieScene/Private/Evaluation/MovieSceneLegacyTrackInstanceTemplate.cpp b/Engine/Source/Runtime/MovieScene/Private/Evaluation/MovieSceneLegacyTrackInstanceTemplate.cpp
--- a/Engine/Source/Runtime/MovieScene/Private/Evaluation/MovieSceneLegacyTrackInstanceTemplate.cpp
+++ b/Engine/Source/Runtime/MovieScene/Private/Evaluation/MovieSceneLegacyTrackInstanceTemplate.cpp
@@ -2,6 +2,8 @@
 
 #include "Evaluation/MovieSceneLegacyTrackInstanceTemplate.h"
 #include "Evaluation/MovieSceneEvaluationTemplateInstance.h"
+#include "IMovieScenePlayer.h"
+
 
 /** Legacy tracks just get a new ID each time (so they always save/restore state, regardless of whether a similar track has already animated the object. That was the old behaviour) */
 struct FMovieSceneLegacyAnimTypeID : FMovieSceneAnimTypeID
diff --git a/Engine/Source/Runtime/MovieScene/Private/MovieScene.cpp b/Engine/Source/Runtime/MovieScene/Private/MovieScene.cpp
--- a/Engine/Source/Runtime/MovieScene/Private/MovieScene.cpp
+++ b/Engine/Source/Runtime/MovieScene/Private/MovieScene.cpp
@@ -4,6 +4,9 @@
 #include "MovieSceneTrack.h"
 #include "MovieSceneFolder.h"
 #include "Evaluation/MovieSceneEvaluationCustomVersion.h"
+#include "Evaluation/MovieSceneEvaluationCustomVersion.h"
+#include "Compilation/MovieSceneSegmentCompiler.h"
+#include "SequencerObjectVersion.h"
 
 /* UMovieScene interface
  *****************************************************************************/
@@ -12,6 +15,9 @@ UMovieScene::UMovieScene(const FObjectInitializer& ObjectInitializer)
 	: Super(ObjectInitializer)
 	, SelectionRange(FFloatRange::Empty())
 	, PlaybackRange(FFloatRange::Empty())
+#if WITH_EDITORONLY_DATA
+	, bPlaybackRangeLocked(false)
+#endif
 	, bForceFixedFrameIntervalPlayback(false)
 	, FixedFrameInterval(0.0f)
 	, InTime_DEPRECATED(FLT_MAX)
@@ -364,6 +370,17 @@ void UMovieScene::SetViewRange(float Start, float End)
 #endif
 }
 
+#if WITH_EDITORONLY_DATA
+bool UMovieScene::IsPlaybackRangeLocked() const
+{
+	return bPlaybackRangeLocked;
+}
+
+void UMovieScene::SetPlaybackRangeLocked(bool bLocked)
+{
+	bPlaybackRangeLocked = bLocked;
+}
+#endif
 
 bool UMovieScene::GetForceFixedFrameIntervalPlayback() const
 {
@@ -472,13 +489,6 @@ bool UMovieScene::AddGivenTrack(UMovieSceneTrack* InTrack, const FGuid& ObjectGu
 	{
 		if (Binding.GetObjectGuid() == ObjectGuid)
 		{
-			for (auto& Track : Binding.GetTracks())
-			{
-				if (Track->GetTrackName() == InTrack->GetTrackName())
-				{
-					return false;
-				}
-			}
 			InTrack->Rename(nullptr, this);
 			check(InTrack);
 			Binding.AddTrack(*InTrack);
@@ -492,7 +502,7 @@ bool UMovieScene::RemoveTrack(UMovieSceneTrack& Track)
 {
 	Modify();
 
-bool bAnythingRemoved = false;
+	bool bAnythingRemoved = false;
 
 	for (auto& Binding : ObjectBindings)
 	{
@@ -674,12 +684,108 @@ void UMovieScene::UpgradeTimeRanges()
 }
 
 
+void UMovieScene::UpgradeTrackRows()
+{
+	int32 NumMasterTracks = MasterTracks.Num();
+	for (int32 MasterTrackIndex = 0; MasterTrackIndex < NumMasterTracks; ++MasterTrackIndex)
+	{
+		UpgradeTrackRow(MasterTracks[MasterTrackIndex]);
+	}
+
+	for (auto& Binding : ObjectBindings)
+	{
+		int32 NumTracks = Binding.GetTracks().Num();
+		for (int32 TrackIndex = 0; TrackIndex < NumTracks; ++TrackIndex)
+		{
+			UpgradeTrackRow(Binding.GetTracks()[TrackIndex]);
+		}
+	}
+}
+
+void UMovieScene::UpgradeTrackRow(UMovieSceneTrack* InTrack)
+{
+	if (GetLinkerCustomVersion(FSequencerObjectVersion::GUID) >= FSequencerObjectVersion::ConvertMultipleRowsToTracks)
+	{
+		return;
+	}
+
+	// Skeletal animation tracks and audio tracks went through a brief period of expanding their rows out to multiple tracks
+	// but that upgrade path was short-lived since we added better support for such tracks in the sequencer front-end.
+	// it was impossible to reliably implement such upgrades due the legacy 'evaluate nearest section' behavior on each of the duplicated tracks
+	static const FName SkeletalAnimationTrack("MovieSceneSkeletalAnimationTrack");
+	static const FName AudioTrack("MovieSceneAudioTrack");
+
+	FName TrackName = InTrack->GetClass()->GetFName();
+	if (TrackName != SkeletalAnimationTrack && TrackName != AudioTrack)
+	{
+		return;
+	}
+
+	// Even still, there is a small amount of upgrade to be done to ensure that blending does not occur on sections that were
+	// previously overridden by another, so we upgrade these cases by weighting such areas with keys at 0 or 1
+
+	// If there aren't sections on multiple rows, disregard
+	auto ContainsMultipleRows = [](UMovieSceneSection* InSection){ return InSection->GetRowIndex() > 0; };
+	if (!InTrack->GetAllSections().ContainsByPredicate(ContainsMultipleRows))
+	{
+		return;
+	}
+
+	// Deal with overlapping sections. For example, skeletal animation sections should be weighted along their evaluation bounds.
+	//
+	// For example,
+	//
+	//          [----A----]
+	// [-----B-----]
+	//
+	//
+	// Evaluation segments result in:
+	//
+	//          [----A----]
+	// [---B---]
+	//
+	// OR (depending upon order in the original track):
+	//
+	//              [--A--]
+	// [------B----]
+	//
+
+	TInlineValue<FMovieSceneSegmentCompilerRules> RowCompilerRules = InTrack->GetRowCompilerRules();
+	FMovieSceneTrackCompiler::FRows TrackRows(InTrack->GetAllSections(), RowCompilerRules.GetPtr(nullptr));
+	FMovieSceneTrackEvaluationField EvaluationField = FMovieSceneTrackCompiler().Compile(TrackRows.Rows, InTrack->GetTrackCompilerRules().GetPtr(nullptr));
+
+	for (const FMovieSceneSegment& EvalSegment : EvaluationField.Segments)
+	{
+		if (EvalSegment.Range.IsDegenerate() || EvalSegment.Range.IsEmpty())
+		{
+			continue;
+		}
+
+		// Last element in original source section array wins. This is a carry over from legacy behavior whereby the last
+		// section in the array would override others, regardless of where it sat organizationally in the sequence.
+		int32 WinningImplIndex = -1;
+		for (const FSectionEvaluationData& EvalData : EvalSegment.Impls)
+		{
+			WinningImplIndex = FMath::Max(WinningImplIndex, EvalData.ImplIndex);
+		}
+
+		UMovieSceneSection* Section = InTrack->GetAllSections()[WinningImplIndex];
+		Section->ConditionalPostLoad();
+
+		Section->PostLoadUpgradeTrackRow(EvalSegment.Range);
+	}
+
+	InTrack->MarkAsChanged();
+}
+
+
 /* UObject interface
  *****************************************************************************/
 
 void UMovieScene::PostLoad()
 {
 	UpgradeTimeRanges();
+	UpgradeTrackRows();
 
 	for (FMovieSceneSpawnable& Spawnable : Spawnables)
 	{
@@ -751,9 +857,10 @@ void UMovieScene::ReplaceBinding(const FGuid& OldGuid, const FGuid& NewGuid, con
 			Binding.SetName(Name);
 
 			// Changing a binding guid invalidates any tracks contained within the binding
+			// Make sure they are written into the transaction buffer by calling modify
 			for (UMovieSceneTrack* Track : Binding.GetTracks())
 			{
-				Track->MarkAsChanged();
+				Track->Modify();
 			}
 			break;
 		}
diff --git a/Engine/Source/Runtime/MovieScene/Private/MovieSceneBindingOverrides.cpp b/Engine/Source/Runtime/MovieScene/Private/MovieSceneBindingOverrides.cpp
new file mode 100644
--- /dev/null
+++ b/Engine/Source/Runtime/MovieScene/Private/MovieSceneBindingOverrides.cpp
@@ -0,0 +1,128 @@
+// Copyright 1998-2016 Epic Games, Inc. All Rights Reserved.
+
+#include "MovieSceneBindingOverrides.h"
+
+UMovieSceneBindingOverrides::UMovieSceneBindingOverrides(const FObjectInitializer& Init)
+	: Super(Init)
+	, bLookupDirty(true)
+{
+}
+
+bool UMovieSceneBindingOverrides::LocateBoundObjects(const FGuid& InBindingId, TArray<UObject*, TInlineAllocator<1>>& OutObjects) const
+{
+	if (bLookupDirty)
+	{
+		RebuildLookupMap();
+	}
+
+	bool bAllowDefault = true;
+
+	for (auto It = LookupMap.CreateConstKeyIterator(InBindingId);It;++It)
+	{
+		const FMovieSceneBindingOverrideData& Data = BindingData[It.Value()];
+		UObject* Object = Data.Object.Get();
+
+		if (Data.bOverridesDefault)
+		{
+			bAllowDefault = false;
+		}
+
+		if (Object)
+		{
+			OutObjects.Add(Object);
+		}
+	}
+
+	return bAllowDefault;
+}
+
+void UMovieSceneBindingOverrides::SetBinding(FMovieSceneObjectBindingPtr Binding, const TArray<UObject*>& Objects, bool bAllowBindingsFromAsset)
+{
+	ResetBinding(Binding);
+
+	for (UObject* Object : Objects)
+	{
+		if (!Object)
+		{
+			continue;
+		}
+
+		LookupMap.Add(Binding.Guid, BindingData.Num());
+
+		FMovieSceneBindingOverrideData NewBinding;
+		NewBinding.ObjectBindingId = Binding;
+		NewBinding.Object = Object;
+		NewBinding.bOverridesDefault = !bAllowBindingsFromAsset;
+		BindingData.Add(NewBinding);
+	}
+}
+
+void UMovieSceneBindingOverrides::AddBinding(FMovieSceneObjectBindingPtr Binding, UObject* Object, bool bAllowBindingsFromAsset)
+{
+	if (Object)
+	{
+		LookupMap.Add(Binding.Guid, BindingData.Num());
+
+		FMovieSceneBindingOverrideData NewBinding;
+		NewBinding.ObjectBindingId = Binding;
+		NewBinding.Object = Object;
+		NewBinding.bOverridesDefault = !bAllowBindingsFromAsset;
+		BindingData.Add(NewBinding);
+	}
+}
+
+void UMovieSceneBindingOverrides::RemoveBinding(FMovieSceneObjectBindingPtr Binding, UObject* Object)
+{
+	int32 NumRemoved = BindingData.RemoveAll([=](const FMovieSceneBindingOverrideData& InBindingData){
+		return InBindingData.Object == Object && Binding.Guid == InBindingData.ObjectBindingId.Guid;
+	});
+
+	if (NumRemoved)
+	{
+		bLookupDirty = true;
+	}
+}
+
+void UMovieSceneBindingOverrides::ResetBinding(FMovieSceneObjectBindingPtr Binding)
+{
+	int32 NumRemoved = BindingData.RemoveAll([=](const FMovieSceneBindingOverrideData& InBindingData){
+		return Binding.Guid == InBindingData.ObjectBindingId.Guid;
+	});
+
+	if (NumRemoved)
+	{
+		bLookupDirty = true;
+	}
+}
+
+void UMovieSceneBindingOverrides::ResetBindings()
+{
+	if (BindingData.Num())
+	{
+		BindingData.Reset();
+		LookupMap.Reset();
+		bLookupDirty = false;
+	}
+}
+
+void UMovieSceneBindingOverrides::RebuildLookupMap() const
+{
+	LookupMap.Reset();
+
+	for (int32 Index = 0; Index < BindingData.Num(); ++Index)
+	{
+		LookupMap.Add(BindingData[Index].ObjectBindingId.Guid, Index);
+	}
+
+	bLookupDirty = false;
+}
+
+#if WITH_EDITOR
+
+void UMovieSceneBindingOverrides::PostEditChangeProperty(FPropertyChangedEvent& PropertyChangedEvent)
+{
+	Super::PostEditChangeProperty(PropertyChangedEvent);
+	RebuildLookupMap();
+}
+
+#endif
\ No newline at end of file
diff --git a/Engine/Source/Runtime/MovieScene/Private/MovieSceneSection.cpp b/Engine/Source/Runtime/MovieScene/Private/MovieSceneSection.cpp
--- a/Engine/Source/Runtime/MovieScene/Private/MovieSceneSection.cpp
+++ b/Engine/Source/Runtime/MovieScene/Private/MovieSceneSection.cpp
@@ -18,6 +18,18 @@ UMovieSceneSection::UMovieSceneSection(const FObjectInitializer& ObjectInitializ
 { }
 
 
+void UMovieSceneSection::PostInitProperties()
+{
+	// Propagate sub object flags from our outer (track) to ourselves. This is required for sections that are stored on blueprints (archetypes) so that they can be referenced in worlds.
+	if (GetOuter()->HasAnyFlags(RF_ClassDefaultObject|RF_ArchetypeObject))
+	{
+		SetFlags(GetOuter()->GetMaskedFlags(RF_PropagateToSubObjects));
+	}
+	
+	Super::PostInitProperties();
+}
+
+
 bool UMovieSceneSection::TryModify(bool bAlwaysMarkDirty)
 {
 	if (IsLocked())
diff --git a/Engine/Source/Runtime/MovieScene/Private/MovieSceneSequencePlayer.cpp b/Engine/Source/Runtime/MovieScene/Private/MovieSceneSequencePlayer.cpp
new file mode 100644
--- /dev/null
+++ b/Engine/Source/Runtime/MovieScene/Private/MovieSceneSequencePlayer.cpp
@@ -0,0 +1,350 @@
+// Copyright 1998-2016 Epic Games, Inc. All Rights Reserved.
+
+#include "MovieSceneSequencePlayer.h"
+#include "MovieScene.h"
+#include "MovieSceneSequence.h"
+
+
+bool FMovieSceneSequencePlaybackSettings::SerializeFromMismatchedTag( const FPropertyTag& Tag, FArchive& Ar )
+{
+	if (Tag.Type == NAME_StructProperty && Tag.StructName == "LevelSequencePlaybackSettings")
+	{
+		StaticStruct()->SerializeItem(Ar, this, nullptr);
+		return true;
+	}
+
+	return false;
+}
+
+UMovieSceneSequencePlayer::UMovieSceneSequencePlayer(const FObjectInitializer& Init)
+	: Super(Init)
+	, bIsPlaying(false)
+	, bReversePlayback(false)
+	, bIsEvaluating(false)
+	, Sequence(nullptr)
+	, TimeCursorPosition(0.0f)
+	, StartTime(0.f)
+	, EndTime(0.f)
+	, CurrentNumLoops(0)
+{
+}
+
+EMovieScenePlayerStatus::Type UMovieSceneSequencePlayer::GetPlaybackStatus() const
+{
+	return bIsPlaying ? EMovieScenePlayerStatus::Playing : EMovieScenePlayerStatus::Stopped;
+}
+
+FMovieSceneSpawnRegister& UMovieSceneSequencePlayer::GetSpawnRegister()
+{
+	return SpawnRegister.IsValid() ? *SpawnRegister : IMovieScenePlayer::GetSpawnRegister();
+}
+
+void UMovieSceneSequencePlayer::ResolveBoundObjects(const FGuid& InBindingId, FMovieSceneSequenceID SequenceID, UMovieSceneSequence& InSequence, UObject* ResolutionContext, TArray<UObject*, TInlineAllocator<1>>& OutObjects) const
+{
+	bool bAllowDefault = PlaybackSettings.BindingOverrides ? PlaybackSettings.BindingOverrides->LocateBoundObjects(InBindingId, OutObjects) : true;
+
+	if (bAllowDefault)
+	{
+		InSequence.LocateBoundObjects(InBindingId, ResolutionContext, OutObjects);
+	}
+}
+
+void UMovieSceneSequencePlayer::Play()
+{
+	bReversePlayback = false;
+	PlayInternal();
+}
+
+void UMovieSceneSequencePlayer::PlayReverse()
+{
+	bReversePlayback = true;
+	PlayInternal();
+}
+
+void UMovieSceneSequencePlayer::ChangePlaybackDirection()
+{
+	bReversePlayback = !bReversePlayback;
+	PlayInternal();
+}
+
+void UMovieSceneSequencePlayer::PlayLooping(int32 NumLoops)
+{
+	PlaybackSettings.LoopCount = NumLoops;
+	PlayInternal();
+}
+
+void UMovieSceneSequencePlayer::PlayInternal()
+{
+	if (!bIsPlaying)
+	{
+		// Start playing
+		StartPlayingNextTick();
+
+		// Update now
+		bPendingFirstUpdate = false;
+		
+		if (PlaybackSettings.bRestoreState)
+		{
+			PreAnimatedState.EnableGlobalCapture();
+		}
+
+		UMovieSceneSequence* MovieSceneSequence = RootTemplateInstance.GetSequence(MovieSceneSequenceID::Root);
+		TOptional<float> FixedFrameInterval = MovieSceneSequence->GetMovieScene() ? MovieSceneSequence->GetMovieScene()->GetOptionalFixedFrameInterval() : TOptional<float>();
+
+		FMovieSceneEvaluationRange Range = PlayPosition.JumpTo(GetSequencePosition(), FixedFrameInterval);
+		UpdateMovieSceneInstance(Range);
+
+		if (OnPlay.IsBound())
+		{
+			OnPlay.Broadcast();
+		}
+	}
+}
+
+void UMovieSceneSequencePlayer::StartPlayingNextTick()
+{
+	if (bIsPlaying || !Sequence || !CanPlay())
+	{
+		return;
+	}
+
+	// @todo Sequencer playback: Should we recreate the instance every time?
+	// We must not recreate the instance since it holds stateful information (such as which objects it has spawned). Recreating the instance would break any 
+	// @todo: Is this still the case now that eval state is stored (correctly) in the player?
+	if (!RootTemplateInstance.IsValid())
+	{
+		RootTemplateInstance.Initialize(*Sequence, *this);
+	}
+
+	OnStartedPlaying();
+
+	bPendingFirstUpdate = true;
+	bIsPlaying = true;
+}
+
+void UMovieSceneSequencePlayer::Pause()
+{
+	if (bIsPlaying)
+	{
+		if (bIsEvaluating)
+		{
+			LatentActions.Add(ELatentAction::Pause);
+			return;
+		}
+
+		bIsPlaying = false;
+
+		// Evaluate the sequence at its current time, with a status of 'stopped' to ensure that animated state pauses correctly
+		{
+			bIsEvaluating = true;
+
+			UMovieSceneSequence* MovieSceneSequence = RootTemplateInstance.GetSequence(MovieSceneSequenceID::Root);
+			TOptional<float> FixedFrameInterval = MovieSceneSequence->GetMovieScene() ? MovieSceneSequence->GetMovieScene()->GetOptionalFixedFrameInterval() : TOptional<float>();
+
+			FMovieSceneEvaluationRange Range = PlayPosition.JumpTo(GetSequencePosition(), FixedFrameInterval);
+
+			const FMovieSceneContext Context(Range, EMovieScenePlayerStatus::Stopped);
+			RootTemplateInstance.Evaluate(Context, *this);
+
+			bIsEvaluating = false;
+		}
+
+		ApplyLatentActions();
+
+		if (OnPause.IsBound())
+		{
+			OnPause.Broadcast();
+		}
+	}
+}
+
+void UMovieSceneSequencePlayer::Stop()
+{
+	if (bIsPlaying)
+	{
+		if (bIsEvaluating)
+		{
+			LatentActions.Add(ELatentAction::Stop);
+			return;
+		}
+
+		bIsPlaying = false;
+		TimeCursorPosition = PlaybackSettings.PlayRate < 0.f ? GetLength() : 0.f;
+		CurrentNumLoops = 0;
+
+		if (PlaybackSettings.bRestoreState)
+		{
+			RestorePreAnimatedState();
+		}
+
+		RootTemplateInstance.Finish(*this);
+
+		OnStopped();
+
+		if (OnStop.IsBound())
+		{
+			OnStop.Broadcast();
+		}
+	}
+}
+
+float UMovieSceneSequencePlayer::GetPlaybackPosition() const
+{
+	return TimeCursorPosition;
+}
+
+void UMovieSceneSequencePlayer::SetPlaybackPosition(float NewPlaybackPosition)
+{
+	UpdateTimeCursorPosition(NewPlaybackPosition);
+}
+
+bool UMovieSceneSequencePlayer::IsPlaying() const
+{
+	return bIsPlaying;
+}
+
+float UMovieSceneSequencePlayer::GetLength() const
+{
+	return EndTime - StartTime;
+}
+
+float UMovieSceneSequencePlayer::GetPlayRate() const
+{
+	return PlaybackSettings.PlayRate;
+}
+
+void UMovieSceneSequencePlayer::SetPlayRate(float PlayRate)
+{
+	PlaybackSettings.PlayRate = PlayRate;
+}
+
+void UMovieSceneSequencePlayer::SetPlaybackRange( const float NewStartTime, const float NewEndTime )
+{
+	StartTime = NewStartTime;
+	EndTime = FMath::Max(NewEndTime, StartTime);
+
+	TimeCursorPosition = FMath::Clamp(TimeCursorPosition, 0.f, GetLength());
+}
+
+bool UMovieSceneSequencePlayer::ShouldStopOrLoop(float NewPosition) const
+{
+	bool bShouldStopOrLoop = false;
+	if (!bReversePlayback)
+	{
+		bShouldStopOrLoop = NewPosition >= GetLength();
+	}
+	else
+	{
+		bShouldStopOrLoop = NewPosition < 0.f;
+	}
+
+	return bShouldStopOrLoop;
+}
+
+void UMovieSceneSequencePlayer::Initialize(UMovieSceneSequence* InSequence, const FMovieSceneSequencePlaybackSettings& InSettings)
+{
+	check(InSequence);
+	
+	Sequence = InSequence;
+	PlaybackSettings = InSettings;
+
+	if (UMovieScene* MovieScene = Sequence->GetMovieScene())
+	{
+		TRange<float> PlaybackRange = MovieScene->GetPlaybackRange();
+		SetPlaybackRange(PlaybackRange.GetLowerBoundValue(), PlaybackRange.GetUpperBoundValue());
+	}
+
+	TimeCursorPosition = PlaybackSettings.bRandomStartTime ? FMath::FRand() * 0.99f * GetLength() : FMath::Clamp(PlaybackSettings.StartTime, 0.f, GetLength());
+
+	RootTemplateInstance.Initialize(*Sequence, *this);
+
+	// Ensure everything is set up, ready for playback
+	Stop();
+}
+
+void UMovieSceneSequencePlayer::Update(const float DeltaSeconds)
+{
+	if (bIsPlaying)
+	{
+		float PlayRate = bReversePlayback ? -PlaybackSettings.PlayRate : PlaybackSettings.PlayRate;
+		UpdateTimeCursorPosition(TimeCursorPosition + DeltaSeconds * PlayRate);
+	}
+}
+
+void UMovieSceneSequencePlayer::UpdateTimeCursorPosition(float NewPosition)
+{
+	float Length = GetLength();
+
+	UMovieSceneSequence* MovieSceneSequence = RootTemplateInstance.GetSequence(MovieSceneSequenceID::Root);
+	TOptional<float> FixedFrameInterval = MovieSceneSequence->GetMovieScene() ? MovieSceneSequence->GetMovieScene()->GetOptionalFixedFrameInterval() : TOptional<float>();
+
+	if (bPendingFirstUpdate)
+	{
+		NewPosition = TimeCursorPosition;
+		bPendingFirstUpdate = false;
+	}
+
+	if (ShouldStopOrLoop(NewPosition))
+	{
+		// loop playback
+		if (PlaybackSettings.LoopCount < 0 || CurrentNumLoops < PlaybackSettings.LoopCount)
+		{
+			++CurrentNumLoops;
+			const float Overplay = FMath::Fmod(NewPosition, Length);
+			TimeCursorPosition = Overplay < 0 ? Length + Overplay : Overplay;
+			
+			PlayPosition.Reset(Overplay < 0 ? Length : 0.f);
+			FMovieSceneEvaluationRange Range = PlayPosition.PlayTo(GetSequencePosition(), FixedFrameInterval);
+
+			if (SpawnRegister.IsValid())
+			{
+				SpawnRegister->ForgetExternallyOwnedSpawnedObjects(State, *this);
+			}
+
+			UpdateMovieSceneInstance(Range);
+
+			OnLooped();
+			return;
+		}
+
+		// stop playback
+		Stop();
+		return;
+	}
+
+	TimeCursorPosition = NewPosition;
+
+	FMovieSceneEvaluationRange Range = PlayPosition.PlayTo(NewPosition + StartTime, FixedFrameInterval);
+	UpdateMovieSceneInstance(Range);
+}
+
+void UMovieSceneSequencePlayer::UpdateMovieSceneInstance(FMovieSceneEvaluationRange InRange)
+{
+	bIsEvaluating = true;
+
+	const FMovieSceneContext Context(InRange, GetPlaybackStatus());
+	RootTemplateInstance.Evaluate(Context, *this);
+
+#if WITH_EDITOR
+	OnMovieSceneSequencePlayerUpdate.Broadcast(*this, Context.GetTime(), Context.GetPreviousTime());
+#endif
+	bIsEvaluating = false;
+
+	ApplyLatentActions();
+}
+
+void UMovieSceneSequencePlayer::ApplyLatentActions()
+{
+	// Swap to a stack array to ensure no reentrancy if we evaluate during a pause, for instance
+	TArray<ELatentAction> TheseActions;
+	Swap(TheseActions, LatentActions);
+
+	for (ELatentAction LatentAction : TheseActions)
+	{
+		switch(LatentAction)
+		{
+		case ELatentAction::Stop:	Stop(); break;
+		case ELatentAction::Pause:	Pause(); break;
+		}
+	}
+}
\ No newline at end of file
diff --git a/Engine/Source/Runtime/MovieScene/Private/MovieSceneTrack.cpp b/Engine/Source/Runtime/MovieScene/Private/MovieSceneTrack.cpp
--- a/Engine/Source/Runtime/MovieScene/Private/MovieSceneTrack.cpp
+++ b/Engine/Source/Runtime/MovieScene/Private/MovieSceneTrack.cpp
@@ -17,6 +17,17 @@ UMovieSceneTrack::UMovieSceneTrack(const FObjectInitializer& InInitializer)
 #endif
 }
 
+void UMovieSceneTrack::PostInitProperties()
+{
+	// Propagate sub object flags from our outer (movie scene) to ourselves. This is required for tracks that are stored on blueprints (archetypes) so that they can be referenced in worlds.
+	if (GetOuter()->HasAnyFlags(RF_ClassDefaultObject|RF_ArchetypeObject))
+	{
+		SetFlags(GetOuter()->GetMaskedFlags(RF_PropagateToSubObjects));
+	}
+	
+	Super::PostInitProperties();
+}
+
 PRAGMA_DISABLE_DEPRECATION_WARNINGS
 TSharedPtr<IMovieSceneTrackInstance> UMovieSceneTrack::CreateLegacyInstance() const
 {
diff --git a/Engine/Source/Runtime/MovieScene/Public/Compilation/MovieSceneCompilerRules.h b/Engine/Source/Runtime/MovieScene/Public/Compilation/MovieSceneCompilerRules.h
--- a/Engine/Source/Runtime/MovieScene/Public/Compilation/MovieSceneCompilerRules.h
+++ b/Engine/Source/Runtime/MovieScene/Public/Compilation/MovieSceneCompilerRules.h
@@ -63,6 +63,11 @@ namespace MovieSceneSegmentCompiler
 			}
 		}
 	}
+
+	static void BlendSegmentUpperBoundExclusive(FMovieSceneSegment& Segment, const TArrayView<const FMovieSceneSectionData>& SourceData)
+	{
+		Segment.Range = FFloatRange(Segment.Range.GetLowerBoundValue(), Segment.Range.GetUpperBoundValue());
+	}
 }
 
 class FMovieSceneAdditiveCameraRules : public FMovieSceneSegmentCompilerRules
diff --git a/Engine/Source/Runtime/MovieScene/Public/Evaluation/MovieSceneEvalTemplate.h b/Engine/Source/Runtime/MovieScene/Public/Evaluation/MovieSceneEvalTemplate.h
--- a/Engine/Source/Runtime/MovieScene/Public/Evaluation/MovieSceneEvalTemplate.h
+++ b/Engine/Source/Runtime/MovieScene/Public/Evaluation/MovieSceneEvalTemplate.h
@@ -150,6 +150,9 @@ struct FMovieSceneEvalTemplatePtr
 	FMovieSceneEvalTemplatePtr(T&& In)
 		: TInlineValue(Forward<T>(In))
 	{
+#if WITH_EDITOR
+		checkf(T::StaticStruct() == &In.GetScriptStruct(), TEXT("%s does not correctly override GetScriptStructImpl. Template will not serialize correctly."), *T::StaticStruct()->GetName());
+#endif
 	}
 
 	/** Copy construction/assignment */
diff --git a/Engine/Source/Runtime/MovieScene/Public/Evaluation/MovieSceneEvalTemplateBase.h b/Engine/Source/Runtime/MovieScene/Public/Evaluation/MovieSceneEvalTemplateBase.h
--- a/Engine/Source/Runtime/MovieScene/Public/Evaluation/MovieSceneEvalTemplateBase.h
+++ b/Engine/Source/Runtime/MovieScene/Public/Evaluation/MovieSceneEvalTemplateBase.h
@@ -7,6 +7,11 @@
 #include "Evaluation/MovieScenePlayback.h"
 #include "MovieSceneEvalTemplateBase.generated.h"
 
+
+struct FPersistentEvaluationData;
+class IMovieScenePlayer;
+
+
 /** Empty struct used for serialization */
 USTRUCT()
 struct FMovieSceneEmptyStruct
diff --git a/Engine/Source/Runtime/MovieScene/Public/Evaluation/MovieSceneEvaluation.h b/Engine/Source/Runtime/MovieScene/Public/Evaluation/MovieSceneEvaluation.h
new file mode 100644
--- /dev/null
+++ b/Engine/Source/Runtime/MovieScene/Public/Evaluation/MovieSceneEvaluation.h
@@ -0,0 +1,15 @@
+// Copyright 1998-2016 Epic Games, Inc. All Rights Reserved.
+
+#pragma once
+
+#include "MovieSceneSequence.h"
+#include "MovieSceneSegment.h"
+#include "MovieSceneEvalTemplate.h"
+#include "MovieSceneTrackImplementation.h"
+#include "MovieScenePreAnimatedState.h"
+#include "MovieSceneSection.h"
+#include "MovieSceneEvaluationState.h"
+#include "MovieScenePlayback.h"
+#include "MovieSceneExecutionTokens.h"
+#include "MovieSceneEvaluationTemplate.h"
+#include "MovieSceneEvaluationTemplateInstance.h"
\ No newline at end of file
diff --git a/Engine/Source/Runtime/MovieScene/Public/Evaluation/MovieSceneEvaluationTemplate.h b/Engine/Source/Runtime/MovieScene/Public/Evaluation/MovieSceneEvaluationTemplate.h
--- a/Engine/Source/Runtime/MovieScene/Public/Evaluation/MovieSceneEvaluationTemplate.h
+++ b/Engine/Source/Runtime/MovieScene/Public/Evaluation/MovieSceneEvaluationTemplate.h
@@ -4,6 +4,8 @@
 
 #include "CoreMinimal.h"
 #include "UObject/ObjectMacros.h"
+#include "UObject/ScriptMacros.h"
+#include "UObject/WeakObjectPtr.h"
 #include "Misc/Guid.h"
 #include "UObject/Class.h"
 #include "MovieSceneTrack.h"
@@ -42,7 +44,7 @@ public:
 	/** Custom serialized - map of track signature to array of track identifiers that it created */
 	TMap<FGuid, TArray<FMovieSceneTrackIdentifier, TInlineAllocator<1>>> TrackSignatureToTrackIdentifier;
 };
-template<> struct TStructOpsTypeTraits<FMovieSceneGenerationLedger> : public TStructOpsTypeTraitsBase { enum { WithSerializer = true }; };
+template<> struct TStructOpsTypeTraits<FMovieSceneGenerationLedger> : public TStructOpsTypeTraitsBase { enum { WithSerializer = true, WithCopy = true }; };
 
 /**
  * Template that is used for efficient runtime evaluation of a movie scene sequence. Potentially serialized into the asset.
@@ -234,8 +236,6 @@ struct FCachedMovieSceneEvaluationTemplate : public FMovieSceneEvaluationTemplat
 
 private:
 
-	void RegenerateImpl(const FMovieSceneTrackCompilationParams& Params);
-
 	/** Transient data */
 	TWeakObjectPtr<UMovieSceneSequence> SourceSequence;
 	FMovieSceneSequenceTemplateStore* Origin;
diff --git a/Engine/Source/Runtime/MovieScene/Public/Evaluation/MovieSceneExecutionTokens.h b/Engine/Source/Runtime/MovieScene/Public/Evaluation/MovieSceneExecutionTokens.h
--- a/Engine/Source/Runtime/MovieScene/Public/Evaluation/MovieSceneExecutionTokens.h
+++ b/Engine/Source/Runtime/MovieScene/Public/Evaluation/MovieSceneExecutionTokens.h
@@ -6,6 +6,7 @@
 #include "Evaluation/MovieSceneEvaluationKey.h"
 #include "Misc/InlineValue.h"
 #include "MovieSceneExecutionToken.h"
+#include "MovieSceneSection.h"
 #include "Evaluation/PersistentEvaluationData.h"
 #include "Evaluation/MovieScenePlayback.h"
 
diff --git a/Engine/Source/Runtime/MovieScene/Public/Evaluation/MovieScenePlayback.h b/Engine/Source/Runtime/MovieScene/Public/Evaluation/MovieScenePlayback.h
--- a/Engine/Source/Runtime/MovieScene/Public/Evaluation/MovieScenePlayback.h
+++ b/Engine/Source/Runtime/MovieScene/Public/Evaluation/MovieScenePlayback.h
@@ -3,8 +3,9 @@
 #pragma once
 
 #include "CoreMinimal.h"
+#include "MovieSceneFwd.h"
 #include "Evaluation/MovieSceneSequenceTransform.h"
-#include "IMovieScenePlayer.h"
+
 
 /** Enumeration specifying whether we're playing forwards or backwards */
 enum class EPlayDirection
@@ -12,6 +13,7 @@ enum class EPlayDirection
 	Forwards, Backwards
 };
 
+
 /** MovieScene evaluation context. Should remain bitwise copyable, and contain no external state since this has the potential to be used on a thread */
 struct MOVIESCENE_API FMovieSceneEvaluationRange
 {
diff --git a/Engine/Source/Runtime/MovieScene/Public/Evaluation/MovieScenePropertyTemplate.h b/Engine/Source/Runtime/MovieScene/Public/Evaluation/MovieScenePropertyTemplate.h
--- a/Engine/Source/Runtime/MovieScene/Public/Evaluation/MovieScenePropertyTemplate.h
+++ b/Engine/Source/Runtime/MovieScene/Public/Evaluation/MovieScenePropertyTemplate.h
@@ -5,6 +5,7 @@
 #include "CoreMinimal.h"
 #include "Stats/Stats.h"
 #include "UObject/ObjectMacros.h"
+#include "IMovieScenePlayer.h"
 #include "MovieSceneFwd.h"
 #include "MovieSceneCommonHelpers.h"
 #include "Evaluation/MovieSceneAnimTypeID.h"
diff --git a/Engine/Source/Runtime/MovieScene/Public/Evaluation/MovieSceneSegment.h b/Engine/Source/Runtime/MovieScene/Public/Evaluation/MovieSceneSegment.h
--- a/Engine/Source/Runtime/MovieScene/Public/Evaluation/MovieSceneSegment.h
+++ b/Engine/Source/Runtime/MovieScene/Public/Evaluation/MovieSceneSegment.h
@@ -99,5 +99,5 @@ struct FMovieSceneSegment
 
 template<> struct TStructOpsTypeTraits<FMovieSceneSegment> : public TStructOpsTypeTraitsBase
 {
-	enum { WithSerializer = true };
+	enum { WithSerializer = true, WithCopy = true };
 };
diff --git a/Engine/Source/Runtime/MovieScene/Public/Evaluation/MovieSceneTrackImplementation.h b/Engine/Source/Runtime/MovieScene/Public/Evaluation/MovieSceneTrackImplementation.h
--- a/Engine/Source/Runtime/MovieScene/Public/Evaluation/MovieSceneTrackImplementation.h
+++ b/Engine/Source/Runtime/MovieScene/Public/Evaluation/MovieSceneTrackImplementation.h
@@ -9,9 +9,12 @@
 #include "Evaluation/MovieSceneEvalTemplateBase.h"
 #include "MovieSceneTrackImplementation.generated.h"
 
+
+struct FMovieSceneEvaluationOperand;
 struct FMovieSceneEvaluationTrack;
 struct FMovieSceneExecutionTokens;
 
+
 /**
  * Structure that allows the implementation of setup/teardown/initialization/evaluation logic at the track level.
  */
@@ -111,6 +114,9 @@ struct FMovieSceneTrackImplementationPtr
 	FMovieSceneTrackImplementationPtr(T&& In)
 		: TInlineValue(Forward<T>(In))
 	{
+#if WITH_EDITOR
+		checkf(T::StaticStruct() == &In.GetScriptStruct(), TEXT("%s type does not correctly override GetScriptStructImpl. Track will not serialize correctly."), *T::StaticStruct()->GetName());
+#endif
 	}
 	
 	/** Copy construction/assignment */
diff --git a/Engine/Source/Runtime/MovieScene/Public/IMovieScenePlayer.h b/Engine/Source/Runtime/MovieScene/Public/IMovieScenePlayer.h
--- a/Engine/Source/Runtime/MovieScene/Public/IMovieScenePlayer.h
+++ b/Engine/Source/Runtime/MovieScene/Public/IMovieScenePlayer.h
@@ -8,27 +8,16 @@
 #include "MovieSceneSequenceInstance.h"
 #include "Evaluation/MovieSceneEvaluationKey.h"
 #include "Evaluation/MovieScenePreAnimatedState.h"
+#include "MovieSceneFwd.h"
 #include "MovieSceneSpawnRegister.h"
 #include "Containers/ArrayView.h"
 #include "Evaluation/MovieSceneEvaluationState.h"
+#include "MovieSceneSequence.h"
+
 
 class FViewportClient;
 struct FMovieSceneRootEvaluationTemplateInstance;
 
-namespace EMovieScenePlayerStatus
-{
-	enum Type
-	{
-		Stopped,
-		Playing,
-		Recording,
-		Scrubbing,
-		Jumping,
-		Stepping,
-		MAX
-	};
-}
-
 
 struct EMovieSceneViewportParams
 {
@@ -55,7 +44,6 @@ struct EMovieSceneViewportParams
 	bool bEnableColorScaling;
 };
 
-
 /**
  * Interface for movie scene players
  * Provides information for playback of a movie scene
@@ -100,6 +88,17 @@ public:
 	*/
 	virtual void SetPlaybackStatus(EMovieScenePlayerStatus::Type InPlaybackStatus) = 0;
 
+	/**
+	 * Resolve objects bound to the specified binding ID
+	 *
+	 * @param InBindingId	The ID relating to the object(s) to resolve
+	 * @param OutObjects	Container to populate with the bound objects
+	 */
+	virtual void ResolveBoundObjects(const FGuid& InBindingId, FMovieSceneSequenceID SequenceID, UMovieSceneSequence& Sequence, UObject* ResolutionContext, TArray<UObject*, TInlineAllocator<1>>& OutObjects) const
+	{
+		Sequence.LocateBoundObjects(InBindingId, ResolutionContext, OutObjects);
+	}
+
 	/**
 	 * Obtain an object responsible for managing movie scene spawnables
 	 */
@@ -227,6 +226,7 @@ public:
 	void RestorePreAnimatedState()
 	{
 		PreAnimatedState.RestorePreAnimatedState(*this);
+		State.ClearObjectCaches();
 	}
 
 	/**
diff --git a/Engine/Source/Runtime/MovieScene/Public/MovieScene.h b/Engine/Source/Runtime/MovieScene/Public/MovieScene.h
--- a/Engine/Source/Runtime/MovieScene/Public/MovieScene.h
+++ b/Engine/Source/Runtime/MovieScene/Public/MovieScene.h
@@ -19,6 +19,15 @@ class UMovieSceneTrack;
 
 MOVIESCENE_API DECLARE_LOG_CATEGORY_EXTERN(LogMovieScene, Log, All);
 
+USTRUCT(BlueprintType)
+struct FMovieSceneObjectBindingPtr
+{
+	GENERATED_BODY()
+
+	UPROPERTY(EditAnywhere, Category="Binding")
+	FGuid Guid;
+};
+
 
 /** @todo: remove this type when support for intrinsics on TMap values is added? */
 USTRUCT()
@@ -81,7 +90,7 @@ struct FMovieSceneTrackLabels
 /**
  * Implements a movie scene asset.
  */
-UCLASS()
+UCLASS(DefaultToInstanced)
 class MOVIESCENE_API UMovieScene
 	: public UMovieSceneSignedObject
 {
@@ -226,8 +235,6 @@ public:
 	/**
 	* Adds a given track.
 	*
-	* Note: Function will only add if the track is not already exist.
-	*
 	* @param InTrack The track to add.
 	* @param ObjectGuid The runtime object guid that the type should bind to.
 	* @see  FindTrack, RemoveTrack
@@ -489,6 +496,18 @@ public:
 	 */
 	void SetViewRange(float Start, float End);
 
+#if WITH_EDITORONLY_DATA
+	/** 
+	 * Return whether the playback range is locked.
+	 */
+	bool IsPlaybackRangeLocked() const;
+
+	/**
+	 * Set whether the playback range is locked.
+	 */
+	void SetPlaybackRangeLocked(bool bLocked);
+#endif
+
 	/**
 	 * Gets whether or not playback should be forced to match the fixed frame interval.  When true all time values will be rounded to a fixed
 	 * frame value which will force editor and runtime playback to match exactly, but will result in duplicate frames if the runtime and editor
@@ -567,6 +586,10 @@ protected:
 	/** Perform legacy upgrade of time ranges */
 	void UpgradeTimeRanges();
 
+	/** Perform legacy upgrade of track rows */
+	void UpgradeTrackRows();
+	void UpgradeTrackRow(UMovieSceneTrack*);
+
 private:
 
 	// Small value added for fixed frame interval calculations to make up for consistency in
@@ -590,11 +613,11 @@ private:
 	TArray<FMovieSceneBinding> ObjectBindings;
 
 	/** Master tracks which are not bound to spawned or possessed objects */
-	UPROPERTY()
+	UPROPERTY(Instanced)
 	TArray<UMovieSceneTrack*> MasterTracks;
 
 	/** The camera cut track is a specialized track for switching between cameras on a cinematic */
-	UPROPERTY()
+	UPROPERTY(Instanced)
 	UMovieSceneTrack* CameraCutTrack;
 
 	/** User-defined selection range. */
@@ -605,6 +628,12 @@ private:
 	UPROPERTY()
 	FFloatRange PlaybackRange;
 
+	/** User-defined playback range is locked. */
+#if WITH_EDITORONLY_DATA
+	UPROPERTY()
+	bool bPlaybackRangeLocked;
+#endif
+
 	UPROPERTY()
 	bool bForceFixedFrameIntervalPlayback;
 
diff --git a/Engine/Source/Runtime/MovieScene/Public/MovieSceneBinding.h b/Engine/Source/Runtime/MovieScene/Public/MovieSceneBinding.h
--- a/Engine/Source/Runtime/MovieScene/Public/MovieSceneBinding.h
+++ b/Engine/Source/Runtime/MovieScene/Public/MovieSceneBinding.h
@@ -126,6 +126,6 @@ private:
 	FString BindingName;
 
 	/** All tracks in this binding */
-	UPROPERTY()
+	UPROPERTY(Instanced)
 	TArray<UMovieSceneTrack*> Tracks;
 };
diff --git a/Engine/Source/Runtime/MovieScene/Public/MovieSceneBindingOverrides.h b/Engine/Source/Runtime/MovieScene/Public/MovieSceneBindingOverrides.h
new file mode 100644
--- /dev/null
+++ b/Engine/Source/Runtime/MovieScene/Public/MovieSceneBindingOverrides.h
@@ -0,0 +1,103 @@
+// Copyright 1998-2016 Epic Games, Inc. All Rights Reserved.
+
+#pragma once
+
+#include "MovieScene.h"
+#include "MovieSceneBindingOverridesInterface.h"
+#include "UObject/Object.h"
+#include "UObject/ObjectMacros.h"
+#include "UObject/ScriptMacros.h"
+
+#include "MovieSceneBindingOverrides.generated.h"
+
+
+/** Movie scene binding override data */
+USTRUCT()
+struct FMovieSceneBindingOverrideData
+{
+	GENERATED_BODY()
+
+	FMovieSceneBindingOverrideData()
+		: bOverridesDefault(true)
+	{
+	}
+
+	/** Specifies the object binding to override. */
+	UPROPERTY(EditAnywhere, Category="Binding")
+	FMovieSceneObjectBindingPtr ObjectBindingId;
+
+	/** Specifies the object to override the binding with. */
+	UPROPERTY(EditAnywhere, Category="Binding")
+	TWeakObjectPtr<UObject> Object;
+
+	/** Specifies whether the default assignment should remain bound (false) or if this should completely override the default binding (false). */
+	UPROPERTY(EditAnywhere, Category="Binding")
+	bool bOverridesDefault;
+};
+
+
+/**
+ * A one-to-many definition of movie scene object binding IDs to overridden objects that should be bound to that binding.
+ */
+UCLASS(DefaultToInstanced, EditInlineNew)
+class MOVIESCENE_API UMovieSceneBindingOverrides
+	: public UObject
+	, public IMovieSceneBindingOverridesInterface
+{
+public:
+
+	GENERATED_BODY()
+
+	/** Default constructor */
+	UMovieSceneBindingOverrides(const FObjectInitializer& Init);
+
+	virtual bool LocateBoundObjects(const FGuid& InBindingId, TArray<UObject*, TInlineAllocator<1>>& OutObjects) const override;
+
+public:
+
+	/**
+	 * Overrides the specified binding with the specified objects,
+	 * optionally still allowing the bindings defined in the sequence.
+	 */
+	void SetBinding(FMovieSceneObjectBindingPtr Binding, const TArray<UObject*>& Objects, bool bAllowBindingsFromAsset = false);
+
+	/**
+	 * Adds the specified actor to the overridden bindings for the specified binding ID,
+	 * optionally still allowing the bindings defined in the sequence.
+	 */
+	void AddBinding(FMovieSceneObjectBindingPtr Binding, UObject* Object, bool bAllowBindingsFromAsset = false);
+
+	/**
+	 * Removes the specified actor from the specified binding's actor array.
+	 */
+	void RemoveBinding(FMovieSceneObjectBindingPtr Binding, UObject* Object);
+
+	/**
+	 * Resets the specified binding back to the defaults defined by the sequence.
+	 */
+	void ResetBinding(FMovieSceneObjectBindingPtr Binding);
+
+	/** Resets all overridden bindings back to the defaults defined by the sequence. */
+	void ResetBindings();
+
+protected:
+
+	/** Rebuild the lookup map for efficient lookup. */
+	void RebuildLookupMap() const;
+
+#if WITH_EDITOR
+
+	virtual void PostEditChangeProperty(FPropertyChangedEvent& PropertyChangedEvent) override;
+
+#endif
+
+private:
+
+	/** The actual binding data */
+	UPROPERTY(EditAnywhere, Category="General", DisplayName="Binding Overrides")
+	TArray<FMovieSceneBindingOverrideData> BindingData;
+
+	/** Runtime lookup map */
+	mutable bool bLookupDirty;
+	mutable TMultiMap<FGuid, int32> LookupMap;
+};
diff --git a/Engine/Source/Runtime/MovieScene/Public/MovieSceneBindingOverridesInterface.h b/Engine/Source/Runtime/MovieScene/Public/MovieSceneBindingOverridesInterface.h
new file mode 100644
--- /dev/null
+++ b/Engine/Source/Runtime/MovieScene/Public/MovieSceneBindingOverridesInterface.h
@@ -0,0 +1,34 @@
+// Copyright 1998-2016 Epic Games, Inc. All Rights Reserved.
+
+#pragma once
+
+#include "UObject/Interface.h"
+#include "UObject/ObjectMacros.h"
+
+#include "MovieSceneBindingOverridesInterface.generated.h"
+
+
+/** Movie scene binding overrides interface */
+UINTERFACE()
+class MOVIESCENE_API UMovieSceneBindingOverridesInterface
+	: public UInterface
+{
+public:
+	GENERATED_BODY()
+};
+
+
+class MOVIESCENE_API IMovieSceneBindingOverridesInterface
+{
+public:
+	GENERATED_BODY()
+
+	/**
+	 * Locate bound objects that relate to the specified binding ID
+	 *
+	 * @param InBindingId 		The GUID of the object binding in the movie scene
+	 * @param OutObjects 		Array to populate with bound objects
+	 * @return true to prevent default lookup of the binding in the sequence itself, else false
+	 */
+	virtual bool LocateBoundObjects(const FGuid& InBindingId, TArray<UObject*, TInlineAllocator<1>>& OutObjects) const = 0;
+};
diff --git a/Engine/Source/Runtime/MovieScene/Public/MovieSceneFwd.h b/Engine/Source/Runtime/MovieScene/Public/MovieSceneFwd.h
--- a/Engine/Source/Runtime/MovieScene/Public/MovieSceneFwd.h
+++ b/Engine/Source/Runtime/MovieScene/Public/MovieSceneFwd.h
@@ -6,6 +6,20 @@
 #include "Stats/Stats.h"
 #include "MovieSceneSequenceID.h"
 
+namespace EMovieScenePlayerStatus
+{
+	enum Type
+	{
+		Stopped,
+		Playing,
+		Recording,
+		Scrubbing,
+		Jumping,
+		Stepping,
+		MAX
+	};
+}
+
 DECLARE_STATS_GROUP(TEXT("Movie Scene Evaluation"), STATGROUP_MovieSceneEval, STATCAT_Advanced);
 
 #ifndef MOVIESCENE_DETAILED_STATS
diff --git a/Engine/Source/Runtime/MovieScene/Public/MovieSceneSection.h b/Engine/Source/Runtime/MovieScene/Public/MovieSceneSection.h
--- a/Engine/Source/Runtime/MovieScene/Public/MovieSceneSection.h
+++ b/Engine/Source/Runtime/MovieScene/Public/MovieSceneSection.h
@@ -52,7 +52,7 @@ struct FMovieSceneSectionEvalOptions
 /**
  * Base class for movie scene sections
  */
-UCLASS(abstract, MinimalAPI)
+UCLASS(abstract, DefaultToInstanced, MinimalAPI)
 class UMovieSceneSection
 	: public UMovieSceneSignedObject
 {
@@ -327,6 +327,17 @@ public:
 	/** Sets the time for the key referenced by the supplied key handle. */
 	virtual void SetKeyTime( FKeyHandle KeyHandle, float Time ) PURE_VIRTUAL( UAISenseEvent::SetKeyTime, );
 	
+	/** For backwards compatibility, allow this section to do work when upgrading the section's track rows
+	 *
+	 * @param InEvaluationRange The new evaluation range for this section
+	 */
+	virtual void PostLoadUpgradeTrackRow(const TRange<float>& InEvaluationRange) { }
+
+protected:
+
+	//~ UObject interface
+	MOVIESCENE_API virtual void PostInitProperties() override;
+
 private:
 
 	/** The start time of the section */
diff --git a/Engine/Source/Runtime/MovieScene/Public/MovieSceneSequence.h b/Engine/Source/Runtime/MovieScene/Public/MovieSceneSequence.h
--- a/Engine/Source/Runtime/MovieScene/Public/MovieSceneSequence.h
+++ b/Engine/Source/Runtime/MovieScene/Public/MovieSceneSequence.h
@@ -12,6 +12,7 @@
 
 class ITargetPlatform;
 class UMovieScene;
+struct FMovieScenePossessable;
 
 /**
  * Abstract base class for movie scene animations (C++ version).
@@ -132,6 +133,14 @@ public:
 	 */
 	virtual UObject* MakeSpawnableTemplateFromInstance(UObject& InSourceObject, FName ObjectName) { return nullptr; }
 
+	/**
+	 * Specifies whether this sequence allows rebinding of the specified possessable
+	 *
+	 * @param InPossessable The possessable to check
+	 * @return true if rebinding this possessable is valid at runtime, false otherwise
+	 */
+	virtual bool CanRebindPossessable(const FMovieScenePossessable& InPossessable) const { return false; }
+
 public:
 
 	MOVIESCENE_API virtual void Serialize(FArchive& Ar) override;
diff --git a/Engine/Source/Runtime/MovieScene/Public/MovieSceneSequenceID.h b/Engine/Source/Runtime/MovieScene/Public/MovieSceneSequenceID.h
--- a/Engine/Source/Runtime/MovieScene/Public/MovieSceneSequenceID.h
+++ b/Engine/Source/Runtime/MovieScene/Public/MovieSceneSequenceID.h
@@ -42,7 +42,7 @@ struct FMovieSceneSequenceID
 
 	FORCEINLINE FMovieSceneSequenceID AccumulateParentID(FMovieSceneSequenceID InParentID)
 	{
-		return FMovieSceneSequenceID(HashCombine(Value, InParentID.Value));
+		return Value == 0 ? InParentID : FMovieSceneSequenceID(HashCombine(Value, InParentID.Value));
 	}
 
 	FORCEINLINE bool Serialize(FArchive& Ar)
@@ -72,7 +72,8 @@ struct TStructOpsTypeTraits<FMovieSceneSequenceID> : public TStructOpsTypeTraits
 {
 	enum
 	{
-		WithSerializer = true
+		WithSerializer = true,
+		WithCopy = true
 	};
 };
 
diff --git a/Engine/Source/Runtime/MovieScene/Public/MovieSceneSequencePlayer.h b/Engine/Source/Runtime/MovieScene/Public/MovieSceneSequencePlayer.h
new file mode 100644
--- /dev/null
+++ b/Engine/Source/Runtime/MovieScene/Public/MovieSceneSequencePlayer.h
@@ -0,0 +1,286 @@
+// Copyright 1998-2016 Epic Games, Inc. All Rights Reserved.
+
+#pragma once
+
+#include "CoreMinimal.h"
+#include "UObject/ObjectMacros.h"
+#include "UObject/ScriptMacros.h"
+#include "IMovieScenePlayer.h"
+#include "Evaluation/MovieSceneEvaluationTemplateInstance.h"
+#include "MovieSceneBindingOverridesInterface.h"
+#include "MovieSceneSequencePlayer.generated.h"
+
+DECLARE_DYNAMIC_MULTICAST_DELEGATE(FOnMovieSceneSequencePlayerEvent);
+
+/**
+ * Settings for the level sequence player actor.
+ */
+USTRUCT(BlueprintType)
+struct FMovieSceneSequencePlaybackSettings
+{
+	FMovieSceneSequencePlaybackSettings()
+		: LoopCount(0)
+		, PlayRate(1.f)
+		, bRandomStartTime(false)
+		, StartTime(0.f)
+		, bRestoreState(false)
+	{ }
+
+	GENERATED_BODY()
+
+	/** Number of times to loop playback. -1 for infinite, else the number of times to loop before stopping */
+	UPROPERTY(EditAnywhere, BlueprintReadWrite, Category="Playback", meta=(UIMin=1, DisplayName="Loop"))
+	int32 LoopCount;
+
+	/** The rate at which to playback the animation */
+	UPROPERTY(EditAnywhere, BlueprintReadWrite, Category="Playback", meta=(Units=Multiplier))
+	float PlayRate;
+
+	/** Start playback at a random time */
+	UPROPERTY(EditAnywhere, BlueprintReadWrite, Category="Playback")
+	bool bRandomStartTime;
+
+	/** Start playback at the specified time */
+	UPROPERTY(EditAnywhere, BlueprintReadWrite, Category="Playback", meta=(Units=s, EditCondition="!bRandomStartTime"))
+	float StartTime;
+
+	/** Flag used to specify whether actor states should be restored on stop */
+	UPROPERTY(EditAnywhere, BlueprintReadWrite, Category="Playback")
+	bool bRestoreState;
+
+	/** Interface that defines overridden bindings for this sequence */
+	UPROPERTY()
+	TScriptInterface<IMovieSceneBindingOverridesInterface> BindingOverrides;
+
+	MOVIESCENE_API bool SerializeFromMismatchedTag(const FPropertyTag& Tag, FArchive& Ar);
+};
+
+template<> struct TStructOpsTypeTraits<FMovieSceneSequencePlaybackSettings> : public TStructOpsTypeTraitsBase
+{
+	enum { WithCopy = true, WithSerializeFromMismatchedTag = true };
+};
+
+/**
+ * Abstract class that provides consistent player behaviour for various animation players
+ */
+UCLASS(Abstract, BlueprintType)
+class MOVIESCENE_API UMovieSceneSequencePlayer
+	: public UObject
+	, public IMovieScenePlayer
+{
+public:
+	GENERATED_BODY()
+
+	UMovieSceneSequencePlayer(const FObjectInitializer&);
+
+	/** Start playback forwards from the current time cursor position, using the current play rate. */
+	UFUNCTION(BlueprintCallable, Category="Game|Cinematic")
+	void Play();
+
+	/** Reverse playback. */
+	UFUNCTION(BlueprintCallable, Category="Game|Cinematic")
+	void PlayReverse();
+
+	/** Changes the direction of playback (go in reverse if it was going forward, or vice versa) */
+	UFUNCTION(BlueprintCallable, Category="Game|Cinematic")
+	void ChangePlaybackDirection();
+
+	/**
+	 * Start playback from the current time cursor position, looping the specified number of times.
+	 * @param NumLoops - The number of loops to play. -1 indicates infinite looping.
+	 */
+	UFUNCTION(BlueprintCallable, Category="Game|Cinematic")
+	void PlayLooping(int32 NumLoops = -1);
+
+	/** Start playback from the current time cursor position, using the current play rate. Does not update the animation until next tick. */
+	UFUNCTION(BlueprintCallable, Category="Game|Cinematic")
+	void StartPlayingNextTick();
+	
+	/** Pause playback. */
+	UFUNCTION(BlueprintCallable, Category="Game|Cinematic")
+	void Pause();
+	
+	/** Stop playback. */
+	UFUNCTION(BlueprintCallable, Category="Game|Cinematic")
+	void Stop();
+
+	/** Get the current playback position */
+	UFUNCTION(BlueprintCallable, Category="Game|Cinematic")
+	float GetPlaybackPosition() const;
+
+	/**
+	 * Set the current playback position
+	 * @param NewPlaybackPosition - The new playback position to set.
+	 * If the animation is currently playing, it will continue to do so from the new position
+	 */
+	UFUNCTION(BlueprintCallable, Category="Game|Cinematic")
+	void SetPlaybackPosition(float NewPlaybackPosition);
+
+	/** Check whether the sequence is actively playing. */
+	UFUNCTION(BlueprintCallable, Category="Game|Cinematic")
+	bool IsPlaying() const;
+
+	/** Get the playback length of the sequence */
+	UFUNCTION(BlueprintCallable, Category="Game|Cinematic")
+	float GetLength() const;
+
+	/** Get the playback rate of this player. */
+	UFUNCTION(BlueprintCallable, Category="Game|Cinematic")
+	float GetPlayRate() const;
+
+	/**
+	 * Set the playback rate of this player. Negative values will play the animation in reverse.
+	 * @param PlayRate - The new rate of playback for the animation.
+	 */
+	UFUNCTION(BlueprintCallable, Category="Game|Cinematic")
+	void SetPlayRate(float PlayRate);
+
+	/**
+	 * Sets the range in time to be played back by this player, overriding the default range stored in the asset
+	 *
+	 * @param	NewStartTime	The new starting time for playback
+	 * @param	NewEndTime		The new ending time for playback.  Must be larger than the start time.
+	 */
+	UFUNCTION(BlueprintCallable, Category="Game|Cinematic")
+	void SetPlaybackRange( const float NewStartTime, const float NewEndTime );
+
+	/** Get the offset within the level sequence to start playing */
+	UFUNCTION(BlueprintCallable, Category="Game|Cinematic")
+	float GetPlaybackStart() const { return StartTime; }
+
+	/** Get the offset within the level sequence to finish playing */
+	UFUNCTION(BlueprintCallable, Category="Game|Cinematic")
+	float GetPlaybackEnd() const { return EndTime; }
+	
+	/** An event that is broadcast each time this level sequence player is updated */
+	DECLARE_EVENT_ThreeParams( UMovieSceneSequencePlayer, FOnMovieSceneSequencePlayerUpdated, const UMovieSceneSequencePlayer&, float /*current time*/, float /*previous time*/ );
+	FOnMovieSceneSequencePlayerUpdated& OnSequenceUpdated() const { return OnMovieSceneSequencePlayerUpdate; }
+
+	/** Event triggered when the level sequence player is played */
+	UPROPERTY(BlueprintAssignable, Category="Game|Cinematic")
+	FOnMovieSceneSequencePlayerEvent OnPlay;
+
+	/** Event triggered when the level sequence player is stopped */
+	UPROPERTY(BlueprintAssignable, Category="Game|Cinematic")
+	FOnMovieSceneSequencePlayerEvent OnStop;
+
+	/** Event triggered when the level sequence player is paused */
+	UPROPERTY(BlueprintAssignable, Category="Game|Cinematic")
+	FOnMovieSceneSequencePlayerEvent OnPause;
+
+public:
+
+	/** Update based on the specified delta seconds */
+	void Update(const float DeltaSeconds);
+
+	/** Initialize this player with a sequence and some settings */
+	void Initialize(UMovieSceneSequence* InSequence, const FMovieSceneSequencePlaybackSettings& InSettings);
+
+public:
+
+	/**
+	 * Access the sequence this player is playing
+	 * @return the sequence currently assigned to this player
+	 */
+	UMovieSceneSequence* GetSequence() const { return Sequence; }
+
+protected:
+
+	void PlayInternal();
+
+	void UpdateMovieSceneInstance(FMovieSceneEvaluationRange InRange);
+
+	void UpdateTimeCursorPosition(float NewPosition);
+
+	bool ShouldStopOrLoop(float NewPosition) const;
+
+	FORCEINLINE float GetSequencePosition() const { return TimeCursorPosition + StartTime; }
+
+protected:
+
+	//~ IMovieScenePlayer interface
+	virtual FMovieSceneRootEvaluationTemplateInstance& GetEvaluationTemplate() override { return RootTemplateInstance; }
+	virtual EMovieScenePlayerStatus::Type GetPlaybackStatus() const override;
+	virtual FMovieSceneSpawnRegister& GetSpawnRegister() override;
+
+	virtual void SetPlaybackStatus(EMovieScenePlayerStatus::Type InPlaybackStatus) override {}
+	virtual void SetViewportSettings(const TMap<FViewportClient*, EMovieSceneViewportParams>& ViewportParamsMap) override {}
+	virtual void GetViewportSettings(TMap<FViewportClient*, EMovieSceneViewportParams>& ViewportParamsMap) const override {}
+	virtual void UpdateCameraCut(UObject* CameraObject, UObject* UnlockIfCameraObject, bool bJumpCut) override {}
+	virtual void ResolveBoundObjects(const FGuid& InBindingId, FMovieSceneSequenceID SequenceID, UMovieSceneSequence& Sequence, UObject* ResolutionContext, TArray<UObject*, TInlineAllocator<1>>& OutObjects) const override;
+
+protected:
+
+	virtual bool CanPlay() const { return true; }
+	virtual void OnStartedPlaying() {}
+	virtual void OnLooped() {}
+	virtual void OnPaused() {}
+	virtual void OnStopped() {}
+	
+private:
+
+	/** Apply any latent actions which may have accumulated while the sequence was being evaluated */
+	void ApplyLatentActions();
+
+protected:
+
+	/** Whether we're currently playing. If false, then sequence playback is paused or was never started. */
+	UPROPERTY()
+	uint32 bIsPlaying : 1;
+
+	/** Whether we're currently playing in reverse. */
+	UPROPERTY()
+	uint32 bReversePlayback : 1;
+
+	/** True where we're waiting for the first update of the sequence after calling StartPlayingNextTick. */
+	UPROPERTY()
+	uint32 bPendingFirstUpdate : 1;
+
+	/** Set to true while evaluating to prevent reentrancy */
+	bool bIsEvaluating : 1;
+
+	/** The sequence to play back */
+	UPROPERTY(transient)
+	UMovieSceneSequence* Sequence;
+
+	/** The current time cursor position within the sequence (in seconds) */
+	UPROPERTY()
+	float TimeCursorPosition;
+
+	/** Time time at which to start playing the sequence (defaults to the lower bound of the sequence's play range) */
+	UPROPERTY()
+	float StartTime;
+
+	/** Time time at which to end playing the sequence (defaults to the upper bound of the sequence's play range) */
+	UPROPERTY()
+	float EndTime;
+
+	/** The number of times we have looped in the current playback */
+	UPROPERTY(transient)
+	int32 CurrentNumLoops;
+
+	enum class ELatentAction
+	{
+		Stop, Pause
+	};
+
+	/** Set of latent actions that are to be performed when the sequence has finished evaluating this frame */
+	TArray<ELatentAction> LatentActions;
+
+	/** Specific playback settings for the animation. */
+	UPROPERTY()
+	FMovieSceneSequencePlaybackSettings PlaybackSettings;
+
+	/** The root template instance we're evaluating */
+	FMovieSceneRootEvaluationTemplateInstance RootTemplateInstance;
+
+	/** Play position helper */
+	FMovieScenePlaybackPosition PlayPosition;
+
+	TSharedPtr<FMovieSceneSpawnRegister> SpawnRegister;
+
+private:
+
+	/** The event that will be broadcast every time the sequence is updated */
+	mutable FOnMovieSceneSequencePlayerUpdated OnMovieSceneSequencePlayerUpdate;
+};
diff --git a/Engine/Source/Runtime/MovieScene/Public/MovieSceneTrack.h b/Engine/Source/Runtime/MovieScene/Public/MovieSceneTrack.h
--- a/Engine/Source/Runtime/MovieScene/Public/MovieSceneTrack.h
+++ b/Engine/Source/Runtime/MovieScene/Public/MovieSceneTrack.h
@@ -121,7 +121,7 @@ enum class EMovieSceneCompileResult : uint8
 /**
  * Base class for a track in a Movie Scene
  */
-UCLASS(abstract, MinimalAPI)
+UCLASS(abstract, DefaultToInstanced, MinimalAPI)
 class UMovieSceneTrack
 	: public UMovieSceneSignedObject
 {
@@ -211,6 +211,11 @@ protected:
 	DEPRECATED(4.15, "Create Instance has been deprecated. Please provide an evaluation template through CreateTemplateForSection instead.")
 	virtual TSharedPtr<class IMovieSceneTrackInstance> CreateInstance() { return nullptr; }
 
+protected:
+
+	//~ UObject interface
+	MOVIESCENE_API virtual void PostInitProperties() override;
+
 public:
 
 	/**
diff --git a/Engine/Source/Runtime/MovieSceneCapture/Private/AutomatedLevelSequenceCapture.cpp b/Engine/Source/Runtime/MovieSceneCapture/Private/AutomatedLevelSequenceCapture.cpp
--- a/Engine/Source/Runtime/MovieSceneCapture/Private/AutomatedLevelSequenceCapture.cpp
+++ b/Engine/Source/Runtime/MovieSceneCapture/Private/AutomatedLevelSequenceCapture.cpp
@@ -473,7 +473,7 @@ void UAutomatedLevelSequenceCapture::Tick(float DeltaSeconds)
 	}
 }
 
-void UAutomatedLevelSequenceCapture::SequenceUpdated(const ULevelSequencePlayer& Player, float CurrentTime, float PreviousTime)
+void UAutomatedLevelSequenceCapture::SequenceUpdated(const UMovieSceneSequencePlayer& Player, float CurrentTime, float PreviousTime)
 {
 	if (bCapturing)
 	{
@@ -575,8 +575,8 @@ void UAutomatedLevelSequenceCapture::ExportEDL()
 	
 	ALevelSequenceActor* Actor = LevelSequenceActor.Get();
 	ULevelSequencePlayer* Player = Actor ? Actor->SequencePlayer : nullptr;
-	ULevelSequence* LevelSequence = Player ? Player->GetLevelSequence() : nullptr;
-	UMovieScene* MovieScene = LevelSequence ? LevelSequence->GetMovieScene() : nullptr;
+	UMovieSceneSequence* Sequence = Player ? Player->GetSequence() : nullptr;
+	UMovieScene* MovieScene = Sequence ? Sequence->GetMovieScene() : nullptr;
 
 	if (!MovieScene)
 	{
diff --git a/Engine/Source/Runtime/MovieSceneCapture/Public/AutomatedLevelSequenceCapture.h b/Engine/Source/Runtime/MovieSceneCapture/Public/AutomatedLevelSequenceCapture.h
--- a/Engine/Source/Runtime/MovieSceneCapture/Public/AutomatedLevelSequenceCapture.h
+++ b/Engine/Source/Runtime/MovieSceneCapture/Public/AutomatedLevelSequenceCapture.h
@@ -84,7 +84,7 @@ private:
 	void UpdateFrameState();
 
 	/** Called when the level sequence has updated the world */
-	void SequenceUpdated(const ULevelSequencePlayer& Player, float CurrentTime, float PreviousTime);
+	void SequenceUpdated(const UMovieSceneSequencePlayer& Player, float CurrentTime, float PreviousTime);
 
 	/** Called to set up the player's playback range */
 	void SetupFrameRange();
diff --git a/Engine/Source/Runtime/MovieSceneCapture/Public/MovieSceneCaptureSettings.h b/Engine/Source/Runtime/MovieSceneCapture/Public/MovieSceneCaptureSettings.h
--- a/Engine/Source/Runtime/MovieSceneCapture/Public/MovieSceneCaptureSettings.h
+++ b/Engine/Source/Runtime/MovieSceneCapture/Public/MovieSceneCaptureSettings.h
@@ -74,7 +74,7 @@ struct MOVIESCENECAPTURE_API FMovieSceneCaptureSettings
 	uint8 ZeroPadFrameNumbers;
 
 	/** The frame rate at which to capture */
-	UPROPERTY(config, EditAnywhere, Category=CaptureSettings)
+	UPROPERTY(config, EditAnywhere, Category=CaptureSettings, meta=(ClampMin=1, UIMin=1, ClampMax=200, UIMax=200))
 	int32 FrameRate;
 
 	/** The resolution at which to capture */
diff --git a/Engine/Source/Runtime/MovieSceneTracks/MovieSceneTracks.Build.cs b/Engine/Source/Runtime/MovieSceneTracks/MovieSceneTracks.Build.cs
--- a/Engine/Source/Runtime/MovieSceneTracks/MovieSceneTracks.Build.cs
+++ b/Engine/Source/Runtime/MovieSceneTracks/MovieSceneTracks.Build.cs
@@ -26,6 +26,7 @@ public class MovieSceneTracks : ModuleRules
 		PrivateDependencyModuleNames.AddRange(
 			new string[] {
 				"SlateCore",
+                "AnimGraphRuntime"
 			}
 		);
 	}
diff --git a/Engine/Source/Runtime/MovieSceneTracks/Private/Evaluation/MovieScene3DAttachTemplate.cpp b/Engine/Source/Runtime/MovieSceneTracks/Private/Evaluation/MovieScene3DAttachTemplate.cpp
--- a/Engine/Source/Runtime/MovieSceneTracks/Private/Evaluation/MovieScene3DAttachTemplate.cpp
+++ b/Engine/Source/Runtime/MovieSceneTracks/Private/Evaluation/MovieScene3DAttachTemplate.cpp
@@ -7,10 +7,14 @@
 #include "Components/SceneComponent.h"
 #include "GameFramework/Actor.h"
 #include "MovieSceneCommonHelpers.h"
+#include "MovieSceneEvaluation.h"
+#include "IMovieScenePlayer.h"
+
 
 DECLARE_CYCLE_STAT(TEXT("Attach Track Evaluate"), MovieSceneEval_AttachTrack_Evaluate, STATGROUP_MovieSceneEval);
 DECLARE_CYCLE_STAT(TEXT("Attach Track Token Execute"), MovieSceneEval_AttachTrack_TokenExecute, STATGROUP_MovieSceneEval);
 
+
 /** A token that sets a component's attach */
 struct F3DAttachTrackToken
 {
diff --git a/Engine/Source/Runtime/MovieSceneTracks/Private/Evaluation/MovieScene3DPathTemplate.cpp b/Engine/Source/Runtime/MovieSceneTracks/Private/Evaluation/MovieScene3DPathTemplate.cpp
--- a/Engine/Source/Runtime/MovieSceneTracks/Private/Evaluation/MovieScene3DPathTemplate.cpp
+++ b/Engine/Source/Runtime/MovieSceneTracks/Private/Evaluation/MovieScene3DPathTemplate.cpp
@@ -1,12 +1,13 @@
 // Copyright 1998-2017 Epic Games, Inc. All Rights Reserved.
 
 #include "Evaluation/MovieScene3DPathTemplate.h"
-
 #include "Evaluation/MovieSceneTemplateCommon.h"
 #include "MovieSceneCommonHelpers.h"
-
 #include "GameFramework/Actor.h"
 #include "Components/SplineComponent.h"
+#include "MovieSceneEvaluation.h"
+#include "IMovieScenePlayer.h"
+
 
 DECLARE_CYCLE_STAT(TEXT("Path Track Evaluate"), MovieSceneEval_PathTrack_Evaluate, STATGROUP_MovieSceneEval);
 DECLARE_CYCLE_STAT(TEXT("Path Track Token Execute"), MovieSceneEval_PathTrack_TokenExecute, STATGROUP_MovieSceneEval);
diff --git a/Engine/Source/Runtime/MovieSceneTracks/Private/Evaluation/MovieScene3DTransformTemplate.cpp b/Engine/Source/Runtime/MovieSceneTracks/Private/Evaluation/MovieScene3DTransformTemplate.cpp
--- a/Engine/Source/Runtime/MovieSceneTracks/Private/Evaluation/MovieScene3DTransformTemplate.cpp
+++ b/Engine/Source/Runtime/MovieSceneTracks/Private/Evaluation/MovieScene3DTransformTemplate.cpp
@@ -1,14 +1,17 @@
 // Copyright 1998-2017 Epic Games, Inc. All Rights Reserved.
 
 #include "Evaluation/MovieScene3DTransformTemplate.h"
-
 #include "Sections/MovieScene3DTransformSection.h"
 #include "Evaluation/MovieSceneTemplateCommon.h"
 #include "MovieSceneCommonHelpers.h"
+#include "MovieSceneEvaluation.h"
+#include "IMovieScenePlayer.h"
+
 
 DECLARE_CYCLE_STAT(TEXT("Transform Track Evaluate"), MovieSceneEval_TransformTrack_Evaluate, STATGROUP_MovieSceneEval);
 DECLARE_CYCLE_STAT(TEXT("Transform Track Token Execute"), MovieSceneEval_TransformTrack_TokenExecute, STATGROUP_MovieSceneEval);
 
+
 /** A movie scene execution token that stores a specific transform, and an operand */
 struct F3DTransformTrackExecutionToken
 	: IMovieSceneExecutionToken
diff --git a/Engine/Source/Runtime/MovieSceneTracks/Private/Evaluation/MovieSceneActorReferenceTemplate.cpp b/Engine/Source/Runtime/MovieSceneTracks/Private/Evaluation/MovieSceneActorReferenceTemplate.cpp
--- a/Engine/Source/Runtime/MovieSceneTracks/Private/Evaluation/MovieSceneActorReferenceTemplate.cpp
+++ b/Engine/Source/Runtime/MovieSceneTracks/Private/Evaluation/MovieSceneActorReferenceTemplate.cpp
@@ -5,6 +5,8 @@
 #include "Sections/MovieSceneActorReferenceSection.h"
 #include "Tracks/MovieSceneActorReferenceTrack.h"
 #include "GameFramework/Actor.h"
+#include "MovieSceneEvaluation.h"
+
 
 namespace PropertyTemplate
 {
diff --git a/Engine/Source/Runtime/MovieSceneTracks/Private/Evaluation/MovieSceneAudioTemplate.cpp b/Engine/Source/Runtime/MovieSceneTracks/Private/Evaluation/MovieSceneAudioTemplate.cpp
--- a/Engine/Source/Runtime/MovieSceneTracks/Private/Evaluation/MovieSceneAudioTemplate.cpp
+++ b/Engine/Source/Runtime/MovieSceneTracks/Private/Evaluation/MovieSceneAudioTemplate.cpp
@@ -10,14 +10,17 @@
 #include "AudioThread.h"
 #include "AudioDevice.h"
 #include "ActiveSound.h"
-
 #include "Sections/MovieSceneAudioSection.h"
 #include "Tracks/MovieSceneAudioTrack.h"
 #include "MovieScene.h"
+#include "MovieSceneEvaluation.h"
+#include "IMovieScenePlayer.h"
+
 
 DECLARE_CYCLE_STAT(TEXT("Audio Track Evaluate"), MovieSceneEval_AudioTrack_Evaluate, STATGROUP_MovieSceneEval);
 DECLARE_CYCLE_STAT(TEXT("Audio Track Token Execute"), MovieSceneEval_AudioTrack_TokenExecute, STATGROUP_MovieSceneEval);
 
+
 /** Stop audio from playing */
 struct FStopAudioPreAnimatedToken : IMovieScenePreAnimatedToken
 {
@@ -232,9 +235,10 @@ struct FAudioSectionExecutionToken : IMovieSceneExecutionToken
 
 FMovieSceneAudioSectionTemplateData::FMovieSceneAudioSectionTemplateData(const UMovieSceneAudioSection& Section)
 	: Sound(Section.GetSound())
-	, AudioStartTime(Section.GetAudioStartTime())
-	, AudioDilationFactor(Section.GetAudioDilationFactor())
-	, AudioVolume(Section.GetAudioVolume())
+	, AudioStartOffset(Section.GetStartOffset())
+	, AudioRange(Section.GetAudioRange())
+	, AudioPitchMultiplierCurve(Section.GetPitchMultiplierCurve())
+	, AudioVolumeCurve(Section.GetSoundVolumeCurve())
 	, RowIndex(Section.GetRowIndex())
 {
 }
@@ -244,18 +248,27 @@ void FMovieSceneAudioSectionTemplateData::EnsureAudioIsPlaying(UAudioComponent&
 	Player.SavePreAnimatedState(AudioComponent, FStopAudioPreAnimatedToken::GetAnimTypeID(), FStopAudioPreAnimatedToken::FProducer());
 
 	bool bPlaySound = !AudioComponent.IsPlaying() || AudioComponent.Sound != Sound;
-	if (bPlaySound)
+	float AudioVolume = AudioVolumeCurve.Eval(Context.GetTime());
+	float PitchMultiplier = AudioPitchMultiplierCurve.Eval(Context.GetTime());
+
+	if (AudioComponent.VolumeMultiplier != AudioVolume)
+	{
+		AudioComponent.SetVolumeMultiplier(AudioVolume);
+	}
+	
+	if (AudioComponent.PitchMultiplier != PitchMultiplier)
 	{
-		float PitchMultiplier = 1.f / AudioDilationFactor;
+		AudioComponent.SetPitchMultiplier(PitchMultiplier);
+	}
 
+	if (bPlaySound)
+	{
 		AudioComponent.bAllowSpatialization = bAllowSpatialization;
 		AudioComponent.bOverrideAttenuation = bAllowSpatialization;
 		AudioComponent.Stop();
 		AudioComponent.SetSound(Sound);
-		AudioComponent.SetVolumeMultiplier(AudioVolume);
-		AudioComponent.SetPitchMultiplier(PitchMultiplier);
 		AudioComponent.bIsUISound = true;
-		AudioComponent.Play(Context.GetTime() - AudioStartTime);
+		AudioComponent.Play(AudioStartOffset < 0 ? 0 : AudioStartOffset);
 
 		if (Context.GetStatus() == EMovieScenePlayerStatus::Scrubbing)
 		{
@@ -295,7 +308,7 @@ void FMovieSceneAudioSectionTemplate::Evaluate(const FMovieSceneEvaluationOperan
 {
 	MOVIESCENE_DETAILED_SCOPE_CYCLE_COUNTER(MovieSceneEval_AudioTrack_Evaluate)
 
-	if (GEngine && GEngine->UseSound() && !Context.HasJumped() && Context.GetStatus() != EMovieScenePlayerStatus::Jumping)
+	if (GEngine && GEngine->UseSound() && !Context.HasJumped() && Context.GetStatus() != EMovieScenePlayerStatus::Jumping && AudioData.AudioRange.Contains(Context.GetTime()))
 	{
 		ExecutionTokens.Add(FAudioSectionExecutionToken(AudioData));
 	}
diff --git a/Engine/Source/Runtime/MovieSceneTracks/Private/Evaluation/MovieSceneAudioTemplate.h b/Engine/Source/Runtime/MovieSceneTracks/Private/Evaluation/MovieSceneAudioTemplate.h
--- a/Engine/Source/Runtime/MovieSceneTracks/Private/Evaluation/MovieSceneAudioTemplate.h
+++ b/Engine/Source/Runtime/MovieSceneTracks/Private/Evaluation/MovieSceneAudioTemplate.h
@@ -26,17 +26,21 @@ struct FMovieSceneAudioSectionTemplateData
 	UPROPERTY()
 	USoundBase* Sound;
 
-	/** The absolute time that the sound starts playing at */
+	/** The offset into the beginning of the audio clip */
 	UPROPERTY()
-	float AudioStartTime;
+	float AudioStartOffset;
 	
+	/** The time range this audio plays */
+	UPROPERTY()
+	FFloatRange AudioRange;
+
 	/** The amount which this audio is time dilated by */
 	UPROPERTY()
-	float AudioDilationFactor;
+	FRichCurve AudioPitchMultiplierCurve;
 
 	/** The volume the sound will be played with. */
 	UPROPERTY()
-	float AudioVolume;
+	FRichCurve AudioVolumeCurve;
 
 	/** The row index of the section */
 	UPROPERTY()
diff --git a/Engine/Source/Runtime/MovieSceneTracks/Private/Evaluation/MovieSceneCameraAnimTemplate.cpp b/Engine/Source/Runtime/MovieSceneTracks/Private/Evaluation/MovieSceneCameraAnimTemplate.cpp
--- a/Engine/Source/Runtime/MovieSceneTracks/Private/Evaluation/MovieSceneCameraAnimTemplate.cpp
+++ b/Engine/Source/Runtime/MovieSceneTracks/Private/Evaluation/MovieSceneCameraAnimTemplate.cpp
@@ -8,7 +8,10 @@
 #include "Camera/CameraAnimInst.h"
 #include "Engine/World.h"
 #include "MovieSceneCommonHelpers.h"
+#include "MovieSceneEvaluation.h"
 #include "UObject/Package.h"
+#include "IMovieScenePlayer.h"
+
 
 DECLARE_CYCLE_STAT(TEXT("Camera Anim Track Token Execute"), MovieSceneEval_CameraAnimTrack_TokenExecute, STATGROUP_MovieSceneEval);
 
diff --git a/Engine/Source/Runtime/MovieSceneTracks/Private/Evaluation/MovieSceneCameraCutTemplate.cpp b/Engine/Source/Runtime/MovieSceneTracks/Private/Evaluation/MovieSceneCameraCutTemplate.cpp
--- a/Engine/Source/Runtime/MovieSceneTracks/Private/Evaluation/MovieSceneCameraCutTemplate.cpp
+++ b/Engine/Source/Runtime/MovieSceneTracks/Private/Evaluation/MovieSceneCameraCutTemplate.cpp
@@ -2,6 +2,9 @@
 
 #include "Evaluation/MovieSceneCameraCutTemplate.h"
 #include "Sections/MovieSceneCameraCutSection.h"
+#include "MovieSceneEvaluation.h"
+#include "IMovieScenePlayer.h"
+
 
 DECLARE_CYCLE_STAT(TEXT("Camera Cut Track Token Execute"), MovieSceneEval_CameraCutTrack_TokenExecute, STATGROUP_MovieSceneEval);
 
diff --git a/Engine/Source/Runtime/MovieSceneTracks/Private/Evaluation/MovieSceneColorTemplate.cpp b/Engine/Source/Runtime/MovieSceneTracks/Private/Evaluation/MovieSceneColorTemplate.cpp
--- a/Engine/Source/Runtime/MovieSceneTracks/Private/Evaluation/MovieSceneColorTemplate.cpp
+++ b/Engine/Source/Runtime/MovieSceneTracks/Private/Evaluation/MovieSceneColorTemplate.cpp
@@ -7,9 +7,13 @@
 #include "MovieSceneCommonHelpers.h"
 #include "Components/LightComponent.h"
 #include "Styling/SlateColor.h"
+#include "MovieSceneEvaluation.h"
+#include "IMovieScenePlayer.h"
+
 
 DECLARE_CYCLE_STAT(TEXT("Color Track Token Execute"), MovieSceneEval_ColorTrack_TokenExecute, STATGROUP_MovieSceneEval);
 
+
 enum class EColorType : uint8
 {
 	/** FSlateColor */
diff --git a/Engine/Source/Runtime/MovieSceneTracks/Private/Evaluation/MovieSceneEventTemplate.cpp b/Engine/Source/Runtime/MovieSceneTracks/Private/Evaluation/MovieSceneEventTemplate.cpp
--- a/Engine/Source/Runtime/MovieSceneTracks/Private/Evaluation/MovieSceneEventTemplate.cpp
+++ b/Engine/Source/Runtime/MovieSceneTracks/Private/Evaluation/MovieSceneEventTemplate.cpp
@@ -6,9 +6,13 @@
 #include "Evaluation/MovieSceneEvaluationTemplateInstance.h"
 #include "EngineGlobals.h"
 #include "MovieScene.h"
+#include "MovieSceneEvaluation.h"
+#include "IMovieScenePlayer.h"
+
 
 DECLARE_CYCLE_STAT(TEXT("Event Track Token Execute"), MovieSceneEval_EventTrack_TokenExecute, STATGROUP_MovieSceneEval);
 
+
 struct FEventData
 {
 	FEventData(const FEventPayload& InPayload, float InGlobalPosition) : Payload(InPayload), GlobalPosition(InGlobalPosition) {}
@@ -50,14 +54,18 @@ struct FEventTrackExecutionToken
 		}
 
 #if !UE_BUILD_SHIPPING
-		UWorld* World = Player.GetPlaybackContext()->GetWorld();
-		if (World)
+		if (PerformanceCaptureEventPositions.Num())
 		{
-			FString LevelSequenceName = Player.GetEvaluationTemplate().GetSequence(MovieSceneSequenceID::Root)->GetName();
-			
-			for (float EventPosition : PerformanceCaptureEventPositions)
+			UObject* PlaybackContext = Player.GetPlaybackContext();
+			UWorld* World = PlaybackContext ? PlaybackContext->GetWorld() : nullptr;
+			if (World)
 			{
-				GEngine->PerformanceCapture(World, World->GetName(), LevelSequenceName, EventPosition);
+				FString LevelSequenceName = Player.GetEvaluationTemplate().GetSequence(MovieSceneSequenceID::Root)->GetName();
+			
+				for (float EventPosition : PerformanceCaptureEventPositions)
+				{
+					GEngine->PerformanceCapture(World, World->GetName(), LevelSequenceName, EventPosition);
+				}
 			}
 		}
 #endif	// UE_BUILD_SHIPPING
diff --git a/Engine/Source/Runtime/MovieSceneTracks/Private/Evaluation/MovieSceneFadeTemplate.cpp b/Engine/Source/Runtime/MovieSceneTracks/Private/Evaluation/MovieSceneFadeTemplate.cpp
--- a/Engine/Source/Runtime/MovieSceneTracks/Private/Evaluation/MovieSceneFadeTemplate.cpp
+++ b/Engine/Source/Runtime/MovieSceneTracks/Private/Evaluation/MovieSceneFadeTemplate.cpp
@@ -7,6 +7,8 @@
 #include "Engine/Engine.h"
 #include "EngineGlobals.h"
 #include "GameFramework/PlayerController.h"
+#include "IMovieScenePlayer.h"
+#include "MovieSceneEvaluation.h"
 
 DECLARE_CYCLE_STAT(TEXT("Fade Track Token Execute"), MovieSceneEval_FadeTrack_TokenExecute, STATGROUP_MovieSceneEval);
 
@@ -34,24 +36,18 @@ struct FFadeTrackToken
 		Player.GetViewportSettings(ViewportParamsMap);
 		for (auto ViewportParamsPair : ViewportParamsMap)
 		{
-			ViewportParamsMap[ViewportParamsPair.Key] = ViewportParams;
+			ViewportParamsMap[ViewportParamsPair.Key] = ViewportParams;	
 		}
 		Player.SetViewportSettings(ViewportParamsMap);
 
 		// Set runtime fade
-		for (const FWorldContext& Context : GEngine->GetWorldContexts())
+		UWorld* World = Cast<UWorld>(Player.GetPlaybackContext());
+		if (World && (World->WorldType == EWorldType::Game || World->WorldType == EWorldType::PIE))
 		{
-			if (Context.WorldType == EWorldType::Game || Context.WorldType == EWorldType::PIE)
+			APlayerController* PlayerController = World->GetGameInstance()->GetFirstLocalPlayerController();
+			if (PlayerController != nullptr && PlayerController->PlayerCameraManager && !PlayerController->PlayerCameraManager->IsPendingKill())
 			{
-				UWorld* World = Context.World();
-				if (World != nullptr)
-				{
-					APlayerController* PlayerController = World->GetGameInstance()->GetFirstLocalPlayerController();
-					if (PlayerController != nullptr && PlayerController->PlayerCameraManager && !PlayerController->PlayerCameraManager->IsPendingKill())
-					{
-						PlayerController->PlayerCameraManager->SetManualCameraFade(FadeValue, FadeColor, bFadeAudio);
-					}
-				}
+				PlayerController->PlayerCameraManager->SetManualCameraFade(FadeValue, FadeColor, bFadeAudio);
 			}
 		}
 	}
diff --git a/Engine/Source/Runtime/MovieSceneTracks/Private/Evaluation/MovieSceneLevelVisibilityTemplate.cpp b/Engine/Source/Runtime/MovieSceneTracks/Private/Evaluation/MovieSceneLevelVisibilityTemplate.cpp
--- a/Engine/Source/Runtime/MovieSceneTracks/Private/Evaluation/MovieSceneLevelVisibilityTemplate.cpp
+++ b/Engine/Source/Runtime/MovieSceneTracks/Private/Evaluation/MovieSceneLevelVisibilityTemplate.cpp
@@ -1,11 +1,13 @@
 // Copyright 1998-2017 Epic Games, Inc. All Rights Reserved.
 
 #include "Evaluation/MovieSceneLevelVisibilityTemplate.h"
-
 #include "Engine/LevelStreaming.h"
 #include "Engine/EngineTypes.h"
 #include "Engine/World.h"
 #include "Misc/PackageName.h"
+#include "MovieSceneEvaluation.h"
+#include "IMovieScenePlayer.h"
+
 
 bool GetLevelVisibility(const ULevelStreaming& Level)
 {
diff --git a/Engine/Source/Runtime/MovieSceneTracks/Private/Evaluation/MovieSceneParameterTemplate.cpp b/Engine/Source/Runtime/MovieSceneTracks/Private/Evaluation/MovieSceneParameterTemplate.cpp
--- a/Engine/Source/Runtime/MovieSceneTracks/Private/Evaluation/MovieSceneParameterTemplate.cpp
+++ b/Engine/Source/Runtime/MovieSceneTracks/Private/Evaluation/MovieSceneParameterTemplate.cpp
@@ -3,6 +3,8 @@
 #include "Evaluation/MovieSceneParameterTemplate.h"
 #include "Tracks/MovieSceneMaterialTrack.h"
 #include "Components/PrimitiveComponent.h"
+#include "MovieSceneEvaluation.h"
+
 
 FMovieSceneParameterSectionTemplate::FMovieSceneParameterSectionTemplate(const UMovieSceneParameterSection& Section)
 	: Scalars(Section.GetScalarParameterNamesAndCurves())
