Sequencer: Template generation now ignores tracks that would be optimized out during cook
  - This prevents the template from containing orphaned objects at runtime with cooked content

diff --git a/Engine/Source/Runtime/MovieScene/Private/Compilation/MovieSceneEvaluationTemplateGenerator.cpp b/Engine/Source/Runtime/MovieScene/Private/Compilation/MovieSceneEvaluationTemplateGenerator.cpp
--- a/Engine/Source/Runtime/MovieScene/Private/Compilation/MovieSceneEvaluationTemplateGenerator.cpp
+++ b/Engine/Source/Runtime/MovieScene/Private/Compilation/MovieSceneEvaluationTemplateGenerator.cpp
@@ -156,6 +156,15 @@ void FMovieSceneEvaluationTemplateGenerator::Generate(FMovieSceneTrackCompilatio
 
 	for (const FMovieSceneBinding& ObjectBinding : MovieScene->GetBindings())
 	{
+#if WITH_EDITOR
+		// Skip object bindings that are optimized out
+		auto ShouldRemoveObject = [](const UMovieSceneTrack* Track){ return EnumHasAnyFlags(Track->GetCookOptimizationFlags(), ECookOptimizationFlags::RemoveObject); };
+		if (ObjectBinding.GetTracks().ContainsByPredicate(ShouldRemoveObject))
+		{
+			continue;
+		}
+#endif
+
 		for (UMovieSceneTrack* Track : ObjectBinding.GetTracks())
 		{
 			ProcessTrack(*Track, ObjectBinding.GetObjectGuid());
@@ -194,6 +203,14 @@ void FMovieSceneEvaluationTemplateGenerator::Generate(FMovieSceneTrackCompilatio
 
 void FMovieSceneEvaluationTemplateGenerator::ProcessTrack(const UMovieSceneTrack& InTrack, const FGuid& ObjectId)
 {
+#if WITH_EDITOR
+	// Skip tracks that are optimized out
+	if (EnumHasAnyFlags(InTrack.GetCookOptimizationFlags(), ECookOptimizationFlags::RemoveTrack | ECookOptimizationFlags::RemoveObject))
+	{
+		return;
+	}
+#endif
+
 	FGuid Signature = InTrack.GetSignature();
 
 	// See if this track signature already exists in the ledger, if it does, we don't need to regenerate it
