Copying //UE4/Dev-Sequencer to //UE4/Dev-Main (Source: //UE4/Dev-Sequencer @ 3800792)

diff --git a/Engine/Source/Runtime/MediaAssets/Private/Misc/MediaTextureResource.cpp b/Engine/Source/Runtime/MediaAssets/Private/Misc/MediaTextureResource.cpp
--- a/Engine/Source/Runtime/MediaAssets/Private/Misc/MediaTextureResource.cpp
+++ b/Engine/Source/Runtime/MediaAssets/Private/Misc/MediaTextureResource.cpp
@@ -128,9 +128,10 @@ namespace MediaTextureResource
 /* FMediaTextureResource structors
  *****************************************************************************/
 
-FMediaTextureResource::FMediaTextureResource(UMediaTexture& InOwner, FIntPoint& InOwnerDim, SIZE_T& InOwnerSize)
+FMediaTextureResource::FMediaTextureResource(UMediaTexture& InOwner, FIntPoint& InOwnerDim, SIZE_T& InOwnerSize, FLinearColor InClearColor, FGuid InTextureGuid)
 	: Cleared(false)
-	, CurrentClearColor(FLinearColor::Transparent)
+	, CurrentClearColor(InClearColor)
+	, InitialTextureGuid(InTextureGuid)
 	, Owner(InOwner)
 	, OwnerDim(InOwnerDim)
 	, OwnerSize(InOwnerSize)
@@ -144,6 +145,9 @@ void FMediaTextureResource::Render(const FRenderParams& Params)
 {
 	check(IsInRenderingThread());
 
+	FLinearColor Rotation(1, 0, 0, 1);
+	FLinearColor Offset(0, 0, 0, 0);
+
 	TSharedPtr<FMediaTextureSampleSource, ESPMode::ThreadSafe> SampleSource = Params.SampleSource.Pin();
 
 	if (SampleSource.IsValid())
@@ -157,7 +161,7 @@ void FMediaTextureResource::Render(const FRenderParams& Params)
 			const FTimespan StartTime = Sample->GetTime();
 			const FTimespan EndTime = StartTime + Sample->GetDuration();
 
-			if (((Params.Rate > 0.0f) && (StartTime > Params.Time)) ||
+			if (((Params.Rate > 0.0f) && (StartTime >= Params.Time)) ||
 				((Params.Rate < 0.0f) && (EndTime <= Params.Time)))
 			{
 				break; // future sample
@@ -166,57 +170,58 @@ void FMediaTextureResource::Render(const FRenderParams& Params)
 			SampleValid = SampleSource->Dequeue(Sample);
 		}
 
-		if (!SampleValid)
-		{
-			return; // no sample to render
-		}
-
-		check(Sample.IsValid());
-
-		// render the sample
-		if (Sample->GetOutputDim().GetMin() <= 0)
+		if (SampleValid)
 		{
-			#if MEDIATEXTURERESOURCE_TRACE_RENDER
-				UE_LOG(LogMediaAssets, VeryVerbose, TEXT("TextureResource %p: Corrupt sample with time %s at time %s"), this, *Sample->GetTime().ToString(), *Params.Time.ToString());
-			#endif
+			// render the sample
+			if (Sample->GetOutputDim().GetMin() <= 0)
+			{
+				#if MEDIATEXTURERESOURCE_TRACE_RENDER
+					UE_LOG(LogMediaAssets, VeryVerbose, TEXT("TextureResource %p: Corrupt sample with time %s at time %s"), this, *Sample->GetTime().ToString(), *Params.Time.ToString());
+				#endif
 
-			ClearTexture(FLinearColor::Red, Params.SrgbOutput); // mark corrupt sample
-		}
-		else if (MediaTextureResource::RequiresConversion(Sample, Params.SrgbOutput))
-		{
-			#if MEDIATEXTURERESOURCE_TRACE_RENDER
-				UE_LOG(LogMediaAssets, VeryVerbose, TEXT("TextureResource %p: Converting sample with time %s at time %s"), this, *Sample->GetTime().ToString(), *Params.Time.ToString());
-			#endif
+				ClearTexture(FLinearColor::Red, Params.SrgbOutput); // mark corrupt sample
+			}
+			else if (MediaTextureResource::RequiresConversion(Sample, Params.SrgbOutput))
+			{
+				#if MEDIATEXTURERESOURCE_TRACE_RENDER
+					UE_LOG(LogMediaAssets, VeryVerbose, TEXT("TextureResource %p: Converting sample with time %s at time %s"), this, *Sample->GetTime().ToString(), *Params.Time.ToString());
+				#endif
 
-			ConvertSample(Sample, Params.ClearColor, Params.SrgbOutput);
-		}
-		else
-		{
-			#if MEDIATEXTURERESOURCE_TRACE_RENDER
-				UE_LOG(LogMediaAssets, VeryVerbose, TEXT("TextureResource %p: Copying sample with time %s at time %s"), this, *Sample->GetTime().ToString(), *Params.Time.ToString());
-			#endif
+				ConvertSample(Sample, Params.ClearColor, Params.SrgbOutput);
+			}
+			else
+			{
+				#if MEDIATEXTURERESOURCE_TRACE_RENDER
+					UE_LOG(LogMediaAssets, VeryVerbose, TEXT("TextureResource %p: Copying sample with time %s at time %s"), this, *Sample->GetTime().ToString(), *Params.Time.ToString());
+				#endif
 
-			CopySample(Sample, Params.ClearColor, Params.SrgbOutput);
-		}
+				CopySample(Sample, Params.ClearColor, Params.SrgbOutput);
+			}
 
-		if (!GSupportsImageExternal && Params.PlayerGuid.IsValid())
-		{
-			FTextureRHIRef VideoTexture = (FTextureRHIRef)Owner.TextureReference.TextureReferenceRHI;
-			FExternalTextureRegistry::Get().RegisterExternalTexture(Params.PlayerGuid, VideoTexture, SamplerStateRHI, Sample->GetScaleRotation(), Sample->GetOffset());
+			Rotation = Sample->GetScaleRotation();
+			Offset = Sample->GetOffset();
 		}
 	}
-	else if (!Cleared)
+	else if (!Cleared || (Params.ClearColor != CurrentClearColor))
 	{
 		#if MEDIATEXTURERESOURCE_TRACE_RENDER
 			UE_LOG(LogMediaAssets, VeryVerbose, TEXT("TextureResource %p: Clearing texture at time %s"), this, *Params.Time.ToString());
 		#endif
 
 		ClearTexture(Params.ClearColor, Params.SrgbOutput);
+	}
 
-		if (!GSupportsImageExternal && Params.PlayerGuid.IsValid())
+	if (!GSupportsImageExternal)
+	{
+		if (Params.TextureGuid.IsValid())
 		{
 			FTextureRHIRef VideoTexture = (FTextureRHIRef)Owner.TextureReference.TextureReferenceRHI;
-			FExternalTextureRegistry::Get().RegisterExternalTexture(Params.PlayerGuid, VideoTexture, SamplerStateRHI, FLinearColor(1.0f, 0.0f, 0.0f, 1.0f), FLinearColor(0.0f, 0.0f, 0.0f, 0.0f));
+			FExternalTextureRegistry::Get().RegisterExternalTexture(Params.TextureGuid, VideoTexture, SamplerStateRHI, Rotation, Offset);
+		}
+
+		if (Params.LastGuid.IsValid() && (Params.LastGuid != Params.TextureGuid))
+		{
+			FExternalTextureRegistry::Get().UnregisterExternalTexture(Params.LastGuid);
 		}
 	}
 }
@@ -263,6 +268,21 @@ void FMediaTextureResource::InitDynamicRHI()
 	);
 
 	SamplerStateRHI = RHICreateSamplerState(SamplerStateInitializer);
+
+	// Note: set up default texture, or we can get sampler bind errors on render
+	// we can't leave here without having a valid bindable resource for some RHIs.
+
+	ClearTexture(CurrentClearColor, Owner.SRGB);
+
+	check(TextureRHI.IsValid());
+	check(RenderTargetTextureRHI.IsValid());
+	check(OutputTarget.IsValid());
+
+	if (!GSupportsImageExternal)
+	{
+		FTextureRHIRef VideoTexture = (FTextureRHIRef)Owner.TextureReference.TextureReferenceRHI;
+		FExternalTextureRegistry::Get().RegisterExternalTexture(InitialTextureGuid, VideoTexture, SamplerStateRHI);
+	}
 }
 
 
diff --git a/Engine/Source/Runtime/MediaAssets/Private/Misc/MediaTextureResource.h b/Engine/Source/Runtime/MediaAssets/Private/Misc/MediaTextureResource.h
--- a/Engine/Source/Runtime/MediaAssets/Private/Misc/MediaTextureResource.h
+++ b/Engine/Source/Runtime/MediaAssets/Private/Misc/MediaTextureResource.h
@@ -35,9 +35,12 @@ public:
 	 * Creates and initializes a new instance.
 	 *
 	 * @param InOwner The Movie texture object to create a resource for (must not be nullptr).
-	 * @param InSink The sink that receives texture samples from the media player.
+	 * @param InOwnerDim Reference to the width and height of the texture that owns this resource (will be updated by resource).
+	 * @param InOWnerSize Reference to the size in bytes of the texture that owns this resource (will be updated by resource).
+	 * @param InClearColor The initial clear color.
+	 * @param InTextureGuid The initial external texture GUID.
 	 */
-	FMediaTextureResource(UMediaTexture& InOwner, FIntPoint& InOwnerDim, SIZE_T& InOwnerSize);
+	FMediaTextureResource(UMediaTexture& InOwner, FIntPoint& InOwnerDim, SIZE_T& InOwnerSize, FLinearColor InClearColor, FGuid InTextureGuid);
 
 	/** Virtual destructor. */
 	virtual ~FMediaTextureResource() { }
@@ -50,8 +53,8 @@ public:
 		/** The clear color to use when clearing the texture. */
 		FLinearColor ClearColor;
 
-		/** Guid associated with media player. */
-		FGuid PlayerGuid;
+		/** The external texture GUID at the previously rendered frame. */
+		FGuid LastGuid;
 
 		/** The player's play rate. */
 		float Rate;
@@ -62,6 +65,9 @@ public:
 		/** Whether output should be in sRGB color space. */
 		bool SrgbOutput;
 
+		/** Guid associated with the texture. */
+		FGuid TextureGuid;
+
 		/** The time of the video frame to render (in player's clock). */
 		FTimespan Time;
 	};
@@ -140,6 +146,9 @@ private:
 	/** Tracks the current clear color. */
 	FLinearColor CurrentClearColor;
 
+	/** The external texture GUID to use when initializing this resource. */
+	FGuid InitialTextureGuid;
+
 	/** Input render target if the texture samples don't provide one (for conversions). */
 	TRefCountPtr<FRHITexture2D> InputTarget;
 
diff --git a/Engine/Source/Runtime/MediaAssets/Public/MediaPlaylist.h b/Engine/Source/Runtime/MediaAssets/Public/MediaPlaylist.h
--- a/Engine/Source/Runtime/MediaAssets/Public/MediaPlaylist.h
+++ b/Engine/Source/Runtime/MediaAssets/Public/MediaPlaylist.h
@@ -19,12 +19,6 @@ class MEDIAASSETS_API UMediaPlaylist
 {
 	GENERATED_BODY()
 
-public:
-
-	/** Whether the play list should loop (default = true). */
-	UPROPERTY(BlueprintReadWrite, EditAnywhere, Category=Playback)
-	uint32 Loop:1;
-
 public:
 
 	/**
diff --git a/Engine/Source/Runtime/MediaAssets/Public/MediaSoundComponent.h b/Engine/Source/Runtime/MediaAssets/Public/MediaSoundComponent.h
--- a/Engine/Source/Runtime/MediaAssets/Public/MediaSoundComponent.h
+++ b/Engine/Source/Runtime/MediaAssets/Public/MediaSoundComponent.h
@@ -50,13 +50,9 @@ class MEDIAASSETS_API UMediaSoundComponent
 public:
 
 	/** Media sound channel type. */
-	UPROPERTY(EditAnywhere, Category="UMediaSoundComponent")
+	UPROPERTY(EditAnywhere, Category="MediaSoundComponent")
 	EMediaSoundChannels Channels;
 
-	/** The media player asset associated with this component. */
-	UPROPERTY(BlueprintReadWrite, EditAnywhere, Category="Media")
-	UMediaPlayer* MediaPlayer;
-
 public:
 
 	/**
@@ -69,6 +65,16 @@ public:
 	/** Virtual destructor. */
 	~UMediaSoundComponent();
 
+public:
+
+	/**
+	 * Set the media player that provides the audio samples.
+	 *
+	 * @param NewMediaPlayer The player to set.
+	 */
+	UFUNCTION(BlueprintCallable, Category="Media|MediaSoundComponent")
+	void SetMediaPlayer(UMediaPlayer* NewMediaPlayer);
+
 public:
 
 	void UpdatePlayer();
@@ -86,6 +92,16 @@ public:
 	virtual void Activate(bool bReset = false) override;
 	virtual void Deactivate() override;
 
+public:
+
+	//~ UObject interface
+
+	virtual void PostLoad() override;
+
+#if WITH_EDITOR
+	virtual void PostEditChangeProperty(FPropertyChangedEvent& PropertyChangedEvent) override;
+#endif
+
 protected:
 
 	//~ USynthComponent interface
@@ -93,11 +109,33 @@ protected:
 	virtual bool Init(int32& SampleRate) override;
 	virtual void OnGenerateAudio(float* OutAudio, int32 NumSamples) override;
 
+protected:
+
+	/**
+	 * The media player asset associated with this component.
+	 *
+	 * This property is meant for design-time convenience. To change the
+	 * associated media player at run-time, use the SetMediaPlayer method.
+	 *
+	 * @see SetMediaPlayer
+	 */
+	UPROPERTY(EditAnywhere, Category="Media")
+	UMediaPlayer* MediaPlayer;
+
 private:
 
+	/** The player's current play rate (cached for use on audio thread). */
+	float CachedRate;
+
+	/** The player's current time (cached for use on audio thread). */
+	FTimespan CachedTime;
+
 	/** Critical section for synchronizing access to PlayerFacadePtr. */
 	FCriticalSection CriticalSection;
 
+	/** The player that is currently associated with this component. */
+	TWeakObjectPtr<UMediaPlayer> CurrentPlayer;
+
 	/** The player facade that's currently providing texture samples. */
 	TWeakPtr<FMediaPlayerFacade, ESPMode::ThreadSafe> CurrentPlayerFacade;
 
diff --git a/Engine/Source/Runtime/MediaAssets/Public/MediaTexture.h b/Engine/Source/Runtime/MediaAssets/Public/MediaTexture.h
--- a/Engine/Source/Runtime/MediaAssets/Public/MediaTexture.h
+++ b/Engine/Source/Runtime/MediaAssets/Public/MediaTexture.h
@@ -46,10 +46,6 @@ class MEDIAASSETS_API UMediaTexture
 	UPROPERTY(EditAnywhere, BlueprintReadWrite, Category="MediaTexture")
 	FLinearColor ClearColor;
 
-	/** The media player asset associated with this texture. */
-	UPROPERTY(EditAnywhere, BlueprintReadWrite, Category="Media")
-	UMediaPlayer* MediaPlayer;
-
 public:
 
 	/**
@@ -70,6 +66,15 @@ public:
 	UFUNCTION(BlueprintCallable, Category="Media|MediaTexture")
 	int32 GetHeight() const;
 
+	/**
+	 * Get the media player that provides the video samples.
+	 *
+	 * @return The texture's media player, or nullptr if not set.
+	 * @see SetMediaPlayer
+	 */
+	UFUNCTION(BlueprintCallable, Category="Media|MediaTexture")
+	UMediaPlayer* GetMediaPlayer() const;
+
 	/**
 	 * Gets the current width of the texture.
 	 *
@@ -79,6 +84,15 @@ public:
 	UFUNCTION(BlueprintCallable, Category="Media|MediaTexture")
 	int32 GetWidth() const;
 
+	/**
+	 * Set the media player that provides the video samples.
+	 *
+	 * @param NewMediaPlayer The player to set.
+	 * @see GetMediaPlayer
+	 */
+	UFUNCTION(BlueprintCallable, Category="Media|MediaTexture")
+	void SetMediaPlayer(UMediaPlayer* NewMediaPlayer);
+
 public:
 
 	//~ UTexture interface.
@@ -96,6 +110,7 @@ public:
 
 	virtual FString GetDesc() override;
 	virtual void GetResourceSizeEx(FResourceSizeEx& CumulativeResourceSize) override;
+	virtual void PostLoad() override;
 
 #if WITH_EDITOR
 	virtual void PostEditChangeProperty(FPropertyChangedEvent& PropertyChangedEvent) override;
@@ -110,8 +125,18 @@ protected:
 	 */
 	void TickResource(FTimespan Timecode);
 
-	/** Unregister the player's external texture GUID. */
-	void UnregisterPlayerGuid();
+protected:
+
+	/**
+	 * The media player asset associated with this texture.
+	 *
+	 * This property is meant for design-time convenience. To change the
+	 * associated media player at run-time, use the SetMediaPlayer method.
+	 *
+	 * @see SetMediaPlayer
+	 */
+	UPROPERTY(EditAnywhere, Category="Media")
+	UMediaPlayer* MediaPlayer;
 
 private:
 
@@ -120,14 +145,23 @@ private:
 	/** The texture's media clock sink. */
 	TSharedPtr<FMediaTextureClockSink, ESPMode::ThreadSafe> ClockSink;
 
-	/** The player facade that's currently providing texture samples. */
-	TWeakPtr<FMediaPlayerFacade, ESPMode::ThreadSafe> CurrentPlayerFacade;
+	/** The default external texture GUID if no media player is assigned. */
+	/*TAtomic<*/FGuid/*>*/ CurrentGuid;
+
+	/** The player that is currently associated with this texture. */
+	TWeakObjectPtr<UMediaPlayer> CurrentPlayer;
+
+	/** The default external texture GUID if no media player is assigned. */
+	const FGuid DefaultGuid;
 
 	/** Current width and height of the resource (in pixels). */
 	FIntPoint Dimensions;
 
-	/** The previously used player GUID. */
-	FGuid LastPlayerGuid;
+	/** The previously clear color. */
+	FLinearColor LastClearColor;
+
+	/** The previously used sRGB flag. */
+	bool LastSrgb;
 
 	/** Texture sample queue. */
 	TSharedPtr<FMediaTextureSampleQueue, ESPMode::ThreadSafe> SampleQueue;
diff --git a/Engine/Source/Runtime/MediaUtils/Public/MediaSampleQueue.h b/Engine/Source/Runtime/MediaUtils/Public/MediaSampleQueue.h
--- a/Engine/Source/Runtime/MediaUtils/Public/MediaSampleQueue.h
+++ b/Engine/Source/Runtime/MediaUtils/Public/MediaSampleQueue.h
@@ -114,8 +114,7 @@ public:
 			return false; // pending flush
 		}
 
-		check(Samples.Pop());
-		check(Sample.IsValid());
+		Samples.Pop();
 
 		FPlatformAtomics::InterlockedDecrement(&NumSamples);
 		check(NumSamples >= 0);
diff --git a/Engine/Source/Runtime/MovieScene/Private/Compilation/MovieSceneCompiler.cpp b/Engine/Source/Runtime/MovieScene/Private/Compilation/MovieSceneCompiler.cpp
new file mode 100644
--- /dev/null
+++ b/Engine/Source/Runtime/MovieScene/Private/Compilation/MovieSceneCompiler.cpp
@@ -0,0 +1,600 @@
+
+#include "MovieSceneCompiler.h"
+
+#include "IMovieSceneModule.h"
+#include "MovieSceneEvaluationTemplate.h"
+#include "MovieSceneEvaluationTemplateInstance.h"
+#include "Compilation/MovieSceneEvaluationTemplateGenerator.h"
+
+#include "MovieScene.h"
+#include "MovieSceneSequence.h"
+#include "Sections/MovieSceneSubSection.h"
+#include "Tracks/MovieSceneSubTrack.h"
+
+/** Parameter structure used for gathering entities for a given time or range */
+struct FGatherParameters
+{
+	FGatherParameters(FMovieSceneRootOverridePath& InRootPath, FMovieSceneSequenceHierarchy& InRootHierarchy, IMovieSceneSequenceTemplateStore& InTemplateStore, TRange<float> InCompileRange)
+		: RootPath(InRootPath)
+		, RootHierarchy(InRootHierarchy)
+		, TemplateStore(InTemplateStore)
+		, RootCompileRange(InCompileRange)
+		, RootClampRange(TRange<float>::All())
+		, LocalCompileRange(RootCompileRange)
+		, LocalClampRange(RootClampRange)
+		, Flags(ESectionEvaluationFlags::None)
+		, HierarchicalBias(0)
+	{}
+
+	FGatherParameters CreateForSubData(const FMovieSceneSubSequenceData& SubData) const
+	{
+		FGatherParameters SubParams = *this;
+
+		SubParams.RootToSequenceTransform   = SubData.RootToSequenceTransform;
+		SubParams.HierarchicalBias          = SubData.HierarchicalBias;
+
+		SubParams.LocalCompileRange         = SubParams.RootCompileRange * SubData.RootToSequenceTransform;
+		SubParams.LocalClampRange           = SubParams.RootClampRange   * SubData.RootToSequenceTransform;
+
+		return SubParams;
+	}
+
+	void SetClampRange(TRange<float> InNewRootClampRange)
+	{
+		RootClampRange  = InNewRootClampRange;
+		LocalClampRange = InNewRootClampRange * RootToSequenceTransform;
+	}
+
+	/** Path from root to current sequence */
+	FMovieSceneRootOverridePath& RootPath;
+	/** Hierarchy for the root sequence template */
+	FMovieSceneSequenceHierarchy& RootHierarchy;
+	/** Store from which to retrieve templates */
+	IMovieSceneSequenceTemplateStore& TemplateStore;
+
+	/** The range that is being compiled in the root's time-space */
+	TRange<float> RootCompileRange;
+	/** A range to clamp compilation to in the root's time-space */
+	TRange<float> RootClampRange;
+
+	/** The range that is being compiled in the current sequence's time-space */
+	TRange<float> LocalCompileRange;
+	/** A range to clamp compilation to in the current sequence's time-space */
+	TRange<float> LocalClampRange;
+
+	/** Evaluation flags for the current sequence */
+	ESectionEvaluationFlags Flags;
+
+	/** Transform from the root time-space to the current sequence's time-space */
+	FMovieSceneSequenceTransform RootToSequenceTransform;
+
+	/** Current accumulated hierarchical bias */
+	int32 HierarchicalBias;
+};
+
+struct FCompileOnTheFlyData
+{
+	/** Primary sort - group */
+	uint16 GroupEvaluationPriority;
+	/** Secondary sort - Hierarchical bias */
+	int32 HierarchicalBias;
+	/** Tertiary sort - Eval priority */
+	int32 EvaluationPriority;
+
+	/** Whether the track requires initialization or not */
+	bool bRequiresInit;
+	/** Cached ptr to the evaluation track */
+	const FMovieSceneEvaluationTrack* Track;
+	/** Cached segment ptr within the above track */
+	FMovieSceneEvaluationFieldSegmentPtr Segment;
+};
+
+/** Gathered data for a given time or range */
+struct FMovieSceneGatheredCompilerData
+{
+	/** Intersection of any empty space that overlaps the currently evaluating time range */
+	TRange<float> EmptyOverlappingRange;
+	/** Tree of tracks to evaluate */
+	TMovieSceneEvaluationTree<FCompileOnTheFlyData> Tracks;
+	/** Tree of active sequences */
+	TMovieSceneEvaluationTree<FMovieSceneSequenceID> Sequences;
+};
+
+IMovieSceneModule& GetMovieSceneModule()
+{
+	static TWeakPtr<IMovieSceneModule> WeakMovieSceneModule;
+
+	TSharedPtr<IMovieSceneModule> Shared = WeakMovieSceneModule.Pin();
+	if (!Shared.IsValid())
+	{
+		WeakMovieSceneModule = IMovieSceneModule::Get().GetWeakPtr();
+		Shared = WeakMovieSceneModule.Pin();
+	}
+	check(Shared.IsValid());
+
+	return *Shared;
+}
+
+void FMovieSceneCompiler::Compile(UMovieSceneSequence& InCompileSequence, IMovieSceneSequenceTemplateStore& InTemplateStore)
+{
+	FMovieSceneEvaluationTemplate& CompileTemplate = InTemplateStore.AccessTemplate(InCompileSequence);
+
+	// Pass down a mutable path to the gather functions
+	FMovieSceneRootOverridePath RootPath;
+
+	// Gather everything that happens, recursively
+	FMovieSceneGatheredCompilerData GatherData;
+	FGatherParameters GatherParams(RootPath, CompileTemplate.Hierarchy, InTemplateStore, TRange<float>::All());
+	GatherCompileOnTheFlyData(InCompileSequence, GatherParams, GatherData);
+
+	// Wipe the current evaluation field for the template
+	CompileTemplate.EvaluationField = FMovieSceneEvaluationField();
+
+	TArray<FCompileOnTheFlyData> CompileData;
+
+	for (FMovieSceneEvaluationTreeRangeIterator It(GatherData.Tracks); It; ++It)
+	{
+		CompileData.Reset();
+
+		for (const FCompileOnTheFlyData& TrackData : GatherData.Tracks.GetAllData(It.Node()))
+		{
+			CompileData.Add(TrackData);
+		}
+
+		// Sort the compilation data based on (in order):
+		//  1. Group
+		//  2. Hierarchical bias
+		//  3. Evaluation priority
+		CompileData.Sort(SortPredicate);
+
+		// Compose the final result for the compiled range
+		FCompiledGroupResult Result(It.Range());
+
+		// Generate the evaluation group by gathering initialization and evaluation ptrs for each unique group
+		PopulateEvaluationGroup(Result, CompileData);
+
+		// Compute meta data for this segment
+		TMovieSceneEvaluationTreeDataIterator<FMovieSceneSequenceID> SubSequences = GatherData.Sequences.GetAllData(GatherData.Sequences.IterateFromLowerBound(It.Range().GetLowerBound()).Node());
+		PopulateMetaData(Result, CompileTemplate.Hierarchy, CompileData, SubSequences);
+
+		CompileTemplate.EvaluationField.Add(Result.Range, MoveTemp(Result.Group), MoveTemp(Result.MetaData));
+	}
+}
+
+TOptional<FCompiledGroupResult> FMovieSceneCompiler::CompileTime(float InGlobalTime, UMovieSceneSequence& InCompileSequence, IMovieSceneSequenceTemplateStore& InTemplateStore)
+{
+	FMovieSceneEvaluationTemplate& CompileTemplate = InTemplateStore.AccessTemplate(InCompileSequence);
+
+	TRange<float> CompileRange = TRange<float>::Inclusive(InGlobalTime, InGlobalTime);
+
+	// Pass down a mutable path to the gather functions
+	FMovieSceneRootOverridePath RootPath;
+
+	// Gather everything that happens at this time, recursively
+	FMovieSceneGatheredCompilerData GatherData;
+	FGatherParameters GatherParams(RootPath, CompileTemplate.Hierarchy, InTemplateStore, CompileRange);
+	GatherCompileOnTheFlyData(InCompileSequence, GatherParams, GatherData);
+
+	FMovieSceneEvaluationTreeRangeIterator NodeAtTime = GatherData.Tracks.IterateFromTime(InGlobalTime);
+	if (ensure(NodeAtTime))
+	{
+		TRange<float> CompiledRange = TRange<float>::Intersection(
+			TRange<float>::Intersection(
+				GatherData.EmptyOverlappingRange,
+				GatherData.Sequences.IterateFromLowerBound(CompileRange.GetLowerBound()).Range()
+			),
+			NodeAtTime.Range());
+
+		// Always include the global time in the compiled range - this is necessary when floating point rounding
+		// introduced during sub sequence transformations shift the range to just outside the current time
+		CompiledRange = TRange<float>::Hull(CompiledRange, CompileRange);
+
+		TArray<FCompileOnTheFlyData> SortedCompileData;
+		for (const FCompileOnTheFlyData& TrackData : GatherData.Tracks.GetAllData(NodeAtTime.Node()))
+		{
+			SortedCompileData.Add(TrackData);
+		}
+
+		// Sort the compilation data based on (in order):
+		//  1. Group
+		//  2. Hierarchical bias
+		//  3. Evaluation priority
+		SortedCompileData.Sort(SortPredicate);
+
+		// Compose the final result for the compiled range
+		FCompiledGroupResult Result(CompiledRange);
+
+		// Generate the evaluation group by gathering initialization and evaluation ptrs for each unique group
+		PopulateEvaluationGroup(Result, SortedCompileData);
+
+		// Compute meta data for this segment
+		TMovieSceneEvaluationTreeDataIterator<FMovieSceneSequenceID> SubSequences = GatherData.Sequences.GetAllData(GatherData.Sequences.IterateFromLowerBound(CompileRange.GetLowerBound()).Node());
+		PopulateMetaData(Result, CompileTemplate.Hierarchy, SortedCompileData, SubSequences);
+
+		return Result;
+	}
+
+	return TOptional<FCompiledGroupResult>();
+}
+
+void FMovieSceneCompiler::CompileHierarchy(const UMovieSceneSequence& InRootSequence, FMovieSceneSequenceHierarchy& OutHierarchy, FMovieSceneSequenceID RootSequenceID, int32 MaxDepth)
+{
+	FMovieSceneRootOverridePath Path;
+	Path.Set(RootSequenceID, OutHierarchy);
+
+	CompileHierarchy(InRootSequence, OutHierarchy, Path, MaxDepth);
+}
+
+void FMovieSceneCompiler::CompileHierarchy(const UMovieSceneSequence& InSequence, FMovieSceneSequenceHierarchy& OutHierarchy, FMovieSceneRootOverridePath& Path, int32 MaxDepth)
+{
+	UMovieScene* MovieScene = InSequence.GetMovieScene();
+	if (!MovieScene)
+	{
+		return;
+	}
+
+	FMovieSceneSequenceID ParentID = Path.Remap(MovieSceneSequenceID::Root);
+
+	// Remove all existing children
+	if (FMovieSceneSequenceHierarchyNode* ExistingNode = OutHierarchy.FindNode(ParentID))
+	{
+		OutHierarchy.Remove(ExistingNode->Children);
+	}
+
+	auto ProcessSection = [&](const UMovieSceneSection* Section, const FGuid& InObjectBindingId)
+	{
+		const UMovieSceneSubSection* SubSection = Cast<UMovieSceneSubSection>(Section);
+		const UMovieSceneSequence* SubSequence = SubSection ? SubSection->GetSequence() : nullptr;
+
+		if (SubSequence)
+		{
+			FMovieSceneSequenceID DeterministicID = SubSection->GetSequenceID();
+
+			GetOrCreateSubSequenceData(Path.Remap(DeterministicID), ParentID, *SubSection, InObjectBindingId, OutHierarchy);
+
+			int32 NewMaxDepth = MaxDepth == -1 ? -1 : MaxDepth - 1;
+			if (NewMaxDepth == -1 || NewMaxDepth > 1)
+			{
+				Path.Push(DeterministicID);
+
+				CompileHierarchy(*SubSequence, OutHierarchy, Path, NewMaxDepth);
+
+				Path.Pop();
+			}
+		}
+	};
+
+	for (const UMovieSceneTrack* Track : MovieScene->GetMasterTracks())
+	{
+		for (const UMovieSceneSection* Section : Track->GetAllSections())
+		{
+			ProcessSection(Section, FGuid());
+		}
+	}
+
+	for (const FMovieSceneBinding& ObjectBinding : MovieScene->GetBindings())
+	{
+		for (const UMovieSceneTrack* Track : ObjectBinding.GetTracks())
+		{
+			for (const UMovieSceneSection* Section : Track->GetAllSections())
+			{
+				ProcessSection(Section, ObjectBinding.GetObjectGuid());
+			}
+		}
+	}
+}
+
+void FMovieSceneCompiler::GatherCompileOnTheFlyData(UMovieSceneSequence& InSequence, const FGatherParameters& Params, FMovieSceneGatheredCompilerData& OutData)
+{
+	// Regenerate the track structure if it's out of date
+	FMovieSceneEvaluationTemplate& Template = Params.TemplateStore.AccessTemplate(InSequence);
+	if (Template.SequenceSignature != InSequence.GetSignature())
+	{
+		FMovieSceneEvaluationTemplateGenerator(InSequence, Template).Generate();
+	}
+
+	// Iterate tracks within this template
+	for (TTuple<FMovieSceneTrackIdentifier, FMovieSceneEvaluationTrack>& Pair : Template.GetTracks())
+	{
+		const bool bTrackMatchesFlags = ( Params.Flags == ESectionEvaluationFlags::None )
+			|| ( EnumHasAnyFlags(Params.Flags, ESectionEvaluationFlags::PreRoll)  && Pair.Value.ShouldEvaluateInPreroll()  )
+			|| ( EnumHasAnyFlags(Params.Flags, ESectionEvaluationFlags::PostRoll) && Pair.Value.ShouldEvaluateInPostroll() );
+
+		if (bTrackMatchesFlags)
+		{
+			GatherCompileDataForTrack(Pair.Value, Pair.Key, Params, OutData);
+		}
+	}
+
+	TRange<float> CompileClampIntersection = TRange<float>::Intersection(Params.LocalCompileRange, Params.LocalClampRange);
+
+	// Iterate sub section ranges that overlap with the compile range
+	FGatherParameters SubSectionGatherParams = Params;
+
+	const TMovieSceneEvaluationTree<FMovieSceneSubSectionData>& SubSectionField = Template.GetSubSectionField();
+
+	// Start iterating the field from the lower bound of the compile range
+	FMovieSceneEvaluationTreeRangeIterator SubSectionIt(SubSectionField.IterateFromLowerBound(CompileClampIntersection.GetLowerBound()));
+
+	// Intersect the unique range in the tree with the current overlapping empty range to constrict the resulting compile range in the case where this is a gap between sub sections
+	OutData.EmptyOverlappingRange = TRange<float>::Intersection(OutData.EmptyOverlappingRange, SubSectionIt.Range() * Params.RootToSequenceTransform.Inverse());
+
+	for ( ; SubSectionIt && SubSectionIt.Range().Overlaps(CompileClampIntersection); ++SubSectionIt)
+	{
+		TRange<float> ThisSegmentRange = SubSectionIt.Range() * Params.RootToSequenceTransform.Inverse();
+		if (ThisSegmentRange.IsEmpty())
+		{
+			continue;
+		}
+
+		SubSectionGatherParams.SetClampRange(ThisSegmentRange);
+
+		// Iterate all sub sections in the current range
+		for (const FMovieSceneSubSectionData& SubSectionData : SubSectionField.GetAllData(SubSectionIt.Node()))
+		{
+			UMovieSceneSubSection* SubSection = SubSectionData.Section.Get();
+			if (!SubSection)
+			{
+				continue;
+			}
+
+			UMovieSceneSubTrack* SubTrack = SubSection->GetTypedOuter<UMovieSceneSubTrack>();
+
+			const bool bTrackMatchesFlags = ( Params.Flags == ESectionEvaluationFlags::None )
+				|| ( EnumHasAnyFlags(Params.Flags, ESectionEvaluationFlags::PreRoll)  && SubTrack && SubTrack->EvalOptions.bEvaluateInPreroll  )
+				|| ( EnumHasAnyFlags(Params.Flags, ESectionEvaluationFlags::PostRoll) && SubTrack && SubTrack->EvalOptions.bEvaluateInPostroll );
+
+			if (bTrackMatchesFlags && SubSection)
+			{
+				SubSectionGatherParams.Flags = SubSectionData.Flags;
+
+				GatherCompileDataForSubSection(*SubSection, SubSectionData.ObjectBindingId, SubSectionGatherParams, OutData);
+			}
+		}
+	}
+}
+
+void FMovieSceneCompiler::GatherCompileDataForSubSection(const UMovieSceneSubSection& SubSection, const FGuid& InObjectBindingId, const FGatherParameters& Params, FMovieSceneGatheredCompilerData& OutData)
+{
+	UMovieSceneSequence* SubSequence = SubSection.GetSequence();
+	if (!SubSequence)
+	{
+		return;
+	}
+
+	FMovieSceneSequenceID UnAccumulatedSequenceID = SubSection.GetSequenceID();
+
+	// Hash this source ID with the outer sequence ID to make it unique
+	FMovieSceneSequenceID ParentSequenceID = Params.RootPath.Remap(MovieSceneSequenceID::Root);
+	FMovieSceneSequenceID InnerSequenceID  = Params.RootPath.Remap(UnAccumulatedSequenceID);
+
+	// Add the active sequence ID to each range. We add each range individually since this range may inform the final compiled range
+	OutData.Sequences.Add(Params.RootClampRange, InnerSequenceID);
+
+	// Add this sub sequence ID to the root path
+	Params.RootPath.Push(UnAccumulatedSequenceID);
+
+	// Find/add sub data in the root template
+	TOptional<FGatherParameters> SubParams;
+
+	{
+		const FMovieSceneSubSequenceData* CompilationSubData = GetOrCreateSubSequenceData(InnerSequenceID, ParentSequenceID, SubSection, InObjectBindingId, Params.RootHierarchy);
+		check(CompilationSubData);
+
+		SubParams = Params.CreateForSubData(*CompilationSubData);
+		// Any code after this point may reallocate the root hierarchy, so CompilationSubData cannot be used
+	}
+
+	GatherCompileOnTheFlyData(*SubSequence, SubParams.GetValue(), OutData);
+
+	// Pop the path off the root path
+	Params.RootPath.Pop();
+}
+
+const FMovieSceneSubSequenceData* FMovieSceneCompiler::GetOrCreateSubSequenceData(FMovieSceneSequenceID InnerSequenceID, FMovieSceneSequenceID ParentSequenceID, const UMovieSceneSubSection& SubSection, const FGuid& InObjectBindingId, FMovieSceneSequenceHierarchy& InOutHierarchy)
+{
+	// Find/add sub data in the root template
+	const FMovieSceneSubSequenceData* SubData = InOutHierarchy.FindSubData(InnerSequenceID);
+	if (SubData)
+	{
+		return SubData;
+	}
+	
+	FSubSequenceInstanceDataParams InstanceParams{ InnerSequenceID, FMovieSceneEvaluationOperand(ParentSequenceID, InObjectBindingId) };
+	FMovieSceneSubSequenceData NewSubData = SubSection.GenerateSubSequenceData(InstanceParams);
+
+	// Intersect this inner sequence's valid play range with the parent's if possible
+	const FMovieSceneSubSequenceData* ParentSubData = ParentSequenceID != MovieSceneSequenceID::Root ? InOutHierarchy.FindSubData(ParentSequenceID) : nullptr;
+	if (ParentSubData)
+	{
+		TRange<float> ParentPlayRangeChildSpace = ParentSubData->PlayRange * NewSubData.RootToSequenceTransform;
+		NewSubData.PlayRange = TRange<float>::Intersection(ParentPlayRangeChildSpace, NewSubData.PlayRange);
+
+		// Accumulate parent transform
+		NewSubData.RootToSequenceTransform = NewSubData.RootToSequenceTransform * ParentSubData->RootToSequenceTransform;
+
+		// Accumulate parent hierarchical bias
+		NewSubData.HierarchicalBias += ParentSubData->HierarchicalBias;
+	}
+
+	// Add the sub data to the root hierarchy
+	InOutHierarchy.Add(NewSubData, InnerSequenceID, ParentSequenceID);
+
+	return InOutHierarchy.FindSubData(InnerSequenceID);
+}
+
+void FMovieSceneCompiler::GatherCompileDataForTrack(FMovieSceneEvaluationTrack& Track, FMovieSceneTrackIdentifier TrackID, const FGatherParameters& Params, FMovieSceneGatheredCompilerData& OutData)
+{
+	auto RequiresInit = [&Track](FSectionEvaluationData EvalData)
+	{
+		return Track.HasChildTemplate(EvalData.ImplIndex) && Track.GetChildTemplate(EvalData.ImplIndex).RequiresInitialization();
+	};
+
+	FMovieSceneSequenceTransform SequenceToRootTransform = Params.RootToSequenceTransform.Inverse();
+	FMovieSceneSequenceID CurrentSequenceID              = Params.RootPath.Remap(MovieSceneSequenceID::Root);
+
+	TArray<FMovieSceneSegmentIdentifier> SegmentIDs      = Track.GetSegmentsInRange(Params.LocalCompileRange);
+	if (SegmentIDs.Num() == 0)
+	{
+		// No segment at this time, so just report the time range of the empty space.
+		TRange<float> EmptyTrackSpace = Track.GetUniqueRangeFromLowerBound(Params.LocalCompileRange.GetLowerBound());
+		TRange<float> ClampedEmptyTrackSpaceRoot = TRange<float>::Intersection(Params.LocalClampRange, EmptyTrackSpace) * SequenceToRootTransform;
+
+		OutData.EmptyOverlappingRange = TRange<float>::Intersection(ClampedEmptyTrackSpaceRoot, OutData.EmptyOverlappingRange);
+	}
+	else for (FMovieSceneSegmentIdentifier SegmentID : SegmentIDs)
+	{
+		const FMovieSceneSegment& ThisSegment = Track.GetSegment(SegmentID);
+
+		FCompileOnTheFlyData Data;
+		Data.Segment = FMovieSceneEvaluationFieldSegmentPtr(CurrentSequenceID, TrackID, SegmentID);
+		Data.GroupEvaluationPriority = GetMovieSceneModule().GetEvaluationGroupParameters(Track.GetEvaluationGroup()).EvaluationPriority;
+		Data.HierarchicalBias = Params.HierarchicalBias;
+		Data.EvaluationPriority = Track.GetEvaluationPriority();
+		Data.Track = &Track;
+		Data.bRequiresInit = ThisSegment.Impls.ContainsByPredicate(RequiresInit);
+
+		TRange<float> IntersectionRange = TRange<float>::Intersection(Params.LocalClampRange, ThisSegment.Range) * SequenceToRootTransform;
+		if (!IntersectionRange.IsEmpty())
+		{
+			OutData.Tracks.Add(IntersectionRange, Data);
+		}
+	}
+}
+
+void FMovieSceneCompiler::PopulateMetaData(FCompiledGroupResult& OutResult, const FMovieSceneSequenceHierarchy& RootHierarchy, const TArray<FCompileOnTheFlyData>& SortedCompileData, TMovieSceneEvaluationTreeDataIterator<FMovieSceneSequenceID> SubSequences)
+{
+	OutResult.MetaData.Reset();
+
+	// Add all the init tracks first
+	uint32 SortOrder = 0;
+	for (const FCompileOnTheFlyData& CompileData : SortedCompileData)
+	{
+		if (!CompileData.bRequiresInit)
+		{
+			continue;
+		}
+
+		FMovieSceneEvaluationFieldSegmentPtr SegmentPtr = CompileData.Segment;
+
+		// Add the track key
+		FMovieSceneEvaluationKey TrackKey(SegmentPtr.SequenceID, SegmentPtr.TrackIdentifier);
+		OutResult.MetaData.ActiveEntities.Add(FMovieSceneOrderedEvaluationKey{ TrackKey, SortOrder++ });
+
+		for (FSectionEvaluationData EvalData : CompileData.Track->GetSegment(SegmentPtr.SegmentID).Impls)
+		{
+			FMovieSceneEvaluationKey SectionKey = TrackKey.AsSection(EvalData.ImplIndex);
+			OutResult.MetaData.ActiveEntities.Add(FMovieSceneOrderedEvaluationKey{ SectionKey, SortOrder++ });
+		}
+	}
+
+	// Then all the eval tracks
+	for (const FCompileOnTheFlyData& CompileData : SortedCompileData)
+	{
+		if (CompileData.bRequiresInit)
+		{
+			continue;
+		}
+
+		FMovieSceneEvaluationFieldSegmentPtr SegmentPtr = CompileData.Segment;
+
+		// Add the track key
+		FMovieSceneEvaluationKey TrackKey(SegmentPtr.SequenceID, SegmentPtr.TrackIdentifier);
+		OutResult.MetaData.ActiveEntities.Add(FMovieSceneOrderedEvaluationKey{ TrackKey, SortOrder++ });
+
+		for (FSectionEvaluationData EvalData : CompileData.Track->GetSegment(SegmentPtr.SegmentID).Impls)
+		{
+			FMovieSceneEvaluationKey SectionKey = TrackKey.AsSection(EvalData.ImplIndex);
+			OutResult.MetaData.ActiveEntities.Add(FMovieSceneOrderedEvaluationKey{ SectionKey, SortOrder++ });
+		}
+	}
+
+	Algo::SortBy(OutResult.MetaData.ActiveEntities, &FMovieSceneOrderedEvaluationKey::Key);
+
+	{
+		OutResult.MetaData.ActiveSequences.Reset();
+		OutResult.MetaData.ActiveSequences.Add(MovieSceneSequenceID::Root);
+
+		for (FMovieSceneSequenceID SequenceID : SubSequences)
+		{
+			const FMovieSceneSubSequenceData* SubData = RootHierarchy.FindSubData(SequenceID);
+			check(SubData);
+
+			UMovieSceneSequence* Sequence = SubData->GetSequence();
+
+			OutResult.MetaData.ActiveSequences.Add(SequenceID);
+			OutResult.MetaData.SubSequenceSignatures.Add(SequenceID, Sequence ? Sequence->GetSignature() : FGuid());
+		}
+
+		OutResult.MetaData.ActiveSequences.Sort();
+	}
+}
+
+bool FMovieSceneCompiler::SortPredicate(const FCompileOnTheFlyData& A, const FCompileOnTheFlyData& B)
+{
+	if (A.GroupEvaluationPriority != B.GroupEvaluationPriority)
+	{
+		return A.GroupEvaluationPriority > B.GroupEvaluationPriority;
+	}
+	else if (A.HierarchicalBias != B.HierarchicalBias)
+	{
+		return A.HierarchicalBias < B.HierarchicalBias;
+	}
+	else
+	{
+		return A.EvaluationPriority > B.EvaluationPriority;
+	}
+}
+
+void FMovieSceneCompiler::AddPtrsToGroup(FMovieSceneEvaluationGroup& Group, TArray<FMovieSceneEvaluationFieldSegmentPtr>& InitPtrs, TArray<FMovieSceneEvaluationFieldSegmentPtr>& EvalPtrs)
+{
+	if (!InitPtrs.Num() && !EvalPtrs.Num())
+	{
+		return;
+	}
+
+	FMovieSceneEvaluationGroupLUTIndex Index;
+
+	Index.LUTOffset = Group.SegmentPtrLUT.Num();
+	Index.NumInitPtrs = InitPtrs.Num();
+	Index.NumEvalPtrs = EvalPtrs.Num();
+
+	Group.LUTIndices.Add(Index);
+	Group.SegmentPtrLUT.Append(InitPtrs);
+	Group.SegmentPtrLUT.Append(EvalPtrs);
+
+	InitPtrs.Reset();
+	EvalPtrs.Reset();
+}
+
+void FMovieSceneCompiler::PopulateEvaluationGroup(FCompiledGroupResult& OutResult, const TArray<FCompileOnTheFlyData>& SortedCompileData)
+{
+	TArray<FMovieSceneEvaluationFieldSegmentPtr> EvalPtrs;
+	TArray<FMovieSceneEvaluationFieldSegmentPtr> InitPtrs;
+
+	// Now iterate the tracks and insert indices for initialization and evaluation
+	FName CurrentEvaluationGroup, LastEvaluationGroup;
+
+	for (const FCompileOnTheFlyData& Data : SortedCompileData)
+	{
+		// If we're now in a different flush group, add the ptrs to the group
+		{
+			CurrentEvaluationGroup = Data.Track->GetEvaluationGroup();
+			if (CurrentEvaluationGroup != LastEvaluationGroup)
+			{
+				AddPtrsToGroup(OutResult.Group, InitPtrs, EvalPtrs);
+			}
+			LastEvaluationGroup = CurrentEvaluationGroup;
+		}
+
+		// If this track requires initialization, add it to the init array
+		if (Data.bRequiresInit)
+		{
+			InitPtrs.Add(Data.Segment);
+		}
+
+		// All tracks require evaluation implicitly
+		EvalPtrs.Add(Data.Segment);
+	}
+	AddPtrsToGroup(OutResult.Group, InitPtrs, EvalPtrs);
+}
diff --git a/Engine/Source/Runtime/MovieScene/Private/Compilation/MovieSceneEvaluationTemplateGenerator.cpp b/Engine/Source/Runtime/MovieScene/Private/Compilation/MovieSceneEvaluationTemplateGenerator.cpp
--- a/Engine/Source/Runtime/MovieScene/Private/Compilation/MovieSceneEvaluationTemplateGenerator.cpp
+++ b/Engine/Source/Runtime/MovieScene/Private/Compilation/MovieSceneEvaluationTemplateGenerator.cpp
@@ -6,144 +6,32 @@
 #include "Algo/Sort.h"
 #include "UObject/ObjectKey.h"
 #include "IMovieSceneModule.h"
+#include "MovieSceneSubTrack.h"
+#include "MovieSceneSubSection.h"
+#include "MovieSceneSegmentCompiler.h"
 
-namespace 
+FMovieSceneEvaluationTemplateGenerator::FMovieSceneEvaluationTemplateGenerator(UMovieSceneSequence& InSequence, FMovieSceneEvaluationTemplate& OutTemplate)
+	: SourceSequence(InSequence), Template(OutTemplate)
 {
-	void AddPtrsToGroup(FMovieSceneEvaluationGroup& Group, TArray<FMovieSceneEvaluationFieldSegmentPtr>& InitPtrs, TArray<FMovieSceneEvaluationFieldSegmentPtr>& EvalPtrs)
-	{
-		if (!InitPtrs.Num() && !EvalPtrs.Num())
-		{
-			return;
-		}
-
-		FMovieSceneEvaluationGroupLUTIndex Index;
-
-		Index.LUTOffset = Group.SegmentPtrLUT.Num();
-		Index.NumInitPtrs = InitPtrs.Num();
-		Index.NumEvalPtrs = EvalPtrs.Num();
-
-		Group.LUTIndices.Add(Index);
-		Group.SegmentPtrLUT.Append(InitPtrs);
-		Group.SegmentPtrLUT.Append(EvalPtrs);
-
-		InitPtrs.Reset();
-		EvalPtrs.Reset();
-	}
-}
 
-FMovieSceneEvaluationTemplateGenerator::FMovieSceneEvaluationTemplateGenerator(UMovieSceneSequence& InSequence, FMovieSceneEvaluationTemplate& OutTemplate, FMovieSceneSequenceTemplateStore& InStore)
-	: SourceSequence(InSequence)
-	, Template(OutTemplate)
-	, TransientArgs(*this, InStore)
-{
-}
-
-void FMovieSceneEvaluationTemplateGenerator::AddLegacyTrack(FMovieSceneEvaluationTrack&& InTrackTemplate, const UMovieSceneTrack& SourceTrack)
-{
-	AddOwnedTrack(MoveTemp(InTrackTemplate), SourceTrack);
-	Template.bHasLegacyTrackInstances = true;
 }
 
 void FMovieSceneEvaluationTemplateGenerator::AddOwnedTrack(FMovieSceneEvaluationTrack&& InTrackTemplate, const UMovieSceneTrack& SourceTrack)
 {
-	// Add the track to the template
-	Template.AddTrack(SourceTrack.GetSignature(), MoveTemp(InTrackTemplate));
 	CompiledSignatures.Add(SourceTrack.GetSignature());
+	Template.AddTrack(SourceTrack.GetSignature(), MoveTemp(InTrackTemplate));
 }
 
-void FMovieSceneEvaluationTemplateGenerator::AddSharedTrack(FMovieSceneEvaluationTrack&& InTrackTemplate, FMovieSceneSharedDataId SharedId, const UMovieSceneTrack& SourceTrack)
-{
-	FSharedPersistentDataKey Key(SharedId, FMovieSceneEvaluationOperand(MovieSceneSequenceID::Root, InTrackTemplate.GetObjectBindingID()));
-
-	if (AddedSharedTracks.Contains(Key))
-	{
-		return;
-	}
-
-	AddedSharedTracks.Add(Key);
-	AddOwnedTrack(MoveTemp(InTrackTemplate), SourceTrack);
-}
-
-void FMovieSceneEvaluationTemplateGenerator::AddExternalSegments(TRange<float> RootRange, TArrayView<const FMovieSceneEvaluationFieldSegmentPtr> SegmentPtrs, ESectionEvaluationFlags Flags)
-{
-	if (RootRange.IsEmpty())
-	{
-		return;
-	}
-
-	SegmentData.Reserve(SegmentData.Num() + SegmentPtrs.Num());
-
-	for (const FMovieSceneEvaluationFieldSegmentPtr& SegmentPtr : SegmentPtrs)
-	{
-		// Add one segment to the external segment map per flag configuration
-		FExternalSegment NewSegment{ -1, Flags };
-		for (auto It = ExternalSegmentLookup.CreateConstKeyIterator(SegmentPtr); It; ++It)
-		{
-			if (It.Value().Flags == Flags)
-			{
-				NewSegment = It.Value();
-				break;
-			}
-		}
-
-		if (NewSegment.Index == -1)
-		{
-			NewSegment.Index = TrackLUT.Num();
-			ExternalSegmentLookup.Add(SegmentPtr, NewSegment);
-			TrackLUT.Add(SegmentPtr);
-		}
-
-		// Flags are irrelevant for the final segment compilation
-		SegmentData.Add(FMovieSceneSectionData(RootRange, FSectionEvaluationData(NewSegment.Index)));
-	}
-}
-
-FMovieSceneSequenceTransform FMovieSceneEvaluationTemplateGenerator::GetSequenceTransform(FMovieSceneSequenceIDRef InSequenceID) const
-{
-	if (InSequenceID == MovieSceneSequenceID::Root)
-	{
-		return FMovieSceneSequenceTransform();
-	}
-
-	const FMovieSceneSubSequenceData* Data = Template.Hierarchy.FindSubData(InSequenceID);
-	return ensure(Data) ? Data->RootToSequenceTransform : FMovieSceneSequenceTransform();
-}
 
-void FMovieSceneEvaluationTemplateGenerator::AddSubSequence(FMovieSceneSubSequenceData SequenceData, FMovieSceneSequenceIDRef ParentID, FMovieSceneSequenceID SequenceID)
+void FMovieSceneEvaluationTemplateGenerator::Generate()
 {
-	FMovieSceneSequenceHierarchyNode* ParentNode = Template.Hierarchy.FindNode(ParentID);
-	checkf(ParentNode, TEXT("Cannot generate a sequence ID for a ParentID that doesn't yet exist"));
-
-	check(SequenceID.IsValid());
-	
-#if WITH_EDITORONLY_DATA
-	if (const FMovieSceneSubSequenceData* ParentSubSequenceData = Template.Hierarchy.FindSubData(ParentID))
-	{
-		// Clamp this sequence's valid play range by its parent's valid play range
-		TRange<float> ParentPlayRangeChildSpace = ParentSubSequenceData->ValidPlayRange * (SequenceData.RootToSequenceTransform * ParentSubSequenceData->RootToSequenceTransform.Inverse());
-		SequenceData.ValidPlayRange = TRange<float>::Intersection(ParentPlayRangeChildSpace, SequenceData.ValidPlayRange);
-	}
-#endif
-
-	// Ensure we have a unique ID. This should never happen in reality.
-	while(!ensureMsgf(!Template.Hierarchy.FindNode(SequenceID), TEXT("CRC collision on deterministic hashes. Manually hashing a random new one.")))
-	{
-		SequenceID = SequenceID.AccumulateParentID(SequenceID);
-	}
-
-	Template.Hierarchy.Add(SequenceData, SequenceID, ParentID);
-}
+	// We always regenerate this
+	Template.ResetFieldData();
 
+	CompiledSignatures.Reset();
 
-void FMovieSceneEvaluationTemplateGenerator::Generate(FMovieSceneTrackCompilationParams InParams)
-{
-	Template.Hierarchy = FMovieSceneSequenceHierarchy();
-
-	// Generate templates for each track in the movie scene
 	UMovieScene* MovieScene = SourceSequence.GetMovieScene();
 
-	TransientArgs.Params = InParams;
-
 	if (UMovieSceneTrack* Track = MovieScene->GetCameraCutTrack())
 	{
 		ProcessTrack(*Track);
@@ -156,261 +44,200 @@ void FMovieSceneEvaluationTemplateGenerator::Generate(FMovieSceneTrackCompilatio
 
 	for (const FMovieSceneBinding& ObjectBinding : MovieScene->GetBindings())
 	{
-#if WITH_EDITOR
-		// Skip object bindings that are optimized out
-		auto ShouldRemoveObject = [](const UMovieSceneTrack* Track){ return EnumHasAnyFlags(Track->GetCookOptimizationFlags(), ECookOptimizationFlags::RemoveObject); };
-		if (ObjectBinding.GetTracks().ContainsByPredicate(ShouldRemoveObject))
-		{
-			continue;
-		}
-#endif
-
 		for (UMovieSceneTrack* Track : ObjectBinding.GetTracks())
 		{
 			ProcessTrack(*Track, ObjectBinding.GetObjectGuid());
 		}
 	}
 
-	// Remove references to old tracks
-	RemoveOldTrackReferences();
-
-	// Add all the tracks in *this* sequence (these exist after any sub section ptrs, not that its important for this algorithm)
-	for (auto& Pair : Template.GetTracks())
-	{
-		TArrayView<const FMovieSceneSegment> Segments = Pair.Value.GetSegments();
-
-		// Add the segment range data to the master collection for overall compilation
-		for (int32 SegmentIndex = 0; SegmentIndex < Segments.Num(); ++SegmentIndex)
-		{
-			SegmentData.Add(FMovieSceneSectionData(Segments[SegmentIndex].Range, FSectionEvaluationData(TrackLUT.Num())));
-			TrackLUT.Add(FMovieSceneEvaluationFieldSegmentPtr(MovieSceneSequenceID::Root, Pair.Key, SegmentIndex));
-		}
-	}
-
-	TMap<FMovieSceneSequenceID, FMovieSceneEvaluationTemplate*> SequenceIdToTemplate;
-	for (auto& Pair : Template.Hierarchy.AllSubSequenceData())
-	{
-		if (UMovieSceneSequence* Sequence = Pair.Value.Sequence)
-		{
-			SequenceIdToTemplate.Add(Pair.Key, &TransientArgs.SubSequenceStore.GetCompiledTemplate(*Sequence, FObjectKey(Pair.Value.SequenceKeyObject)));
-		}
-	}
+	Template.RemoveStaleData(CompiledSignatures);
 
-	// Compile the new evaluation field
-	TArray<FMovieSceneSegment> NewSegments = FMovieSceneSegmentCompiler().Compile(SegmentData);
-	UpdateEvaluationField(NewSegments, TrackLUT, SequenceIdToTemplate);
+	Template.SequenceSignature = SourceSequence.GetSignature();
 }
 
-void FMovieSceneEvaluationTemplateGenerator::ProcessTrack(const UMovieSceneTrack& InTrack, const FGuid& ObjectId)
+void FMovieSceneEvaluationTemplateGenerator::ProcessTrack(const UMovieSceneTrack& Track, const FGuid& ObjectBindingId)
 {
-#if WITH_EDITOR
-	// Skip tracks that are optimized out
-	if (EnumHasAnyFlags(InTrack.GetCookOptimizationFlags(), ECookOptimizationFlags::RemoveTrack | ECookOptimizationFlags::RemoveObject))
-	{
-		return;
-	}
-#endif
+	FMovieSceneTrackSegmentBlenderPtr TrackBlender = Track.GetTrackSegmentBlender();
 
-	FGuid Signature = InTrack.GetSignature();
-
-	// See if this track signature already exists in the ledger, if it does, we don't need to regenerate it
-	if (Template.FindTracks(Signature).Num())
+	// Deal with sub tracks specifically
+	if (const UMovieSceneSubTrack* SubTrack = Cast<const UMovieSceneSubTrack>(&Track))
 	{
-		CompiledSignatures.Add(Signature);
-		return;
+		ProcessSubTrack(*SubTrack, ObjectBindingId);
 	}
 
-	TransientArgs.ObjectBindingId = ObjectId;
-	TransientArgs.DefaultCompletionMode = SourceSequence.DefaultCompletionMode;
-
-	// Potentially expensive generation is required
-	InTrack.GenerateTemplate(TransientArgs);
-}
-
-void FMovieSceneEvaluationTemplateGenerator::RemoveOldTrackReferences()
-{
-	TArray<FGuid> SignaturesToRemove;
-
-	// Go through the template ledger, and remove anything that is no longer referenced
-	for (auto& Pair : Template.GetLedger().TrackSignatureToTrackIdentifier)
+	// If the ledger already contains this track, just update the uncompiled track field
+	else if (FMovieSceneTrackIdentifier TrackID = Template.GetLedger().FindTrack(Track.GetSignature()))
 	{
-		if (!CompiledSignatures.Contains(Pair.Key))
-		{
-			SignaturesToRemove.Add(Pair.Key);
-		}
+		CompiledSignatures.Add(Track.GetSignature());
+
+		// Don't invalidate the evaluation field if for this track as this track hasn't changed
+		const bool bInvalidateEvaluationField = false;
+		Template.DefineTrackStructure(TrackID, bInvalidateEvaluationField);
 	}
 
-	// Remove the signatures, updating entries in the evaluation field as we go
-	for (const FGuid& Signature : SignaturesToRemove)
+	// Else the track doesn't exist - we need to generate it from scratch
+	else
 	{
-		Template.RemoveTrack(Signature);
+		FMovieSceneTrackCompilerArgs Args(*this);
+		Args.ObjectBindingId = ObjectBindingId;
+		Args.DefaultCompletionMode = SourceSequence.DefaultCompletionMode;
+		Track.GenerateTemplate(Args);
 	}
 }
 
-void FMovieSceneEvaluationTemplateGenerator::UpdateEvaluationField(const TArray<FMovieSceneSegment>& Segments, const TArray<FMovieSceneEvaluationFieldSegmentPtr>& Ptrs, const TMap<FMovieSceneSequenceID, FMovieSceneEvaluationTemplate*>& Templates)
-{
-	IMovieSceneModule& MovieSceneModule = IMovieSceneModule::Get_Concurrent();
 
-	FMovieSceneEvaluationField& Field = Template.EvaluationField;
+void FMovieSceneEvaluationTemplateGenerator::ProcessSubTrack(const UMovieSceneSubTrack& SubTrack, const FGuid& ObjectBindingId)
+{
+	// Get the segment blender required to blend at the track level for this segment
+	FMovieSceneTrackSegmentBlenderPtr    TrackSegmentBlender = SubTrack.GetTrackSegmentBlender();
+	FMovieSceneTrackRowSegmentBlenderPtr RowSegmentBlender   = SubTrack.GetRowSegmentBlender();
 
-	Field = FMovieSceneEvaluationField();
+	// If the track has no blenders at all, we can just add all the sub section ranges directly
+	const bool bRequiresSegmentBlending = TrackSegmentBlender.IsValid() || RowSegmentBlender.IsValid();
 
-	TArray<FMovieSceneEvaluationFieldSegmentPtr> AllTracksInSegment;
+	TMovieSceneEvaluationTree<FSectionEvaluationData> SubSectionBlendTree;
 
-	for (int32 Index = 0; Index < Segments.Num(); ++Index)
+	// Iterate all the sub track's sections, adding them to the template field
+	const TArray<UMovieSceneSection*>& AllSections = SubTrack.GetAllSections();
+	for (int32 SectionIndex = 0; SectionIndex < AllSections.Num(); ++SectionIndex)
 	{
-		const FMovieSceneSegment& Segment = Segments[Index];
-		if (Segment.Impls.Num() == 0)
+		UMovieSceneSubSection* SubSection = CastChecked<UMovieSceneSubSection>(AllSections[SectionIndex]);
+		if (!SubSection || !SubSection->IsActive())
 		{
 			continue;
 		}
 
-		Field.Ranges.Add(Segment.Range);
+		CompiledSignatures.Add(SubSection->GetSignature());
 
-		AllTracksInSegment.Reset();
-		for (const FSectionEvaluationData& LUTData : Segment.Impls)
+		if (bRequiresSegmentBlending)
 		{
-			AllTracksInSegment.Add(Ptrs[LUTData.ImplIndex]);
-		}
+			const TRange<float> SectionRange = SubSection->GetRange();
 
-		// Sort the track ptrs, and define flush ranges
-		AllTracksInSegment.Sort(
-			[&](const FMovieSceneEvaluationFieldSegmentPtr& A, const FMovieSceneEvaluationFieldSegmentPtr& B){
-				return SortPredicate(A, B, Templates, MovieSceneModule);
+			// Process the actual section range
+			if (bRequiresSegmentBlending)
+			{
+				SubSectionBlendTree.Add(SectionRange, FSectionEvaluationData(SectionIndex, ESectionEvaluationFlags::None));
 			}
-		);
-
-		FMovieSceneEvaluationGroup& Group = Field.Groups[Field.Groups.Emplace()];
-
-		TArray<FMovieSceneEvaluationFieldSegmentPtr> EvalPtrs;
-		TArray<FMovieSceneEvaluationFieldSegmentPtr> InitPtrs;
-
-		// Now iterate the tracks and insert indices for initialization and evaluation
-		FName CurrentEvaluationGroup, LastEvaluationGroup;
-		TOptional<uint32> LastPriority;
-
-		for (const FMovieSceneEvaluationFieldSegmentPtr& Ptr : AllTracksInSegment)
-		{
-			const FMovieSceneEvaluationTrack* Track = LookupTrack(Ptr, Templates);
-			if (!ensure(Track))
+			else
 			{
-				continue;
+				Template.AddSubSectionRange(*SubSection, ObjectBindingId, SectionRange, ESectionEvaluationFlags::None);
 			}
 
-			// If we're now in a different flush group, add the ptrs
-			CurrentEvaluationGroup = Track->GetEvaluationGroup();
-			if (CurrentEvaluationGroup != LastEvaluationGroup)
+			// Process the section preroll range
+			if (!SectionRange.GetLowerBound().IsOpen() && SubSection->GetPreRollTime() > 0)
 			{
-				AddPtrsToGroup(Group, InitPtrs, EvalPtrs);
-			}
-			LastEvaluationGroup = Track->GetEvaluationGroup();
+				TRange<float> PreRollRange(SectionRange.GetLowerBoundValue() - SubSection->GetPreRollTime(), TRangeBound<float>::FlipInclusion(SectionRange.GetLowerBoundValue()));
 
-			const bool bRequiresInitialization = Track->GetSegment(Ptr.SegmentIndex).Impls.ContainsByPredicate(
-				[Track](FSectionEvaluationData EvalData)
+				if (bRequiresSegmentBlending)
 				{
-					return Track->GetChildTemplate(EvalData.ImplIndex).RequiresInitialization();
+					SubSectionBlendTree.Add(PreRollRange, FSectionEvaluationData(SectionIndex, ESectionEvaluationFlags::PreRoll));
 				}
-			);
+				else
+				{
+					Template.AddSubSectionRange(*SubSection, ObjectBindingId, PreRollRange, ESectionEvaluationFlags::PreRoll);
+				}
+			}
 
-			if (bRequiresInitialization)
+			// Process the section postroll range
+			if (!SectionRange.GetUpperBound().IsOpen() && SubSection->GetPostRollTime() > 0)
 			{
-				InitPtrs.Add(Ptr);
-			}
+				TRange<float> PostRollRange(TRangeBound<float>::FlipInclusion(SectionRange.GetUpperBoundValue()), SectionRange.GetUpperBoundValue() + SubSection->GetPostRollTime());
 
-			EvalPtrs.Add(Ptr);
+				if (bRequiresSegmentBlending)
+				{
+					SubSectionBlendTree.Add(PostRollRange, FSectionEvaluationData(SectionIndex, ESectionEvaluationFlags::PostRoll));
+				}
+				else
+				{
+					Template.AddSubSectionRange(*SubSection, ObjectBindingId, PostRollRange, ESectionEvaluationFlags::PostRoll);
+				}
+			}
 		}
+	}
 
-		AddPtrsToGroup(Group, InitPtrs, EvalPtrs);
-
-		// Copmpute meta data for this segment
-		FMovieSceneEvaluationMetaData& MetaData = Field.MetaData[Field.MetaData.Emplace()];
-		InitializeMetaData(MetaData, Group, Templates);
+	// Nothing more to do if we don't need to blend the segments
+	if (!bRequiresSegmentBlending)
+	{
+		return;
 	}
-}
 
-void FMovieSceneEvaluationTemplateGenerator::InitializeMetaData(FMovieSceneEvaluationMetaData& MetaData, FMovieSceneEvaluationGroup& Group, const TMap<FMovieSceneSequenceID, FMovieSceneEvaluationTemplate*>& Templates)
-{
-	MetaData.Reset();
+	// Keep accumulating similar segments to reduce fragmentation in the field
+	TOptional<FMovieSceneSegment> AccumulateSegment;
+
+	TMap<int32, FSegmentBlendData> RowBlendData;
 
-	TSet<FMovieSceneSequenceID> ActiveSequenceIDs;
-	TSet<FMovieSceneEvaluationKey> ActiveEntitySet;
-	for (const FMovieSceneEvaluationFieldSegmentPtr& SegmentPtr : Group.SegmentPtrLUT)
+	// Go through each unique range within SubSectionBlendTree, blending any entities present appropriately, then adding them to the tree
+	for (FMovieSceneEvaluationTreeRangeIterator SubSectionIt(SubSectionBlendTree); SubSectionIt; ++SubSectionIt)
 	{
-		const FMovieSceneEvaluationTrack* Track = LookupTrack(SegmentPtr, Templates);
-		if (!ensure(Track))
+		TMovieSceneEvaluationTreeDataIterator<FSectionEvaluationData> DataIt = SubSectionBlendTree.GetAllData(SubSectionIt.Node());
+
+		// If there are no sub sections here, continue to the next range
+		if (!DataIt)
 		{
 			continue;
 		}
 
-		// Add the active sequence to the meta data
-		MetaData.ActiveSequences.AddUnique(SegmentPtr.SequenceID);
+		FSegmentBlendData TrackBlendData;
 
-		// Add the track key
-		FMovieSceneEvaluationKey TrackKey(SegmentPtr.SequenceID, SegmentPtr.TrackIdentifier);
-		if (!ActiveEntitySet.Contains(TrackKey))
+		// Compile at the row level first if necessary
+		if (RowSegmentBlender.IsValid())
 		{
-			ActiveEntitySet.Add(TrackKey);
-			MetaData.ActiveEntities.Add(FMovieSceneOrderedEvaluationKey{ TrackKey, uint32(MetaData.ActiveEntities.Num()) });
-		}
+			RowBlendData.Reset();
 
-		for (FSectionEvaluationData EvalData : Track->GetSegment(SegmentPtr.SegmentIndex).Impls)
-		{
-			FMovieSceneEvaluationKey SectionKey = TrackKey.AsSection(EvalData.ImplIndex);
-			if (!ActiveEntitySet.Contains(SectionKey))
+			// Add every FSectionEvaluationData for the current time range to the section data
+			for (const FSectionEvaluationData& EvalData : DataIt)
 			{
-				MetaData.ActiveEntities.Add(FMovieSceneOrderedEvaluationKey{ SectionKey, uint32(MetaData.ActiveEntities.Num()) });
+				const UMovieSceneSection* Section = AllSections[EvalData.ImplIndex];
+				check(Section);
+				RowBlendData.FindOrAdd(Section->GetRowIndex()).Add(FMovieSceneSectionData(Section, EvalData.ImplIndex, EvalData.Flags));
 			}
-		}
-	}
 
-	Algo::SortBy(MetaData.ActiveEntities, &FMovieSceneOrderedEvaluationKey::Key);
-	MetaData.ActiveSequences.Sort();
-}
+			// Blend and append
+			for (TPair<int32, FSegmentBlendData>& Pair : RowBlendData)
+			{
+				RowSegmentBlender->Blend(Pair.Value);
+				TrackBlendData.Append(Pair.Value);
+			}
+		}
 
-const FMovieSceneEvaluationTrack* FMovieSceneEvaluationTemplateGenerator::LookupTrack(const FMovieSceneEvaluationFieldTrackPtr& InPtr, const TMap<FMovieSceneSequenceID, FMovieSceneEvaluationTemplate*>& Templates)
-{
-	if (InPtr.SequenceID == MovieSceneSequenceID::Root)
-	{
-		return Template.FindTrack(InPtr.TrackIdentifier);
-	}
-	else if (const FMovieSceneEvaluationTemplate* SubTemplate = Templates.FindRef(InPtr.SequenceID))
-	{
-		return SubTemplate->FindTrack(InPtr.TrackIdentifier);
-	}
-	ensure(false);
-	return nullptr;
-}
+		// Otherwise just add everything to the track blend data
+		else for (const FSectionEvaluationData& EvalData : DataIt)
+		{
+			const UMovieSceneSection* Section = AllSections[EvalData.ImplIndex];
+			check(Section);
+			TrackBlendData.Add(FMovieSceneSectionData(Section, EvalData.ImplIndex, EvalData.Flags));
+		}
 
-bool FMovieSceneEvaluationTemplateGenerator::SortPredicate(const FMovieSceneEvaluationFieldTrackPtr& InPtrA, const FMovieSceneEvaluationFieldTrackPtr& InPtrB, const TMap<FMovieSceneSequenceID, FMovieSceneEvaluationTemplate*>& Templates, IMovieSceneModule& MovieSceneModule)
-{
-	const FMovieSceneEvaluationTrack* A = LookupTrack(InPtrA, Templates);
-	const FMovieSceneEvaluationTrack* B = LookupTrack(InPtrB, Templates);
-	if (!ensure(A && B))
-	{
-		return false;
-	}
+		// Compile at the track level
+		if (TrackSegmentBlender.IsValid())
+		{
+			TrackSegmentBlender->Blend(TrackBlendData);
+		}
 
-	FMovieSceneEvaluationGroupParameters GroupA = MovieSceneModule.GetEvaluationGroupParameters(A->GetEvaluationGroup());
-	FMovieSceneEvaluationGroupParameters GroupB = MovieSceneModule.GetEvaluationGroupParameters(B->GetEvaluationGroup());
+		FMovieSceneSegment NextSegment(SubSectionIt.Range());
+		TrackBlendData.AddToSegment(NextSegment);
 
-	if (GroupA.EvaluationPriority != GroupB.EvaluationPriority)
-	{
-		return GroupA.EvaluationPriority > GroupB.EvaluationPriority;
+		if (!AccumulateSegment.IsSet())
+		{
+			AccumulateSegment = MoveTemp(NextSegment);
+		}
+		else if (!AccumulateSegment->CombineWith(NextSegment))
+		{
+			// This segment is different from the previous one - add the previous one and start accumulating into this one
+			for (FSectionEvaluationData EvalData : AccumulateSegment->Impls)
+			{
+				Template.AddSubSectionRange(*CastChecked<UMovieSceneSubSection>(AllSections[EvalData.ImplIndex]), ObjectBindingId, AccumulateSegment->Range, EvalData.Flags);
+			}
+			AccumulateSegment = MoveTemp(NextSegment);
+		}
 	}
 
-	const FMovieSceneSubSequenceData* DataA = Template.Hierarchy.FindSubData(InPtrA.SequenceID);
-	const FMovieSceneSubSequenceData* DataB = Template.Hierarchy.FindSubData(InPtrB.SequenceID);
-
-	if (DataA || DataB)
+	// If there is any segment outstanding, add it to the template
+	if (AccumulateSegment.IsSet())
 	{
-		int32 HierarchicalBiasA = DataA ? DataA->HierarchicalBias : 0;
-		int32 HierarchicalBiasB = DataB ? DataB->HierarchicalBias : 0;
-
-		if (HierarchicalBiasA != HierarchicalBiasB)
+		for (FSectionEvaluationData EvalData : AccumulateSegment->Impls)
 		{
-			return HierarchicalBiasA < HierarchicalBiasB;
+			Template.AddSubSectionRange(*CastChecked<UMovieSceneSubSection>(AllSections[EvalData.ImplIndex]), ObjectBindingId, AccumulateSegment->Range, EvalData.Flags);
 		}
 	}
-
-	return A->GetEvaluationPriority() > B->GetEvaluationPriority();
-}
+}
\ No newline at end of file
diff --git a/Engine/Source/Runtime/MovieScene/Private/Compilation/MovieSceneSegmentCompiler.cpp b/Engine/Source/Runtime/MovieScene/Private/Compilation/MovieSceneSegmentCompiler.cpp
--- a/Engine/Source/Runtime/MovieScene/Private/Compilation/MovieSceneSegmentCompiler.cpp
+++ b/Engine/Source/Runtime/MovieScene/Private/Compilation/MovieSceneSegmentCompiler.cpp
@@ -1,419 +1,42 @@
 // Copyright 1998-2017 Epic Games, Inc. All Rights Reserved.
 
 #include "Compilation/MovieSceneSegmentCompiler.h"
+#include "Compilation/MovieSceneCompilerRules.h"
+#include "MovieSceneEvaluationTree.h"
 
-FMovieSceneSectionData::FMovieSceneSectionData(const TRange<float>& InBounds, FSectionEvaluationData InEvalData, FOptionalMovieSceneBlendType InBlendType, int32 InPriority)
-	: Bounds(InBounds)
-	, EvalData(InEvalData)
-	, BlendType(InBlendType)
-	, Priority(InPriority)
-{}
-
-void FMovieSceneSegmentCompilerRules::ProcessSegments(TArray<FMovieSceneSegment>& Segments, const TArrayView<const FMovieSceneSectionData>& SourceData) const
+PRAGMA_DISABLE_DEPRECATION_WARNINGS
+void FMovieSceneSegmentCompilerRules::Blend(FSegmentBlendData& BlendData) const
 {
-	for (int32 Index = 0; Index < Segments.Num(); ++Index)
-	{
-		FMovieSceneSegment& Segment = Segments[Index];
-
-		BlendSegment(Segment, SourceData);
-		if (!bAllowEmptySegments && Segment.Impls.Num() == 0)
-		{
-			Segments.RemoveAt(Index, 1, false);
-		}
-	}
-
-	if (!Segments.Num())
-	{
-		return;
-	}
-
-	int32 Index = 0;
+	FMovieSceneSegment TmpSegment;
 
-	// First off, add an open segment before the first segment if necessary/possible
-	if (!Segments[0].Range.GetLowerBound().IsOpen())
+	for (int32 Index = 0; Index < BlendData.Num(); ++Index)
 	{
-		InsertSegment(Segments, 0, SourceData);
+		TmpSegment.Impls.Add(FSectionEvaluationData(Index, BlendData[Index].Flags));
 	}
-	++Index;
 
-	// Then add any segments inbetween valid segments
-	for (; Index < Segments.Num(); ++Index)
-	{
-		if (InsertSegment(Segments, Index, SourceData))
-		{
-			// Increment the counter to account for the added segment
-			++Index;
-		}
-	}
+	BlendSegment(TmpSegment, BlendData);
 
-	if (!Segments.Last().Range.GetUpperBound().IsOpen())
+	FSegmentBlendData NewBlendData;
+	for (FSectionEvaluationData EvalData : TmpSegment.Impls)
 	{
-		InsertSegment(Segments, Segments.Num(), SourceData);
+		NewBlendData.Add(BlendData[EvalData.ImplIndex]);
+		NewBlendData.Last().Flags = EvalData.Flags;
 	}
-
-	PostProcessSegments(Segments, SourceData);
 }
+PRAGMA_ENABLE_DEPRECATION_WARNINGS
 
-bool FMovieSceneSegmentCompilerRules::InsertSegment(TArray<FMovieSceneSegment>& Segments, int32 Index, const TArrayView<const FMovieSceneSectionData>& SourceData) const
+bool FMovieSceneAdditiveCameraTrackBlender::SortByStartTime(const FMovieSceneSectionData& A, const FMovieSceneSectionData& B)
 {
-	const FMovieSceneSegment* PreviousSegment = Segments.IsValidIndex(Index-1) ? &Segments[Index-1] : nullptr;
-	const FMovieSceneSegment* NextSegment = Segments.IsValidIndex(Index) ? &Segments[Index] : nullptr;
-
-	TRange<float> EmptyRange(
-		PreviousSegment 	? TRangeBound<float>::FlipInclusion(PreviousSegment->Range.GetUpperBound()) : TRangeBound<float>(),
-		NextSegment 		? TRangeBound<float>::FlipInclusion(NextSegment->Range.GetLowerBound()) 	: TRangeBound<float>()
-		);
-
-	if (EmptyRange.IsEmpty())
+	if (A.Section->IsInfinite())
 	{
-		return false;
+		return true;
 	}
-
-	TOptional<FMovieSceneSegment> NewSeg = InsertEmptySpace(EmptyRange, PreviousSegment, NextSegment);
-	if (!NewSeg.IsSet())
+	else if (B.Section->IsInfinite())
 	{
 		return false;
 	}
-	else if (!ensureMsgf(EmptyRange.Contains(NewSeg.GetValue().Range), TEXT("Attempting to insert an range that overflows the empty space. Correcting....")))
-	{
-		NewSeg.GetValue().Range = TRange<float>::Intersection(NewSeg.GetValue().Range, EmptyRange);
-	}
-
-	BlendSegment(NewSeg.GetValue(), SourceData);
-	Segments.Insert(MoveTemp(NewSeg.GetValue()), Index);
-
-	return true;
-}
-
-TArray<FMovieSceneSegment> FMovieSceneSegmentCompiler::Compile(TArrayView<const FMovieSceneSectionData> Data, const FMovieSceneSegmentCompilerRules* Rules, EMovieSceneSegmentIndexSpace InIndexSpace)
-{
-	SourceData = Data;
-	IndexSpace = InIndexSpace;
-	OverlappingSections.Reset(16);
-	OverlappingRefCounts.Reset(16);
-	LowerBounds.Reset(Data.Num());
-	UpperBounds.Reset(Data.Num());
-	CompiledSegments.Reset();
-
-	LowerReadIndex = 0;
-	UpperReadIndex = 0;
-
-	// Populate the lists of lower/upper bounds
-	for (int32 Index = 0; Index < Data.Num(); ++Index)
-	{
-		const FMovieSceneSectionData& Section = Data[Index];
-		if (!Section.Bounds.IsEmpty())
-		{
-			FSectionEvaluationData EvalData = Section.EvalData;
-
-			// Regardless of what IndexSpace was specified, we always run the compiler with the source data index
-			// then translate afterwards so we have a consistent way of writing compiler rules
-			EvalData.ImplIndex = Index;
-
-			LowerBounds.Add(FBound(EvalData, Section.Bounds.GetLowerBound()));
-			UpperBounds.Add(FBound(EvalData, Section.Bounds.GetUpperBound()));
-		}
-	}
-
-	LowerBounds.Sort(
-		[](const FBound& A, const FBound& B)
-		{
-			return FFloatRangeBound::MinLower(A.Bound, B.Bound) == A.Bound;
-		}
-	);
-
-	UpperBounds.Sort(
-		[](const FBound& A, const FBound& B)
-		{
-			return FFloatRangeBound::MinUpper(A.Bound, B.Bound) == A.Bound;
-		}
-	);
-	
-	while (LowerReadIndex < LowerBounds.Num())
-	{
-		CloseCompletedSegments();
-
-		TRangeBound<float> OpeningBound = LowerBounds[LowerReadIndex].Bound;
-
-		const bool bHadOverlappingSections = OverlappingSections.Num() != 0;
-
-		// Add the currently overlapping sections for any sections starting at exactly this time
-		do
-		{
-			// Reference count how many times this section is overlapping the current time with the same flags. This is to support multiple references to the same section.
-			const int32 OverlapIndex = FindOverlappingIndex(LowerBounds[LowerReadIndex].EvalData);
-			if (OverlapIndex == INDEX_NONE)
-			{
-				OverlappingSections.Add(LowerBounds[LowerReadIndex].EvalData);
-				OverlappingRefCounts.Add(1);
-			}
-			else
-			{
-				++OverlappingRefCounts[OverlapIndex];
-			}
-		}
-		while(++LowerReadIndex < LowerBounds.Num() && LowerBounds[LowerReadIndex].Bound == OpeningBound);
-
-		CompiledSegments.Emplace(
-			TRange<float>(OpeningBound, TRangeBound<float>()),
-			OverlappingSections
-			);
-	}
-
-	CloseCompletedSegments();
-
-	ensure(OverlappingSections.Num() == 0);
-
-	if (Rules)
-	{
-		Rules->ProcessSegments(CompiledSegments, Data);
-	}
-
-	if (IndexSpace == EMovieSceneSegmentIndexSpace::ActualImplIndex)
-	{
-		for (int32 Index = 0; Index < CompiledSegments.Num();)
-		{
-			FMovieSceneSegment& CompiledSegment = CompiledSegments[Index];
-			// Assign what are currently source data indices to the actual implemntation index specified in the source data
-			for (FSectionEvaluationData& EvalData : CompiledSegment.Impls)
-			{
-				EvalData.ImplIndex = Data[EvalData.ImplIndex].EvalData.ImplIndex;
-			}
-
-			if (Index > 0 && (CompiledSegment.Impls.Num() != 0 || (Rules && Rules->AllowEmptySegments())))
-			{
-				FMovieSceneSegment& PreviousSegment = CompiledSegments[Index - 1];
-
-				// If this is the same as the previous segment, and it ajoins the previous segment's range, just increase the range of the previous segment
-				if (PreviousSegment.Range.Adjoins(CompiledSegment.Range) && PreviousSegment.Impls == CompiledSegment.Impls)
-				{
-					PreviousSegment.Range = FFloatRange::Hull(PreviousSegment.Range, CompiledSegment.Range);
-					
-					CompiledSegments.RemoveAt(Index, 1, false);
-					// continue immediately to avoid incrementing the index
-					continue;
-				}
-			}
-
-			++Index;
-		}
-	}
-
-	return MoveTemp(CompiledSegments);
-}
-
-int32 FMovieSceneSegmentCompiler::FindOverlappingIndex(FSectionEvaluationData In) const
-{
-	In.ImplIndex = SourceData[In.ImplIndex].EvalData.ImplIndex;
-
-	return OverlappingSections.IndexOfByPredicate(
-		[this, In](FSectionEvaluationData InEvalData)
-		{
-			InEvalData.ImplIndex = SourceData[InEvalData.ImplIndex].EvalData.ImplIndex;
-			return InEvalData == In;
-		}
-	);
-}
-
-void FMovieSceneSegmentCompiler::CloseCompletedSegments()
-{
-	if (!CompiledSegments.Num())
-	{
-		return;
-	}
-
-	while(UpperReadIndex < UpperBounds.Num())
+	else
 	{
-		FMovieSceneSegment& LastSegment = CompiledSegments.Last();
-
-		// If there is a non-empty range between the next lower bound and upper bound, we can't close any more
-		const bool bHasOpeningRange = LowerReadIndex < LowerBounds.Num() && !TRange<float>(LowerBounds[LowerReadIndex].Bound, UpperBounds[UpperReadIndex].Bound).IsEmpty();
-		if (bHasOpeningRange)
-		{
-			if (OverlappingSections.Num())
-			{
-				TRangeBound<float> ClosingBound = TRangeBound<float>::FlipInclusion(LowerBounds[LowerReadIndex].Bound);
-				TRange<float> NewRange(LastSegment.Range.GetLowerBound(), ClosingBound);
-				if (!NewRange.IsEmpty())
-				{
-					// Just set the closing bound of the last segment and return
-					LastSegment.Range = NewRange;
-				}
-				else
-				{
-					// If it's empty, there's not point adding a segment, just allow the next segment to include the current OverlappingSegments
-					CompiledSegments.RemoveAt(CompiledSegments.Num() - 1, 1, false);
-				}
-			}
-			return;
-		}
-
-		TRangeBound<float> ClosingBound = UpperBounds[UpperReadIndex].Bound;
-		// Update the last segment's closing range
-		LastSegment.Range = TRange<float>(
-			LastSegment.Range.GetLowerBound(),
-			ClosingBound);
-
-		ensure(!LastSegment.Range.IsEmpty());
-
-		// Remove all sections that finish on this time
-		while (UpperReadIndex < UpperBounds.Num() && UpperBounds[UpperReadIndex].Bound == ClosingBound)
-		{
-			const int32 OverlappingIndex = FindOverlappingIndex(UpperBounds[UpperReadIndex].EvalData);
-			if (ensure(OverlappingIndex != INDEX_NONE))
-			{
-				if (--OverlappingRefCounts[OverlappingIndex] == 0)
-				{
-					OverlappingSections.RemoveAtSwap(OverlappingIndex, 1, false);
-					OverlappingRefCounts.RemoveAtSwap(OverlappingIndex, 1, false);
-				}
-			}
-			++UpperReadIndex;
-		}
-
-		// If there are any more sections still active, create a new segment for those
-		if (OverlappingSections.Num())
-		{
-			CompiledSegments.Emplace(
-				TRange<float>(TRangeBound<float>::FlipInclusion(ClosingBound), TRangeBound<float>()),
-				OverlappingSections
-				);
-		}
+		return A.Section->GetStartTime() < B.Section->GetStartTime();
 	}
 }
-
-FMovieSceneTrackCompiler::FRows::FRows(const TArray<UMovieSceneSection*>& Sections, const FMovieSceneSegmentCompilerRules* CompileRules)
-{
-	for (int32 Index = 0; Index < Sections.Num(); ++Index)
-	{
-		const UMovieSceneSection* Section = Sections[Index];
-		if (!ensure(Section) || !Section->IsActive())
-		{
-			continue;
-		}
-		
-		const int32 RowIndex = Section->GetRowIndex();
-		if (!Rows.IsValidIndex(RowIndex))
-		{
-			Rows.AddDefaulted(RowIndex - Rows.Num() + 1);
-		}
-
-		const TRange<float> Range = Section->IsInfinite() ? TRange<float>::All() : Section->GetRange();
-
-		FSectionEvaluationData EvalData(Index);
-
-		Rows[RowIndex].Sections.Add(
-			FMovieSceneSectionData(Range, EvalData, Section->GetBlendType(), Section->GetOverlapPriority())
-			);
-
-		if (!Range.GetLowerBound().IsOpen() && Section->GetPreRollTime() > 0)
-		{
-			EvalData.Flags = ESectionEvaluationFlags::PreRoll;
-			TRange<float> PreRollRange(Range.GetLowerBoundValue() - Section->GetPreRollTime(), TRangeBound<float>::FlipInclusion(Range.GetLowerBoundValue()));
-			Rows[RowIndex].Sections.Add(
-				FMovieSceneSectionData(PreRollRange, EvalData, Section->GetBlendType(), Section->GetOverlapPriority())
-				);
-		}
-		if (!Range.GetUpperBound().IsOpen() && Section->GetPostRollTime() > 0)
-		{
-			EvalData.Flags = ESectionEvaluationFlags::PostRoll;
-			TRange<float> PostRollRange(TRangeBound<float>::FlipInclusion(Range.GetUpperBoundValue()), Range.GetUpperBoundValue() + Section->GetPostRollTime());
-			Rows[RowIndex].Sections.Add(
-				FMovieSceneSectionData(PostRollRange, EvalData, Section->GetBlendType(), Section->GetOverlapPriority())
-				);
-		}
-	}
-
-	Rows.RemoveAll(
-		[](const FRow& Row)
-		{
-			return Row.Sections.Num() == 0;
-		});
-
-	for (FRow& Row : Rows)
-	{
-		Row.CompileRules = CompileRules;
-	}
-}
-
-FMovieSceneTrackEvaluationField FMovieSceneTrackCompiler::Compile(TArrayView<FRow> Rows, const FMovieSceneSegmentCompilerRules* Rules)
-{
-	FMovieSceneTrackEvaluationField Result;
-
-	// Methodology:
-	//   - We initially run one segment compilation per row, specifying each FRow::Sections for the source data.
-	//     This allows us to do per-row blending (like handling overlapping sections)
-	//   - After each row, we accumulate a single array of compiled segments to be considered for compilation at the track level
-	//     This allows us to do blending on a row basis, without considering individual row blending rules
-	// 
-
-	TArray<FMovieSceneSectionData> TrackCompileData;
-
-	// Compile each row
-	for (int32 RowIndex = 0; RowIndex < Rows.Num(); ++RowIndex)
-	{
-		const FRow& Row = Rows[RowIndex];
-		if (!Row.Sections.Num())
-		{
-			continue;
-		}
-		
-		FMovieSceneSegmentCompiler Compiler;
-		
-		// Compile this row into segments
-		TArray<FMovieSceneSegment> RowSegments = Compiler.Compile(
-			TArray<FMovieSceneSectionData>(Row.Sections),
-			Row.CompileRules,
-			EMovieSceneSegmentIndexSpace::SourceDataIndex);
-
-		const int32 Priority = Rows.Num() - RowIndex;
-		for (FMovieSceneSegment& Segment : RowSegments)
-		{
-			// Add each implementation in this segment as a separate entry in the source data to ensure that the correct evaluation flags are compiled
-			for (FSectionEvaluationData& EvalData : Segment.Impls)
-			{
-				const FMovieSceneSectionData& SectionData = Row.Sections[EvalData.ImplIndex];
-				// Remap to the actual section index
-				EvalData.ImplIndex = SectionData.EvalData.ImplIndex;
-				TrackCompileData.Add(FMovieSceneSectionData(Segment.Range, EvalData, SectionData.BlendType, Priority));
-			}
-		}
-	}
-
-	// Boil down each row into a single, blended field
-	FMovieSceneSegmentCompiler Compiler;
-	TArray<FMovieSceneSegment> TrackSegments = Compiler.Compile(TrackCompileData, nullptr, EMovieSceneSegmentIndexSpace::SourceDataIndex);
-
-	// At this point, ImplIndex members correspond to the source data array (TrackCompileData)
-	if (Rules)
-	{
-		Rules->ProcessSegments(TrackSegments, TrackCompileData);
-	}
-
-	// Compile each of the row segments into a single segment comprising all overlapping sections on the track, blended and ordered appropriately
-	// There should be no empty space by this point
-	for (FMovieSceneSegment& Segment : TrackSegments)
-	{
-		for (FSectionEvaluationData& SectionEvalData : Segment.Impls)
-		{
-			// Remap the index to the actual section index
-			SectionEvalData.ImplIndex = TrackCompileData[SectionEvalData.ImplIndex].EvalData.ImplIndex;
-		}
-		
-		if (Segment.Impls.Num() != 0 || (Rules && Rules->AllowEmptySegments()))
-		{
-			// If this is the same as the previous segment, and it ajoins the previous segment's range, just increase the range of the previous segment
-			if (Result.Segments.Num() != 0 && Result.Segments.Last().Range.Adjoins(Segment.Range) && Result.Segments.Last().Impls == Segment.Impls)
-			{
-				Result.Segments.Last().Range = FFloatRange::Hull(Result.Segments.Last().Range, Segment.Range);
-			}
-			else
-			{
-				Result.Segments.Add(MoveTemp(Segment));
-			}
-		}
-	}
-
-	return Result;
-}
-
diff --git a/Engine/Source/Runtime/MovieScene/Private/Evaluation/MovieSceneEvalTemplate.cpp b/Engine/Source/Runtime/MovieScene/Private/Evaluation/MovieSceneEvalTemplate.cpp
--- a/Engine/Source/Runtime/MovieScene/Private/Evaluation/MovieSceneEvalTemplate.cpp
+++ b/Engine/Source/Runtime/MovieScene/Private/Evaluation/MovieSceneEvalTemplate.cpp
@@ -10,5 +10,5 @@ float FMovieSceneEvalTemplate::EvaluateEasing(float CurrentTime) const
 
 bool FMovieSceneEvalTemplatePtr::Serialize(FArchive& Ar)
 {
-	return SerializeEvaluationTemplate(*this, Ar);
+	return SerializeInlineValue(*this, Ar);
 }
diff --git a/Engine/Source/Runtime/MovieScene/Private/Evaluation/MovieSceneEvalTemplateSerializer.h b/Engine/Source/Runtime/MovieScene/Private/Evaluation/MovieSceneEvalTemplateSerializer.h
--- a/Engine/Source/Runtime/MovieScene/Private/Evaluation/MovieSceneEvalTemplateSerializer.h
+++ b/Engine/Source/Runtime/MovieScene/Private/Evaluation/MovieSceneEvalTemplateSerializer.h
@@ -8,9 +8,9 @@
 
 struct FMovieSceneEmptyStruct;
 
-/** Serialize the template */
+/** Serialize an inline value that has a GetScriptStruct method */
 template<typename T, uint8 N>
-bool SerializeEvaluationTemplate(TInlineValue<T, N>& Impl, FArchive& Ar)
+bool SerializeInlineValue(TInlineValue<T, N>& Impl, FArchive& Ar)
 {
 	Ar.UsingCustomVersion(FMovieSceneEvaluationCustomVersion::GUID);
 	
@@ -31,7 +31,7 @@ bool SerializeEvaluationTemplate(TInlineValue<T, N>& Impl, FArchive& Ar)
 		{
 #if !WITH_EDITORONLY_DATA
 			// We only throw this warning in cooked builds as that is the only place where this deserialized data matters
-			UE_LOG(LogMovieScene, Warning, TEXT("Unknown or invalid track type (%s) found in serialized data. This track will no longer work. Please recompile template data."), *TypeName);
+			UE_LOG(LogMovieScene, Warning, TEXT("Unknown or invalid type (%s) found in serialized data. This will no longer work. Please recompile template data."), *TypeName);
 #endif
 			// If it wasn't found, just deserialize an empty struct instead, and set ourselves to default
 			FMovieSceneEmptyStruct Empty;
diff --git a/Engine/Source/Runtime/MovieScene/Private/Evaluation/MovieSceneEvaluationField.cpp b/Engine/Source/Runtime/MovieScene/Private/Evaluation/MovieSceneEvaluationField.cpp
--- a/Engine/Source/Runtime/MovieScene/Private/Evaluation/MovieSceneEvaluationField.cpp
+++ b/Engine/Source/Runtime/MovieScene/Private/Evaluation/MovieSceneEvaluationField.cpp
@@ -1,8 +1,12 @@
 // Copyright 1998-2017 Epic Games, Inc. All Rights Reserved.
 
 #include "Evaluation/MovieSceneEvaluationField.h"
+#include "Evaluation/MovieSceneEvaluationTemplateInstance.h"
+#include "MovieSceneCommonHelpers.h"
 #include "Algo/Sort.h"
 
+#include "MovieSceneSequence.h"
+
 int32 FMovieSceneEvaluationField::GetSegmentFromTime(float Time) const
 {
 	// Linear search
@@ -40,20 +44,72 @@ TRange<int32> FMovieSceneEvaluationField::OverlapRange(TRange<float> Range) cons
 	return Num != 0 ? TRange<int32>(StartIndex, StartIndex + Num) : TRange<int32>::Empty();
 }
 
-void FMovieSceneEvaluationMetaData::RemapSequenceIDsForRoot(FMovieSceneSequenceID OverrideRootID)
+void FMovieSceneEvaluationField::Invalidate(TRange<float> Range)
 {
-	if (OverrideRootID == MovieSceneSequenceID::Root)
+	TRange<int32> OverlappingRange = OverlapRange(Range);
+	if (!OverlappingRange.IsEmpty())
+	{
+		Ranges.RemoveAt(OverlappingRange.GetLowerBoundValue(), OverlappingRange.Size<int32>(), false);
+		Groups.RemoveAt(OverlappingRange.GetLowerBoundValue(), OverlappingRange.Size<int32>(), false);
+		MetaData.RemoveAt(OverlappingRange.GetLowerBoundValue(), OverlappingRange.Size<int32>(), false);
+
+		Signature = FGuid::NewGuid();
+	}
+}
+
+int32 FMovieSceneEvaluationField::Insert(TRange<float> InRange, FMovieSceneEvaluationGroup&& InGroup, FMovieSceneEvaluationMetaData&& InMetaData)
+{
+	const int32 InsertIndex = Algo::UpperBoundBy(Ranges, InRange.GetLowerBound(), &TRange<float>::GetLowerBound, MovieSceneHelpers::SortLowerBounds);
+
+	// While we have floating point representations of ranges, there are a number of situations that can lead to the check firing:
+	// 		1. Consider ranges of (0,9.16666603f) and (9.16666698f,20) offset in an inner sequence by 98.8333511f:
+	//			Such transformation would result in the upper and lower bounds rounding to the same inclusive value (108.000015f) making them overlap in the parent space.
+	//		2. When compiling a segment at a particular time, that time transformed into an inner sequence could result in a compiled range that doesn't actually
+	//			include the global time when transformed back into the master space. In this scenario we always inflate the compiled range to include the global time, but this shouldn't be necessary
+	//			Such logic can lead to us attempting to add overlapping ranges into the field
+
+	// @todo: Remove this code and enforce the check below outright when we have proper time representation
+	if (Ranges.IsValidIndex(InsertIndex  ) && Ranges[InsertIndex  ].Overlaps(InRange))
 	{
-		return;
+		InRange = TRange<float>(InRange.GetLowerBound(), TRangeBound<float>::FlipInclusion(Ranges[InsertIndex].GetLowerBound()));
+	}
+	if (Ranges.IsValidIndex(InsertIndex-1) && Ranges[InsertIndex-1].Overlaps(InRange))
+	{
+		InRange = TRange<float>(TRangeBound<float>::FlipInclusion(Ranges[InsertIndex-1].GetUpperBound()), InRange.GetUpperBound());
 	}
 
-	for (FMovieSceneSequenceID& SequenceID : ActiveSequences)
+	if (!ensure(!InRange.IsEmpty()))
 	{
-		SequenceID = SequenceID.AccumulateParentID(OverrideRootID);
+		return INDEX_NONE;
 	}
-	for (FMovieSceneOrderedEvaluationKey& OrderedKey : ActiveEntities)
+
+	const bool bOverlapping = 
+		(Ranges.IsValidIndex(InsertIndex  ) && Ranges[InsertIndex  ].Overlaps(InRange)) ||
+		(Ranges.IsValidIndex(InsertIndex-1) && Ranges[InsertIndex-1].Overlaps(InRange));
+
+	if (!ensureAlwaysMsgf(!bOverlapping, TEXT("Attempting to insert an overlapping range into the evaluation field.")))
 	{
-		OrderedKey.Key.SequenceID = OrderedKey.Key.SequenceID.AccumulateParentID(OverrideRootID);
+		return INDEX_NONE;
+	}
+
+	Ranges.Insert(InRange, InsertIndex);
+	MetaData.Insert(MoveTemp(InMetaData), InsertIndex);
+	Groups.Insert(MoveTemp(InGroup), InsertIndex);
+
+	Signature = FGuid::NewGuid();
+
+	return InsertIndex;
+}
+
+void FMovieSceneEvaluationField::Add(TRange<float> InRange, FMovieSceneEvaluationGroup&& InGroup, FMovieSceneEvaluationMetaData&& InMetaData)
+{
+	if (ensureAlwaysMsgf(!Ranges.Num() || !Ranges.Last().Overlaps(InRange), TEXT("Attempting to add overlapping ranges to sequence evaluation field.")))
+	{
+		Ranges.Add(InRange);
+		MetaData.Add(MoveTemp(InMetaData));
+		Groups.Add(MoveTemp(InGroup));
+
+		Signature = FGuid::NewGuid();
 	}
 }
 
@@ -185,3 +241,28 @@ void FMovieSceneEvaluationMetaData::DiffEntities(const FMovieSceneEvaluationMeta
 		Algo::SortBy(*NewKeys, &FMovieSceneOrderedEvaluationKey::EvaluationIndex);
 	}
 }
+
+bool FMovieSceneEvaluationMetaData::IsDirty(UMovieSceneSequence& RootSequence, const FMovieSceneSequenceHierarchy& RootHierarchy, TRange<float>* OutSubRangeToInvalidate) const
+{
+	bool bDirty = false;
+
+	for (const TTuple<FMovieSceneSequenceID, FGuid>& Pair : SubSequenceSignatures)
+	{
+		// Sequence IDs at this point are relative to the root override template
+		const FMovieSceneSubSequenceData* SubData = RootHierarchy.FindSubData(Pair.Key);
+		UMovieSceneSequence* SubSequence = SubData ? SubData->GetSequence() : nullptr;
+
+		if (!SubSequence || SubSequence->GetSignature() != Pair.Value)
+		{
+			bDirty = true;
+
+			if (OutSubRangeToInvalidate)
+			{
+				TRange<float> DirtyRange = SubData ? TRange<float>::Hull(TRange<float>::Hull(SubData->PreRollRange, SubData->PlayRange), SubData->PostRollRange) * SubData->RootToSequenceTransform.Inverse() : TRange<float>::All();
+				*OutSubRangeToInvalidate = TRange<float>::Hull(*OutSubRangeToInvalidate, DirtyRange);
+			}
+		}
+	}
+
+	return bDirty;
+}
diff --git a/Engine/Source/Runtime/MovieScene/Private/Evaluation/MovieSceneEvaluationTemplate.cpp b/Engine/Source/Runtime/MovieScene/Private/Evaluation/MovieSceneEvaluationTemplate.cpp
--- a/Engine/Source/Runtime/MovieScene/Private/Evaluation/MovieSceneEvaluationTemplate.cpp
+++ b/Engine/Source/Runtime/MovieScene/Private/Evaluation/MovieSceneEvaluationTemplate.cpp
@@ -2,175 +2,217 @@
 
 #include "Evaluation/MovieSceneEvaluationTemplate.h"
 #include "MovieSceneSequence.h"
+#include "Sections/MovieSceneSubSection.h"
 
-FMovieSceneSequenceCachedSignature::FMovieSceneSequenceCachedSignature(UMovieSceneSequence& InSequence)
-	: Sequence(&InSequence)
-	, CachedSignature(InSequence.GetSignature())
+FMovieSceneSubSectionData::FMovieSceneSubSectionData(UMovieSceneSubSection& InSubSection, const FGuid& InObjectBindingId, ESectionEvaluationFlags InFlags)
+	: Section(&InSubSection), ObjectBindingId(InObjectBindingId), Flags(InFlags)
+{}
+
+FMovieSceneTrackIdentifier FMovieSceneTemplateGenerationLedger::FindTrack(const FGuid& InSignature) const
 {
+	return TrackSignatureToTrackIdentifier.FindRef(InSignature);
 }
 
-#if WITH_EDITORONLY_DATA
+void FMovieSceneTemplateGenerationLedger::AddTrack(const FGuid& InSignature, FMovieSceneTrackIdentifier Identifier)
+{
+	ensure(!TrackSignatureToTrackIdentifier.Contains(InSignature));
+	TrackSignatureToTrackIdentifier.Add(InSignature, Identifier);
+}
 
-void FCachedMovieSceneEvaluationTemplate::Initialize(UMovieSceneSequence& InSequence, FMovieSceneSequenceTemplateStore* InOrigin)
+void FMovieSceneEvaluationTemplate::PostSerialize(const FArchive& Ar)
 {
-	SourceSequence = &InSequence;
-	Origin = InOrigin;
+	if (Ar.IsLoading())
+	{
+		for (auto& Pair : Tracks)
+		{
+			if (TemplateLedger.LastTrackIdentifier == FMovieSceneTrackIdentifier::Invalid() || TemplateLedger.LastTrackIdentifier.Value < Pair.Key.Value)
+			{
+				// Reset previously serialized, invalid data
+				*this = FMovieSceneEvaluationTemplate();
+				break;
+			}
+		}
+	}
 }
 
-void FCachedMovieSceneEvaluationTemplate::Regenerate()
+void FMovieSceneEvaluationTemplate::ResetFieldData()
 {
-	Regenerate(CachedCompilationParams);
+	TrackFieldData.Field.Reset();
+	SubSectionFieldData.Field.Reset();
 }
 
-void FCachedMovieSceneEvaluationTemplate::Regenerate(const FMovieSceneTrackCompilationParams& Params)
+const TMovieSceneEvaluationTree<FMovieSceneTrackIdentifier>& FMovieSceneEvaluationTemplate::GetTrackField() const
 {
-	if (IsOutOfDate(Params))
-	{
-		RegenerateImpl(Params);
-	}
+	return TrackFieldData.Field;
 }
 
-void FCachedMovieSceneEvaluationTemplate::ForceRegenerate(const FMovieSceneTrackCompilationParams& Params)
+const TMovieSceneEvaluationTree<FMovieSceneSubSectionData>& FMovieSceneEvaluationTemplate::GetSubSectionField() const
 {
-	ResetGeneratedData();
-	RegenerateImpl(Params);
+	return SubSectionFieldData.Field;
 }
 
-void FCachedMovieSceneEvaluationTemplate::RegenerateImpl(const FMovieSceneTrackCompilationParams& Params)
+void FMovieSceneEvaluationTemplate::AddSubSectionRange(UMovieSceneSubSection& InSubSection, const FGuid& InObjectBindingId, const TRange<float>& InRange, ESectionEvaluationFlags InFlags)
 {
-	if (Params.bDuringBlueprintCompile != CachedCompilationParams.bDuringBlueprintCompile)
+	if (!ensure(!InSubSection.IsInfinite()))
 	{
-		ResetGeneratedData();
+		// @todo: Infinite sub sections??
+		return;
 	}
 
-	CachedSignatures.Reset();
-	CachedCompilationParams = Params;
+	// Add the sub section to the field, but we don't invalidate the evaluation field unless we know the section has actually changed
+	SubSectionFieldData.Field.Add(InRange, FMovieSceneSubSectionData(InSubSection, InObjectBindingId, InFlags));
 
-	if (UMovieSceneSequence* Sequence = SourceSequence.Get())
+	// Don't need to do anything else if the section was already generated
+	if (!TemplateLedger.ContainsSubSection(InSubSection.GetSignature()))
 	{
-		FMovieSceneSequenceTemplateStore DefaultStore;
-		Sequence->GenerateEvaluationTemplate(*this, Params, Origin ? *Origin : DefaultStore);
-
-		CachedSignatures.Add(FMovieSceneSequenceCachedSignature(*Sequence));
-		for (auto& Pair : Hierarchy.AllSubSequenceData())
+		if (InSubSection.GetSequence())
 		{
-			if (UMovieSceneSequence* SubSequence = Pair.Value.Sequence)
-			{
-				CachedSignatures.Add(FMovieSceneSequenceCachedSignature(*SubSequence));
-			}
+			FMovieSceneSequenceID SubSequenceID = InSubSection.GetSequenceID();
+			FSubSequenceInstanceDataParams InstanceParams{ SubSequenceID, FMovieSceneEvaluationOperand(MovieSceneSequenceID::Root, InObjectBindingId) };
+
+			FMovieSceneSubSequenceData NewSubData = InSubSection.GenerateSubSequenceData(InstanceParams);
+
+			Hierarchy.Add(NewSubData, SubSequenceID, MovieSceneSequenceID::Root);
 		}
+
+		TRange<float> EntireSectionRange(InSubSection.GetStartTime() - InSubSection.GetPreRollTime(), TRangeBound<float>::Inclusive(InSubSection.GetEndTime() + InSubSection.GetPostRollTime()));
+
+		// Add the section to the ledger
+		TemplateLedger.SubSectionRanges.Add(InSubSection.GetSignature(), EntireSectionRange);
+
+		// Invalidate the overlapping field
+		EvaluationField.Invalidate(EntireSectionRange);
 	}
 }
 
-bool FCachedMovieSceneEvaluationTemplate::IsOutOfDate(const FMovieSceneTrackCompilationParams& Params) const
+FMovieSceneTrackIdentifier FMovieSceneEvaluationTemplate::AddTrack(const FGuid& InSignature, FMovieSceneEvaluationTrack&& InTrack)
+{
+	FMovieSceneTrackIdentifier NewIdentifier = ++TemplateLedger.LastTrackIdentifier;
+
+	InTrack.SetupOverrides();
+	Tracks.Add(NewIdentifier, MoveTemp(InTrack));
+	TemplateLedger.AddTrack(InSignature, NewIdentifier);
+
+	// Add this track's segments to the unsorted track field, invalidating anything in the compiled evaluation field
+	DefineTrackStructure(NewIdentifier, true);
+
+	return NewIdentifier;
+}
+
+void FMovieSceneEvaluationTemplate::DefineTrackStructure(FMovieSceneTrackIdentifier TrackIdentifier, bool bInvalidateEvaluationField)
 {
-	if (Params != CachedCompilationParams || CachedSignatures.Num() == 0)
+	const FMovieSceneEvaluationTrack* Track = FindTrack(TrackIdentifier);
+	if (!ensure(Track))
 	{
-		return true;
+		return;
 	}
 
-	// out of date if the cached signautres contains any out of date sigs
-	return CachedSignatures.ContainsByPredicate(
-		[](const FMovieSceneSequenceCachedSignature& Sig){
-			UMovieSceneSequence* Sequence = Sig.Sequence.Get();
-			return !Sequence || Sequence->GetSignature() != Sig.CachedSignature;
-		}
-	);
-}
+	FMovieSceneTrackSegmentBlenderPtr TrackBlender = Track->GetSourceTrack()->GetTrackSegmentBlender();
+	const bool bAddEmptySpace = TrackBlender.IsValid() && TrackBlender->CanFillEmptySpace();
 
-#endif // WITH_EDITORONLY_DATA
+	if (bAddEmptySpace && bInvalidateEvaluationField)
+	{
+		// Optimization - when tracks can add empty space, we just invalidate the entire field once
+		EvaluationField.Invalidate(TRange<float>::All());
+		bInvalidateEvaluationField = false;
+	}
 
-TArrayView<const FMovieSceneTrackIdentifier> FMovieSceneTemplateGenerationLedger::FindTracks(const FGuid& InSignature) const
-{
-	if (auto* Tracks = TrackSignatureToTrackIdentifier.Find(InSignature))
+	// Add each range
+	for (FMovieSceneEvaluationTreeRangeIterator It(Track->Iterate()); It; ++It)
 	{
-		return Tracks->Data;
+		const bool bShouldAddEntry = bAddEmptySpace || Track->GetData(It.Node());
+		if (!bShouldAddEntry)
+		{
+			continue;
+		}
+
+		TrackFieldData.Field.Add(It.Range(), TrackIdentifier);
+
+		if (bInvalidateEvaluationField)
+		{
+			EvaluationField.Invalidate(It.Range());
+		}
 	}
-	return TArrayView<FMovieSceneTrackIdentifier>();
 }
 
-void FMovieSceneTemplateGenerationLedger::AddTrack(const FGuid& InSignature, FMovieSceneTrackIdentifier Identifier)
+void FMovieSceneEvaluationTemplate::RemoveTrack(const FGuid& InSignature)
 {
-	TrackSignatureToTrackIdentifier.FindOrAdd(InSignature).Data.Add(Identifier);
-	++TrackReferenceCounts.FindOrAdd(Identifier);
-}
+	FMovieSceneTrackIdentifier TrackIdentifier = TemplateLedger.FindTrack(InSignature);
+	if (TrackIdentifier == FMovieSceneTrackIdentifier::Invalid())
+	{
+		return;
+	}
 
-void FMovieSceneEvaluationTemplate::PostSerialize(const FArchive& Ar)
-{
-	if (Ar.IsLoading())
+	FMovieSceneEvaluationTrack* Track = Tracks.Find(TrackIdentifier);
+	if (Track)
 	{
-		for (auto& Pair : Tracks)
+		// Invalidate any ranges occupied by this track
+		for (FMovieSceneEvaluationTreeRangeIterator It = Track->Iterate(); It; ++It)
 		{
-			if (TemplateLedger.LastTrackIdentifier == FMovieSceneTrackIdentifier::Invalid() || TemplateLedger.LastTrackIdentifier.Value < Pair.Key)
+			if (Track->GetData(It.Node()))
 			{
-				// Reset previously serialized, invalid data
-				ResetGeneratedData();
-				break;
+				EvaluationField.Invalidate(It.Range());
 			}
 		}
+
+		// Add the track to the stale tracks map if necessary
+		if (bKeepStaleTracks)
+		{
+			StaleTracks.Add(TrackIdentifier, MoveTemp(*Track));
+		}
 	}
-}
 
-void FMovieSceneEvaluationTemplate::ResetGeneratedData()
-{
-	TemplateLedger.TrackSignatureToTrackIdentifier.Reset();
-	TemplateLedger.TrackReferenceCounts.Reset();
-
-	Tracks.Reset();
-	StaleTracks.Reset();
-	EvaluationField = FMovieSceneEvaluationField();
-	Hierarchy = FMovieSceneSequenceHierarchy();
-	bHasLegacyTrackInstances = false;
+	Tracks.Remove(TrackIdentifier);
+	TemplateLedger.TrackSignatureToTrackIdentifier.Remove(InSignature);
 }
 
-FMovieSceneTrackIdentifier FMovieSceneEvaluationTemplate::AddTrack(const FGuid& InSignature, FMovieSceneEvaluationTrack&& InTrack)
+void FMovieSceneEvaluationTemplate::RemoveStaleData(const TSet<FGuid>& ActiveSignatures)
 {
-	FMovieSceneTrackIdentifier NewIdentifier = ++TemplateLedger.LastTrackIdentifier;
-	
-	InTrack.SetupOverrides();
-	Tracks.Add(NewIdentifier.Value, MoveTemp(InTrack));
-	TemplateLedger.AddTrack(InSignature, NewIdentifier);
+	{
+		TArray<FGuid> SignaturesToRemove;
 
-	return NewIdentifier;
-}
+		// Go through the template ledger, and remove anything that is no longer referenced
+		for (auto& Pair : TemplateLedger.TrackSignatureToTrackIdentifier)
+		{
+			if (!ActiveSignatures.Contains(Pair.Key))
+			{
+				SignaturesToRemove.Add(Pair.Key);
+			}
+		}
 
-void FMovieSceneEvaluationTemplate::RemoveTrack(const FGuid& InSignature)
-{
-	for (FMovieSceneTrackIdentifier TrackIdentifier : TemplateLedger.FindTracks(InSignature))
+		// Remove the signatures, updating entries in the evaluation field as we go
+		for (const FGuid& Signature : SignaturesToRemove)
+		{
+			RemoveTrack(Signature);
+		}
+	}
+
+	// Remove stale sub sections
 	{
-		int32* RefCount = TemplateLedger.TrackReferenceCounts.Find(TrackIdentifier);
-		if (ensure(RefCount) && --(*RefCount) == 0)
+		TArray<TTuple<FGuid, FFloatRange>> SubSectionsToRemove;
+
+		for (const TTuple<FGuid, FFloatRange>& Pair : TemplateLedger.SubSectionRanges)
 		{
-			if (bKeepStaleTracks)
+			if (!ActiveSignatures.Contains(Pair.Key))
 			{
-				if (FMovieSceneEvaluationTrack* Track = Tracks.Find(TrackIdentifier.Value))
-				{
-					StaleTracks.Add(TrackIdentifier.Value, MoveTemp(*Track));
-				}
+				SubSectionsToRemove.Add(Pair);
 			}
-			
-			Tracks.Remove(TrackIdentifier.Value);
-			TemplateLedger.TrackReferenceCounts.Remove(TrackIdentifier);
+		}
+
+		for (const TTuple<FGuid, FFloatRange>& Pair : SubSectionsToRemove)
+		{
+			TemplateLedger.SubSectionRanges.Remove(Pair.Key);
+			EvaluationField.Invalidate(Pair.Value);
 		}
 	}
-	TemplateLedger.TrackSignatureToTrackIdentifier.Remove(InSignature);
 }
 
 const TMap<FMovieSceneTrackIdentifier, FMovieSceneEvaluationTrack>& FMovieSceneEvaluationTemplate::GetTracks() const
 {
-	// Reinterpret the uint32 as FMovieSceneTrackIdentifier
-	static_assert(sizeof(FMovieSceneTrackIdentifier) == sizeof(uint32), "FMovieSceneTrackIdentifier is not convertible directly to/from uint32.");
-	return *reinterpret_cast<const TMap<FMovieSceneTrackIdentifier, FMovieSceneEvaluationTrack>*>(&Tracks);
+	return Tracks;
 }
 
 TMap<FMovieSceneTrackIdentifier, FMovieSceneEvaluationTrack>& FMovieSceneEvaluationTemplate::GetTracks()
 {
-	// Reinterpret the uint32 as FMovieSceneTrackIdentifier
-	static_assert(sizeof(FMovieSceneTrackIdentifier) == sizeof(uint32), "FMovieSceneTrackIdentifier is not convertible directly to/from uint32.");
-	return *reinterpret_cast<TMap<FMovieSceneTrackIdentifier, FMovieSceneEvaluationTrack>*>(&Tracks);
-}
-
-TArrayView<const FMovieSceneTrackIdentifier> FMovieSceneEvaluationTemplate::FindTracks(const FGuid& InSignature) const
-{
-	return TemplateLedger.FindTracks(InSignature);
+	return Tracks;
 }
