Copying //UE4/Dev-Sequencer to //UE4/Dev-Main (Source: //UE4/Dev-Sequencer @ 3800792)

diff --git a/Engine/Source/Runtime/MovieScene/Private/Evaluation/MovieSceneEvaluationTemplateInstance.cpp b/Engine/Source/Runtime/MovieScene/Private/Evaluation/MovieSceneEvaluationTemplateInstance.cpp
--- a/Engine/Source/Runtime/MovieScene/Private/Evaluation/MovieSceneEvaluationTemplateInstance.cpp
+++ b/Engine/Source/Runtime/MovieScene/Private/Evaluation/MovieSceneEvaluationTemplateInstance.cpp
@@ -4,7 +4,12 @@
 #include "IMovieScenePlayer.h"
 #include "MovieSceneSequence.h"
 #include "MovieSceneSequence.h"
-#include "MovieSceneSequenceInstance.h"
+#include "MovieSceneCompiler.h"
+#include "Sections/MovieSceneSubSection.h"
+#include "Compilation/MovieSceneEvaluationTemplateGenerator.h"
+
+#include "IMovieSceneModule.h"
+#include "Algo/Sort.h"
 
 DECLARE_CYCLE_STAT(TEXT("Entire Evaluation Cost"), MovieSceneEval_EntireEvaluationCost, STATGROUP_MovieSceneEval);
 DECLARE_CYCLE_STAT(TEXT("Gather Entries For Frame"), MovieSceneEval_GatherEntries, STATGROUP_MovieSceneEval);
@@ -44,79 +49,41 @@ private:
 	TArray<FMovieSceneEvaluationKey> KeysToRestore;
 };
 
-FMovieSceneEvaluationTemplateInstance::FMovieSceneEvaluationTemplateInstance(UMovieSceneSequence& InSequence, const FMovieSceneEvaluationTemplate& InTemplate)
-	: Sequence(&InSequence)
-	, Template(&InTemplate)
-	, PreRollRange(TRange<float>::Empty())
-	, PostRollRange(TRange<float>::Empty())
-	, HierarchicalBias(0)
-{
-	if (Template->bHasLegacyTrackInstances)
-	{
-		// Initialize the legacy sequence instance
-		LegacySequenceInstance = MakeShareable(new FMovieSceneSequenceInstance(InSequence, MovieSceneSequenceID::Root));
-	}
-}
+FMovieSceneEvaluationTemplateInstance::FMovieSceneEvaluationTemplateInstance()
+	: Sequence(nullptr), Template(nullptr), SubData(nullptr)
+{}
 
-FMovieSceneEvaluationTemplateInstance::FMovieSceneEvaluationTemplateInstance(const FMovieSceneSubSequenceData& InSubData, const FMovieSceneEvaluationTemplate& InTemplate, FMovieSceneSequenceIDRef InSequenceID)
-	: Sequence(InSubData.Sequence)
-	, RootToSequenceTransform(InSubData.RootToSequenceTransform)
-	, Template(&InTemplate)
-	, PreRollRange(InSubData.PreRollRange)
-	, PostRollRange(InSubData.PostRollRange)
-	, HierarchicalBias(InSubData.HierarchicalBias)
+FMovieSceneEvaluationTemplateInstance::FMovieSceneEvaluationTemplateInstance(UMovieSceneSequence& InSequence, FMovieSceneEvaluationTemplate* InTemplate)
+	: Sequence(&InSequence), Template(InTemplate), SubData(nullptr)
+{}
+
+FMovieSceneEvaluationTemplateInstance::FMovieSceneEvaluationTemplateInstance(const FMovieSceneSubSequenceData* InSubData, IMovieSceneSequenceTemplateStore& TemplateStore)
+	: Sequence(nullptr), Template(nullptr), SubData(InSubData)
 {
-	if (Template->bHasLegacyTrackInstances)
+	Sequence = InSubData ? InSubData->GetSequence() : nullptr;
+	if (Sequence)
 	{
-		// Initialize the legacy sequence instance
-		LegacySequenceInstance = MakeShareable(new FMovieSceneSequenceInstance(*InSubData.Sequence, InSequenceID));
+		Template = &TemplateStore.AccessTemplate(*Sequence);
 	}
 }
 
 FMovieSceneRootEvaluationTemplateInstance::FMovieSceneRootEvaluationTemplateInstance()
 	: RootSequence(nullptr)
-	, TemplateStore(MakeShareable(new FMovieSceneSequenceTemplateStore))
-	, bIsDirty(false)
+	, TemplateStore(MakeShared<FMovieSceneSequencePrecompiledTemplateStore>())
 {
-
 }
 
 FMovieSceneRootEvaluationTemplateInstance::~FMovieSceneRootEvaluationTemplateInstance()
 {
-	Reset();
-}
-
-void FMovieSceneRootEvaluationTemplateInstance::Reset()
-{
-	if (TemplateStore->AreTemplatesVolatile())
-	{
-		if (UMovieSceneSequence* Sequence = RootInstance.Sequence.Get())
-		{
-			Sequence->OnSignatureChanged().RemoveAll(this);
-		}
-
-		for (auto& Pair : SubInstances)
-		{
-			if (UMovieSceneSequence* Sequence = Pair.Value.Sequence.Get())
-			{
-				Sequence->OnSignatureChanged().RemoveAll(this);
-			}
-		}
-	}
-
-	SubInstances.Reset();
 }
 
-void FMovieSceneRootEvaluationTemplateInstance::Initialize(UMovieSceneSequence& InRootSequence, IMovieScenePlayer& Player, TSharedRef<FMovieSceneSequenceTemplateStore> InTemplateStore)
+void FMovieSceneRootEvaluationTemplateInstance::Initialize(UMovieSceneSequence& InRootSequence, IMovieScenePlayer& Player, TSharedRef<IMovieSceneSequenceTemplateStore> InTemplateStore)
 {
 	if (RootSequence != &InRootSequence)
 	{
 		Finish(Player);
 	}
 
-	// Ensure we reset everything before we overwrite the template store (which potentially owns templates we've previously referenced)
-	Reset();
-
 	TemplateStore = MoveTemp(InTemplateStore);
 
 	Initialize(InRootSequence, Player);
@@ -124,8 +91,6 @@ void FMovieSceneRootEvaluationTemplateInstance::Initialize(UMovieSceneSequence&
 
 void FMovieSceneRootEvaluationTemplateInstance::Initialize(UMovieSceneSequence& InRootSequence, IMovieScenePlayer& Player)
 {
-	Reset();
-
 	if (RootSequence.Get() != &InRootSequence)
 	{
 		// Always ensure that there is no persistent data when initializing a new sequence
@@ -136,46 +101,13 @@ void FMovieSceneRootEvaluationTemplateInstance::Initialize(UMovieSceneSequence&
 		LastFrameMetaData.Reset();
 		ThisFrameMetaData.Reset();
 		ExecutionTokens = FMovieSceneExecutionTokens();
-	}
-
-	const bool bAddEvents = TemplateStore->AreTemplatesVolatile();
-	
-	RootSequence = &InRootSequence;
-
-	const FMovieSceneEvaluationTemplate& RootTemplate = TemplateStore->GetCompiledTemplate(InRootSequence);
-	
-	Player.State.AssignSequence(MovieSceneSequenceID::Root, InRootSequence, Player);
-	RootInstance = FMovieSceneEvaluationTemplateInstance(InRootSequence, RootTemplate);
-
-	if (bAddEvents)
-	{
-		InRootSequence.OnSignatureChanged().AddRaw(this, &FMovieSceneRootEvaluationTemplateInstance::OnSequenceChanged);
-	}
-
-	for (auto& Pair : RootTemplate.Hierarchy.AllSubSequenceData())
-	{
-		const FMovieSceneSubSequenceData& Data = Pair.Value;
-		FMovieSceneSequenceID SequenceID = Pair.Key;
-
-		if (!Data.Sequence)
-		{
-			continue;
-		}
 
-		Player.State.AssignSequence(SequenceID, *Data.Sequence, Player);
-
-		const FMovieSceneEvaluationTemplate& ChildTemplate = TemplateStore->GetCompiledTemplate(*Data.Sequence, FObjectKey(Data.SequenceKeyObject));
-		FMovieSceneEvaluationTemplateInstance& Instance = SubInstances.Add(SequenceID, FMovieSceneEvaluationTemplateInstance(Data, ChildTemplate, SequenceID));
-
-		if (bAddEvents)
-		{
-			Data.Sequence->OnSignatureChanged().AddRaw(this, &FMovieSceneRootEvaluationTemplateInstance::OnSequenceChanged);
-		}
+		TransientInstances.ResetSubInstances();
 	}
 
-	bIsDirty = false;
-
-	OnUpdatedEvent.Broadcast();
+	RootSequence = &InRootSequence;
+	RootTemplate = &TemplateStore->AccessTemplate(InRootSequence);
+	TransientInstances.RootInstance = FMovieSceneEvaluationTemplateInstance(InRootSequence, RootTemplate);
 }
 
 void FMovieSceneRootEvaluationTemplateInstance::Finish(IMovieScenePlayer& Player)
@@ -184,57 +116,31 @@ void FMovieSceneRootEvaluationTemplateInstance::Finish(IMovieScenePlayer& Player
 	ThisFrameMetaData.Reset();
 
 	CallSetupTearDown(Player);
+
+	TransientInstances.Reset();
 }
 
-void FMovieSceneRootEvaluationTemplateInstance::Evaluate(FMovieSceneContext Context, IMovieScenePlayer& Player, FMovieSceneSequenceID OverrideRootID)
+void FMovieSceneRootEvaluationTemplateInstance::Evaluate(FMovieSceneContext Context, IMovieScenePlayer& Player, FMovieSceneSequenceID InOverrideRootID)
 {
-	if (bIsDirty && RootSequence.IsValid())
-	{
-		Initialize(*RootSequence.Get(), Player);
-	}
-
 	SCOPE_CYCLE_COUNTER(MovieSceneEval_EntireEvaluationCost);
 
 	Swap(ThisFrameMetaData, LastFrameMetaData);
 	ThisFrameMetaData.Reset();
 
-	const FMovieSceneEvaluationTemplateInstance* Instance = GetInstance(OverrideRootID);
-
-	ensureMsgf(Instance, TEXT("Could not find instance for supplied sequence ID."));
-
-	const int32 FieldIndex = (Instance && Instance->Template) ? Instance->Template->EvaluationField.GetSegmentFromTime(Context.GetTime() * Instance->RootToSequenceTransform) : INDEX_NONE;
-	if (FieldIndex == INDEX_NONE)
+	if (TransientInstances.RootID != InOverrideRootID)
 	{
+		// Tear everything down if we're evaluating a different root sequence
 		CallSetupTearDown(Player);
-		return;
-	}
-
-	// Construct a path that allows us to remap sequence IDs from the local (OverrideRootID) template, to the master template
-	ReverseOverrideRootPath.Reset();
-	{
-		FMovieSceneSequenceID CurrentSequenceID = OverrideRootID;
-		const FMovieSceneSequenceHierarchy& Hierarchy = GetHierarchy();
-		while (CurrentSequenceID != MovieSceneSequenceID::Root)
-		{
-			const FMovieSceneSequenceHierarchyNode* CurrentNode = Hierarchy.FindNode(CurrentSequenceID);
-			const FMovieSceneSubSequenceData* SubData = Hierarchy.FindSubData(CurrentSequenceID);
-			if (!ensureAlwaysMsgf(CurrentNode && SubData, TEXT("Malformed sequence hierarchy")))
-			{
-				return;
-			}
-
-			ReverseOverrideRootPath.Add(SubData->DeterministicSequenceID);
-			CurrentSequenceID = CurrentNode->ParentID;
-		}
+		LastFrameMetaData.Reset();
 	}
 
-	const FMovieSceneEvaluationGroup& Group = Instance->Template->EvaluationField.Groups[FieldIndex];
+	RootOverridePath.Set(InOverrideRootID, GetHierarchy());
 
-	// Take a copy of this frame's meta-data
-	ThisFrameMetaData = Instance->Template->EvaluationField.MetaData[FieldIndex];
-	if (OverrideRootID != MovieSceneSequenceID::Root)
+	const FMovieSceneEvaluationGroup* GroupToEvaluate = SetupFrame(Player, InOverrideRootID, Context);
+	if (!GroupToEvaluate)
 	{
-		ThisFrameMetaData.RemapSequenceIDsForRoot(OverrideRootID);
+		CallSetupTearDown(Player);
+		return;
 	}
 
 	// Cause stale tracks to not restore until after evaluation. This fixes issues when tracks that are set to 'Restore State' are regenerated, causing the state to be restored then re-animated by the new track
@@ -248,18 +154,111 @@ void FMovieSceneRootEvaluationTemplateInstance::Evaluate(FMovieSceneContext Cont
 	Player.State.InvalidateExpiredObjects();
 
 	// Accumulate execution tokens into this structure
-	EvaluateGroup(Group, Context, Player);
+	EvaluateGroup(*GroupToEvaluate, Context, Player);
 
 	// Process execution tokens
 	ExecutionTokens.Apply(Context, Player);
 }
 
+const FMovieSceneEvaluationGroup* FMovieSceneRootEvaluationTemplateInstance::SetupFrame(IMovieScenePlayer& Player, FMovieSceneSequenceID InOverrideRootID, FMovieSceneContext Context)
+{
+	UMovieSceneSequence* RootSequencePtr = RootSequence.Get();
+	if (!RootSequencePtr)
+	{
+		return nullptr;
+	}
+
+	// Create the instance for the current root override
+	FMovieSceneEvaluationTemplateInstance RootOverrideInstance;
+
+	if (InOverrideRootID != MovieSceneSequenceID::Root)
+	{
+		if (const FMovieSceneSubSequenceData* OverrideSubData = GetHierarchy().FindSubData(InOverrideRootID))
+		{
+			RootOverrideInstance = FMovieSceneEvaluationTemplateInstance(OverrideSubData, *TemplateStore);
+		}
+	}
+	else
+	{
+		RootOverrideInstance = FMovieSceneEvaluationTemplateInstance(*RootSequencePtr, RootTemplate);
+	}
+
+	if (!ensureMsgf(RootOverrideInstance.IsValid(), TEXT("Could not find valid sequence or template for supplied sequence ID.")))
+	{
+		return nullptr;
+	}
+
+	// Ensure the root is up to date
+	if (RootOverrideInstance.Template->SequenceSignature != RootOverrideInstance.Sequence->GetSignature())
+	{
+		FMovieSceneEvaluationTemplateGenerator(*RootOverrideInstance.Sequence, *RootOverrideInstance.Template).Generate();
+	}
+
+	// First off, attempt to find the evaluation group in the existing evaluation field data from the template
+	int32 TemplateFieldIndex = RootOverrideInstance.Template->EvaluationField.GetSegmentFromTime(Context.GetTime());
+
+	if (TemplateFieldIndex != INDEX_NONE)
+	{
+		const FMovieSceneEvaluationMetaData& FieldMetaData = RootOverrideInstance.Template->EvaluationField.GetMetaData(TemplateFieldIndex);
+
+		// Verify that this field entry is still valid (all its cached signatures are still the same)
+		TRange<float> InvalidatedSubSequenceRange = TRange<float>::Empty();
+		if (FieldMetaData.IsDirty(*RootOverrideInstance.Sequence, RootOverrideInstance.Template->Hierarchy, &InvalidatedSubSequenceRange))
+		{
+			TemplateFieldIndex = INDEX_NONE;
+
+			if (!InvalidatedSubSequenceRange.IsEmpty())
+			{
+				// Invalidate the evaluation field for the root template (it may not exist until we compile below)
+				RootOverrideInstance.Template->EvaluationField.Invalidate(InvalidatedSubSequenceRange);
+			}
+		}
+	}
+
+	if (TemplateFieldIndex == INDEX_NONE)
+	{
+		// We need to compile an entry in the evaluation field
+		static bool bFullCompile = false;
+		if (bFullCompile)
+		{
+			FMovieSceneCompiler::Compile(*RootOverrideInstance.Sequence, *TemplateStore);
+			TemplateFieldIndex = RootOverrideInstance.Template->EvaluationField.GetSegmentFromTime(Context.GetTime());
+		}
+		else
+		{
+			TOptional<FCompiledGroupResult> CompileResult = FMovieSceneCompiler::CompileTime(Context.GetTime(), *RootOverrideInstance.Sequence, *TemplateStore);
+
+			if (CompileResult.IsSet())
+			{
+				TRange<float> FieldRange = CompileResult->Range;
+				TemplateFieldIndex = RootOverrideInstance.Template->EvaluationField.Insert(
+					FieldRange,
+					MoveTemp(CompileResult->Group),
+					MoveTemp(CompileResult->MetaData)
+					);
+			}
+		}
+	}
+
+	if (TemplateFieldIndex != INDEX_NONE)
+	{
+		// Set meta-data
+		ThisFrameMetaData = RootOverrideInstance.Template->EvaluationField.GetMetaData(TemplateFieldIndex);
+
+		RecreateInstances(RootOverrideInstance, InOverrideRootID, Player);
+
+		return &RootOverrideInstance.Template->EvaluationField.GetGroup(TemplateFieldIndex);
+	}
+
+	return nullptr;
+}
+
 void FMovieSceneRootEvaluationTemplateInstance::EvaluateGroup(const FMovieSceneEvaluationGroup& Group, const FMovieSceneContext& RootContext, IMovieScenePlayer& Player)
 {
 	MOVIESCENE_DETAILED_SCOPE_CYCLE_COUNTER(MovieSceneEval_EvaluateGroup);
 
-	FPersistentEvaluationData PersistentDataProxy(Player.State.PersistentEntityData, Player.State.PersistentSharedData);
-	
+	FPersistentEvaluationData PersistentDataProxy(Player);
+
 	FMovieSceneEvaluationOperand Operand;
 
 	FMovieSceneContext Context = RootContext;
@@ -275,9 +274,9 @@ void FMovieSceneRootEvaluationTemplateInstance::EvaluateGroup(const FMovieSceneE
 			FMovieSceneEvaluationFieldSegmentPtr SegmentPtr = Group.SegmentPtrLUT[TrackIndex];
 
 			// Ensure we're able to find the sequence instance in our root if we've overridden
-			SegmentPtr.SequenceID = GetSequenceIdForRoot(SegmentPtr.SequenceID);
+			SegmentPtr.SequenceID = RootOverridePath.Remap(SegmentPtr.SequenceID);
 
-			const FMovieSceneEvaluationTemplateInstance& Instance = GetInstanceChecked(SegmentPtr.SequenceID);
+			const FMovieSceneEvaluationTemplateInstance& Instance = TransientInstances.GetChecked(SegmentPtr.SequenceID);
 			const FMovieSceneEvaluationTrack* Track = Instance.Template->FindTrack(SegmentPtr.TrackIdentifier);
 
 			if (Track)
@@ -291,16 +290,16 @@ void FMovieSceneRootEvaluationTemplateInstance::EvaluateGroup(const FMovieSceneE
 				Player.PreAnimatedState.SetCaptureEntity(TrackKey, EMovieSceneCompletionMode::KeepState);
 
 				SubContext = Context;
-				if (SegmentPtr.SequenceID != MovieSceneSequenceID::Root)
+				if (Instance.SubData)
 				{
-					SubContext = Context.Transform(Instance.RootToSequenceTransform);
+					SubContext = Context.Transform(Instance.SubData->RootToSequenceTransform);
 
 					// Hittest against the sequence's pre and postroll ranges
-					SubContext.ReportOuterSectionRanges(Instance.PreRollRange, Instance.PostRollRange);
-					SubContext.SetHierarchicalBias(Instance.HierarchicalBias);
+					SubContext.ReportOuterSectionRanges(Instance.SubData->PreRollRange, Instance.SubData->PostRollRange);
+					SubContext.SetHierarchicalBias(Instance.SubData->HierarchicalBias);
 				}
 
-				Track->Initialize(SegmentPtr.SegmentIndex, Operand, SubContext, PersistentDataProxy, Player);
+				Track->Initialize(SegmentPtr.SegmentID, Operand, SubContext, PersistentDataProxy, Player);
 			}
 		}
 
@@ -313,11 +312,11 @@ void FMovieSceneRootEvaluationTemplateInstance::EvaluateGroup(const FMovieSceneE
 		for (; TrackIndex < Index.LUTOffset + Index.NumInitPtrs + Index.NumEvalPtrs; ++TrackIndex)
 		{
 			FMovieSceneEvaluationFieldSegmentPtr SegmentPtr = Group.SegmentPtrLUT[TrackIndex];
-			
+
 			// Ensure we're able to find the sequence instance in our root if we've overridden
-			SegmentPtr.SequenceID = GetSequenceIdForRoot(SegmentPtr.SequenceID);
+			SegmentPtr.SequenceID = RootOverridePath.Remap(SegmentPtr.SequenceID);
 
-			const FMovieSceneEvaluationTemplateInstance& Instance = GetInstanceChecked(SegmentPtr.SequenceID);
+			const FMovieSceneEvaluationTemplateInstance& Instance = TransientInstances.GetChecked(SegmentPtr.SequenceID);
 			const FMovieSceneEvaluationTrack* Track = Instance.Template->FindTrack(SegmentPtr.TrackIdentifier);
 
 			if (Track)
@@ -333,17 +332,17 @@ void FMovieSceneRootEvaluationTemplateInstance::EvaluateGroup(const FMovieSceneE
 				ExecutionTokens.SetCurrentScope(FMovieSceneEvaluationScope(TrackKey, EMovieSceneCompletionMode::KeepState));
 
 				SubContext = Context;
-				if (SegmentPtr.SequenceID != MovieSceneSequenceID::Root)
+				if (Instance.SubData)
 				{
-					SubContext = Context.Transform(Instance.RootToSequenceTransform);
+					SubContext = Context.Transform(Instance.SubData->RootToSequenceTransform);
 
 					// Hittest against the sequence's pre and postroll ranges
-					SubContext.ReportOuterSectionRanges(Instance.PreRollRange, Instance.PostRollRange);
-					SubContext.SetHierarchicalBias(Instance.HierarchicalBias);
+					SubContext.ReportOuterSectionRanges(Instance.SubData->PreRollRange, Instance.SubData->PostRollRange);
+					SubContext.SetHierarchicalBias(Instance.SubData->HierarchicalBias);
 				}
 
 				Track->Evaluate(
-					SegmentPtr.SegmentIndex,
+					SegmentPtr.SegmentID,
 					Operand,
 					SubContext,
 					PersistentDataProxy,
@@ -359,7 +358,7 @@ void FMovieSceneRootEvaluationTemplateInstance::CallSetupTearDown(IMovieScenePla
 {
 	MOVIESCENE_DETAILED_SCOPE_CYCLE_COUNTER(MovieSceneEval_CallSetupTearDown);
 
-	FPersistentEvaluationData PersistentDataProxy(Player.State.PersistentEntityData, Player.State.PersistentSharedData);
+	FPersistentEvaluationData PersistentDataProxy(Player);
 
 	TArray<FMovieSceneOrderedEvaluationKey> ExpiredEntities;
 	TArray<FMovieSceneOrderedEvaluationKey> NewEntities;
@@ -367,10 +366,13 @@ void FMovieSceneRootEvaluationTemplateInstance::CallSetupTearDown(IMovieScenePla
 	
 	for (const FMovieSceneOrderedEvaluationKey& OrderedKey : ExpiredEntities)
 	{
-		const FMovieSceneEvaluationKey Key = OrderedKey.Key;
+		FMovieSceneEvaluationKey Key = OrderedKey.Key;
+
+		// Ensure we're able to find the sequence instance in our root if we've overridden
+		Key.SequenceID = RootOverridePath.Remap(Key.SequenceID);
 
-		const FMovieSceneEvaluationTemplateInstance* Instance = FindInstance(Key.SequenceID);
-		if (Instance)
+		const FMovieSceneEvaluationTemplateInstance* Instance = TransientInstances.Find(Key.SequenceID);
+		if (Instance && Instance->IsValid())
 		{
 			const FMovieSceneEvaluationTrack* Track = Instance->Template->FindTrack(Key.TrackIdentifier);
 			const bool bStaleTrack = Instance->Template->IsTrackStale(Key.TrackIdentifier);
@@ -378,7 +380,7 @@ void FMovieSceneRootEvaluationTemplateInstance::CallSetupTearDown(IMovieScenePla
 			// Track data key may be required by both tracks and sections
 			PersistentDataProxy.SetTrackKey(Key.AsTrack());
 
-			if (Key.SectionIdentifier == uint32(-1))
+			if (Key.SectionIndex == uint32(-1))
 			{
 				if (Track)
 				{
@@ -390,9 +392,9 @@ void FMovieSceneRootEvaluationTemplateInstance::CallSetupTearDown(IMovieScenePla
 			else
 			{
 				PersistentDataProxy.SetSectionKey(Key);
-				if (Track)
+				if (Track && Track->HasChildTemplate(Key.SectionIndex))
 				{
-					Track->GetChildTemplate(Key.SectionIdentifier).OnEndEvaluation(PersistentDataProxy, Player);
+					Track->GetChildTemplate(Key.SectionIndex).OnEndEvaluation(PersistentDataProxy, Player);
 				}
 
 				PersistentDataProxy.ResetSectionData();
@@ -411,22 +413,26 @@ void FMovieSceneRootEvaluationTemplateInstance::CallSetupTearDown(IMovieScenePla
 
 	for (const FMovieSceneOrderedEvaluationKey& OrderedKey : NewEntities)
 	{
-		const FMovieSceneEvaluationKey Key = OrderedKey.Key;
+		FMovieSceneEvaluationKey Key = OrderedKey.Key;
 
-		const FMovieSceneEvaluationTemplateInstance& Instance = GetInstanceChecked(Key.SequenceID);
+		// Ensure we're able to find the sequence instance in our root if we've overridden
+		Key.SequenceID = RootOverridePath.Remap(Key.SequenceID);
+
+		const FMovieSceneEvaluationTemplateInstance& Instance = TransientInstances.GetChecked(Key.SequenceID);
+		check(Instance.IsValid());
 
 		if (const FMovieSceneEvaluationTrack* Track = Instance.Template->FindTrack(Key.TrackIdentifier))
 		{
 			PersistentDataProxy.SetTrackKey(Key.AsTrack());
 
-			if (Key.SectionIdentifier == uint32(-1))
+			if (Key.SectionIndex == uint32(-1))
 			{
 				Track->OnBeginEvaluation(PersistentDataProxy, Player);
 			}
 			else
 			{
 				PersistentDataProxy.SetSectionKey(Key);
-				Track->GetChildTemplate(Key.SectionIdentifier).OnBeginEvaluation(PersistentDataProxy, Player);
+				Track->GetChildTemplate(Key.SectionIndex).OnBeginEvaluation(PersistentDataProxy, Player);
 			}
 		}
 	}
@@ -438,8 +444,18 @@ void FMovieSceneRootEvaluationTemplateInstance::CallSetupTearDown(IMovieScenePla
 	ThisFrameMetaData.DiffSequences(LastFrameMetaData, nullptr, &ExpiredSequenceIDs);
 	for (FMovieSceneSequenceID ExpiredID : ExpiredSequenceIDs)
 	{
-		Register.OnSequenceExpired(ExpiredID, Player);
+		Register.OnSequenceExpired(RootOverridePath.Remap(ExpiredID), Player);
+	}
+}
+
+bool FMovieSceneRootEvaluationTemplateInstance::IsDirty() const
+{
+	if (TransientInstances.RootInstance.IsValid())
+	{
+		return LastFrameMetaData.IsDirty(*TransientInstances.RootInstance.Sequence, TransientInstances.RootInstance.Template->Hierarchy);
 	}
+
+	return true;
 }
 
 void FMovieSceneRootEvaluationTemplateInstance::CopyActuators(FMovieSceneBlendingAccumulator& Accumulator) const
@@ -447,7 +463,60 @@ void FMovieSceneRootEvaluationTemplateInstance::CopyActuators(FMovieSceneBlendin
 	Accumulator.Actuators = ExecutionTokens.GetBlendingAccumulator().Actuators;
 }
 
-void FMovieSceneRootEvaluationTemplateInstance::OnSequenceChanged()
+UMovieSceneSequence* FMovieSceneRootEvaluationTemplateInstance::GetSequence(FMovieSceneSequenceIDRef SequenceID) const
 {
-	bIsDirty = true;
+	if (SequenceID == MovieSceneSequenceID::Root)
+	{
+		return RootSequence.Get();
+	}
+	
+	const FMovieSceneSubSequenceData* SubData = GetHierarchy().FindSubData(SequenceID);
+	return SubData ? SubData->GetSequence() : nullptr;
 }
+
+FMovieSceneEvaluationTemplate* FMovieSceneRootEvaluationTemplateInstance::FindTemplate(FMovieSceneSequenceIDRef SequenceID)
+{
+	if (SequenceID == MovieSceneSequenceID::Root)
+	{
+		return RootTemplate;
+	}
+
+	const FMovieSceneSubSequenceData* SubData = GetHierarchy().FindSubData(SequenceID);
+	UMovieSceneSequence* Sequence = SubData ? SubData->GetSequence() : nullptr;
+
+	return Sequence ? &TemplateStore->AccessTemplate(*Sequence) : nullptr;
+}
+
+void FMovieSceneRootEvaluationTemplateInstance::RecreateInstances(const FMovieSceneEvaluationTemplateInstance& RootOverrideInstance, FMovieSceneSequenceID InOverrideRootID, IMovieScenePlayer& Player)
+{
+	check(RootOverrideInstance.IsValid());
+
+	TransientInstances.RootID = InOverrideRootID;
+	TransientInstances.RootInstance = RootOverrideInstance;
+	
+	Player.State.AssignSequence(TransientInstances.RootID, *RootOverrideInstance.Sequence, Player);
+
+	TransientInstances.ResetSubInstances();
+
+	// We recreate all necessary sequence instances for the current and previous frames by diffing the sequences active last frame, with this frame
+	TArray<FMovieSceneSequenceID> PreviousAndCurrentFrameSequenceIDs = LastFrameMetaData.ActiveSequences;
+	ThisFrameMetaData.DiffSequences(LastFrameMetaData, &PreviousAndCurrentFrameSequenceIDs, nullptr);
+
+	for (FMovieSceneSequenceID SequenceID : PreviousAndCurrentFrameSequenceIDs)
+	{
+		const FMovieSceneSubSequenceData* SubData = RootOverrideInstance.Template->Hierarchy.FindSubData(SequenceID);
+		if (SubData)
+		{
+			FMovieSceneEvaluationTemplateInstance NewInstance(SubData, *TemplateStore);
+			
+			// Always add sequence IDs as the root
+			SequenceID = RootOverridePath.Remap(SequenceID);
+			TransientInstances.Add(SequenceID, NewInstance);
+
+			if (NewInstance.Sequence)
+			{
+				Player.State.AssignSequence(SequenceID, *NewInstance.Sequence, Player);
+			}
+		}
+	}
+}
\ No newline at end of file
diff --git a/Engine/Source/Runtime/MovieScene/Private/Evaluation/MovieSceneEvaluationTrack.cpp b/Engine/Source/Runtime/MovieScene/Private/Evaluation/MovieSceneEvaluationTrack.cpp
--- a/Engine/Source/Runtime/MovieScene/Private/Evaluation/MovieSceneEvaluationTrack.cpp
+++ b/Engine/Source/Runtime/MovieScene/Private/Evaluation/MovieSceneEvaluationTrack.cpp
@@ -3,8 +3,117 @@
 #include "Evaluation/MovieSceneEvaluationTrack.h"
 #include "MovieSceneExecutionTokens.h"
 
+#include "MovieSceneSegmentCompiler.h"
+#include "MovieSceneCompilerRules.h"
+#include "MovieSceneTrack.h"
+#include "MovieSceneCommonHelpers.h"
+
+#include "Algo/BinarySearch.h"
+
+
+#if WITH_DEV_AUTOMATION_TESTS
+
+	struct FGetTrackSegmentBlender
+	{
+		FMovieSceneTrackSegmentBlender* operator()(const UMovieSceneTrack* Track, FMovieSceneTrackSegmentBlenderPtr& OutOwnerPtr) const
+		{
+			if (Override.IsValid())
+			{
+				return &Override.GetValue();
+			}
+			else if (Track)
+			{
+				OutOwnerPtr = Track->GetTrackSegmentBlender();
+			}
+
+			return OutOwnerPtr.GetPtr();
+		}
+
+		static FMovieSceneTrackSegmentBlenderPtr Override;
+
+	} GetTrackSegmentBlender;
+
+	FMovieSceneTrackSegmentBlenderPtr FGetTrackSegmentBlender::Override;
+
+	FScopedOverrideTrackSegmentBlender::FScopedOverrideTrackSegmentBlender(FMovieSceneTrackSegmentBlenderPtr&& InTrackSegmentBlender)
+	{
+		check(!GetTrackSegmentBlender.Override.IsValid());
+		GetTrackSegmentBlender.Override = MoveTemp(InTrackSegmentBlender);
+	}
+
+	FScopedOverrideTrackSegmentBlender::~FScopedOverrideTrackSegmentBlender()
+	{
+		GetTrackSegmentBlender.Override.Reset();
+	}
+
+#else	// WITH_DEV_AUTOMATION_TESTS
+
+	FMovieSceneTrackSegmentBlender* GetTrackSegmentBlender(const UMovieSceneTrack* Track, FMovieSceneTrackSegmentBlenderPtr& OutOwnerPtr)
+	{
+		if (Track)
+		{
+			OutOwnerPtr = Track->GetTrackSegmentBlender();
+		}
+		return OutOwnerPtr.GetPtr();
+	}
+
+#endif	// WITH_DEV_AUTOMATION_TESTS
+
+FMovieSceneSegmentIdentifier FMovieSceneEvaluationTrackSegments::Add(FMovieSceneSegment&& Segment)
+{
+	int32 InsertIndex = 0;
+	if (Segment.Range.GetLowerBound().IsClosed())
+	{
+		// Determine where to insert the segment
+		auto GetLowerBound = [](const FMovieSceneSegment& In) {
+			return In.Range.GetLowerBound();
+		};
+
+		// Find the index of the first segment with a lower bound > than this lower bound.
+		InsertIndex = Algo::UpperBoundBy(SortedSegments, Segment.Range.GetLowerBound(), GetLowerBound, MovieSceneHelpers::SortLowerBounds);
+		InsertIndex = FMath::Clamp(InsertIndex, 0, SortedSegments.Num());
+	}
+
+	if (!ensureMsgf(   ( !SortedSegments.IsValidIndex(InsertIndex  ) || !SortedSegments[InsertIndex  ].Range.Overlaps(Segment.Range) ) && 
+			           ( !SortedSegments.IsValidIndex(InsertIndex+1) || !SortedSegments[InsertIndex+1].Range.Overlaps(Segment.Range) ), TEXT("Attempting to add overlapping segment to segment array. This is invalid.")))
+	{
+		return FMovieSceneSegmentIdentifier();
+	}
+
+	// Attempt to combine into the previous or next segment
+	if (SortedSegments.IsValidIndex(InsertIndex-1) && SortedSegments[InsertIndex-1].CombineWith(Segment))
+	{
+		return FMovieSceneSegmentIdentifier(SegmentIdentifierToIndex.Add(InsertIndex-1));
+	}
+	if (SortedSegments.IsValidIndex(InsertIndex  ) && SortedSegments[InsertIndex  ].CombineWith(Segment))
+	{
+		return FMovieSceneSegmentIdentifier(SegmentIdentifierToIndex.Add(InsertIndex  ));
+	}
+
+	// Fixup any indices that occur after this index
+	for (int32& Index : SegmentIdentifierToIndex)
+	{
+		if (Index >= InsertIndex)
+		{
+			++Index;
+		}
+	}
+
+	FMovieSceneSegmentIdentifier ID(SegmentIdentifierToIndex.Num());
+	Segment.ID = ID;
+
+	SegmentIdentifierToIndex.Add(InsertIndex);
+	SortedSegments.Insert(MoveTemp(Segment), InsertIndex);
+
+	return ID;
+}
 
 FMovieSceneEvaluationTrack::FMovieSceneEvaluationTrack()
+	: EvaluationPriority(1000)
+	, EvaluationMethod(EEvaluationMethod::Static)
+	, SourceTrack(nullptr)
+	, bEvaluateInPreroll(true)
+	, bEvaluateInPostroll(true)
 {
 }
 
@@ -12,6 +121,7 @@ FMovieSceneEvaluationTrack::FMovieSceneEvaluationTrack(const FGuid& InObjectBind
 	: ObjectBindingID(InObjectBindingID)
 	, EvaluationPriority(1000)
 	, EvaluationMethod(EEvaluationMethod::Static)
+	, SourceTrack(nullptr)
 	, bEvaluateInPreroll(true)
 	, bEvaluateInPostroll(true)
 {
@@ -34,7 +144,7 @@ void FMovieSceneEvaluationTrack::PostSerialize(const FArchive& Ar)
 
 		if (ImplsToRemove.Num())
 		{
-			for (FMovieSceneSegment& Segment : Segments)
+			for (FMovieSceneSegment& Segment : Segments.GetSorted())
 			{
 				Segment.Impls.RemoveAll([&](const FSectionEvaluationData& In) { return ImplsToRemove.Contains(In.ImplIndex); });
 			}
@@ -49,7 +159,7 @@ void FMovieSceneEvaluationTrack::DefineAsSingleTemplate(FMovieSceneEvalTemplateP
 	Segments.Reset(1);
 
 	ChildTemplates.Add(MoveTemp(InTemplate));
-	
+
 	FSectionEvaluationData EvalData(0);
 	Segments.Add(FMovieSceneSegment(TRange<float>::All(), TArrayView<FSectionEvaluationData>(&EvalData, 1)));
 }
@@ -59,41 +169,203 @@ int32 FMovieSceneEvaluationTrack::AddChildTemplate(FMovieSceneEvalTemplatePtr&&
 	return ChildTemplates.Add(MoveTemp(InTemplate));
 }
 
-void FMovieSceneEvaluationTrack::SetSegments(TArray<FMovieSceneSegment>&& InSegments)
+FMovieSceneSegmentIdentifier FMovieSceneEvaluationTrack::CompileSegment(FMovieSceneEvaluationTreeRangeIterator CurrentNode)
+{
+	// Define a new segment for the specified iterator's range
+	FMovieSceneSegment NewSegment(CurrentNode.Range());
+
+	// Get the segment blender required to blend at the track level for this segment
+	FMovieSceneTrackSegmentBlenderPtr TrackSegmentBlenderValue;
+	FMovieSceneTrackSegmentBlender* TrackSegmentBlender = GetTrackSegmentBlender(SourceTrack, TrackSegmentBlenderValue);
+
+	// Iterate all the section evaluation data for the current segment range
+	TMovieSceneEvaluationTreeDataIterator<FSectionEvaluationData> DataIterator = GetData(CurrentNode.Node());
+	if (DataIterator)
+	{
+		// Track section data
+		FSegmentBlendData TrackBlendData;
+
+		FMovieSceneTrackRowSegmentBlenderPtr RowSegmentBlender = SourceTrack ? SourceTrack->GetRowSegmentBlender() : FMovieSceneTrackRowSegmentBlenderPtr();
+		// Compile at the row level first
+		if (RowSegmentBlender.IsValid())
+		{
+			TMap<int32, FSegmentBlendData> RowBlendData;
+
+			// Add every FSectionEvaluationData for the current time range to the section data
+			for (const FSectionEvaluationData& EvalData : DataIterator)
+			{
+				const UMovieSceneSection* Section = ChildTemplates[EvalData.ImplIndex]->GetSourceSection();
+				if (Section)
+				{
+					RowBlendData.FindOrAdd(Section->GetRowIndex()).Add(FMovieSceneSectionData(Section, EvalData.ImplIndex, EvalData.Flags));
+				}
+			}
+
+			for (TPair<int32, FSegmentBlendData>& Pair : RowBlendData)
+			{
+				RowSegmentBlender->Blend(Pair.Value);
+				TrackBlendData.Append(Pair.Value);
+			}
+		}
+		else for (const FSectionEvaluationData& EvalData : DataIterator)
+		{
+			const UMovieSceneSection* Section = ChildTemplates.IsValidIndex(EvalData.ImplIndex) ? ChildTemplates[EvalData.ImplIndex]->GetSourceSection() : nullptr;
+			TrackBlendData.Add(FMovieSceneSectionData(Section, EvalData.ImplIndex, EvalData.Flags));
+		}
+
+		// Compile at the track level
+		if (TrackSegmentBlender)
+		{
+			TrackSegmentBlender->Blend(TrackBlendData);
+		}
+
+		TrackBlendData.AddToSegment(NewSegment);
+	}
+
+	// Insert an empty segment
+	if (NewSegment.Impls.Num() == 0)
+	{
+		if (TrackSegmentBlender && TrackSegmentBlender->CanFillEmptySpace())
+		{
+			FMovieSceneEvaluationTreeRangeIterator PrevSegmentIt = CurrentNode.Previous();
+			FMovieSceneEvaluationTreeRangeIterator NextSegmentIt = CurrentNode.Next();
+
+			ensure(
+				( !PrevSegmentIt || !PrevSegmentIt.Range().Overlaps(CurrentNode.Range()) ) &&
+				( !NextSegmentIt || !NextSegmentIt.Range().Overlaps(CurrentNode.Range()) )
+			);
+
+			FMovieSceneSegmentIdentifier PrevSegmentID = PrevSegmentIt ? GetSegmentFromIterator(PrevSegmentIt) : FMovieSceneSegmentIdentifier();
+			FMovieSceneSegmentIdentifier NextSegmentID = NextSegmentIt ? GetSegmentFromIterator(NextSegmentIt) : FMovieSceneSegmentIdentifier();
+
+			FMovieSceneSegment* PrevSegment = PrevSegmentID.IsValid() ? &Segments[PrevSegmentID] : nullptr;
+			FMovieSceneSegment* NextSegment = NextSegmentID.IsValid() ? &Segments[NextSegmentID] : nullptr;
+
+			TOptional<FMovieSceneSegment> EmptySegment = TrackSegmentBlender->InsertEmptySpace(NewSegment.Range, PrevSegment, NextSegment);
+			if (EmptySegment.IsSet())
+			{
+				NewSegment.Range = EmptySegment->Range;
+				NewSegment.Impls = EmptySegment->Impls;
+			}
+		}
+	}
+
+	NewSegment.bAllowEmpty = TrackSegmentBlender ? TrackSegmentBlender->AllowEmptySegments() : false;
+	if (!NewSegment.bAllowEmpty && NewSegment.Impls.Num() == 0)
+	{
+		return FMovieSceneSegmentIdentifier();
+	}
+
+	return Segments.Add(MoveTemp(NewSegment));
+}
+
+FMovieSceneSegmentIdentifier FMovieSceneEvaluationTrack::FindFirstSegment(TRange<float> InLocalRange)
 {
-	Segments = MoveTemp(InSegments);
+	// Binary search for an existing compiled segment in the sorted array
+	auto GetLowerBound = [](const FMovieSceneSegment& In) {
+		return In.Range.GetLowerBound();
+	};
+
+	// Find a segment with a lowerbound >= to this iterator range
+	TArrayView<const FMovieSceneSegment> SortedSegments = Segments.GetSorted();
+	int32 SegmentIndex = Algo::LowerBoundBy(SortedSegments, InLocalRange.GetLowerBound(), GetLowerBound, MovieSceneHelpers::SortLowerBounds);
+	if (SortedSegments.IsValidIndex(SegmentIndex) && SortedSegments[SegmentIndex].Range.GetLowerBound() == InLocalRange.GetLowerBound())
+	{
+		return SortedSegments[SegmentIndex].ID;
+	}
+	else if (SortedSegments.IsValidIndex(SegmentIndex-1) && SortedSegments[SegmentIndex-1].Range.Overlaps(InLocalRange))
+	{
+		return SortedSegments[SegmentIndex-1].ID;
+	}
 
-	ValidateSegments();
+	return FMovieSceneSegmentIdentifier();
 }
 
-void FMovieSceneEvaluationTrack::ValidateSegments()
+TArray<FMovieSceneSegmentIdentifier> FMovieSceneEvaluationTrack::GetSegmentsInRange(TRange<float> InLocalRange)
 {
-	// We don't remove segments as this may break ptrs that have been set up in the evaluation field. Instead we just remove invalid track indices.
+	TArray<FMovieSceneSegmentIdentifier> SegmentsInRange;
 
-	for (FMovieSceneSegment& Segment : Segments)
+	while (!InLocalRange.IsEmpty())
 	{
-		for (int32 Index = Segment.Impls.Num() - 1; Index >= 0; --Index)
+		FMovieSceneSegmentIdentifier SegmentID = FindFirstSegment(InLocalRange);
+
+		if (!SegmentID.IsValid())
 		{
-			int32 TemplateIndex = Segment.Impls[Index].ImplIndex;
-			if (!HasChildTemplate(TemplateIndex))
+			FMovieSceneEvaluationTreeRangeIterator NodeAtTime = EvaluationTree.Tree.IterateFromLowerBound(InLocalRange.GetLowerBound());
+			check(NodeAtTime)
+			SegmentID = CompileSegment(NodeAtTime);
+
+			if (!SegmentID.IsValid())
 			{
-				Segment.Impls.RemoveAt(Index, 1, false);
+				if (NodeAtTime.Range().GetUpperBound().IsOpen())
+				{
+					break;
+				}
+
+				InLocalRange = TRange<float>(TRangeBound<float>::FlipInclusion(NodeAtTime.Range().GetUpperBound()), InLocalRange.GetUpperBound());
+				continue;
 			}
 		}
+
+		SegmentsInRange.Add(SegmentID);
+		if (Segments[SegmentID].Range.GetUpperBound().IsOpen())
+		{
+			break;
+		}
+
+		InLocalRange = TRange<float>(TRangeBound<float>::FlipInclusion(Segments[SegmentID].Range.GetUpperBound()), InLocalRange.GetUpperBound());
+	}
+
+	return SegmentsInRange;
+}
+
+FMovieSceneSegmentIdentifier FMovieSceneEvaluationTrack::GetSegmentFromTime(float InTime)
+{
+	FMovieSceneSegmentIdentifier Existing = FindFirstSegment(TRange<float>::Inclusive(InTime, InTime));
+	if (Existing.IsValid())
+	{
+		return Existing;
 	}
+
+	FMovieSceneEvaluationTreeRangeIterator NodeAtTime = EvaluationTree.Tree.IterateFromTime(InTime);
+	check(NodeAtTime)
+	return CompileSegment(NodeAtTime);
+}
+
+FMovieSceneSegmentIdentifier FMovieSceneEvaluationTrack::GetSegmentFromIterator(FMovieSceneEvaluationTreeRangeIterator Iterator)
+{
+	check(Iterator);
+
+	FMovieSceneSegmentIdentifier Existing = FindFirstSegment(Iterator.Range());
+	if (Existing.IsValid())
+	{
+		ensure(Segments[Existing].Range.Contains(Iterator.Range()));
+		return Existing;
+	}
+
+	return CompileSegment(Iterator);
 }
 
-int32 FMovieSceneEvaluationTrack::FindSegmentIndex(float InTime) const
+TRange<float> FMovieSceneEvaluationTrack::GetUniqueRangeFromLowerBound(TRangeBound<float> InLowerBound) const
 {
-	for (int32 Index = 0; Index < Segments.Num(); ++Index)
+	return EvaluationTree.Tree.IterateFromLowerBound(InLowerBound).Range();
+}
+
+void FMovieSceneEvaluationTrack::ValidateSegments()
+{
+	// We don't remove segments as this may break ptrs that have been set up in the evaluation field. Instead we just remove invalid track indices.
+
+	for (FMovieSceneSegment& Segment : Segments.GetSorted())
 	{
-		if (Segments[Index].Range.Contains(InTime))
+		for (int32 Index = Segment.Impls.Num() - 1; Index >= 0; --Index)
 		{
-			return Index;
+			int32 TemplateIndex = Segment.Impls[Index].ImplIndex;
+			if (!HasChildTemplate(TemplateIndex))
+			{
+				Segment.Impls.RemoveAt(Index, 1, false);
+			}
 		}
 	}
-
-	return INDEX_NONE;
 }
 
 void FMovieSceneEvaluationTrack::SetupOverrides()
@@ -112,21 +384,21 @@ void FMovieSceneEvaluationTrack::SetupOverrides()
 	}
 }
 
-void FMovieSceneEvaluationTrack::Initialize(int32 SegmentIndex, const FMovieSceneEvaluationOperand& Operand, const FMovieSceneContext& Context, FPersistentEvaluationData& PersistentData, IMovieScenePlayer& Player) const
+void FMovieSceneEvaluationTrack::Initialize(FMovieSceneSegmentIdentifier SegmentID, const FMovieSceneEvaluationOperand& Operand, const FMovieSceneContext& Context, FPersistentEvaluationData& PersistentData, IMovieScenePlayer& Player) const
 {
 	if (!TrackTemplate.IsValid() || !TrackTemplate->HasCustomInitialize())
 	{
-		DefaultInitialize(SegmentIndex, Operand, Context, PersistentData, Player);
+		DefaultInitialize(SegmentID, Operand, Context, PersistentData, Player);
 	}
 	else
 	{
-		TrackTemplate->Initialize(*this, SegmentIndex, Operand, Context, PersistentData, Player);
+		TrackTemplate->Initialize(*this, SegmentID, Operand, Context, PersistentData, Player);
 	}
 }
 
-void FMovieSceneEvaluationTrack::DefaultInitialize(int32 SegmentIndex, const FMovieSceneEvaluationOperand& Operand, FMovieSceneContext Context, FPersistentEvaluationData& PersistentData, IMovieScenePlayer& Player) const
+void FMovieSceneEvaluationTrack::DefaultInitialize(FMovieSceneSegmentIdentifier SegmentID, const FMovieSceneEvaluationOperand& Operand, FMovieSceneContext Context, FPersistentEvaluationData& PersistentData, IMovieScenePlayer& Player) const
 {
-	const FMovieSceneSegment& Segment = Segments[SegmentIndex];
+	const FMovieSceneSegment& Segment = Segments[SegmentID];
 	for (const FSectionEvaluationData& EvalData : Segment.Impls)
 	{
 		const FMovieSceneEvalTemplate& Template = GetChildTemplate(EvalData.ImplIndex);
@@ -143,34 +415,34 @@ void FMovieSceneEvaluationTrack::DefaultInitialize(int32 SegmentIndex, const FMo
 	}
 }
 
-void FMovieSceneEvaluationTrack::Evaluate(int32 SegmentIndex, const FMovieSceneEvaluationOperand& Operand, const FMovieSceneContext& Context, const FPersistentEvaluationData& PersistentData, FMovieSceneExecutionTokens& ExecutionTokens) const
+void FMovieSceneEvaluationTrack::Evaluate(FMovieSceneSegmentIdentifier SegmentID, const FMovieSceneEvaluationOperand& Operand, const FMovieSceneContext& Context, const FPersistentEvaluationData& PersistentData, FMovieSceneExecutionTokens& ExecutionTokens) const
 {
 	if (!TrackTemplate.IsValid() || !TrackTemplate->HasCustomEvaluate())
 	{
-		DefaultEvaluate(SegmentIndex, Operand, Context, PersistentData, ExecutionTokens);
+		DefaultEvaluate(SegmentID, Operand, Context, PersistentData, ExecutionTokens);
 	}
 	else
 	{
-		TrackTemplate->Evaluate(*this, SegmentIndex, Operand, Context, PersistentData, ExecutionTokens);
+		TrackTemplate->Evaluate(*this, SegmentID, Operand, Context, PersistentData, ExecutionTokens);
 	}
 }
 
-void FMovieSceneEvaluationTrack::DefaultEvaluate(int32 SegmentIndex, const FMovieSceneEvaluationOperand& Operand, const FMovieSceneContext& Context, const FPersistentEvaluationData& PersistentData, FMovieSceneExecutionTokens& ExecutionTokens) const
+void FMovieSceneEvaluationTrack::DefaultEvaluate(FMovieSceneSegmentIdentifier SegmentID, const FMovieSceneEvaluationOperand& Operand, const FMovieSceneContext& Context, const FPersistentEvaluationData& PersistentData, FMovieSceneExecutionTokens& ExecutionTokens) const
 {
 	switch (EvaluationMethod)
 	{
 	case EEvaluationMethod::Static:
-		EvaluateStatic(SegmentIndex, Operand, Context, PersistentData, ExecutionTokens);
+		EvaluateStatic(SegmentID, Operand, Context, PersistentData, ExecutionTokens);
 		break;
 	case EEvaluationMethod::Swept:
-		EvaluateSwept(SegmentIndex, Operand, Context, PersistentData, ExecutionTokens);
+		EvaluateSwept(SegmentID, Operand, Context, PersistentData, ExecutionTokens);
 		break;
 	}
 }
 
-void FMovieSceneEvaluationTrack::EvaluateStatic(int32 SegmentIndex, const FMovieSceneEvaluationOperand& Operand, FMovieSceneContext Context, const FPersistentEvaluationData& PersistentData, FMovieSceneExecutionTokens& ExecutionTokens) const
+void FMovieSceneEvaluationTrack::EvaluateStatic(FMovieSceneSegmentIdentifier SegmentID, const FMovieSceneEvaluationOperand& Operand, FMovieSceneContext Context, const FPersistentEvaluationData& PersistentData, FMovieSceneExecutionTokens& ExecutionTokens) const
 {
-	for (const FSectionEvaluationData& EvalData : GetSegment(SegmentIndex).Impls)
+	for (const FSectionEvaluationData& EvalData : GetSegment(SegmentID).Impls)
 	{
 		const FMovieSceneEvalTemplate& Template = GetChildTemplate(EvalData.ImplIndex);
 
@@ -211,7 +483,7 @@ namespace
 		return true;
 	}
 
-	void GatherSweptSegments(TRange<float> TraversedRange, int32 CurrentSegmentIndex, const TArray<FMovieSceneSegment>& Segments, TMap<int32, TRange<float>>& ImplToAccumulatedRange)
+	void GatherSweptSegments(TRange<float> TraversedRange, int32 CurrentSegmentIndex, TArrayView<const FMovieSceneSegment> Segments, TMap<int32, TRange<float>>& ImplToAccumulatedRange)
 	{
 		// Search backwards from the current segment for any segments intersecting the traversed range
 		{
@@ -230,12 +502,18 @@ namespace
 	}
 }
 
-void FMovieSceneEvaluationTrack::EvaluateSwept(int32 SegmentIndex, const FMovieSceneEvaluationOperand& Operand, const FMovieSceneContext& Context, const FPersistentEvaluationData& PersistentData, FMovieSceneExecutionTokens& ExecutionTokens) const
+void FMovieSceneEvaluationTrack::EvaluateSwept(FMovieSceneSegmentIdentifier SegmentID, const FMovieSceneEvaluationOperand& Operand, const FMovieSceneContext& Context, const FPersistentEvaluationData& PersistentData, FMovieSceneExecutionTokens& ExecutionTokens) const
 {
 	// Accumulate the relevant ranges that each section intersects with the evaluated range
 	TMap<int32, TRange<float>> ImplToAccumulatedRange;
 
-	GatherSweptSegments(Context.GetRange(), SegmentIndex, Segments, ImplToAccumulatedRange);
+	int32 SortedIndex = Segments.GetSortedIndex(SegmentID);
+	if (SortedIndex == INDEX_NONE)
+	{
+		return;
+	}
+
+	GatherSweptSegments(Context.GetRange(), SortedIndex, Segments.GetSorted(), ImplToAccumulatedRange);
 
 	for (auto& Pair : ImplToAccumulatedRange)
 	{
@@ -255,22 +533,22 @@ void FMovieSceneEvaluationTrack::EvaluateSwept(int32 SegmentIndex, const FMovieS
 	}
 }
 
-void FMovieSceneEvaluationTrack::Interrogate(const FMovieSceneContext& Context, FMovieSceneInterrogationData& Container, UObject* BindingOverride) const
+void FMovieSceneEvaluationTrack::Interrogate(const FMovieSceneContext& Context, FMovieSceneInterrogationData& Container, UObject* BindingOverride)
 {
 	if (TrackTemplate.IsValid() && TrackTemplate->Interrogate(Context, Container, BindingOverride))
 	{
 		return;
 	}
 
-	const int32 SegmentIndex = FindSegmentIndex(Context.GetTime());
-	if (SegmentIndex == INDEX_NONE)
+	FMovieSceneSegmentIdentifier SegmentID = GetSegmentFromTime(Context.GetTime());
+	if (!SegmentID.IsValid())
 	{
 		return;
 	}
 
 	if (EvaluationMethod == EEvaluationMethod::Static)
 	{
-		for (const FSectionEvaluationData& EvalData : GetSegment(SegmentIndex).Impls)
+		for (const FSectionEvaluationData& EvalData : GetSegment(SegmentID).Impls)
 		{
 			GetChildTemplate(EvalData.ImplIndex).Interrogate(Context, Container, BindingOverride);
 		}
@@ -280,7 +558,7 @@ void FMovieSceneEvaluationTrack::Interrogate(const FMovieSceneContext& Context,
 		// Accumulate the relevant ranges that each section intersects with the evaluated range
 		TMap<int32, TRange<float>> ImplToAccumulatedRange;
 
-		GatherSweptSegments(Context.GetRange(), SegmentIndex, Segments, ImplToAccumulatedRange);
+		GatherSweptSegments(Context.GetRange(), Segments.GetSortedIndex(SegmentID), Segments.GetSorted(), ImplToAccumulatedRange);
 		for (auto& Pair : ImplToAccumulatedRange)
 		{
 			GetChildTemplate(Pair.Key).Interrogate(Context, Pair.Value, Container, BindingOverride);
diff --git a/Engine/Source/Runtime/MovieScene/Private/Evaluation/MovieSceneEvaluationTree.cpp b/Engine/Source/Runtime/MovieScene/Private/Evaluation/MovieSceneEvaluationTree.cpp
new file mode 100644
--- /dev/null
+++ b/Engine/Source/Runtime/MovieScene/Private/Evaluation/MovieSceneEvaluationTree.cpp
@@ -0,0 +1,287 @@
+// Copyright 1998-2017 Epic Games, Inc. All Rights Reserved.
+
+#include "MovieSceneEvaluationTree.h"
+#include "CoreMinimal.h"
+#include "Misc/AutomationTest.h"
+#include "Compilation/MovieSceneSegmentCompiler.h"
+#include "Compilation/MovieSceneCompilerRules.h"
+#include "Tests/MovieSceneTestsCommon.h"
+#include "MovieSceneEvaluationField.h"
+#include "MovieSceneCommonHelpers.h"
+#include "Package.h"
+#include "Tuple.h"
+
+
+FMovieSceneEvaluationTreeRangeIterator::FMovieSceneEvaluationTreeRangeIterator(const FMovieSceneEvaluationTree& InTree)
+	: CurrentRange(TRange<float>::All()), CurrentNodeHandle(FMovieSceneEvaluationTreeNodeHandle::Root()), Tree(InTree)
+{
+	// Compute the starting range by inspecting the front of the tree
+	const FMovieSceneEvaluationTreeNode* CurrentNode = &Tree.GetRootNode();
+	TArrayView<const FMovieSceneEvaluationTreeNode> Children;
+
+	// Find the first child-most node that doesn't match the opening range bound
+	for(;;)
+	{
+		Children = Tree.GetChildren(*CurrentNode);
+		if (Children.Num() && Children[0].Range.GetLowerBound().IsOpen())
+		{
+			// Recurse into the child
+			CurrentNodeHandle = FMovieSceneEvaluationTreeNodeHandle(CurrentNode->ChildrenID, 0);
+			CurrentNode = &Tree.GetNode(CurrentNodeHandle);
+		}
+		else
+		{
+			break;
+		}
+	}
+
+	// The upper bound is either the current node's upper bound (if it has no children), or the inverse of its first child's lower bound
+	TRangeBound<float> TrailingBound = Children.Num() ? TRangeBound<float>::FlipInclusion(Children[0].Range.GetLowerBound()) : CurrentNode->Range.GetUpperBound();
+	CurrentRange = TRange<float>(TRangeBound<float>::Open(), TrailingBound);
+}
+
+FMovieSceneEvaluationTreeRangeIterator::FMovieSceneEvaluationTreeRangeIterator(const FMovieSceneEvaluationTree& InTree, TRangeBound<float> StartingBound)
+	: CurrentNodeHandle(FMovieSceneEvaluationTreeNodeHandle::Root()), Tree(InTree)
+{
+	auto GetLowerBound = [](const FMovieSceneEvaluationTreeNode& In) {
+		return In.Range.GetLowerBound();
+	};
+
+	TRange<float> CompareRange(StartingBound, TRangeBound<float>::Open());
+
+	for (;;)
+	{
+		const FMovieSceneEvaluationTreeNode& ThisNode = Tree.GetNode(CurrentNodeHandle);
+
+		// Keep binary searching child nodes until we hit a leaf
+		CurrentRange = ThisNode.Range;
+		if (!ThisNode.ChildrenID.IsValid())
+		{
+			break;
+		}
+
+		TArrayView<const FMovieSceneEvaluationTreeNode> Children = Tree.GetChildren(ThisNode);
+
+		// Binary search children's lower bounds for the first that's >= the starting bound. That results in either ChildIndex or ChildIndex-1 being the overlapping range (if any).
+		int32 ChildIndex = Algo::LowerBoundBy(Children, StartingBound, GetLowerBound, MovieSceneHelpers::SortLowerBounds);
+		if (Children.IsValidIndex(ChildIndex) && Children[ChildIndex].Range.GetLowerBound() == StartingBound)
+		{
+			CurrentNodeHandle = FMovieSceneEvaluationTreeNodeHandle(ThisNode.ChildrenID, ChildIndex);
+		}
+		else if (Children.IsValidIndex(ChildIndex-1) && Children[ChildIndex-1].Range.Overlaps(CompareRange))
+		{
+			CurrentNodeHandle = FMovieSceneEvaluationTreeNodeHandle(ThisNode.ChildrenID, ChildIndex - 1);
+		}
+		else
+		{
+			// We're at some empty space between children
+			CurrentRange = TRange<float>(
+				Children.IsValidIndex(ChildIndex-1) ? TRangeBound<float>::FlipInclusion(Children[ChildIndex-1].Range.GetUpperBound()) : CurrentRange.GetLowerBound(),
+				Children.IsValidIndex(ChildIndex) ? TRangeBound<float>::FlipInclusion(Children[ChildIndex].Range.GetLowerBound()) : CurrentRange.GetUpperBound()
+			);
+
+			break;
+		}
+	}
+
+	check(CurrentRange.Overlaps(CompareRange));
+}
+
+void FMovieSceneEvaluationTreeRangeIterator::Iter(bool bForwards)
+{
+	if (!CurrentNodeHandle.IsValid())
+	{
+		return;
+	}
+
+	const FMovieSceneEvaluationTreeNode& CurrentNode = Tree.GetNode(CurrentNodeHandle);
+	if (GetTrailingBound(bForwards, CurrentNode.Range) == GetTrailingBound(bForwards, CurrentRange))
+	{
+		// We're done with this node, iterate forwards from the same position in the parent
+		CurrentNodeHandle = CurrentNode.Parent;
+		Iter(bForwards);
+		return;
+	}
+
+	TRangeBound<float> NewLeadingBound = CurrentRange.IsEmpty() ? TRangeBound<float>::Open() : TRangeBound<float>::FlipInclusion(GetTrailingBound(bForwards, CurrentRange));
+
+	// Iterate into children when possible (where the leading bound matches the current new leading bound)
+	FMovieSceneEvaluationTreeNodeHandle NextChildHandle = FindNextChild(CurrentNodeHandle, NewLeadingBound, bForwards);
+	while (NextChildHandle.IsValid() && CompareBound(bForwards, Tree.GetNode(NextChildHandle).Range, NewLeadingBound))
+	{
+		CurrentNodeHandle = NextChildHandle;
+		NextChildHandle = FindNextChild(CurrentNodeHandle, NewLeadingBound, bForwards);
+	}
+
+	// The new traililng bound is either the trailing bound of the whole node, or the inverse of the leading bound of the next child (which will only be valid if it's not the same as the leading bound).
+	// The latter means this is space in-between child nodes
+	TRangeBound<float> NewTrailingBound = GetTrailingBound(bForwards, Tree.GetNode(CurrentNodeHandle).Range);
+	if (NextChildHandle.IsValid())
+	{
+		const FMovieSceneEvaluationTreeNode& NextChildNode = Tree.GetNode(NextChildHandle);
+		if (!CompareBound(bForwards, NextChildNode.Range, NewLeadingBound))
+		{
+			NewTrailingBound = TRangeBound<float>::FlipInclusion(GetLeadingBound(bForwards, NextChildNode.Range));
+		}
+	}
+
+	CurrentRange = TRange<float>(bForwards ? NewLeadingBound : NewTrailingBound, bForwards ? NewTrailingBound : NewLeadingBound);
+}
+
+FMovieSceneEvaluationTreeNodeHandle FMovieSceneEvaluationTreeRangeIterator::FindNextChild(FMovieSceneEvaluationTreeNodeHandle ParentNodeHandle, TRangeBound<float> PredicateBound, bool bForwards)
+{
+	auto GetLowerBound = [](const FMovieSceneEvaluationTreeNode& In) {
+		return In.Range.GetLowerBound();
+	};
+	auto GetUpperBound = [](const FMovieSceneEvaluationTreeNode& In) {
+		return In.Range.GetUpperBound();
+	};
+
+	const FMovieSceneEvaluationTreeNode& ParentNode = Tree.GetNode(ParentNodeHandle);
+	TArrayView<const FMovieSceneEvaluationTreeNode> Children = Tree.GetChildren(ParentNode);
+	if (Children.Num())
+	{
+		// Find the index of the next >= bound
+		int32 SearchIndex = bForwards
+			? Algo::LowerBoundBy(Children, PredicateBound, GetLowerBound, MovieSceneHelpers::SortLowerBounds)
+			: Algo::UpperBoundBy(Children, PredicateBound, GetUpperBound, MovieSceneHelpers::SortUpperBounds)-1;
+
+		if (Children.IsValidIndex(SearchIndex))
+		{
+			return FMovieSceneEvaluationTreeNodeHandle(ParentNode.ChildrenID, SearchIndex);
+		}
+	}
+
+	return FMovieSceneEvaluationTreeNodeHandle::Invalid();
+}
+
+
+TArrayView<const FMovieSceneEvaluationTreeNode> FMovieSceneEvaluationTree::GetChildren(const FMovieSceneEvaluationTreeNode& InNode) const
+{
+	return InNode.ChildrenID.IsValid() ? ChildNodes.Get(InNode.ChildrenID) : TArrayView<const FMovieSceneEvaluationTreeNode>();
+}
+
+TArrayView<FMovieSceneEvaluationTreeNode> FMovieSceneEvaluationTree::GetChildren(const FMovieSceneEvaluationTreeNode& InNode)
+{
+	return InNode.ChildrenID.IsValid() ? ChildNodes.Get(InNode.ChildrenID) : TArrayView<FMovieSceneEvaluationTreeNode>();
+}
+
+FMovieSceneEvaluationTreeRangeIterator FMovieSceneEvaluationTree::IterateFromTime(float InTime) const
+{
+	return IterateFromLowerBound(TRangeBound<float>::Inclusive(InTime));
+}
+
+FMovieSceneEvaluationTreeRangeIterator FMovieSceneEvaluationTree::IterateFromLowerBound(TRangeBound<float> InStartingLowerBound) const
+{
+	return FMovieSceneEvaluationTreeRangeIterator(*this, InStartingLowerBound);
+}
+
+void FMovieSceneEvaluationTree::InsertNewChild(TRange<float> InEffectiveRange, const IMovieSceneEvaluationTreeNodeOperator& InOperator, int32 InsertIndex, FMovieSceneEvaluationTreeNodeHandle InParent)
+{
+	// Check if this parent node has any 
+	FEvaluationTreeEntryHandle ChildID = GetNode(InParent).ChildrenID;
+	if (!ChildID.IsValid())
+	{
+		// AllocateEntry may reallocate the whole node array
+		ChildID = ChildNodes.AllocateEntry(1);
+		GetNode(InParent).ChildrenID = ChildID;
+	}
+	else
+	{
+		// Fixup parent handles for any grandchild nodes that reside as subsequent children
+		TArrayView<FMovieSceneEvaluationTreeNode> ExistingChildren = ChildNodes.Get(ChildID);
+		for (int32 Index = InsertIndex; Index < ExistingChildren.Num(); ++Index)
+		{
+			FMovieSceneEvaluationTreeNodeHandle FixedUpParentNode(ChildID, Index+1);
+			for (FMovieSceneEvaluationTreeNode& GrandChild : GetChildren(ExistingChildren[Index]))
+			{
+				GrandChild.Parent = FixedUpParentNode;
+			}
+		}
+	}
+
+	FMovieSceneEvaluationTreeNode NewChildNode(InEffectiveRange, InParent);
+	InOperator(NewChildNode);
+	ChildNodes.Insert(ChildID, InsertIndex, MoveTemp(NewChildNode));
+}
+
+void FMovieSceneEvaluationTree::AddTimeRange(TRange<float> InTimeRange)
+{
+	struct FNullOp : IMovieSceneEvaluationTreeNodeOperator
+	{
+		virtual void operator()(FMovieSceneEvaluationTreeNode& Node) const {}
+	};
+
+	AddTimeRange(InTimeRange, FNullOp(), FMovieSceneEvaluationTreeNodeHandle::Root());
+}
+
+void FMovieSceneEvaluationTree::AddTimeRange(TRange<float> InTimeRange, const IMovieSceneEvaluationTreeNodeOperator& InOperator, FMovieSceneEvaluationTreeNodeHandle InParent)
+{
+	// Take a temporary copy of the node as the container may be reallocated in this function
+	FMovieSceneEvaluationTreeNode ThisNode = GetNode(InParent);
+	if (!ensure(ThisNode.Range.Overlaps(InTimeRange)))
+	{
+		return;
+	}
+
+	if (InTimeRange.Contains(ThisNode.Range))
+	{
+		// If the section range entirely contains this node's range or it's exactly the same, just add the data to this entry
+		// We must be sure to pass a reference to the actual here, not the temporary copy
+		InOperator(GetNode(InParent));
+	}
+	else if (!ThisNode.ChildrenID.IsValid())
+	{
+		// We have no child nodes yet, so just add the intersection as a child at the start
+		TRange<float> Intersection = TRange<float>::Intersection(ThisNode.Range, InTimeRange);
+		InsertNewChild(Intersection, InOperator, 0, InParent);
+	}
+	else
+	{
+		// It must be a sub-division of this range, find the intersection and add new children as necessary
+		TRangeBound<float> LastBound = ThisNode.Range.GetLowerBound();
+
+		// Iterate all existing children, adding new children for gaps that the new section intersects
+		int32 ChildIndex = 0;
+		while (ChildIndex < ChildNodes.Get(ThisNode.ChildrenID).Num())
+		{
+			FMovieSceneEvaluationTreeNode ChildNode = GetNode(ThisNode.ChildrenID, ChildIndex);
+
+			if (!ChildNode.Range.GetLowerBound().IsOpen())
+			{
+				TRange<float> PrecedingSpace(LastBound, TRangeBound<float>::FlipInclusion(ChildNode.Range.GetLowerBound()));
+				TRange<float> PrecedingIntersection = TRange<float>::Intersection(PrecedingSpace, InTimeRange);
+
+				if (!PrecedingIntersection.IsEmpty())
+				{
+					// Insert this child at the current child index
+					InsertNewChild(PrecedingIntersection, InOperator, ChildIndex, InParent);
+
+					// Skip over the child we just added
+					ChildIndex++;
+				}
+			}
+
+			if (ChildNode.Range.Overlaps(InTimeRange))
+			{
+				// Find the node again since InsertNewChild may have re-allocated the nodes
+				AddTimeRange(InTimeRange, InOperator, FMovieSceneEvaluationTreeNodeHandle(ThisNode.ChildrenID, ChildIndex));
+			}
+			LastBound = TRangeBound<float>::FlipInclusion(ChildNode.Range.GetUpperBound());
+			ChildIndex++;
+		}
+
+		// If the last child didn't fully fill this parent node range, add a new one for this new section
+		if (!LastBound.IsOpen())
+		{
+			TRange<float> ProceedingSpace(LastBound, ThisNode.Range.GetUpperBound());
+			TRange<float> ProceedingIntersection = TRange<float>::Intersection(ProceedingSpace, InTimeRange);
+
+			if (!ProceedingIntersection.IsEmpty())
+			{
+				// Insert this child at the end
+				InsertNewChild(ProceedingIntersection, InOperator, ChildNodes.Get(ThisNode.ChildrenID).Num(), InParent);
+			}
+		}
+	}
+}
\ No newline at end of file
diff --git a/Engine/Source/Runtime/MovieScene/Private/Evaluation/MovieSceneExecutionTokens.cpp b/Engine/Source/Runtime/MovieScene/Private/Evaluation/MovieSceneExecutionTokens.cpp
--- a/Engine/Source/Runtime/MovieScene/Private/Evaluation/MovieSceneExecutionTokens.cpp
+++ b/Engine/Source/Runtime/MovieScene/Private/Evaluation/MovieSceneExecutionTokens.cpp
@@ -17,7 +17,7 @@ void FMovieSceneExecutionTokens::Apply(const FMovieSceneContext& RootContext, IM
 {
 	MOVIESCENE_DETAILED_SCOPE_CYCLE_COUNTER(MovieSceneEval_ApplyExecutionTokens);
 
-	FPersistentEvaluationData PersistentDataProxy(Player.State.PersistentEntityData, Player.State.PersistentSharedData);
+	FPersistentEvaluationData PersistentDataProxy(Player);
 
 	TArray<TInlineValue<IMovieSceneSharedExecutionToken, 32>> SortedSharedTokens;
 	SortedSharedTokens.Reserve(SharedTokens.Num());
diff --git a/Engine/Source/Runtime/MovieScene/Private/Evaluation/MovieScenePreAnimatedState.cpp b/Engine/Source/Runtime/MovieScene/Private/Evaluation/MovieScenePreAnimatedState.cpp
--- a/Engine/Source/Runtime/MovieScene/Private/Evaluation/MovieScenePreAnimatedState.cpp
+++ b/Engine/Source/Runtime/MovieScene/Private/Evaluation/MovieScenePreAnimatedState.cpp
@@ -195,8 +195,8 @@ void TMovieSceneSavedTokens<TokenType>::Restore(IMovieScenePlayer& Player, TFunc
 		{
 			MovieSceneImpl::RestorePreAnimatedToken(PreAnimatedTokens[TokenIndex], Player, ResolvedPayload);
 
-			AllAnimatedTypeIDs.RemoveAtSwap(TokenIndex, 1, false);
-			PreAnimatedTokens.RemoveAtSwap(TokenIndex, 1, false);
+			AllAnimatedTypeIDs.RemoveAt(TokenIndex, 1, false);
+			PreAnimatedTokens.RemoveAt(TokenIndex, 1, false);
 
 			AnimatedEntities.RemoveAll(
 				[=](const FMovieSceneEntityAndAnimTypeID& InEntityAndAnimType)
@@ -225,7 +225,7 @@ bool TMovieSceneSavedTokens<TokenType>::RestoreEntity(IMovieScenePlayer& Player,
 				AnimTypesToRestore.Add(EntityAndAnimType.AnimTypeID);
 
 				// This entity is no longer animating this anim type ID
-				AnimatedEntities.RemoveAtSwap(LUTIndex);
+				AnimatedEntities.RemoveAt(LUTIndex);
 			}
 			else
 			{
@@ -251,8 +251,8 @@ bool TMovieSceneSavedTokens<TokenType>::RestoreEntity(IMovieScenePlayer& Player,
 			}
 			else
 			{
-				AllAnimatedTypeIDs.RemoveAtSwap(TokenIndex, 1, false);
-				PreAnimatedTokens.RemoveAtSwap(TokenIndex, 1, false);
+				AllAnimatedTypeIDs.RemoveAt(TokenIndex, 1, false);
+				PreAnimatedTokens.RemoveAt(TokenIndex, 1, false);
 			}
 		}
 	}
diff --git a/Engine/Source/Runtime/MovieScene/Private/Evaluation/MovieSceneRootOverridePath.cpp b/Engine/Source/Runtime/MovieScene/Private/Evaluation/MovieSceneRootOverridePath.cpp
new file mode 100644
--- /dev/null
+++ b/Engine/Source/Runtime/MovieScene/Private/Evaluation/MovieSceneRootOverridePath.cpp
@@ -0,0 +1,29 @@
+// Copyright 1998-2017 Epic Games, Inc. All Rights Reserved.
+
+#include "MovieSceneRootOverridePath.h"
+#include "MovieSceneSequenceHierarchy.h"
+
+void FMovieSceneRootOverridePath::Reset()
+{
+	ReverseOverrideRootPath.Reset();
+}
+
+void FMovieSceneRootOverridePath::Set(FMovieSceneSequenceID OverrideRootID, const FMovieSceneSequenceHierarchy& RootHierarchy)
+{
+	ReverseOverrideRootPath.Reset();
+
+	FMovieSceneSequenceID CurrentSequenceID = OverrideRootID;
+
+	while (CurrentSequenceID != MovieSceneSequenceID::Root)
+	{
+		const FMovieSceneSequenceHierarchyNode* CurrentNode = RootHierarchy.FindNode(CurrentSequenceID);
+		const FMovieSceneSubSequenceData* OuterSubData = RootHierarchy.FindSubData(CurrentSequenceID);
+		if (!ensureAlwaysMsgf(CurrentNode && OuterSubData, TEXT("Malformed sequence hierarchy")))
+		{
+			return;
+		}
+
+		ReverseOverrideRootPath.Add(OuterSubData->DeterministicSequenceID);
+		CurrentSequenceID = CurrentNode->ParentID;
+	}
+}
diff --git a/Engine/Source/Runtime/MovieScene/Private/Evaluation/MovieSceneSequenceHierarchy.cpp b/Engine/Source/Runtime/MovieScene/Private/Evaluation/MovieSceneSequenceHierarchy.cpp
new file mode 100644
--- /dev/null
+++ b/Engine/Source/Runtime/MovieScene/Private/Evaluation/MovieSceneSequenceHierarchy.cpp
@@ -0,0 +1,128 @@
+// Copyright 1998-2017 Epic Games, Inc. All Rights Reserved.
+
+#include "MovieSceneSequenceHierarchy.h"
+#include "Sections/MovieSceneSubSection.h"
+#include "MovieSceneSequence.h"
+#include "MovieScene.h"
+
+FMovieSceneSubSequenceData::FMovieSceneSubSequenceData()
+	: Sequence(nullptr)
+{}
+
+FMovieSceneSubSequenceData::FMovieSceneSubSequenceData(const UMovieSceneSubSection& InSubSection)
+	: Sequence(InSubSection.GetSequence())
+	, DeterministicSequenceID(InSubSection.GetSequenceID())
+	, PreRollRange(TRange<float>::Empty())
+	, PostRollRange(TRange<float>::Empty())
+	, HierarchicalBias(InSubSection.Parameters.HierarchicalBias)
+#if WITH_EDITORONLY_DATA
+	, SectionPath(*InSubSection.GetPathNameInMovieScene())
+#endif
+{
+	UMovieSceneSequence* SequencePtr = GetSequence();
+
+	checkf(SequencePtr, TEXT("Attempting to construct sub sequence data with a null sequence."));
+
+	const float StartTime = InSubSection.GetStartTime();
+	const float EndTime = InSubSection.GetEndTime();
+
+	RootToSequenceTransform = 
+		FMovieSceneSequenceTransform(SequencePtr->GetMovieScene()->GetPlaybackRange().GetLowerBoundValue() + InSubSection.Parameters.StartOffset) *		// Inner play offset
+		FMovieSceneSequenceTransform(0.f, InSubSection.Parameters.TimeScale) *		// Inner play rate
+		FMovieSceneSequenceTransform(-StartTime);					// Outer section start time
+
+	PlayRange = InSubSection.GetRange() * RootToSequenceTransform;
+
+	// Make sure pre/postroll ranges are in the inner sequence's time space
+	if (InSubSection.GetPreRollTime() > 0)
+	{
+		PreRollRange = TRange<float>(StartTime - InSubSection.GetPreRollTime(), TRangeBound<float>::Exclusive(StartTime)) * RootToSequenceTransform;
+	}
+	if (InSubSection.GetPostRollTime() > 0)
+	{
+		PostRollRange = TRange<float>(TRangeBound<float>::Exclusive(EndTime), EndTime + InSubSection.GetPostRollTime()) * RootToSequenceTransform;
+	}
+}
+
+/**
+ * Get this sub sequence's sequence asset
+ */
+UMovieSceneSequence* FMovieSceneSubSequenceData::GetSequence() const
+{
+	UMovieSceneSequence* ResolvedSequence = CachedSequence.Get();
+	if (!ResolvedSequence)
+	{
+		ResolvedSequence = Cast<UMovieSceneSequence>(Sequence.ResolveObject());
+		CachedSequence = ResolvedSequence;
+	}
+	return ResolvedSequence;
+}
+
+void FMovieSceneSequenceHierarchy::Add(const FMovieSceneSubSequenceData& Data, FMovieSceneSequenceIDRef ThisSequenceID, FMovieSceneSequenceIDRef ParentID)
+{
+	check(ParentID != MovieSceneSequenceID::Invalid);
+
+	// Add (or update) the sub sequence data
+	SubSequences.Add(ThisSequenceID, Data);
+
+	// Set up the hierarchical information if we don't have any, or its wrong
+	FMovieSceneSequenceHierarchyNode* ExistingHierarchyNode = Hierarchy.Find(ThisSequenceID);
+	if (!ExistingHierarchyNode || ExistingHierarchyNode->ParentID != ParentID)
+	{
+		if (!ExistingHierarchyNode)
+		{
+			// The node doesn't yet exist - create it
+			FMovieSceneSequenceHierarchyNode Node(ParentID);
+			Hierarchy.Add(ThisSequenceID, Node);
+		}
+		else
+		{
+			// The node exists already but under the wrong parent - we need to move it
+			FMovieSceneSequenceHierarchyNode* Parent = Hierarchy.Find(ExistingHierarchyNode->ParentID);
+			check(Parent);
+			// Remove it from its parent's children
+			Parent->Children.Remove(ThisSequenceID);
+
+			// Set the parent ID
+			ExistingHierarchyNode->ParentID = ParentID;
+		}
+
+		// Add the node to its parent's children array
+		FMovieSceneSequenceHierarchyNode* Parent = Hierarchy.Find(ParentID);
+		check(Parent);
+		ensure(!Parent->Children.Contains(ThisSequenceID));
+		Parent->Children.Add(ThisSequenceID);
+	}
+}
+
+void FMovieSceneSequenceHierarchy::Remove(TArrayView<const FMovieSceneSequenceID> SequenceIDs)
+{
+	TArray<FMovieSceneSequenceID, TInlineAllocator<16>> IDsToRemove;
+	IDsToRemove.Append(SequenceIDs.GetData(), SequenceIDs.Num());
+
+	while (IDsToRemove.Num())
+	{
+		int32 NumRemaining = IDsToRemove.Num();
+		for (int32 Index = 0; Index < NumRemaining; ++Index)
+		{
+			FMovieSceneSequenceID ID = IDsToRemove[Index];
+
+			SubSequences.Remove(ID);
+
+			// Remove all children too
+			if (const FMovieSceneSequenceHierarchyNode* Node = FindNode(ID))
+			{
+				FMovieSceneSequenceHierarchyNode* Parent = Hierarchy.Find(Node->ParentID);
+				if (Parent)
+				{
+					Parent->Children.Remove(ID);
+				}
+
+				IDsToRemove.Append(Node->Children);
+				Hierarchy.Remove(ID);
+			}
+		}
+
+		IDsToRemove.RemoveAt(0, NumRemaining);
+	}
+}
\ No newline at end of file
diff --git a/Engine/Source/Runtime/MovieScene/Private/Evaluation/MovieSceneSequenceTemplateStore.cpp b/Engine/Source/Runtime/MovieScene/Private/Evaluation/MovieSceneSequenceTemplateStore.cpp
--- a/Engine/Source/Runtime/MovieScene/Private/Evaluation/MovieSceneSequenceTemplateStore.cpp
+++ b/Engine/Source/Runtime/MovieScene/Private/Evaluation/MovieSceneSequenceTemplateStore.cpp
@@ -4,45 +4,12 @@
 #include "MovieSceneSequence.h"
 #include "UObject/ObjectKey.h"
 
-FMovieSceneEvaluationTemplate& FMovieSceneSequenceTemplateStore::GetCompiledTemplate(UMovieSceneSequence& Sequence)
+FMovieSceneEvaluationTemplate& FMovieSceneSequencePrecompiledTemplateStore::AccessTemplate(UMovieSceneSequence& Sequence)
 {
-	return GetCompiledTemplate(Sequence, FObjectKey(&Sequence));
+	return Sequence.PrecompiledEvaluationTemplate;
 }
 
-FMovieSceneEvaluationTemplate& FMovieSceneSequenceTemplateStore::GetCompiledTemplate(UMovieSceneSequence& Sequence, FObjectKey InSequenceKey)
+FMovieSceneEvaluationTemplate& FMovieSceneSequenceTemporaryTemplateStore::AccessTemplate(UMovieSceneSequence& Sequence)
 {
-	UObject* ObjectKey = InSequenceKey.ResolveObjectPtr();
-#if WITH_EDITORONLY_DATA
-	if (AreTemplatesVolatile())
-	{
-		// Lookup into map here if we have a key that is not the sequence
-		if (ObjectKey == &Sequence)
-		{
-			Sequence.EvaluationTemplate.Regenerate(Sequence.TemplateParameters);
-		}
-		else
-		{
-			// assume this is an instanced subsequence
-			FCachedMovieSceneEvaluationTemplate* EvaluationTemplatePtr = Sequence.InstancedSubSequenceEvaluationTemplates.Find(ObjectKey);
-			if (EvaluationTemplatePtr == nullptr)
-			{
-				EvaluationTemplatePtr = &Sequence.InstancedSubSequenceEvaluationTemplates.Add(ObjectKey);
-				EvaluationTemplatePtr->Initialize(Sequence, this);
-			}
-
-			EvaluationTemplatePtr->Regenerate(Sequence.TemplateParameters);
-		}
-	}
-#endif
-
-	// Lookup into map here if we have a key that is not the sequence
-	if (ObjectKey == &Sequence)
-	{
-		return Sequence.EvaluationTemplate;
-	}
-	else
-	{
-		// assume this is an instanced subsequence
-		return Sequence.InstancedSubSequenceEvaluationTemplates.FindChecked(ObjectKey);
-	}
-}
+	return Templates.FindOrAdd(&Sequence);
+}
\ No newline at end of file
diff --git a/Engine/Source/Runtime/MovieScene/Private/Evaluation/MovieSceneTrackImplementation.cpp b/Engine/Source/Runtime/MovieScene/Private/Evaluation/MovieSceneTrackImplementation.cpp
--- a/Engine/Source/Runtime/MovieScene/Private/Evaluation/MovieSceneTrackImplementation.cpp
+++ b/Engine/Source/Runtime/MovieScene/Private/Evaluation/MovieSceneTrackImplementation.cpp
@@ -6,5 +6,5 @@
 
 bool FMovieSceneTrackImplementationPtr::Serialize(FArchive& Ar)
 {
-	return SerializeEvaluationTemplate(*this, Ar);
+	return SerializeInlineValue(*this, Ar);
 }
diff --git a/Engine/Source/Runtime/MovieScene/Private/Evaluation/PersistentEvaluationData.cpp b/Engine/Source/Runtime/MovieScene/Private/Evaluation/PersistentEvaluationData.cpp
new file mode 100644
--- /dev/null
+++ b/Engine/Source/Runtime/MovieScene/Private/Evaluation/PersistentEvaluationData.cpp
@@ -0,0 +1,17 @@
+// Copyright 1998-2017 Epic Games, Inc. All Rights Reserved.
+
+#include "PersistentEvaluationData.h"
+#include "IMovieScenePlayer.h"
+#include "Evaluation/MovieSceneEvaluationTemplateInstance.h"
+
+FPersistentEvaluationData::FPersistentEvaluationData(IMovieScenePlayer& InPlayer)
+	: Player(InPlayer)
+	, EntityData(Player.State.PersistentEntityData)
+	, SharedData(Player.State.PersistentSharedData)
+{}
+
+const FMovieSceneSequenceInstanceData* FPersistentEvaluationData::GetInstanceData() const
+{
+	const FMovieSceneSubSequenceData* SubData = Player.GetEvaluationTemplate().GetHierarchy().FindSubData(TrackKey.SequenceID);
+	return SubData && SubData->InstanceData.IsValid() ? &SubData->InstanceData.GetValue() : nullptr;
+}
\ No newline at end of file
diff --git a/Engine/Source/Runtime/MovieScene/Private/MovieScene.cpp b/Engine/Source/Runtime/MovieScene/Private/MovieScene.cpp
--- a/Engine/Source/Runtime/MovieScene/Private/MovieScene.cpp
+++ b/Engine/Source/Runtime/MovieScene/Private/MovieScene.cpp
@@ -102,6 +102,17 @@ FGuid UMovieScene::AddSpawnable( const FString& Name, UObject& ObjectTemplate )
 	return NewSpawnable.GetGuid();
 }
 
+void UMovieScene::AddSpawnable(const FMovieSceneSpawnable& InNewSpawnable, const FMovieSceneBinding& InNewBinding)
+{
+	Modify();
+
+	FMovieSceneSpawnable NewSpawnable;
+	NewSpawnable = InNewSpawnable;
+	Spawnables.Add(NewSpawnable);
+
+	FMovieSceneBinding NewBinding = InNewBinding;
+	ObjectBindings.Add(NewBinding);
+}
 
 bool UMovieScene::RemoveSpawnable( const FGuid& Guid )
 {
@@ -167,6 +178,19 @@ FGuid UMovieScene::AddPossessable( const FString& Name, UClass* Class )
 }
 
 
+void UMovieScene::AddPossessable(const FMovieScenePossessable& InNewPossessable, const FMovieSceneBinding& InNewBinding)
+{
+	Modify();
+
+	FMovieScenePossessable NewPossessable;
+	NewPossessable = InNewPossessable;
+	Possessables.Add(NewPossessable);
+
+	FMovieSceneBinding NewBinding = InNewBinding;
+	ObjectBindings.Add(NewBinding);
+}
+
+
 bool UMovieScene::RemovePossessable( const FGuid& PossessableGuid )
 {
 	bool bAnythingRemoved = false;
@@ -723,6 +747,21 @@ void UMovieScene::PostLoad()
 		}
 	}
 
+	for ( int32 ObjectBindingIndex = 0; ObjectBindingIndex < ObjectBindings.Num(); ++ObjectBindingIndex)
+	{
+		for( int32 TrackIndex = 0; TrackIndex < ObjectBindings[ObjectBindingIndex].GetTracks().Num(); )
+		{
+			if (ObjectBindings[ObjectBindingIndex].GetTracks()[TrackIndex] == nullptr)
+			{
+				ObjectBindings[ObjectBindingIndex].RemoveTrack(*ObjectBindings[ObjectBindingIndex].GetTracks()[TrackIndex]);
+			}
+			else
+			{
+				++TrackIndex;
+			}
+		}
+	}
+
 	UpgradeTimeRanges();
 
 	for (FMovieSceneSpawnable& Spawnable : Spawnables)
diff --git a/Engine/Source/Runtime/MovieScene/Private/MovieSceneModule.cpp b/Engine/Source/Runtime/MovieScene/Private/MovieSceneModule.cpp
--- a/Engine/Source/Runtime/MovieScene/Private/MovieSceneModule.cpp
+++ b/Engine/Source/Runtime/MovieScene/Private/MovieSceneModule.cpp
@@ -15,13 +15,28 @@ DEFINE_LOG_CATEGORY(LogMovieScene);
  */
 class FMovieSceneModule
 	: public IMovieSceneModule
+	, public TSharedFromThis<FMovieSceneModule>
 {
 public:
 
 	// IModuleInterface interface
+	~FMovieSceneModule()
+	{
+		ensure(ModuleHandle.IsUnique());
+	}
 
-	virtual void StartupModule() override { }
-	virtual void ShutdownModule() override { }
+	virtual void StartupModule() override
+	{
+		struct FNoopDefaultDeleter
+		{
+			void operator()(FMovieSceneModule* Object) const {}
+		};
+		ModuleHandle = MakeShareable(this, FNoopDefaultDeleter());
+	}
+
+	virtual void ShutdownModule() override
+	{
+	}
 
 	virtual void RegisterEvaluationGroupParameters(FName GroupName, const FMovieSceneEvaluationGroupParameters& GroupParameters) override
 	{
@@ -41,7 +56,13 @@ public:
 		return EvaluationGroupParameters.FindRef(GroupName);
 	}
 
+	virtual TWeakPtr<IMovieSceneModule> GetWeakPtr() override
+	{
+		return ModuleHandle;
+	}
+
 private:
+	TSharedPtr<FMovieSceneModule> ModuleHandle;
 	TMap<FName, FMovieSceneEvaluationGroupParameters> EvaluationGroupParameters;
 };
 
diff --git a/Engine/Source/Runtime/MovieScene/Private/MovieScenePlayer.cpp b/Engine/Source/Runtime/MovieScene/Private/MovieScenePlayer.cpp
--- a/Engine/Source/Runtime/MovieScene/Private/MovieScenePlayer.cpp
+++ b/Engine/Source/Runtime/MovieScene/Private/MovieScenePlayer.cpp
@@ -1,23 +1,10 @@
 // Copyright 1998-2017 Epic Games, Inc. All Rights Reserved.
 
 #include "CoreMinimal.h"
-#include "MovieSceneSequenceInstance.h"
 #include "IMovieScenePlayer.h"
 #include "Evaluation/MovieSceneEvaluationTemplateInstance.h"
 #include "MovieSceneSequence.h"
 
-TSharedRef<FMovieSceneSequenceInstance> IMovieScenePlayer::GetRootMovieSceneSequenceInstance()
-{
-	return GetEvaluationTemplate().GetInstance(MovieSceneSequenceID::Root)->LegacySequenceInstance.ToSharedRef();
-}
-
-void IMovieScenePlayer::GetRuntimeObjects(TSharedRef<FMovieSceneSequenceInstance> MovieSceneInstance, const FGuid& Guid, TArray<TWeakObjectPtr<>>& OutRuntimeObjects)
-{
-	for (TWeakObjectPtr<> Obj : FindBoundObjects(Guid, MovieSceneInstance->GetSequenceID()))
-	{
-		OutRuntimeObjects.Add(Obj);
-	}
-}
 
 void IMovieScenePlayer::ResolveBoundObjects(const FGuid& InBindingId, FMovieSceneSequenceID SequenceID, UMovieSceneSequence& Sequence, UObject* ResolutionContext, TArray<UObject*, TInlineAllocator<1>>& OutObjects) const
 {
diff --git a/Engine/Source/Runtime/MovieScene/Private/MovieSceneSequence.cpp b/Engine/Source/Runtime/MovieScene/Private/MovieSceneSequence.cpp
--- a/Engine/Source/Runtime/MovieScene/Private/MovieSceneSequence.cpp
+++ b/Engine/Source/Runtime/MovieScene/Private/MovieSceneSequence.cpp
@@ -1,20 +1,18 @@
 // Copyright 1998-2017 Epic Games, Inc. All Rights Reserved.
 
 #include "MovieSceneSequence.h"
-#include "Compilation/MovieSceneEvaluationTemplateGenerator.h"
 #include "Evaluation/MovieSceneEvaluationCustomVersion.h"
+#include "Evaluation/MovieSceneSequenceTemplateStore.h"
 #include "MovieScene.h"
 #include "EditorObjectVersion.h"
+#include "Tracks/MovieSceneSubTrack.h"
+#include "Sections/MovieSceneSubSection.h"
+#include "Compilation/MovieSceneCompiler.h"
 
 UMovieSceneSequence::UMovieSceneSequence(const FObjectInitializer& Init)
 	: Super(Init)
 {
 	bParentContextsAreSignificant = false;
-
-	TemplateParameters.bForEditorPreview = false;
-#if WITH_EDITORONLY_DATA
-	EvaluationTemplate.Initialize(*this);
-#endif
 }
 
 #if WITH_EDITORONLY_DATA
@@ -22,7 +20,8 @@ void UMovieSceneSequence::PostDuplicate(bool bDuplicateForPIE)
 {
 	if (bDuplicateForPIE)
 	{
-		EvaluationTemplate.Regenerate(TemplateParameters);
+		FMovieSceneSequencePrecompiledTemplateStore Store;
+		FMovieSceneCompiler::Compile(*this, Store);
 	}
 
 	Super::PostDuplicate(bDuplicateForPIE);
@@ -33,7 +32,7 @@ void UMovieSceneSequence::PostLoad()
 {
 #if WITH_EDITORONLY_DATA
 	// Wipe compiled data on editor load to ensure we don't try and iteratively compile previously saved content. In a cooked game, this will contain our up-to-date compiled template.
-	EvaluationTemplate = FCachedMovieSceneEvaluationTemplate(*this);
+	PrecompiledEvaluationTemplate = FMovieSceneEvaluationTemplate();
 #endif
 
 	Super::PostLoad();
@@ -47,23 +46,19 @@ void UMovieSceneSequence::Serialize(FArchive& Ar)
 #if WITH_EDITORONLY_DATA
 	if (Ar.IsCooking() && !HasAnyFlags(RF_ClassDefaultObject))
 	{
-		EvaluationTemplate.Regenerate(TemplateParameters);
+		FMovieSceneSequencePrecompiledTemplateStore Store;
+		FMovieSceneCompiler::Compile(*this, Store);
 	}
 	else if (Ar.IsSaving())
 	{
 		// Don't save template data unless we're cooking
-		EvaluationTemplate = FCachedMovieSceneEvaluationTemplate(*this);
+		PrecompiledEvaluationTemplate = FMovieSceneEvaluationTemplate();
 	}
 #endif
 	
 	Super::Serialize(Ar);
 }
 
-void UMovieSceneSequence::GenerateEvaluationTemplate(FMovieSceneEvaluationTemplate& Template, const FMovieSceneTrackCompilationParams& Params, FMovieSceneSequenceTemplateStore& Store)
-{
-	FMovieSceneEvaluationTemplateGenerator(*this, Template, Store).Generate(Params);
-}
-
 FGuid UMovieSceneSequence::FindPossessableObjectId(UObject& Object, UObject* Context) const
 {
 	UMovieScene* MovieScene = GetMovieScene();
diff --git a/Engine/Source/Runtime/MovieScene/Private/MovieSceneSequencePlayer.cpp b/Engine/Source/Runtime/MovieScene/Private/MovieSceneSequencePlayer.cpp
--- a/Engine/Source/Runtime/MovieScene/Private/MovieSceneSequencePlayer.cpp
+++ b/Engine/Source/Runtime/MovieScene/Private/MovieSceneSequencePlayer.cpp
@@ -27,6 +27,7 @@ UMovieSceneSequencePlayer::UMovieSceneSequencePlayer(const FObjectInitializer& I
 	, StartTime(0.f)
 	, EndTime(0.f)
 	, CurrentNumLoops(0)
+	, LatentPlaybackPosition(0.f)
 {
 }
 
@@ -166,20 +167,17 @@ void UMovieSceneSequencePlayer::Pause()
 
 		Status = EMovieScenePlayerStatus::Paused;
 
-		// Evaluate the sequence at its current time, with a status of 'stopped' to ensure that animated state pauses correctly
+		// Evaluate the sequence at its current time, with a status of 'stopped' to ensure that animated state pauses correctly. (ie. audio sounds should stop/pause)
 		{
 			bIsEvaluating = true;
 
 			UMovieSceneSequence* MovieSceneSequence = RootTemplateInstance.GetSequence(MovieSceneSequenceID::Root);
 			TOptional<float> FixedFrameInterval = MovieSceneSequence->GetMovieScene() ? MovieSceneSequence->GetMovieScene()->GetOptionalFixedFrameInterval() : TOptional<float>();
 
-			if (!PlayPosition.GetPreviousPosition().IsSet() || PlayPosition.GetPreviousPosition().GetValue() != GetSequencePosition())
-			{
-				FMovieSceneEvaluationRange Range = PlayPosition.JumpTo(GetSequencePosition(), FixedFrameInterval);
+			FMovieSceneEvaluationRange Range = PlayPosition.JumpTo(GetSequencePosition(), FixedFrameInterval);
 
-				const FMovieSceneContext Context(Range, EMovieScenePlayerStatus::Stopped);
-				RootTemplateInstance.Evaluate(Context, *this);
-			}
+			const FMovieSceneContext Context(Range, EMovieScenePlayerStatus::Stopped);
+			RootTemplateInstance.Evaluate(Context, *this);
 
 			bIsEvaluating = false;
 		}
@@ -258,6 +256,13 @@ float UMovieSceneSequencePlayer::GetPlaybackPosition() const
 
 void UMovieSceneSequencePlayer::SetPlaybackPosition(float NewPlaybackPosition)
 {
+	if (bIsEvaluating)
+	{
+		LatentActions.Add(ELatentAction::SetPlaybackPosition);
+		LatentPlaybackPosition = NewPlaybackPosition;
+		return;
+	}
+
 	UpdateTimeCursorPosition(NewPlaybackPosition);
 }
 
@@ -441,6 +446,7 @@ void UMovieSceneSequencePlayer::ApplyLatentActions()
 		{
 		case ELatentAction::Stop:	Stop(); break;
 		case ELatentAction::Pause:	Pause(); break;
+		case ELatentAction::SetPlaybackPosition: SetPlaybackPosition(LatentPlaybackPosition); break;
 		}
 	}
 }
diff --git a/Engine/Source/Runtime/MovieScene/Private/MovieSceneTrack.cpp b/Engine/Source/Runtime/MovieScene/Private/MovieSceneTrack.cpp
--- a/Engine/Source/Runtime/MovieScene/Private/MovieSceneTrack.cpp
+++ b/Engine/Source/Runtime/MovieScene/Private/MovieSceneTrack.cpp
@@ -5,10 +5,10 @@
 #include "Evaluation/MovieSceneSegment.h"
 #include "Compilation/MovieSceneSegmentCompiler.h"
 #include "Compilation/MovieSceneCompilerRules.h"
+#include "Compilation/IMovieSceneTemplateGenerator.h"
 
 #include "Evaluation/MovieSceneEvaluationTrack.h"
 #include "Evaluation/MovieSceneEvaluationTemplate.h"
-#include "Evaluation/MovieSceneLegacyTrackInstanceTemplate.h"
 
 #include "MovieSceneEvaluationCustomVersion.h"
 
@@ -103,67 +103,69 @@ void UMovieSceneTrack::UpdateEasing()
 				MaxEaseOut = MaxEaseIn = Max * 0.25f;
 			}
 
-			CurrentSection->Modify();
-			CurrentSection->Easing.AutoEaseInTime = FMath::Clamp(MaxEaseIn, 0.f, Max);
-			CurrentSection->Easing.AutoEaseOutTime = FMath::Clamp(MaxEaseOut, 0.f, Max);
+			// Only modify the section if the ease in or out times have actually changed
+			MaxEaseIn = FMath::Clamp(MaxEaseIn, 0.f, Max);
+			MaxEaseOut = FMath::Clamp(MaxEaseOut, 0.f, Max);
+
+			if (CurrentSection->Easing.AutoEaseInTime != MaxEaseIn || CurrentSection->Easing.AutoEaseOutTime != MaxEaseOut)
+			{
+				CurrentSection->Modify();
+
+				CurrentSection->Easing.AutoEaseInTime = MaxEaseIn;
+				CurrentSection->Easing.AutoEaseOutTime = MaxEaseOut;
+			}
 		}
 	}
 }
 
 PRAGMA_DISABLE_DEPRECATION_WARNINGS
-TSharedPtr<IMovieSceneTrackInstance> UMovieSceneTrack::CreateLegacyInstance() const
+TInlineValue<FMovieSceneSegmentCompilerRules> UMovieSceneTrack::GetRowCompilerRules() const
 {
-	// Ugly const cast required due to lack of const-correctness in old system
-	return const_cast<UMovieSceneTrack*>(this)->CreateInstance();
+	return TInlineValue<FMovieSceneSegmentCompilerRules>();
 }
-PRAGMA_ENABLE_DEPRECATION_WARNINGS
 
-TInlineValue<FMovieSceneSegmentCompilerRules> UMovieSceneTrack::GetRowCompilerRules() const
+TInlineValue<FMovieSceneSegmentCompilerRules> UMovieSceneTrack::GetTrackCompilerRules() const
 {
-	// By default we only evaluate the section with the highest Z-Order if they overlap on the same row
-	struct FDefaultCompilerRules : FMovieSceneSegmentCompilerRules
-	{
-		virtual void BlendSegment(FMovieSceneSegment& Segment, const TArrayView<const FMovieSceneSectionData>& SourceData) const
-		{
-			MovieSceneSegmentCompiler::BlendSegmentHighPass(Segment, SourceData);
-		}
-	};
-	return FDefaultCompilerRules();
+	return TInlineValue<FMovieSceneSegmentCompilerRules>();
 }
 
-TInlineValue<FMovieSceneSegmentCompilerRules> UMovieSceneTrack::GetTrackCompilerRules() const
+FMovieSceneTrackRowSegmentBlenderPtr UMovieSceneTrack::GetRowSegmentBlender() const
 {
-	struct FRoundToNearstSectionRules : FMovieSceneSegmentCompilerRules
+	// Handle legacy row compiler rules
+	TInlineValue<FMovieSceneSegmentCompilerRules> LegacyRules = GetRowCompilerRules();
+	if (LegacyRules.IsValid())
 	{
-		virtual TOptional<FMovieSceneSegment> InsertEmptySpace(const TRange<float>& Range, const FMovieSceneSegment* PreviousSegment, const FMovieSceneSegment* NextSegment) const
-		{
-			return MovieSceneSegmentCompiler::EvaluateNearestSegment(Range, PreviousSegment, NextSegment);
-		}
-	};
+		return TLegacyTrackRowSegmentBlender<FMovieSceneTrackRowSegmentBlender>(MoveTemp(LegacyRules));
+	}
+	else
+	{
+		return FDefaultTrackRowSegmentBlender();
+	}
+}
 
-	// Evaluate according to bEvalNearestSection preference
-	if (EvalOptions.bCanEvaluateNearestSection && EvalOptions.bEvalNearestSection)
+FMovieSceneTrackSegmentBlenderPtr UMovieSceneTrack::GetTrackSegmentBlender() const
+{
+	// Handle legacy track compiler rules
+	TInlineValue<FMovieSceneSegmentCompilerRules> LegacyRules = GetTrackCompilerRules();
+	if (LegacyRules.IsValid())
 	{
-		return FRoundToNearstSectionRules();
+		return TLegacyTrackRowSegmentBlender<FMovieSceneTrackSegmentBlender>(MoveTemp(LegacyRules));
+	}
+	else if (EvalOptions.bCanEvaluateNearestSection && EvalOptions.bEvalNearestSection)
+	{
+		return FEvaluateNearestSegmentBlender();
 	}
 	else
 	{
-		return FMovieSceneSegmentCompilerRules();
+		return FMovieSceneTrackSegmentBlenderPtr();
 	}
 }
+PRAGMA_ENABLE_DEPRECATION_WARNINGS
 
 void UMovieSceneTrack::GenerateTemplate(const FMovieSceneTrackCompilerArgs& Args) const
 {
 	FMovieSceneEvaluationTrack NewTrackTemplate(Args.ObjectBindingId);
 
-	// Legacy path
-	if (CreateLegacyInstance().IsValid())
-	{
-		NewTrackTemplate.DefineAsSingleTemplate(FMovieSceneLegacyTrackInstanceTemplate(this));
-		Args.Generator.AddLegacyTrack(MoveTemp(NewTrackTemplate), *this);
-		return;
-	}
-
 	if (Compile(NewTrackTemplate, Args) != EMovieSceneCompileResult::Success)
 	{
 		return;
@@ -176,148 +178,62 @@ FMovieSceneEvaluationTrack UMovieSceneTrack::GenerateTrackTemplate() const
 {
 	FMovieSceneEvaluationTrack TrackTemplate = FMovieSceneEvaluationTrack(FGuid());
 
-	// Legacy path
-	if (CreateLegacyInstance().IsValid())
-	{
-		TrackTemplate.DefineAsSingleTemplate(FMovieSceneLegacyTrackInstanceTemplate(this));
-	}
-	else
+	// For this path, we don't have a generator, so we just pass through a stub
+	struct FTemplateGeneratorStub : IMovieSceneTemplateGenerator
 	{
-		// For this path, we don't have a generator, so we just pass through a stub
-		struct FTemplateGeneratorStub : IMovieSceneTemplateGenerator
-		{
-			virtual void AddOwnedTrack(FMovieSceneEvaluationTrack&& InTrackTemplate, const UMovieSceneTrack& SourceTrack) override {}
-			virtual void AddSharedTrack(FMovieSceneEvaluationTrack&& InTrackTemplate, FMovieSceneSharedDataId SharedId, const UMovieSceneTrack& SourceTrack) override {}
-			virtual void AddLegacyTrack(FMovieSceneEvaluationTrack&& InTrackTemplate, const UMovieSceneTrack& SourceTrack) override {}
-			virtual void AddExternalSegments(TRange<float> RootRange, TArrayView<const FMovieSceneEvaluationFieldSegmentPtr> SegmentPtrs, ESectionEvaluationFlags Flags) override {}
-			virtual FMovieSceneSequenceTransform GetSequenceTransform(FMovieSceneSequenceIDRef InSequenceID) const override { return FMovieSceneSequenceTransform(); }
-			virtual void AddSubSequence(FMovieSceneSubSequenceData SequenceData, FMovieSceneSequenceIDRef ParentID, FMovieSceneSequenceID SpecificID) override { }
-		} Generator;
-
-		FMovieSceneSequenceTemplateStore Store;
-		Compile(TrackTemplate, FMovieSceneTrackCompilerArgs(Generator, Store));
-	}
+		virtual void AddOwnedTrack(FMovieSceneEvaluationTrack&& InTrackTemplate, const UMovieSceneTrack& SourceTrack) override {}
+	} Generator;
+
+	Compile(TrackTemplate, FMovieSceneTrackCompilerArgs(Generator));
+
 	return TrackTemplate;
 }
 
-struct FSegmentRemapper
+EMovieSceneCompileResult UMovieSceneTrack::Compile(FMovieSceneEvaluationTrack& OutTrack, const FMovieSceneTrackCompilerArgs& Args) const
 {
-	explicit FSegmentRemapper(bool bInAllowEmptySegments)
-		: bAllowEmptySegments(bInAllowEmptySegments)
-	{}
-
-	void ProcessSegments(const TArray<FMovieSceneSegment>& SourceSegments, FMovieSceneEvaluationTrack& OutTrack, const TFunctionRef<FMovieSceneEvalTemplatePtr(int32)>& TemplateFactory)
-	{
-		NewSegments.Reset(SourceSegments.Num());
-		NewIndices.Reset();
-
-		for (const FMovieSceneSegment& Segment : SourceSegments)
-		{
-			AddSegment(Segment, OutTrack, TemplateFactory);
-		}
+	OutTrack.SetPreAndPostrollConditions(EvalOptions.bEvaluateInPreroll, EvalOptions.bEvaluateInPostroll);
 
-		OutTrack.SetSegments(MoveTemp(NewSegments));
-	}
+	EMovieSceneCompileResult Result = CustomCompile(OutTrack, Args);
 
-	void AddSegment(const FMovieSceneSegment& SourceSegment, FMovieSceneEvaluationTrack& OutTrack, const TFunctionRef<FMovieSceneEvalTemplatePtr(int32)>& TemplateFactory)
+	if (Result == EMovieSceneCompileResult::Unimplemented)
 	{
-		FMovieSceneSegment NewSegment;
-		NewSegment.Range = SourceSegment.Range;
-
-		for (const FSectionEvaluationData& EvalData : SourceSegment.Impls)
+		for (const UMovieSceneSection* Section : GetAllSections())
 		{
-			EnsureIndexIsValid(EvalData.ImplIndex);
-
-			// Ensure all our templates have been added to the track
-			if (NewIndices[EvalData.ImplIndex] == NotCreatedYet)
-			{
-				FMovieSceneEvalTemplatePtr NewTemplate = TemplateFactory(EvalData.ImplIndex);
-				if (NewTemplate.IsValid())
-				{
-					NewIndices[EvalData.ImplIndex] = OutTrack.AddChildTemplate(MoveTemp(NewTemplate));
-				}
-				else
-				{
-					NewIndices[EvalData.ImplIndex] = NullTemplate;
-				}
-			}
-
-			if (NewIndices[EvalData.ImplIndex] == NullTemplate)
+			if (!Section->IsActive())
 			{
 				continue;
 			}
 
-			FSectionEvaluationData NewData = EvalData;
-			NewData.ImplIndex = NewIndices[EvalData.ImplIndex];
-			NewSegment.Impls.Add(NewData);
-		}
-
-		if (bAllowEmptySegments || NewSegment.Impls.Num())
-		{
-			NewSegments.Add(NewSegment);
-		}
-	}
-
-private:
-	bool bAllowEmptySegments;
-
-	static const int32 NotCreatedYet;
-	static const int32 NullTemplate;
-
-	TArray<FMovieSceneSegment> NewSegments;
-	TArray<int32> NewIndices;
+			FMovieSceneEvalTemplatePtr NewTemplate = CreateTemplateForSection(*Section);
+			if (NewTemplate.IsValid())
+			{
+				NewTemplate->SetCompletionMode(Section->EvalOptions.CompletionMode == EMovieSceneCompletionMode::ProjectDefault ? Args.DefaultCompletionMode : Section->EvalOptions.CompletionMode);
+				NewTemplate->SetSourceSection(Section);
 
-	void EnsureIndexIsValid(int32 InSourceIndex)
-	{
-		NewIndices.Reserve(InSourceIndex + 1);
-		for (int32 Index = NewIndices.Num(); Index < InSourceIndex + 1; ++Index)
-		{
-			NewIndices.Add(NotCreatedYet);
-		}
-	}
-};
+				int32 TemplateIndex = OutTrack.AddChildTemplate(MoveTemp(NewTemplate));
 
-const int32 FSegmentRemapper::NotCreatedYet = -1;
-const int32 FSegmentRemapper::NullTemplate = -2;
+				const TRange<float> SectionRange = Section->IsInfinite() ? TRange<float>::All() : Section->GetRange();
+				OutTrack.AddTreeData(SectionRange, FSectionEvaluationData(TemplateIndex, ESectionEvaluationFlags::None));
 
-EMovieSceneCompileResult UMovieSceneTrack::Compile(FMovieSceneEvaluationTrack& OutTrack, const FMovieSceneTrackCompilerArgs& Args) const
-{
-	OutTrack.SetPreAndPostrollConditions(EvalOptions.bEvaluateInPreroll, EvalOptions.bEvaluateInPostroll);
-
-	EMovieSceneCompileResult Result = CustomCompile(OutTrack, Args);
+				if (!SectionRange.GetLowerBound().IsOpen() && Section->GetPreRollTime() > 0)
+				{
+					TRange<float> PreRollRange(SectionRange.GetLowerBoundValue() - Section->GetPreRollTime(), TRangeBound<float>::FlipInclusion(SectionRange.GetLowerBoundValue()));
+					OutTrack.AddTreeData(PreRollRange, FSectionEvaluationData(TemplateIndex, ESectionEvaluationFlags::PreRoll));
+				}
 
-	if (Result == EMovieSceneCompileResult::Unimplemented)
-	{
-		// Default implementation
-		const TArray<UMovieSceneSection*>& AllSections = GetAllSections();
-		
-		TInlineValue<FMovieSceneSegmentCompilerRules> RowCompilerRules = GetRowCompilerRules();
-		FMovieSceneTrackCompiler::FRows TrackRows(AllSections, RowCompilerRules.GetPtr(nullptr));
-
-		FMovieSceneTrackCompiler Compiler;
-		TInlineValue<FMovieSceneSegmentCompilerRules> Rules = GetTrackCompilerRules();
-		FMovieSceneTrackEvaluationField EvaluationField = Compiler.Compile(TrackRows.Rows, Rules.GetPtr(nullptr));
-
-		const bool bAllowEmptySegments = Rules.IsValid() && Rules->AllowEmptySegments(); 
-
-		FSegmentRemapper Remapper(bAllowEmptySegments);
-		Remapper.ProcessSegments(EvaluationField.Segments, OutTrack,
-			[&](int32 SectionIndex){
-				FMovieSceneEvalTemplatePtr NewTemplate = CreateTemplateForSection(*AllSections[SectionIndex]);
-				if (NewTemplate.IsValid())
+				if (!SectionRange.GetUpperBound().IsOpen() && Section->GetPostRollTime() > 0)
 				{
-					NewTemplate->SetCompletionMode(AllSections[SectionIndex]->EvalOptions.CompletionMode == EMovieSceneCompletionMode::ProjectDefault ? Args.DefaultCompletionMode : AllSections[SectionIndex]->EvalOptions.CompletionMode);
-					NewTemplate->SetSourceSection(AllSections[SectionIndex]);
+					TRange<float> PostRollRange(TRangeBound<float>::FlipInclusion(SectionRange.GetUpperBoundValue()), SectionRange.GetUpperBoundValue() + Section->GetPostRollTime());
+					OutTrack.AddTreeData(PostRollRange, FSectionEvaluationData(TemplateIndex, ESectionEvaluationFlags::PostRoll));
 				}
-				return NewTemplate;
 			}
-		);
-
+		}
 		Result = EMovieSceneCompileResult::Success;
 	}
 
 	if (Result == EMovieSceneCompileResult::Success)
 	{
+		OutTrack.SetSourceTrack(this);
 		PostCompile(OutTrack, Args);
 	}
 
diff --git a/Engine/Source/Runtime/MovieScene/Private/Sections/MovieSceneSubSection.cpp b/Engine/Source/Runtime/MovieScene/Private/Sections/MovieSceneSubSection.cpp
new file mode 100644
--- /dev/null
+++ b/Engine/Source/Runtime/MovieScene/Private/Sections/MovieSceneSubSection.cpp
@@ -0,0 +1,206 @@
+// Copyright 1998-2017 Epic Games, Inc. All Rights Reserved.
+
+#include "Sections/MovieSceneSubSection.h"
+#include "MovieScene.h"
+#include "MovieSceneTrack.h"
+#include "MovieSceneSequence.h"
+#include "Evaluation/MovieSceneEvaluationTemplate.h"
+
+TWeakObjectPtr<UMovieSceneSubSection> UMovieSceneSubSection::TheRecordingSection;
+
+float DeprecatedMagicNumber = TNumericLimits<float>::Lowest();
+
+/* UMovieSceneSubSection structors
+ *****************************************************************************/
+
+UMovieSceneSubSection::UMovieSceneSubSection()
+	: StartOffset_DEPRECATED(DeprecatedMagicNumber)
+	, TimeScale_DEPRECATED(DeprecatedMagicNumber)
+	, PrerollTime_DEPRECATED(DeprecatedMagicNumber)
+{
+}
+
+FString UMovieSceneSubSection::GetPathNameInMovieScene() const
+{
+	UMovieScene* OuterMovieScene = GetTypedOuter<UMovieScene>();
+	check(OuterMovieScene);
+	return GetPathName(OuterMovieScene);
+}
+
+FMovieSceneSequenceID UMovieSceneSubSection::GetSequenceID() const
+{
+	FString FullPath = GetPathNameInMovieScene();
+	if (SubSequence)
+	{
+		FullPath += TEXT(" / ");
+		FullPath += SubSequence->GetPathName();
+	}
+
+	return FMovieSceneSequenceID(FCrc::Strihash_DEPRECATED(*FullPath));
+}
+
+void UMovieSceneSubSection::PostLoad()
+{
+	if (StartOffset_DEPRECATED != DeprecatedMagicNumber)
+	{
+		Parameters.StartOffset = StartOffset_DEPRECATED;
+	}
+
+	if (TimeScale_DEPRECATED != DeprecatedMagicNumber)
+	{
+		Parameters.TimeScale = TimeScale_DEPRECATED;
+	}
+
+	if (PrerollTime_DEPRECATED != DeprecatedMagicNumber)
+	{
+		Parameters.PrerollTime_DEPRECATED = PrerollTime_DEPRECATED;
+	}
+
+	// Pre and post roll is now supported generically
+	if (Parameters.PrerollTime_DEPRECATED > 0.f)
+	{
+		SetPreRollTime(Parameters.PrerollTime_DEPRECATED);
+	}
+
+	if (Parameters.PostrollTime_DEPRECATED > 0.f)
+	{
+		SetPostRollTime(Parameters.PostrollTime_DEPRECATED);
+	}
+
+	Super::PostLoad();
+}
+
+void UMovieSceneSubSection::SetSequence(UMovieSceneSequence* Sequence)
+{
+	SubSequence = Sequence;
+
+#if WITH_EDITOR
+	OnSequenceChangedDelegate.ExecuteIfBound(SubSequence);
+#endif
+}
+
+UMovieSceneSequence* UMovieSceneSubSection::GetSequence() const
+{
+	// when recording we need to act as if we have no sequence
+	// the sequence is patched at the end of recording
+	if(GetRecordingSection() == this)
+	{
+		return nullptr;
+	}
+	else
+	{
+		return SubSequence;
+	}
+}
+
+UMovieSceneSubSection* UMovieSceneSubSection::GetRecordingSection()
+{
+	// check if the section is still valid and part of a track (i.e. it has not been deleted or GCed)
+	if(TheRecordingSection.IsValid())
+	{
+		UMovieSceneTrack* TrackOuter = Cast<UMovieSceneTrack>(TheRecordingSection->GetOuter());
+		if(TrackOuter)
+		{
+			if(TrackOuter->HasSection(*TheRecordingSection.Get()))
+			{
+				return TheRecordingSection.Get();
+			}
+		}
+	}
+
+	return nullptr;
+}
+
+void UMovieSceneSubSection::SetAsRecording(bool bRecord)
+{
+	if(bRecord)
+	{
+		TheRecordingSection = this;
+	}
+	else
+	{
+		TheRecordingSection = nullptr;
+	}
+}
+
+bool UMovieSceneSubSection::IsSetAsRecording()
+{
+	return GetRecordingSection() != nullptr;
+}
+
+AActor* UMovieSceneSubSection::GetActorToRecord()
+{
+	UMovieSceneSubSection* RecordingSection = GetRecordingSection();
+	if(RecordingSection)
+	{
+		return RecordingSection->ActorToRecord.Get();
+	}
+
+	return nullptr;
+}
+
+#if WITH_EDITOR
+void UMovieSceneSubSection::PostEditChangeProperty(FPropertyChangedEvent& PropertyChangedEvent)
+{
+	Super::PostEditChangeProperty(PropertyChangedEvent);
+
+	// recreate runtime instance when sequence is changed
+	if (PropertyChangedEvent.Property && PropertyChangedEvent.Property->GetFName() == GET_MEMBER_NAME_CHECKED(UMovieSceneSubSection, SubSequence))
+	{
+		OnSequenceChangedDelegate.ExecuteIfBound(SubSequence);
+	}
+}
+#endif
+
+UMovieSceneSection* UMovieSceneSubSection::SplitSection( float SplitTime )
+{
+	if ( !IsTimeWithinSection( SplitTime ) )
+	{
+		return nullptr;
+	}
+
+	float InitialStartTime = GetStartTime();
+	float InitialStartOffset = Parameters.StartOffset;
+	float NewStartOffset = ( SplitTime - InitialStartTime ) / Parameters.TimeScale;
+	NewStartOffset += InitialStartOffset;
+
+	// Ensure start offset is not less than 0
+	NewStartOffset = FMath::Max( NewStartOffset, 0.f );
+
+	UMovieSceneSubSection* NewSection = Cast<UMovieSceneSubSection>( UMovieSceneSection::SplitSection( SplitTime ) );
+	if ( NewSection )
+	{
+		NewSection->Parameters.StartOffset = NewStartOffset;
+		return NewSection;
+	}
+
+	return nullptr;
+}
+
+void UMovieSceneSubSection::TrimSection( float TrimTime, bool bTrimLeft )
+{
+	if ( !IsTimeWithinSection( TrimTime ) )
+	{
+		return;
+	}
+
+	float InitialStartTime = GetStartTime();
+	float InitialStartOffset = Parameters.StartOffset;
+
+	UMovieSceneSection::TrimSection( TrimTime, bTrimLeft );
+
+	// If trimming off the left, set the offset of the shot
+	if ( bTrimLeft )
+	{
+		float NewStartOffset = ( TrimTime - InitialStartTime ) / Parameters.TimeScale;
+		NewStartOffset += InitialStartOffset;
+
+		// Ensure start offset is not less than 0
+		Parameters.StartOffset = FMath::Max( NewStartOffset, 0.f );
+	}
+}
+
+FMovieSceneSubSequenceData UMovieSceneSubSection::GenerateSubSequenceData(const FSubSequenceInstanceDataParams& Params) const
+{
+	return FMovieSceneSubSequenceData(*this);
+}
diff --git a/Engine/Source/Runtime/MovieScene/Private/Tests/MovieSceneBlendingTests.cpp b/Engine/Source/Runtime/MovieScene/Private/Tests/MovieSceneBlendingTests.cpp
--- a/Engine/Source/Runtime/MovieScene/Private/Tests/MovieSceneBlendingTests.cpp
+++ b/Engine/Source/Runtime/MovieScene/Private/Tests/MovieSceneBlendingTests.cpp
@@ -56,7 +56,7 @@ bool FMovieSceneBlendingTest::RunTest(const FString& Parameters)
 	const FMovieSceneContext Context(FMovieSceneEvaluationRange(0.f));
 	
 	FNullPlayer Player;
-	FPersistentEvaluationData PersistentDataProxy(Player.State.PersistentEntityData, Player.State.PersistentSharedData);
+	FPersistentEvaluationData PersistentDataProxy(Player);
 
 	FMovieSceneBlendingAccumulator Accumulator;
 	Accumulator.DefineActuator(ID, MakeShared<FInt32Actuator>());
@@ -146,7 +146,7 @@ bool FMovieSceneBlendingStressTest::RunTest(const FString& Parameters)
 	const FMovieSceneContext Context(FMovieSceneEvaluationRange(0.f));
 	
 	FNullPlayer Player;
-	FPersistentEvaluationData PersistentDataProxy(Player.State.PersistentEntityData, Player.State.PersistentSharedData);
+	FPersistentEvaluationData PersistentDataProxy(Player);
 
 	FMovieSceneBlendingAccumulator Accumulator;
 
