Copying //UE4/Dev-Sequencer to //UE4/Dev-Main (Source: //UE4/Dev-Sequencer @ 3800792)

diff --git a/Engine/Source/Runtime/MovieScene/Private/Tests/MovieSceneSegmentCompilerTests.cpp b/Engine/Source/Runtime/MovieScene/Private/Tests/MovieSceneSegmentCompilerTests.cpp
--- a/Engine/Source/Runtime/MovieScene/Private/Tests/MovieSceneSegmentCompilerTests.cpp
+++ b/Engine/Source/Runtime/MovieScene/Private/Tests/MovieSceneSegmentCompilerTests.cpp
@@ -3,23 +3,26 @@
 #include "MovieSceneSegmentCompilerTests.h"
 #include "CoreMinimal.h"
 #include "Misc/AutomationTest.h"
+#include "Compilation/MovieSceneCompiler.h"
 #include "Compilation/MovieSceneSegmentCompiler.h"
 #include "Compilation/MovieSceneCompilerRules.h"
 #include "MovieSceneTestsCommon.h"
+#include "Evaluation/MovieSceneEvaluationTrack.h"
+#include "Evaluation/MovieSceneEvaluationField.h"
 #include "Package.h"
 
+#if WITH_DEV_AUTOMATION_TESTS
+
 namespace Impl
 {
 	static const TRangeBound<float> Inf = TRangeBound<float>::Open();
-	static const FOptionalMovieSceneBlendType NoBlending;
+
+	typedef TTuple<TRange<float>, TArray<int32>> FEvaluationTreeIteratorResult;
 
 	/** Compiler rules to sort by priority */
-	struct FSortByPriorityCompilerRules : FMovieSceneSegmentCompilerRules
+	struct FSortByPrioritySegmentBlender : FMovieSceneTrackSegmentBlender
 	{
-		int32 MaxImplIndex;
-
-		FSortByPriorityCompilerRules(int32 InMaxImplIndex, bool bInAllowEmptySegments)
-			: MaxImplIndex(InMaxImplIndex)
+		FSortByPrioritySegmentBlender(bool bInAllowEmptySegments)
 		{
 			bAllowEmptySegments = bInAllowEmptySegments;
 		}
@@ -29,53 +32,67 @@ namespace Impl
 			return bAllowEmptySegments ? FMovieSceneSegment(Range) : TOptional<FMovieSceneSegment>();
 		}
 
-		virtual void BlendSegment(FMovieSceneSegment& Segment, const TArrayView<const FMovieSceneSectionData>& SourceData) const
+		virtual void Blend(FSegmentBlendData& BlendData) const
 		{
-			Segment.Impls.Sort(
-				[&](const FSectionEvaluationData& A, const FSectionEvaluationData& B)
+			BlendData.Sort(
+				[](const FMovieSceneSectionData& A, const FMovieSceneSectionData& B)
 				{
-					return SourceData[A.ImplIndex].Priority > SourceData[B.ImplIndex].Priority;
+					// Sort by template index and flags when we don't have any sections
+					if (!A.Section || !B.Section)
+					{
+						if (A.Flags == ESectionEvaluationFlags::None && B.Flags == ESectionEvaluationFlags::None)
+						{
+							return A.TemplateIndex < B.TemplateIndex;
+						}
+
+						return A.Flags != ESectionEvaluationFlags::None;
+					}
+					return A.Section->GetOverlapPriority() > B.Section->GetOverlapPriority();
 				}
 			);
 		}
+	};
 
-		virtual void PostProcessSegments(TArray<FMovieSceneSegment>& Segments, const TArrayView<const FMovieSceneSectionData>& SourceData) const override
+	FString Join(TArrayView<const FSectionEvaluationData> Stuff)
+	{
+		FString Result;
+		for (const FSectionEvaluationData& Thing : Stuff)
 		{
-			for (FMovieSceneSegment& Segment : Segments)
+			if (Result.Len())
 			{
-				for (FSectionEvaluationData EvalData : Segment.Impls)
-				{
-					ensureMsgf(SourceData.IsValidIndex(EvalData.ImplIndex),
-						TEXT("Compiled segment data does not correctly map to the source data array"));
-
-					const int32 ThisImpl = SourceData[EvalData.ImplIndex].EvalData.ImplIndex;
-					ensureMsgf(ThisImpl >= 0 && ThisImpl <= MaxImplIndex,
-						TEXT("Compiled segment data does not correctly map to either the designated implementation range"));
-				}
+				Result += TEXT(", ");
 			}
+
+			Result += FString::Printf(TEXT("(Impl: %d, ForcedTime: %.7f, Flags: %u)"), Thing.ImplIndex, Thing.ForcedTime, (uint8)Thing.Flags);
 		}
 
-	};
-	
-	FORCEINLINE void AssertSegmentValues(FAutomationTestBase* Test, TArrayView<const FMovieSceneSegment> Expected, TArrayView<const FMovieSceneSegment> Actual)
+		return Result;
+	}
+
+	FORCEINLINE_DEBUGGABLE void AssertSegmentImpls(FAutomationTestBase* Test, TArrayView<const FSectionEvaluationData> ExpectedImpls, TArrayView<const FSectionEvaluationData> ActualImpls, const TCHAR* AdditionalText = TEXT(""))
 	{
-		auto Join = [](TArrayView<const FSectionEvaluationData> Stuff) -> FString
-		{
-			FString Result;
-			for (const FSectionEvaluationData& Thing : Stuff)
-			{
-				if (Result.Len())
-				{
-					Result += TEXT(", ");
-				}
+		FString ActualImplsString = Join(ActualImpls);
+		FString ExpectedImplsString = Join(ExpectedImpls);
 
-				Result += FString::Printf(TEXT("(Impl: %d, ForcedTime: %.7f, Flags: %u)"), Thing.ImplIndex, Thing.ForcedTime, (uint8)Thing.Flags);
-			}
+		if (ActualImplsString != ExpectedImplsString)
+		{
+			Test->AddError(FString::Printf(TEXT("Compiled data does not match for segment %s.\nExpected: %s\nActual:   %s."), AdditionalText, *ExpectedImplsString, *ActualImplsString));
+		}
+	}
+	
+	FORCEINLINE_DEBUGGABLE void AssertSegmentValues(FAutomationTestBase* Test, TArrayView<const FMovieSceneSegment> Expected, TArrayView<const FMovieSceneSegment> Actual)
+	{
+		using namespace Lex;
 
-			return Result;
-		};
+		for (const FMovieSceneSegment& Segment : Actual)
+		{
+			Test->AddInfo(FString::Printf(TEXT("Actual   %s, Impls: %s"), *ToString(Segment.Range), *Join(Segment.Impls)));
+		}
+		for (const FMovieSceneSegment& Segment : Expected)
+		{
+			Test->AddInfo(FString::Printf(TEXT("Expected %s, Impls: %s"), *ToString(Segment.Range), *Join(Segment.Impls)));
+		}
 
-		using namespace Lex;
 		if (Actual.Num() != Expected.Num())
 		{
 			Test->AddError(FString::Printf(TEXT("Wrong number of compiled segments. Expected %d, actual %d."), Expected.Num(), Actual.Num()));
@@ -87,22 +104,59 @@ namespace Impl
 
 			if (ExpectedSegment.Range != ActualSegment.Range)
 			{
-				Test->AddError(FString::Printf(TEXT("Incorrect compiled segment range at segment index %d. Expected:\n%s\nActual:\n%s"), Index, *Lex::ToString(ExpectedSegment.Range), *Lex::ToString(ActualSegment.Range)));
+				Test->AddError(FString::Printf(TEXT("Incorrect compiled segment range at segment index %d. Expected:\n%s\nActual:\n%s"), Index, *ToString(ExpectedSegment.Range), *ToString(ActualSegment.Range)));
 			}
-			else if (ExpectedSegment.Impls.Num() != ActualSegment.Impls.Num())
+
+			AssertSegmentImpls(Test, ExpectedSegment.Impls, ActualSegment.Impls, *FString::Printf(TEXT("index %d"), Index));
+		}
+	}
+
+	FORCEINLINE_DEBUGGABLE void AssertSegmentAtTime(FAutomationTestBase* Test, FMovieSceneEvaluationTrack& InTrack, float InTime, TArrayView<const FSectionEvaluationData> ExpectedImpls)
+	{
+		FMovieSceneSegmentIdentifier ID = InTrack.GetSegmentFromTime(InTime);
+		if (!ID.IsValid())
+		{
+			Test->AddError(TEXT("No segment compiled for time %.fs"), InTime);
+		}
+		else
+		{
+			AssertSegmentImpls(Test, ExpectedImpls, InTrack.GetSegment(ID).Impls, *FString::Printf(TEXT("time %.8fs"), InTime));
+		}
+	}
+
+	FORCEINLINE_DEBUGGABLE void AssertEvaluationTree(FAutomationTestBase* Test, const TMovieSceneEvaluationTree<int32>& Tree, TArrayView<const FEvaluationTreeIteratorResult> Expected)
+	{
+		using namespace Lex;
+		int32 Index = 0;
+		for (FMovieSceneEvaluationTreeRangeIterator It(Tree); It; ++It, ++Index)
+		{
+			if (Index >= Expected.Num())
 			{
-				Test->AddError(FString::Printf(TEXT("Incorrect number of implementation references compiled into segment index %d. Expected %d, actual %d."), Index, ExpectedSegment.Impls.Num(), ActualSegment.Impls.Num()));
+				Test->AddError(FString::Printf(TEXT("Too many ranges iterated (expected %d, iterated >= %d)"), Expected.Num(), Index + 1));
+				break;
 			}
-			else 
+
+			if (It.Range() != Expected[Index].Get<0>())
 			{
-				FString ActualImpls = Join(ActualSegment.Impls);
-				FString ExpectedImpls = Join(ExpectedSegment.Impls);
+				Test->AddError(FString::Printf(TEXT("Incorrect range iterated at index %d. %s != %s"), Index, *ToString(It.Range()), *ToString(Expected[Index].Get<0>())));
+				continue;
+			}
 
-				if (ActualImpls != ExpectedImpls)
+			int32 Count = 0;
+			const TArray<int32>& ExpectedData = Expected[Index].Get<1>();
+			for (auto AllDataIt = Tree.GetAllData(It.Node()); AllDataIt; ++AllDataIt, ++Count)
+			{
+				if (!ExpectedData.Contains(*AllDataIt))
 				{
-					Test->AddError(FString::Printf(TEXT("Compiled data does not match for segment index %d.\nExpected: %s\nActual:   %s."), Index, *ExpectedImpls, *ActualImpls));
+					Test->AddError(FString::Printf(TEXT("Incorrect data encountered at index %d. Count not find %d in expected data."), Index, *AllDataIt));
+					continue;
 				}
 			}
+
+			if (Count != ExpectedData.Num())
+			{
+				Test->AddError(FString::Printf(TEXT("Incorrect number of data entries iterated at index %d. Found %d entries, expected %d"), Index, Count, ExpectedData.Num()));
+			}
 		}
 	}
 }
@@ -123,28 +177,33 @@ bool FMovieSceneCompilerBasicTest::RunTest(const FString& Parameters)
 	//----------------------------------------------------------------------------------------------------------------------------------------
 	//	Expected Impls		[ 		3			|			0,2,3			|		1,2		|		1,2,4		|			4				]
 
-	const FMovieSceneSectionData SegmentData[] = {
-		FMovieSceneSectionData(TRange<float>(10.f,	20.f), 									FSectionEvaluationData(0), 	NoBlending,		4),
-		FMovieSceneSectionData(TRange<float>(20.f,	30.f), 									FSectionEvaluationData(1), 	NoBlending,		3),
-		FMovieSceneSectionData(TRange<float>(10.f,	30.f), 									FSectionEvaluationData(2), 	NoBlending,		2),
-		FMovieSceneSectionData(TRange<float>(Inf,	TRangeBound<float>::Exclusive(20.f)), 	FSectionEvaluationData(3), 	NoBlending,		1),
-		FMovieSceneSectionData(TRange<float>(25.f,	Inf), 									FSectionEvaluationData(4), 	NoBlending,		0)
-	};
+	// Override the blender (so it doesn't try and get it from the null track ptr)
+	FScopedOverrideTrackSegmentBlender ScopedBlenderOverride(FSortByPrioritySegmentBlender(false));
 
-	FSortByPriorityCompilerRules DefaultRules(4, false);
-	TArray<FMovieSceneSegment> Segments = FMovieSceneSegmentCompiler().Compile(SegmentData, &DefaultRules);
+	FMovieSceneEvaluationTrack Track;
+
+	Track.AddTreeData(TRange<float>(10.f,	20.f),										FSectionEvaluationData(0));
+	Track.AddTreeData(TRange<float>(20.f,	30.f),										FSectionEvaluationData(1));
+	Track.AddTreeData(TRange<float>(10.f,	30.f),										FSectionEvaluationData(2));
+	Track.AddTreeData(TRange<float>(Inf,	TRangeBound<float>::Exclusive(20.f)),		FSectionEvaluationData(3));
+	Track.AddTreeData(TRange<float>(25.f,	Inf),										FSectionEvaluationData(4));
 
 	FMovieSceneSegment Expected[] = {
-		FMovieSceneSegment(FFloatRange(Inf, 	TRangeBound<float>::Exclusive(10.f)), { FSectionEvaluationData(3) }),
-		FMovieSceneSegment(FFloatRange(10.f,	20.f), { FSectionEvaluationData(0), FSectionEvaluationData(2), FSectionEvaluationData(3) }),
-		FMovieSceneSegment(FFloatRange(20.f,	25.f), { FSectionEvaluationData(1), FSectionEvaluationData(2) }),
-		FMovieSceneSegment(FFloatRange(25.f,	30.f), { FSectionEvaluationData(1), FSectionEvaluationData(2), FSectionEvaluationData(4) }),
-		FMovieSceneSegment(FFloatRange(30.f,	Inf), { FSectionEvaluationData(4) })
+		FMovieSceneSegment(FFloatRange(Inf, 	TRangeBound<float>::Exclusive(10.f)),	{ FSectionEvaluationData(3) }),
+		FMovieSceneSegment(FFloatRange(10.f,	20.f),									{ FSectionEvaluationData(0), FSectionEvaluationData(2), FSectionEvaluationData(3) }),
+		FMovieSceneSegment(FFloatRange(20.f,	25.f),									{ FSectionEvaluationData(1), FSectionEvaluationData(2) }),
+		FMovieSceneSegment(FFloatRange(25.f,	30.f),									{ FSectionEvaluationData(1), FSectionEvaluationData(2), FSectionEvaluationData(4) }),
+		FMovieSceneSegment(FFloatRange(30.f,	Inf),									{ FSectionEvaluationData(4) })
 	};
-	AssertSegmentValues(this, Expected, Segments);
+
+	// Compile all the segments
+	Track.GetSegmentsInRange(TRange<float>::All());
+
+	AssertSegmentValues(this, Expected, Track.GetSortedSegments());
 	return true;
 }
 
+
 IMPLEMENT_SIMPLE_AUTOMATION_TEST(FMovieSceneCompilerEmptySpaceTest, "System.Engine.Sequencer.Compiler.Empty Space", EAutomationTestFlags::EditorContext | EAutomationTestFlags::EngineFilter)
 bool FMovieSceneCompilerEmptySpaceTest::RunTest(const FString& Parameters)
 {
@@ -157,56 +216,58 @@ bool FMovieSceneCompilerEmptySpaceTest::RunTest(const FString& Parameters)
 	//----------------------------------------------------------------------------------------------------------------------------------------
 	//	Expected Impls	[	Empty		|		0		|	Empty		|		1		|	Empty		]
 
-	const FMovieSceneSectionData SegmentData[] = {
-		FMovieSceneSectionData(TRange<float>(10.f,	20.f), 									FSectionEvaluationData(0), NoBlending, 	1),
-		FMovieSceneSectionData(TRange<float>(30.f,	40.f), 									FSectionEvaluationData(1), NoBlending, 	0)
-	};
+	FMovieSceneEvaluationTrack Track;
+
+	Track.AddTreeData(TRange<float>(10.f,	20.f), 		FSectionEvaluationData(0));
+	Track.AddTreeData(TRange<float>(30.f,	40.f), 		FSectionEvaluationData(1));
 
-	FSortByPriorityCompilerRules DefaultRules(1, true);
-	TArray<FMovieSceneSegment> Segments = FMovieSceneSegmentCompiler().Compile(SegmentData, &DefaultRules);
+	// Override the blender (so it doesn't try and get it from the null track ptr)
+	FScopedOverrideTrackSegmentBlender ScopedBlenderOverride(FSortByPrioritySegmentBlender(true));
+
+	// Compile all the segments
+	Track.GetSegmentsInRange(TRange<float>::All());
 
 	FMovieSceneSegment Expected[] = {
-		FMovieSceneSegment(FFloatRange(Inf, 	TRangeBound<float>::Exclusive(10.f)), { }),
-		FMovieSceneSegment(FFloatRange(10.f,	20.f), { FSectionEvaluationData(0) }),
-		FMovieSceneSegment(FFloatRange(20.f,	30.f), {  }),
-		FMovieSceneSegment(FFloatRange(30.f,	40.f), { FSectionEvaluationData(1) }),
-		FMovieSceneSegment(FFloatRange(40.f,	Inf), {  })
+		FMovieSceneSegment(FFloatRange(Inf, 	TRangeBound<float>::Exclusive(10.f)), 	{ }),
+		FMovieSceneSegment(FFloatRange(10.f,	20.f), 									{ FSectionEvaluationData(0) }),
+		FMovieSceneSegment(FFloatRange(20.f,	30.f), 									{  }),
+		FMovieSceneSegment(FFloatRange(30.f,	40.f), 									{ FSectionEvaluationData(1) }),
+		FMovieSceneSegment(FFloatRange(40.f,	Inf), 									{  })
 	};
-	AssertSegmentValues(this, Expected, Segments);
+	AssertSegmentValues(this, Expected, Track.GetSortedSegments());
 	return true;
 }
 
+
 IMPLEMENT_SIMPLE_AUTOMATION_TEST(FMovieSceneCustomCompilerTest, "System.Engine.Sequencer.Compiler.Custom", EAutomationTestFlags::EditorContext | EAutomationTestFlags::EngineFilter)
 bool FMovieSceneCustomCompilerTest::RunTest(const FString& Parameters)
 {
 	using namespace Impl;
 
-	// Specify descending priorities on the segments so we always sort the compiled segments in the order they're defined within the data
+	//	Test that we don't get duplicate entries in the resulting segments except for differing flag sets
+
 	//	This is our test layout
 	//	Time			-Inf			10			15			20			25			30						40				Inf
 	//									[===== 0 (preroll) =====]
 	//												[========== 0 ==========]
 	//															[========== 0 ==========][========== 0 =========]
-	//----------------------------------------------------------------------------------------------------------------------------------------
-	//	Expected Impls					[	0 (p)	| (0, 0(p)) |						0						]
 
-	const FMovieSceneSectionData SegmentData[] = {
-		FMovieSceneSectionData(TRange<float>(TRangeBound<float>::Inclusive(10.f),	TRangeBound<float>::Exclusive(20.f)), FSectionEvaluationData(0, ESectionEvaluationFlags::PreRoll),	NoBlending,	4),
-		FMovieSceneSectionData(TRange<float>(TRangeBound<float>::Inclusive(15.f),	TRangeBound<float>::Exclusive(25.f)), FSectionEvaluationData(0),									NoBlending,	1),
-		FMovieSceneSectionData(TRange<float>(TRangeBound<float>::Inclusive(20.f),	TRangeBound<float>::Exclusive(30.f)), FSectionEvaluationData(0),									NoBlending,	1),
-		FMovieSceneSectionData(TRange<float>(TRangeBound<float>::Inclusive(30.f),	TRangeBound<float>::Exclusive(40.f)), FSectionEvaluationData(0),									NoBlending,	1),
-	};
+	FMovieSceneEvaluationTrack Track;
 
-	FSortByPriorityCompilerRules DefaultRules(0, false);
-	TArray<FMovieSceneSegment> Segments = FMovieSceneSegmentCompiler().Compile(SegmentData, &DefaultRules);
+	Track.AddUniqueTreeData(TRange<float>(TRangeBound<float>::Inclusive(10.f),	TRangeBound<float>::Exclusive(20.f)), FSectionEvaluationData(0, ESectionEvaluationFlags::PreRoll));
+	Track.AddUniqueTreeData(TRange<float>(TRangeBound<float>::Inclusive(15.f),	TRangeBound<float>::Exclusive(25.f)), FSectionEvaluationData(0));
+	Track.AddUniqueTreeData(TRange<float>(TRangeBound<float>::Inclusive(20.f),	TRangeBound<float>::Exclusive(30.f)), FSectionEvaluationData(0));
+	Track.AddUniqueTreeData(TRange<float>(TRangeBound<float>::Inclusive(30.f),	TRangeBound<float>::Exclusive(40.f)), FSectionEvaluationData(0));
 
-	FMovieSceneSegment Expected[] = {
-		FMovieSceneSegment(TRange<float>(TRangeBound<float>::Inclusive(10.f),	TRangeBound<float>::Exclusive(15.f)), { FSectionEvaluationData(0, ESectionEvaluationFlags::PreRoll) }),
-		FMovieSceneSegment(TRange<float>(TRangeBound<float>::Inclusive(15.f),	TRangeBound<float>::Exclusive(20.f)), { FSectionEvaluationData(0, ESectionEvaluationFlags::PreRoll), FSectionEvaluationData(0) }),
-		FMovieSceneSegment(TRange<float>(TRangeBound<float>::Inclusive(20.f),	TRangeBound<float>::Exclusive(40.f)), { FSectionEvaluationData(0) }),
-	};
+	// Override the blender (so it doesn't try and get it from the null track ptr)
+	FScopedOverrideTrackSegmentBlender ScopedBlenderOverride(FSortByPrioritySegmentBlender(false));
+
+	AssertSegmentAtTime(this, Track, 12.5f,		{ FSectionEvaluationData(0, ESectionEvaluationFlags::PreRoll) });
+	AssertSegmentAtTime(this, Track, 17.5f,		{ FSectionEvaluationData(0, ESectionEvaluationFlags::PreRoll), FSectionEvaluationData(0) });
+	AssertSegmentAtTime(this, Track, 22.5f,		{ FSectionEvaluationData(0) });
+	AssertSegmentAtTime(this, Track, 27.5f,		{ FSectionEvaluationData(0) });
+	AssertSegmentAtTime(this, Track, 30.f,		{ FSectionEvaluationData(0) });
 
-	AssertSegmentValues(this, Expected, Segments);
 	return true;
 }
 
@@ -243,26 +304,27 @@ bool FMovieSceneTrackCompilerTest::RunTest(const FString& Parameters)
 		Track->SectionArray.Add(Section0);
 		Track->SectionArray.Add(Section1);
 
-		TInlineValue<FMovieSceneSegmentCompilerRules> RowCompilerRules = Track->GetRowCompilerRules();
-		FMovieSceneTrackCompiler::FRows Rows(Track->SectionArray, RowCompilerRules.GetPtr(nullptr));
-
 		// Test compiling the track with the additive camera rules
 		{
-			FMovieSceneAdditiveCameraRules AdditiveCameraRules(Track);
-			FMovieSceneTrackEvaluationField Field = FMovieSceneTrackCompiler().Compile(Rows.Rows, &AdditiveCameraRules);
+			FScopedOverrideTrackSegmentBlender ScopedBlenderOverride{FMovieSceneAdditiveCameraTrackBlender()};
+
+			FMovieSceneEvaluationTrack EvalTrack = Track->GenerateTrackTemplate();
+			EvalTrack.GetSegmentsInRange(TRange<float>::All());
 
 			FMovieSceneSegment Expected[] = {
 				FMovieSceneSegment(FFloatRange(TRangeBound<float>::Inclusive(10.f),	TRangeBound<float>::Exclusive(20.f)), { FSectionEvaluationData(0) }),
 				FMovieSceneSegment(FFloatRange(TRangeBound<float>::Inclusive(20.f),	TRangeBound<float>::Inclusive(25.f)), { FSectionEvaluationData(0), FSectionEvaluationData(1) }),
 				FMovieSceneSegment(FFloatRange(TRangeBound<float>::Exclusive(25.f),	TRangeBound<float>::Inclusive(30.f)), { FSectionEvaluationData(1) }),
 			};
-			AssertSegmentValues(this, Expected, Field.Segments);
+			AssertSegmentValues(this, Expected, EvalTrack.GetSortedSegments());
 		}
 
 		// Test compiling with 'evaluate nearest section' enabled
 		{
 			Track->EvalOptions.bEvalNearestSection = true;
-			FMovieSceneTrackEvaluationField Field = FMovieSceneTrackCompiler().Compile(Rows.Rows, Track->GetTrackCompilerRules().GetPtr(nullptr));
+
+			FMovieSceneEvaluationTrack EvalTrack = Track->GenerateTrackTemplate();
+			EvalTrack.GetSegmentsInRange(TRange<float>::All());
 
 			FMovieSceneSegment Expected[] = {
 				FMovieSceneSegment(FFloatRange(Inf,									TRangeBound<float>::Exclusive(10.f)),	{ FSectionEvaluationData(0, 10.f) }),
@@ -271,33 +333,37 @@ bool FMovieSceneTrackCompilerTest::RunTest(const FString& Parameters)
 				FMovieSceneSegment(FFloatRange(TRangeBound<float>::Exclusive(25.f),	TRangeBound<float>::Inclusive(30.f)),	{ FSectionEvaluationData(1) }),
 				FMovieSceneSegment(FFloatRange(TRangeBound<float>::Exclusive(30.f),	Inf), 									{ FSectionEvaluationData(1, 30.f) }),
 			};
-			AssertSegmentValues(this, Expected, Field.Segments);
+			AssertSegmentValues(this, Expected, EvalTrack.GetSortedSegments());
 		}
 
 		// Test compiling without 'evaluate nearest section' enabled
 		{
 			Track->EvalOptions.bEvalNearestSection = false;
-			FMovieSceneTrackEvaluationField Field = FMovieSceneTrackCompiler().Compile(Rows.Rows, Track->GetTrackCompilerRules().GetPtr(nullptr));
+
+			FMovieSceneEvaluationTrack EvalTrack = Track->GenerateTrackTemplate();
+			EvalTrack.GetSegmentsInRange(TRange<float>::All());
 
 			FMovieSceneSegment Expected[] = {
 				FMovieSceneSegment(FFloatRange(TRangeBound<float>::Inclusive(10.f),	TRangeBound<float>::Exclusive(20.f)), { FSectionEvaluationData(0) }),
 				FMovieSceneSegment(FFloatRange(TRangeBound<float>::Inclusive(20.f),	TRangeBound<float>::Inclusive(25.f)), { FSectionEvaluationData(0), FSectionEvaluationData(1) }),
 				FMovieSceneSegment(FFloatRange(TRangeBound<float>::Exclusive(25.f),	TRangeBound<float>::Inclusive(30.f)), { FSectionEvaluationData(1) }),
 			};
-			AssertSegmentValues(this, Expected, Field.Segments);
+			AssertSegmentValues(this, Expected, EvalTrack.GetSortedSegments());
 		}
 
 		// Test high-pass filter
 		{
 			Track->EvalOptions.bEvalNearestSection = false;
 			Track->bHighPassFilter = true;
-			FMovieSceneTrackEvaluationField Field = FMovieSceneTrackCompiler().Compile(Rows.Rows, Track->GetTrackCompilerRules().GetPtr(nullptr));
+
+			FMovieSceneEvaluationTrack EvalTrack = Track->GenerateTrackTemplate();
+			EvalTrack.GetSegmentsInRange(TRange<float>::All());
 
 			FMovieSceneSegment Expected[] = {
 				FMovieSceneSegment(FFloatRange(TRangeBound<float>::Inclusive(10.f),	TRangeBound<float>::Inclusive(25.f)), { FSectionEvaluationData(0) }),
 				FMovieSceneSegment(FFloatRange(TRangeBound<float>::Exclusive(25.f),	TRangeBound<float>::Inclusive(30.f)), { FSectionEvaluationData(1) }),
 			};
-			AssertSegmentValues(this, Expected, Field.Segments);
+			AssertSegmentValues(this, Expected, EvalTrack.GetSortedSegments());
 		}
 	}
 
@@ -340,9 +406,6 @@ bool FMovieSceneTrackCompilerTest::RunTest(const FString& Parameters)
 		Track->SectionArray.Add(Section2);
 		Track->SectionArray.Add(Section3);
 
-		auto RowCompilerRules = Track->GetRowCompilerRules();
-		FMovieSceneTrackCompiler::FRows Rows(Track->SectionArray, RowCompilerRules.GetPtr(nullptr));
-
 		// Additive camera rules prescribe that they are evaluated in order of start time
 		FMovieSceneSegment Expected[] = {
 			FMovieSceneSegment(FFloatRange(Inf,											TRangeBound<float>::Exclusive(10.f)),	{ }),
@@ -354,8 +417,10 @@ bool FMovieSceneTrackCompilerTest::RunTest(const FString& Parameters)
 
 		// Test compiling the track with the additive camera rules
 		{
-			FMovieSceneAdditiveCameraRules AdditiveCameraRules(Track);
-			FMovieSceneTrackEvaluationField Field = FMovieSceneTrackCompiler().Compile(Rows.Rows, &AdditiveCameraRules);
+			FScopedOverrideTrackSegmentBlender ScopedBlenderOverride{FMovieSceneAdditiveCameraTrackBlender()};
+
+			FMovieSceneEvaluationTrack EvalTrack = Track->GenerateTrackTemplate();
+			EvalTrack.GetSegmentsInRange(TRange<float>::All());
 
 			// Additive camera rules prescribe that they are evaluated in order of start time
 			Expected[0].Impls = { FSectionEvaluationData(1) };
@@ -364,13 +429,15 @@ bool FMovieSceneTrackCompilerTest::RunTest(const FString& Parameters)
 			Expected[3].Impls = { FSectionEvaluationData(1), FSectionEvaluationData(0), FSectionEvaluationData(2) };
 			Expected[4].Impls = { FSectionEvaluationData(1) };
 
-			AssertSegmentValues(this, Expected, Field.Segments);
+			AssertSegmentValues(this, Expected, EvalTrack.GetSortedSegments());
 		}
 
 		// Test compiling with 'evaluate nearest section' enabled
 		{
 			Track->EvalOptions.bEvalNearestSection = true;
-			FMovieSceneTrackEvaluationField Field = FMovieSceneTrackCompiler().Compile(Rows.Rows, Track->GetTrackCompilerRules().GetPtr(nullptr));
+
+			FMovieSceneEvaluationTrack EvalTrack = Track->GenerateTrackTemplate();
+			EvalTrack.GetSegmentsInRange(TRange<float>::All());
 
 			Expected[0].Impls = { FSectionEvaluationData(1) };
 			Expected[1].Impls = { FSectionEvaluationData(0), FSectionEvaluationData(1) };
@@ -378,13 +445,15 @@ bool FMovieSceneTrackCompilerTest::RunTest(const FString& Parameters)
 			Expected[3].Impls = { FSectionEvaluationData(2), FSectionEvaluationData(0), FSectionEvaluationData(1) };
 			Expected[4].Impls = { FSectionEvaluationData(1) };
 
-			AssertSegmentValues(this, Expected, Field.Segments);
+			AssertSegmentValues(this, Expected, EvalTrack.GetSortedSegments());
 		}
 
 		// Test compiling without 'evaluate nearest section' enabled
 		{
 			Track->EvalOptions.bEvalNearestSection = false;
-			FMovieSceneTrackEvaluationField Field = FMovieSceneTrackCompiler().Compile(Rows.Rows, Track->GetTrackCompilerRules().GetPtr(nullptr));
+
+			FMovieSceneEvaluationTrack EvalTrack = Track->GenerateTrackTemplate();
+			EvalTrack.GetSegmentsInRange(TRange<float>::All());
 
 			Expected[0].Impls = { FSectionEvaluationData(1) };
 			Expected[1].Impls = { FSectionEvaluationData(0), FSectionEvaluationData(1) };
@@ -392,14 +461,16 @@ bool FMovieSceneTrackCompilerTest::RunTest(const FString& Parameters)
 			Expected[3].Impls = { FSectionEvaluationData(2), FSectionEvaluationData(0), FSectionEvaluationData(1) };
 			Expected[4].Impls = { FSectionEvaluationData(1) };
 
-			AssertSegmentValues(this, Expected, Field.Segments);
+			AssertSegmentValues(this, Expected, EvalTrack.GetSortedSegments());
 		}
 
 		// Test compiling with high pass filter
 		{
 			Track->EvalOptions.bEvalNearestSection = false;
 			Track->bHighPassFilter = true;
-			FMovieSceneTrackEvaluationField Field = FMovieSceneTrackCompiler().Compile(Rows.Rows, Track->GetTrackCompilerRules().GetPtr(nullptr));
+
+			FMovieSceneEvaluationTrack EvalTrack = Track->GenerateTrackTemplate();
+			EvalTrack.GetSegmentsInRange(TRange<float>::All());
 
 			Expected[0].Impls = { FSectionEvaluationData(1) };
 			Expected[1].Impls = { FSectionEvaluationData(0) };
@@ -407,41 +478,164 @@ bool FMovieSceneTrackCompilerTest::RunTest(const FString& Parameters)
 			Expected[3].Impls = { FSectionEvaluationData(2) };
 			Expected[4].Impls = { FSectionEvaluationData(1) };
 
-			AssertSegmentValues(this, Expected, Field.Segments);
+			AssertSegmentValues(this, Expected, EvalTrack.GetSortedSegments());
 		}
 	}
 
 	return true;
 }
 
+IMPLEMENT_SIMPLE_AUTOMATION_TEST(FMovieSceneCompilerTreeBasicTest, "System.Engine.Sequencer.Compiler.Tree Basic", EAutomationTestFlags::EditorContext | EAutomationTestFlags::EngineFilter)
+bool FMovieSceneCompilerTreeBasicTest::RunTest(const FString& Parameters)
+{
+	using namespace Impl;
+
+	// Specify descending priorities on the segments so we always sort the compiled segments in the order they're defined within the data
+	//	This is our test layout
+	//	Time				-inf				10							20				25					30							inf
+	//											[============= 0 ===========]
+	//																		[============== 1 ==================]
+	//											[========================== 2 ==================================]
+	//						[================== 3 ==========================]
+	//																						[================== 4 ==========================]
+	//----------------------------------------------------------------------------------------------------------------------------------------
+	//	Expected Impls		[ 		3			|			0,2,3			|		1,2		|		1,2,4		|			4				]
+
+	TMovieSceneEvaluationTree<FSectionEvaluationData> EvalTree;
+
+	EvalTree.Add(TRange<float>(10.f,	20.f), 								FSectionEvaluationData(0));
+	EvalTree.Add(TRange<float>(20.f,	30.f), 								FSectionEvaluationData(1));
+	EvalTree.Add(TRange<float>(10.f,	30.f), 								FSectionEvaluationData(2));
+	EvalTree.Add(TRange<float>(Inf,	TRangeBound<float>::Exclusive(20.f)), 	FSectionEvaluationData(3));
+	EvalTree.Add(TRange<float>(25.f,	Inf), 								FSectionEvaluationData(4));
+
+	EvalTree.Compact();
+
+	TArray<FMovieSceneSegment> Segments;
+
+	for (FMovieSceneEvaluationTreeRangeIterator It(EvalTree); It; ++It)
+	{
+		FMovieSceneSegment NewSegment(It.Range());
+
+		for (FSectionEvaluationData EvalData : EvalTree.GetAllData(It.Node()))
+		{
+			NewSegment.Impls.Add(EvalData);
+		}
+
+		Algo::SortBy(NewSegment.Impls, &FSectionEvaluationData::ImplIndex);
+		Segments.Add(NewSegment);
+	}
+
+	FMovieSceneSegment Expected[] = {
+		FMovieSceneSegment(FFloatRange(Inf, 	TRangeBound<float>::Exclusive(10.f)), { FSectionEvaluationData(3) }),
+		FMovieSceneSegment(FFloatRange(10.f,	20.f), { FSectionEvaluationData(0), FSectionEvaluationData(2), FSectionEvaluationData(3) }),
+		FMovieSceneSegment(FFloatRange(20.f,	25.f), { FSectionEvaluationData(1), FSectionEvaluationData(2) }),
+		FMovieSceneSegment(FFloatRange(25.f,	30.f), { FSectionEvaluationData(1), FSectionEvaluationData(2), FSectionEvaluationData(4) }),
+		FMovieSceneSegment(FFloatRange(30.f,	Inf), { FSectionEvaluationData(4) })
+	};
+	AssertSegmentValues(this, Expected, Segments);
+	return true;
+}
+
+IMPLEMENT_SIMPLE_AUTOMATION_TEST(FMovieSceneCompilerTreeIteratorTest, "System.Engine.Sequencer.Compiler.Tree Iterator", EAutomationTestFlags::EditorContext | EAutomationTestFlags::EngineFilter)
+bool FMovieSceneCompilerTreeIteratorTest::RunTest(const FString& Parameters)
+{
+	using namespace Impl;
+	using namespace Lex;
+
+	//	Time				-inf				10							20				25					30							inf
+	//											[============= 0 ===========]
+	//																		[============== 1 ==================]
+	//											[========================== 2 ==================================]
+	//						[================== 3 ==========================]
+	//																						[================== 4 ==========================]
+	//----------------------------------------------------------------------------------------------------------------------------------------
+	//	Expected Impls		[ 		3			|			0,2,3			|		1,2		|		1,2,4		|			4				]
+
+	TMovieSceneEvaluationTree<int32> Tree;
+
+	Tree.Add(TRange<float>(10.f,	20.f), 								0);
+	Tree.Add(TRange<float>(20.f,	30.f), 								1);
+	Tree.Add(TRange<float>(10.f,	30.f), 								2);
+	Tree.Add(TRange<float>(Inf,	TRangeBound<float>::Exclusive(20.f)), 	3);
+	Tree.Add(TRange<float>(25.f,	Inf), 								4);
+
+	FEvaluationTreeIteratorResult Expected[] = {
+		MakeTuple( TRange<float>(Inf, 	TRangeBound<float>::Exclusive(10.f)),   TArray<int32>({ 3 })         ),
+		MakeTuple( TRange<float>(10.f,	20.f),                                  TArray<int32>({ 0, 2, 3 })   ),
+		MakeTuple( TRange<float>(20.f,	25.f),                                  TArray<int32>({ 1, 2 })      ),
+		MakeTuple( TRange<float>(25.f,	30.f),                                  TArray<int32>({ 1, 2, 4 })   ),
+		MakeTuple( TRange<float>(30.f,	Inf),                                   TArray<int32>({ 4 })         ),
+	};
+
+	AssertEvaluationTree(this, Tree, Expected);
+	return true;
+}
+
 
-TInlineValue<FMovieSceneSegmentCompilerRules> UMovieSceneSegmentCompilerTestTrack::GetTrackCompilerRules() const
+IMPLEMENT_SIMPLE_AUTOMATION_TEST(FMovieSceneCompilerTreeIteratorBoundsTest, "System.Engine.Sequencer.Compiler.Tree Iterator Bounds", EAutomationTestFlags::EditorContext | EAutomationTestFlags::EngineFilter)
+bool FMovieSceneCompilerTreeIteratorBoundsTest::RunTest(const FString& Parameters)
 {
-	struct FRules : FMovieSceneSegmentCompilerRules
+	using namespace Impl;
+	using namespace Lex;
+
+	//	Time				-inf				10							20				25					30							inf
+	//						[======== 0 ========][========= 1 ==============][============= 2 ==================][=========== 3 =============]
+
+	TMovieSceneEvaluationTree<int32> Tree;
+
+	Tree.Add(TRange<float>(Inf,		TRangeBound<float>::Exclusive(10.f)), 		0);
+	Tree.Add(TRange<float>(10.f,	TRangeBound<float>::Exclusive(20.f)), 		1);
+	Tree.Add(TRange<float>(20.f,	TRangeBound<float>::Exclusive(30.f)), 		2);
+	Tree.Add(TRange<float>(30.f,	Inf), 										3);
+
+
+	FEvaluationTreeIteratorResult Expected[] = {
+		MakeTuple( TRange<float>(Inf,	TRangeBound<float>::Exclusive(10.f)), 		TArray<int32>({ 0 }) ),
+		MakeTuple( TRange<float>(10.f,	TRangeBound<float>::Exclusive(20.f)), 		TArray<int32>({ 1 }) ),
+		MakeTuple( TRange<float>(20.f,	TRangeBound<float>::Exclusive(30.f)), 		TArray<int32>({ 2 }) ),
+		MakeTuple( TRange<float>(30.f,	Inf), 										TArray<int32>({ 3 }) ),
+	};
+
+	AssertEvaluationTree(this, Tree, Expected);
+
+	return true;
+}
+
+#endif // WITH_DEV_AUTOMATION_TESTS
+
+
+
+FMovieSceneEvalTemplatePtr UMovieSceneSegmentCompilerTestTrack::CreateTemplateForSection(const UMovieSceneSection& InSection) const
+{
+	return FTestMovieSceneEvalTemplate();
+}
+
+FMovieSceneTrackSegmentBlenderPtr UMovieSceneSegmentCompilerTestTrack::GetTrackSegmentBlender() const
+{
+	struct FSegmentBlender : FMovieSceneTrackSegmentBlender
 	{
 		bool bHighPass;
 		bool bEvaluateNearest;
 
-		FRules(bool bInHighPassFilter, bool bInEvaluateNearest)
+		FSegmentBlender(bool bInHighPassFilter, bool bInEvaluateNearest)
 			: bHighPass(bInHighPassFilter)
 			, bEvaluateNearest(bInEvaluateNearest)
-		{}
+		{
+			bCanFillEmptySpace = bEvaluateNearest;
+		}
 
-		virtual void BlendSegment(FMovieSceneSegment& Segment, const TArrayView<const FMovieSceneSectionData>& SourceData) const
+		virtual void Blend(FSegmentBlendData& BlendData) const
 		{
 			if (bHighPass)
 			{
-				MovieSceneSegmentCompiler::BlendSegmentHighPass(Segment, SourceData);
+				MovieSceneSegmentCompiler::ChooseLowestRowIndex(BlendData);
 			}
 
 			// Always sort by priority
-			Segment.Impls.Sort(
-				[&](const FSectionEvaluationData& A, const FSectionEvaluationData& B)
-				{
-					return SourceData[A.ImplIndex].Priority > SourceData[B.ImplIndex].Priority;
-				}
-			);
+			Algo::SortBy(BlendData, [](const FMovieSceneSectionData& In){ return In.Section->GetRowIndex(); });
 		}
+
 		virtual TOptional<FMovieSceneSegment> InsertEmptySpace(const TRange<float>& Range, const FMovieSceneSegment* PreviousSegment, const FMovieSceneSegment* NextSegment) const
 		{
 			return bEvaluateNearest ? MovieSceneSegmentCompiler::EvaluateNearestSegment(Range, PreviousSegment, NextSegment) : TOptional<FMovieSceneSegment>();
@@ -449,5 +643,5 @@ TInlineValue<FMovieSceneSegmentCompilerRules> UMovieSceneSegmentCompilerTestTrac
 	};
 
 	// Evaluate according to bEvalNearestSection preference
-	return FRules(bHighPassFilter, EvalOptions.bCanEvaluateNearestSection && EvalOptions.bEvalNearestSection);
+	return FSegmentBlender(bHighPassFilter, EvalOptions.bCanEvaluateNearestSection && EvalOptions.bEvalNearestSection);
 }
\ No newline at end of file
diff --git a/Engine/Source/Runtime/MovieScene/Private/Tests/MovieSceneSegmentCompilerTests.h b/Engine/Source/Runtime/MovieScene/Private/Tests/MovieSceneSegmentCompilerTests.h
--- a/Engine/Source/Runtime/MovieScene/Private/Tests/MovieSceneSegmentCompilerTests.h
+++ b/Engine/Source/Runtime/MovieScene/Private/Tests/MovieSceneSegmentCompilerTests.h
@@ -6,8 +6,17 @@
 #include "UObject/ObjectMacros.h"
 #include "MovieSceneTrack.h"
 #include "MovieSceneSection.h"
+#include "MovieSceneEvalTemplate.h"
 #include "MovieSceneSegmentCompilerTests.generated.h"
 
+USTRUCT()
+struct FTestMovieSceneEvalTemplate : public FMovieSceneEvalTemplate
+{
+	GENERATED_BODY()
+
+	virtual UScriptStruct& GetScriptStructImpl() const { return *StaticStruct(); }
+};
+
 UCLASS(MinimalAPI)
 class UMovieSceneSegmentCompilerTestTrack : public UMovieSceneTrack
 {
@@ -16,11 +25,12 @@ public:
 	GENERATED_BODY()
 
 	virtual const TArray<UMovieSceneSection*>& GetAllSections() const override { return SectionArray; }
-	virtual TInlineValue<FMovieSceneSegmentCompilerRules> GetTrackCompilerRules() const override;
+	virtual FMovieSceneEvalTemplatePtr CreateTemplateForSection(const UMovieSceneSection& InSection) const override;
+	virtual FMovieSceneTrackSegmentBlenderPtr GetTrackSegmentBlender() const override;
 
 	UPROPERTY()
 	bool bHighPassFilter;
-	
+
 	UPROPERTY()
 	TArray<UMovieSceneSection*> SectionArray;
 };
diff --git a/Engine/Source/Runtime/MovieScene/Private/Tracks/MovieSceneSubTrack.cpp b/Engine/Source/Runtime/MovieScene/Private/Tracks/MovieSceneSubTrack.cpp
new file mode 100644
--- /dev/null
+++ b/Engine/Source/Runtime/MovieScene/Private/Tracks/MovieSceneSubTrack.cpp
@@ -0,0 +1,233 @@
+// Copyright 1998-2017 Epic Games, Inc. All Rights Reserved.
+
+#include "Tracks/MovieSceneSubTrack.h"
+#include "Evaluation/MovieSceneEvaluationField.h"
+#include "Evaluation/MovieSceneSequenceTransform.h"
+#include "Evaluation/MovieSceneSequenceHierarchy.h"
+#include "MovieSceneSequence.h"
+#include "Sections/MovieSceneSubSection.h"
+#include "MovieScene.h"
+#include "Compilation/MovieSceneSegmentCompiler.h"
+
+
+#define LOCTEXT_NAMESPACE "MovieSceneSubTrack"
+
+
+/* UMovieSceneSubTrack interface
+ *****************************************************************************/
+UMovieSceneSubTrack::UMovieSceneSubTrack( const FObjectInitializer& ObjectInitializer )
+	: Super(ObjectInitializer)
+{
+#if WITH_EDITORONLY_DATA
+	TrackTint = FColor(180, 0, 40, 65);
+#endif
+}
+
+UMovieSceneSubSection* UMovieSceneSubTrack::AddSequence(UMovieSceneSequence* Sequence, float StartTime, float Duration, const bool& bInsertSequence)
+{
+	Modify();
+
+	// If inserting, push all existing sections out to make space for the new one
+	if (bInsertSequence)
+	{
+		// If there's a shot that starts at the same time as this new shot, push the shots forward to make space for this one
+		bool bPushShotsForward = false;
+		for (auto Section : Sections)
+		{
+			float SectionStartTime = Section->GetStartTime();
+			if (FMath::IsNearlyEqual(SectionStartTime, StartTime))
+			{
+				bPushShotsForward = true;
+				break;
+			}
+		}
+		if (bPushShotsForward)
+		{
+			for (auto Section : Sections)
+			{
+				float SectionStartTime = Section->GetStartTime();
+				if (SectionStartTime >= StartTime)
+				{
+					Section->SetStartTime(SectionStartTime + Duration);
+					Section->SetEndTime(Section->GetEndTime() + Duration);
+				}
+			}
+		}
+		// Otherwise, see if there's a shot after the start time and clamp the duration to that next shot
+		else
+		{
+			bool bFoundAfterShot = false;
+			float MinDiff = FLT_MAX;
+			for (auto Section : Sections)
+			{
+				float SectionStartTime = Section->GetStartTime();
+				if (SectionStartTime > StartTime)
+				{
+					bFoundAfterShot = true;
+					float Diff = SectionStartTime - StartTime;
+					MinDiff = FMath::Min(MinDiff, Diff);
+				}
+			}
+
+			if (MinDiff != FLT_MAX)
+			{
+				Duration = MinDiff;
+			}
+		}
+	}
+
+	UMovieSceneSubSection* NewSection = CastChecked<UMovieSceneSubSection>(CreateNewSection());
+	{
+		NewSection->SetSequence(Sequence);
+		NewSection->SetStartTime(StartTime);
+		NewSection->SetEndTime(StartTime + Duration);
+	}
+
+	Sections.Add(NewSection);
+
+	return NewSection;
+}
+
+UMovieSceneSubSection* UMovieSceneSubTrack::AddSequenceToRecord()
+{
+	UMovieScene* MovieScene = CastChecked<UMovieScene>(GetOuter());
+	TRange<float> PlaybackRange = MovieScene->GetPlaybackRange();
+
+	int32 MaxRowIndex = -1;
+	for (auto Section : Sections)
+	{
+		MaxRowIndex = FMath::Max(Section->GetRowIndex(), MaxRowIndex);
+	}
+
+	UMovieSceneSubSection* NewSection = CastChecked<UMovieSceneSubSection>(CreateNewSection());
+	{
+		NewSection->SetRowIndex(MaxRowIndex + 1);
+		NewSection->SetAsRecording(true);
+		NewSection->SetStartTime(PlaybackRange.GetLowerBoundValue());
+		NewSection->SetEndTime(PlaybackRange.GetUpperBoundValue());
+	}
+
+	Sections.Add(NewSection);
+
+	return NewSection;
+}
+
+bool UMovieSceneSubTrack::ContainsSequence(const UMovieSceneSequence& Sequence, bool Recursively) const
+{
+	for (const auto& Section : Sections)
+	{
+		const auto SubSection = CastChecked<UMovieSceneSubSection>(Section);
+
+		// is the section referencing the sequence?
+		const UMovieSceneSequence* SubSequence = SubSection->GetSequence();
+
+		if (SubSequence == nullptr)
+		{
+			continue;
+		}
+
+		if (SubSequence == &Sequence)
+		{
+			return true;
+		}
+
+		if (!Recursively)
+		{
+			continue;
+		}
+
+		// does the section have sub-tracks referencing the sequence?
+		const UMovieScene* SubMovieScene = SubSequence->GetMovieScene();
+
+		if (SubMovieScene == nullptr)
+		{
+			continue;
+		}
+
+		UMovieSceneSubTrack* SubSubTrack = SubMovieScene->FindMasterTrack<UMovieSceneSubTrack>();
+
+		if ((SubSubTrack != nullptr) && SubSubTrack->ContainsSequence(Sequence))
+		{
+			return true;
+		}
+	}
+
+	return false;
+}
+
+
+/* UMovieSceneTrack interface
+ *****************************************************************************/
+
+void UMovieSceneSubTrack::AddSection(UMovieSceneSection& Section)
+{
+	if (Section.IsA<UMovieSceneSubSection>())
+	{
+		Sections.Add(&Section);
+	}
+}
+
+
+UMovieSceneSection* UMovieSceneSubTrack::CreateNewSection()
+{
+	return NewObject<UMovieSceneSubSection>(this, NAME_None, RF_Transactional);
+}
+
+
+const TArray<UMovieSceneSection*>& UMovieSceneSubTrack::GetAllSections() const 
+{
+	return Sections;
+}
+
+
+TRange<float> UMovieSceneSubTrack::GetSectionBoundaries() const
+{
+	TArray<TRange<float>> Bounds;
+
+	for (int32 SectionIndex = 0; SectionIndex < Sections.Num(); ++SectionIndex)
+	{
+		Bounds.Add(Sections[SectionIndex]->GetRange());
+	}
+
+	return TRange<float>::Hull(Bounds);
+}
+
+
+bool UMovieSceneSubTrack::HasSection(const UMovieSceneSection& Section) const
+{
+	return Sections.Contains(&Section);
+}
+
+
+bool UMovieSceneSubTrack::IsEmpty() const
+{
+	return Sections.Num() == 0;
+}
+
+
+void UMovieSceneSubTrack::RemoveAllAnimationData()
+{
+	Sections.Empty();
+}
+
+
+void UMovieSceneSubTrack::RemoveSection(UMovieSceneSection& Section)
+{
+	Sections.Remove(&Section);
+}
+
+
+bool UMovieSceneSubTrack::SupportsMultipleRows() const
+{
+	return true;
+}
+
+#if WITH_EDITORONLY_DATA
+FText UMovieSceneSubTrack::GetDefaultDisplayName() const
+{
+	return LOCTEXT("TrackName", "Subscenes");
+}
+#endif
+
+
+#undef LOCTEXT_NAMESPACE
diff --git a/Engine/Source/Runtime/MovieScene/Public/Compilation/IMovieSceneTemplateGenerator.h b/Engine/Source/Runtime/MovieScene/Public/Compilation/IMovieSceneTemplateGenerator.h
--- a/Engine/Source/Runtime/MovieScene/Public/Compilation/IMovieSceneTemplateGenerator.h
+++ b/Engine/Source/Runtime/MovieScene/Public/Compilation/IMovieSceneTemplateGenerator.h
@@ -5,17 +5,17 @@
 #include "CoreMinimal.h"
 #include "MovieSceneSequenceID.h"
 #include "MovieSceneFwd.h"
-#include "MovieSceneSegment.h"
 #include "Containers/ArrayView.h"
-#include "MovieSceneSequenceID.h"
+#include "MovieSceneSequenceTransform.h"
 
 class UMovieSceneTrack;
 struct FMovieSceneEvaluationFieldSegmentPtr;
 struct FMovieSceneEvaluationTrack;
-struct FMovieSceneSequenceTransform;
 struct FMovieSceneSharedDataId;
 struct FMovieSceneSubSequenceData;
 
+enum class ESectionEvaluationFlags : uint8;
+
 /** Abstract base class used to generate evaluation templates */
 struct IMovieSceneTemplateGenerator
 {
@@ -27,48 +27,23 @@ struct IMovieSceneTemplateGenerator
 	 */
 	virtual void AddOwnedTrack(FMovieSceneEvaluationTrack&& InTrackTemplate, const UMovieSceneTrack& SourceTrack) = 0;
 
-	/**
-	 * Add a new track that is potentially shared between multiple tracks. Only one instance of SharedId can exist.
-	 *
-	 * @param InTrackTemplate			The track template to add
-	 * @param SharedId					The identifier of the shared track (for cross referencing between tracks)
-	 * @param SourceTrack				The originating track
-	 */
-	virtual void AddSharedTrack(FMovieSceneEvaluationTrack&& InTrackTemplate, FMovieSceneSharedDataId SharedId, const UMovieSceneTrack& SourceTrack) = 0;
 
-	/**
-	 * Add a legacy track to the template
-	 *
-	 * @param InTrackTemplate			The legacy track template
-	 * @param SourceTrack				The originating track
-	 */
-	virtual void AddLegacyTrack(FMovieSceneEvaluationTrack&& InTrackTemplate, const UMovieSceneTrack& SourceTrack) = 0;
+public:
 
-	/**
-	 * Add a set of remapped segments from a sub sequence to this template
-	 *
-	 * @param RootRange					The range in which to add the specified segments (at the root level)
-	 * @param SegmentPtrs				The remapped segment ptrs
-	 * @param Flags						Flags for these external segments
-	 */
-	virtual void AddExternalSegments(TRange<float> RootRange, TArrayView<const FMovieSceneEvaluationFieldSegmentPtr> SegmentPtrs, ESectionEvaluationFlags Flags) = 0;
+	DEPRECATED(4.19, "Shared tracks are no longer supported, please emit shared Execution tokens instead (FMovieSceneExecutionTokens::AddShared")
+	void AddSharedTrack(FMovieSceneEvaluationTrack&& InTrackTemplate, const FMovieSceneSharedDataId& SharedId, const UMovieSceneTrack& SourceTrack){}
 
-	/**
-	 * Get a sequence's transform from its ID
-	 *
-	 * @param SequenceID				The ID of the sequence to get a transform for
-	 * @return The sequence's transform
-	 */
-	virtual FMovieSceneSequenceTransform GetSequenceTransform(FMovieSceneSequenceIDRef InSequenceID) const = 0;
+	DEPRECATED(4.19, "Support for legacy tracks has now been removed")
+	void AddLegacyTrack(FMovieSceneEvaluationTrack&& InTrackTemplate, const UMovieSceneTrack& SourceTrack){}
 
-	/**
-	 * Add the specified sub sequence data to this generator
-	 *
-	 * @param SequenceData				Data pertaining to the sequence to add
-	 * @param ParentID					ID of the parent sequence
-	 * @param SequenceID				Specific ID to use to identify this sub sequence in the template. Must not already exist in the template.
-	 */
-	virtual void AddSubSequence(FMovieSceneSubSequenceData SequenceData, FMovieSceneSequenceIDRef ParentID, FMovieSceneSequenceID SequenceID) = 0;
+	DEPRECATED(4.19, "External segments are now handled internally")
+	void AddExternalSegments(TRange<float> RootRange, TArrayView<const FMovieSceneEvaluationFieldSegmentPtr> SegmentPtrs, ESectionEvaluationFlags Flags){}
+
+	DEPRECATED(4.19, "Sub sequence templates are no longer generated at compile time within template generators")
+	FMovieSceneSequenceTransform GetSequenceTransform(FMovieSceneSequenceIDRef InSequenceID) const { return FMovieSceneSequenceTransform(); }
+
+	DEPRECATED(4.19, "Sub sequences are now handled internally")
+	void AddSubSequence(const FMovieSceneSubSequenceData& SequenceData, FMovieSceneSequenceIDRef ParentID, FMovieSceneSequenceID SequenceID) {}
 
 protected:
 	virtual ~IMovieSceneTemplateGenerator() { }
diff --git a/Engine/Source/Runtime/MovieScene/Public/Compilation/MovieSceneCompiler.h b/Engine/Source/Runtime/MovieScene/Public/Compilation/MovieSceneCompiler.h
new file mode 100644
--- /dev/null
+++ b/Engine/Source/Runtime/MovieScene/Public/Compilation/MovieSceneCompiler.h
@@ -0,0 +1,106 @@
+
+#pragma once
+
+#include "Range.h"
+#include "Optional.h"
+#include "Array.h"
+
+#include "MovieSceneSequenceID.h"
+#include "MovieSceneSequenceTransform.h"
+#include "MovieSceneEvaluationField.h"
+#include "MovieSceneEvaluationTree.h"
+#include "Evaluation/MovieSceneEvaluationTemplateInstance.h"
+
+struct FGatherParameters;
+struct FCompileOnTheFlyData;
+struct FMovieSceneEvaluationTrack;
+struct FMovieSceneGatheredCompilerData;
+struct IMovieSceneSequenceTemplateStore;
+
+class UMovieSceneSequence;
+class UMovieSceneSubSection;
+
+/**
+ * Structure that represents the result for a single compiled range in a sequence
+ */
+struct FCompiledGroupResult
+{
+	/** Construction to an empty result for the specified range */
+	FCompiledGroupResult(TRange<float> InRange) : Range(InRange) {}
+
+	/** The range within the sequence that the goup and meta-data apply to */
+	TRange<float> Range;
+
+	/** Evaluation group specifying all the entities that occur in the sequence, sorted and redy for evaluation */
+	FMovieSceneEvaluationGroup Group;
+
+	/** Meta-data pertaining to this current time range */
+	FMovieSceneEvaluationMetaData MetaData;
+};
+
+class FMovieSceneCompiler
+{
+public:
+	/**
+	 * Fully (and recursively) compile the specified sequence
+	 *
+	 * @param InCompileSequence 		The Sequence to compile
+	 * @param InTemplateStore 			A template store that contains (or knows how to fetch) evaluation templates for a given sequence
+	 */
+	MOVIESCENE_API static void Compile(UMovieSceneSequence& InCompileSequence, IMovieSceneSequenceTemplateStore& InTemplateStore);
+
+	/**
+	 * Compile the evaluation data that is relevant for the specified time
+	 *
+	 * @param InGlobalTime 				The time at which to compile the evaluation data, in the time-space of InCompileSequence
+	 * @param InCompileSequence 		The Sequence to compile
+	 * @param InTemplateStore 			A template store that contains (or knows how to fetch) evaluation templates for a given sequence
+	 * @return The compiled data, ready for insertion into an FMovieSceneEvaluationField
+	 */
+	MOVIESCENE_API static TOptional<FCompiledGroupResult> CompileTime(float InGlobalTime, UMovieSceneSequence& InCompileSequence, IMovieSceneSequenceTemplateStore& InTemplateStore);
+
+	/**
+	 * Compile the hierarchy data for the specified sequence into a specific hierarchy
+	 *
+	 * @param InRootSequence			The sequence to generate hierarchy information from
+	 * @param OutHierarchy				The hierarchy to populate
+	 * @param RootSequenceID			The sequence identifier of InRootSequence within the specified OutHierarchy
+	 * @param MaxDepth					(Optional) Maximum depth to generate, or -1 to generate all hierarchy information
+	 */
+	MOVIESCENE_API static void CompileHierarchy(const UMovieSceneSequence& InRootSequence, FMovieSceneSequenceHierarchy& OutHierarchy, FMovieSceneSequenceID RootSequenceID = MovieSceneSequenceID::Root, int32 MaxDepth = -1);
+
+private:
+
+	/*~ Helper functions */
+	static void CompileHierarchy(const UMovieSceneSequence& InSequence, FMovieSceneSequenceHierarchy& OutHierarchy, FMovieSceneRootOverridePath& Path, int32 MaxDepth);
+
+	static void AddPtrsToGroup(FMovieSceneEvaluationGroup& Group, TArray<FMovieSceneEvaluationFieldSegmentPtr>& InitPtrs, TArray<FMovieSceneEvaluationFieldSegmentPtr>& EvalPtrs);
+
+	static bool SortPredicate(const FCompileOnTheFlyData& A, const FCompileOnTheFlyData& B);
+
+	static void PopulateEvaluationGroup(FCompiledGroupResult& OutResult, const TArray<FCompileOnTheFlyData>& SortedCompileData);
+
+	static void PopulateMetaData(FCompiledGroupResult& OutResult, const FMovieSceneSequenceHierarchy& RootHierarchy, const TArray<FCompileOnTheFlyData>& SortedCompileData, TMovieSceneEvaluationTreeDataIterator<FMovieSceneSequenceID> SubSequences);
+
+private:
+
+	static void GatherCompileOnTheFlyData(UMovieSceneSequence& InSequence, const FGatherParameters& Params, FMovieSceneGatheredCompilerData& OutData);
+
+	static void GatherCompileDataForSubSection(const UMovieSceneSubSection& SubSection, const FGuid& InObjectBindingId, const FGatherParameters& Params, FMovieSceneGatheredCompilerData& OutData);
+
+	static void GatherCompileDataForTrack(FMovieSceneEvaluationTrack& Track, FMovieSceneTrackIdentifier TrackID, const FGatherParameters& Params, FMovieSceneGatheredCompilerData& OutData);
+
+private:
+
+	/**
+	 * Find or create sub sequence data for the specified parameters in a hierarchy
+	 *
+	 * @param InnerSequenceID			The sequence ID to find or create sub sequence data for
+	 * @param ParentSequenceID			The sequence's parent ID
+	 * @param SubSection				The sub section that's including the sequence
+	 * @param InObjectBindingId			The object binding ID that the sub section belongs to (if any)
+	 * @param Hierarchy					The hierarchy within which to lookup or insert the data
+	 * @return The sub sequence data, or nullptr
+	 */
+	static const FMovieSceneSubSequenceData* GetOrCreateSubSequenceData(FMovieSceneSequenceID InnerSequenceID, FMovieSceneSequenceID ParentSequenceID, const UMovieSceneSubSection& SubSection, const FGuid& InObjectBindingId, FMovieSceneSequenceHierarchy& Hierarchy);
+};
\ No newline at end of file
diff --git a/Engine/Source/Runtime/MovieScene/Public/Compilation/MovieSceneCompilerRules.h b/Engine/Source/Runtime/MovieScene/Public/Compilation/MovieSceneCompilerRules.h
--- a/Engine/Source/Runtime/MovieScene/Public/Compilation/MovieSceneCompilerRules.h
+++ b/Engine/Source/Runtime/MovieScene/Public/Compilation/MovieSceneCompilerRules.h
@@ -39,88 +39,136 @@ namespace MovieSceneSegmentCompiler
 		return TOptional<FMovieSceneSegment>();
 	}
 
-	static void BlendSegmentHighPass(FMovieSceneSegment& Segment, const TArrayView<const FMovieSceneSectionData>& SourceData)
+	static bool AlwaysEvaluateSection(const FMovieSceneSectionData& InSectionData)
 	{
-		if (!Segment.Impls.Num())
+		return InSectionData.Section->GetBlendType().IsValid() || EnumHasAnyFlags(InSectionData.Flags, ESectionEvaluationFlags::PreRoll | ESectionEvaluationFlags::PostRoll);
+	}
+
+	static void FilterOutUnderlappingSections(FSegmentBlendData& BlendData)
+	{
+		if (!BlendData.Num())
 		{
 			return;
 		}
 
-		int32 HighestPriority = TNumericLimits<int32>::Lowest();
-		for (const FSectionEvaluationData& Impl : Segment.Impls)
+		int32 HighestOverlap = TNumericLimits<int32>::Lowest();
+		for (const FMovieSceneSectionData& SectionData : BlendData)
 		{
-			const FMovieSceneSectionData& SectionData = SourceData[Impl.ImplIndex];
-			if (!SectionData.BlendType.IsValid() &&
-				!SectionData.EvalData.IsPreRoll() &&
-				!SectionData.EvalData.IsPostRoll())
+			if (!AlwaysEvaluateSection(SectionData))
 			{
-				HighestPriority = FMath::Max(HighestPriority, SectionData.Priority);
+				HighestOverlap = FMath::Max(HighestOverlap, SectionData.Section->GetOverlapPriority());
 			}
 		}
 
 		// Remove anything that's not the highest priority, (excluding pre/postroll sections)
-		for (int32 RemoveAtIndex = Segment.Impls.Num() - 1; RemoveAtIndex >= 0; --RemoveAtIndex)
+		for (int32 RemoveAtIndex = BlendData.Num() - 1; RemoveAtIndex >= 0; --RemoveAtIndex)
 		{
-			const FMovieSceneSectionData& SectionData = SourceData[Segment.Impls[RemoveAtIndex].ImplIndex];
-			if (SectionData.Priority != HighestPriority &&
-				!SectionData.BlendType.IsValid() &&
-				!SectionData.EvalData.IsPreRoll() && 
-				!SectionData.EvalData.IsPostRoll())
+			const FMovieSceneSectionData& SectionData = BlendData[RemoveAtIndex];
+			if (SectionData.Section->GetOverlapPriority() != HighestOverlap && !AlwaysEvaluateSection(SectionData))
 			{
-				Segment.Impls.RemoveAt(RemoveAtIndex, 1, false);
+				BlendData.RemoveAt(RemoveAtIndex, 1, false);
 			}
 		}
 	}
 
-	// Reduces the evaluated sections to only the section that resides last in the source data. Legacy behaviour from various track instances.
-	static void BlendSegmentLegacySectionOrder(FMovieSceneSegment& Segment, const TArrayView<const FMovieSceneSectionData>& SourceData)
+	static void ChooseLowestRowIndex(FSegmentBlendData& BlendData)
 	{
-		if (Segment.Impls.Num() <= 1)
+		if (!BlendData.Num())
 		{
 			return;
 		}
 
-		Segment.Impls.Sort(
-			[&](const FSectionEvaluationData& A, const FSectionEvaluationData& B)
+		int32 LowestRowIndex = TNumericLimits<int32>::Max();
+		for (const FMovieSceneSectionData& SectionData : BlendData)
+		{
+			if (!AlwaysEvaluateSection(SectionData))
+			{
+				LowestRowIndex = FMath::Min(LowestRowIndex, SectionData.Section->GetRowIndex());
+			}
+		}
+
+		// Remove anything that's not the highest priority, (excluding pre/postroll sections)
+		for (int32 RemoveAtIndex = BlendData.Num() - 1; RemoveAtIndex >= 0; --RemoveAtIndex)
+		{
+			const FMovieSceneSectionData& SectionData = BlendData[RemoveAtIndex];
+			if (SectionData.Section->GetRowIndex() > LowestRowIndex && !AlwaysEvaluateSection(SectionData))
 			{
-				return A.ImplIndex > B.ImplIndex;
+				BlendData.RemoveAt(RemoveAtIndex, 1, false);
 			}
-		);
-		Segment.Impls.SetNum(1, false);
+		}
+	}
+
+	// Reduces the evaluated sections to only the section that resides last in the source data. Legacy behaviour from various track instances.
+	static void BlendSegmentLegacySectionOrder(FSegmentBlendData& BlendData)
+	{
+		if (BlendData.Num() > 1)
+		{
+			Algo::SortBy(BlendData, &FMovieSceneSectionData::TemplateIndex);
+			BlendData.RemoveAt(1, BlendData.Num() - 1, false);
+		}
 	}
 }
 
-class FMovieSceneAdditiveCameraRules : public FMovieSceneSegmentCompilerRules
+/** Default track row segment blender for all tracks */
+struct FDefaultTrackRowSegmentBlender : FMovieSceneTrackRowSegmentBlender
+{
+	virtual void Blend(FSegmentBlendData& BlendData) const override
+	{
+		// By default we only evaluate the section with the highest Z-Order if they overlap on the same row
+		MovieSceneSegmentCompiler::FilterOutUnderlappingSections(BlendData);
+	}
+};
+
+/** Track segment blender that evaluates the nearest segment in empty space */
+struct FEvaluateNearestSegmentBlender : FMovieSceneTrackSegmentBlender
 {
-public:
-	FMovieSceneAdditiveCameraRules(const UMovieSceneTrack* InTrack)
-		: Sections(InTrack->GetAllSections())
+	FEvaluateNearestSegmentBlender()
 	{
+		bCanFillEmptySpace = true;
 	}
 
-	virtual void BlendSegment(FMovieSceneSegment& Segment, const TArrayView<const FMovieSceneSectionData>& SourceData) const
+	virtual TOptional<FMovieSceneSegment> InsertEmptySpace(const TRange<float>& Range, const FMovieSceneSegment* PreviousSegment, const FMovieSceneSegment* NextSegment) const
+	{
+		return MovieSceneSegmentCompiler::EvaluateNearestSegment(Range, PreviousSegment, NextSegment);
+	}
+};
+
+struct FMovieSceneAdditiveCameraTrackBlender : public FMovieSceneTrackSegmentBlender
+{
+	virtual void Blend(FSegmentBlendData& BlendData) const override
 	{
 		// sort by start time to match application order of player camera
-		Segment.Impls.Sort(
-			[&](FSectionEvaluationData A, FSectionEvaluationData B){
-				UMovieSceneSection* SectionA = Sections[SourceData[A.ImplIndex].EvalData.ImplIndex];
-				UMovieSceneSection* SectionB = Sections[SourceData[B.ImplIndex].EvalData.ImplIndex];
-
-				if (SectionA->IsInfinite())
-				{
-					return true;
-				}
-				else if (SectionB->IsInfinite())
-				{
-					return false;
-				}
-				else
-				{
-					return SectionA->GetStartTime() < SectionB->GetStartTime();
-				}
-			}
-		);
+		BlendData.Sort(SortByStartTime);
 	}
 
-	const TArray<UMovieSceneSection*>& Sections;
+private:
+
+	MOVIESCENE_API static bool SortByStartTime(const FMovieSceneSectionData& A, const FMovieSceneSectionData& B);
 };
+
+
+PRAGMA_DISABLE_DEPRECATION_WARNINGS
+
+/** Track row segment blender that acts as a proxy to an instance of the legacy FMovieSceneSegmentCompilerRules */
+template<typename BaseType>
+struct TLegacyTrackRowSegmentBlender : BaseType
+{
+	TLegacyTrackRowSegmentBlender(TInlineValue<FMovieSceneSegmentCompilerRules>&& InRules)
+		: LegacyRules(MoveTemp(InRules))
+	{}
+
+	~TLegacyTrackRowSegmentBlender()
+	{}
+
+	TLegacyTrackRowSegmentBlender(TLegacyTrackRowSegmentBlender&&) = default;
+
+	virtual void Blend(FSegmentBlendData& BlendData) const override
+	{
+		return LegacyRules->Blend(BlendData);
+	}
+
+private:
+	TInlineValue<FMovieSceneSegmentCompilerRules> LegacyRules;
+};
+
+PRAGMA_ENABLE_DEPRECATION_WARNINGS
\ No newline at end of file
diff --git a/Engine/Source/Runtime/MovieScene/Public/Compilation/MovieSceneEvaluationTemplateGenerator.h b/Engine/Source/Runtime/MovieScene/Public/Compilation/MovieSceneEvaluationTemplateGenerator.h
--- a/Engine/Source/Runtime/MovieScene/Public/Compilation/MovieSceneEvaluationTemplateGenerator.h
+++ b/Engine/Source/Runtime/MovieScene/Public/Compilation/MovieSceneEvaluationTemplateGenerator.h
@@ -16,9 +16,11 @@
 #include "Evaluation/MovieSceneSequenceHierarchy.h"
 #include "Evaluation/MovieSceneEvaluationTemplate.h"
 #include "Compilation/MovieSceneSegmentCompiler.h"
+#include "MovieSceneEvaluationTree.h"
 
 class IMovieSceneModule;
 class UMovieSceneSequence;
+class UMovieSceneSubTrack;
 
 /**
  * Class responsible for generating up-to-date evaluation template data
@@ -30,14 +32,12 @@ public:
 	/**
 	 * Construction from a sequence and a destination template
 	 */
-	FMovieSceneEvaluationTemplateGenerator(UMovieSceneSequence& InSequence, FMovieSceneEvaluationTemplate& OutTemplate, FMovieSceneSequenceTemplateStore& InStore);
+	FMovieSceneEvaluationTemplateGenerator(UMovieSceneSequence& InSequence, FMovieSceneEvaluationTemplate& OutTemplate);
 
 	/**
-	 * Generate the template with the specified parameters
-	 *
-	 * @param InParams 			Compilation parameters
+	 * Generate the template with the specified sequence
 	 */
-	void Generate(FMovieSceneTrackCompilationParams InParams);
+	void Generate();
 
 private:
 
@@ -49,80 +49,27 @@ private:
 	 */
 	void ProcessTrack(const UMovieSceneTrack& InTrack, const FGuid& ObjectId = FGuid());
 
-
-	/**
-	 * Remove any references to old track templates that are no longer required to exist
-	 */
-	void RemoveOldTrackReferences();
-
-	/**
-	 * Update the evaluation field with the specified segments and segment ptrs
-	 *
-	 * @param Segments 			The segments to populate the evaluation field with
-	 * @param Impls 			Array of implementation pointers that each segment refers to
-	 */
-	void UpdateEvaluationField(const TArray<FMovieSceneSegment>& Segments, const TArray<FMovieSceneEvaluationFieldSegmentPtr>& Impls, const TMap<FMovieSceneSequenceID, FMovieSceneEvaluationTemplate*>& Templates);
-
 	/**
-	 * Initialize segment meta-data for the specified group
+	 * Process a sub track
 	 *
-	 * @param MetaData			Meta-data container for the group
-	 * @param Group				The evaluation group to determine meta data for
-	 * @param Templates			Template store for locating tracks
+	 * @param SubTrack 			The track to process
+	 * @param ObjectId 			Object binding ID of the track
 	 */
-	void InitializeMetaData(FMovieSceneEvaluationMetaData& MetaData, FMovieSceneEvaluationGroup& Group, const TMap<FMovieSceneSequenceID, FMovieSceneEvaluationTemplate*>& Templates);
+	void ProcessSubTrack(const UMovieSceneSubTrack& SubTrack, const FGuid& ObjectId = FGuid());
 
 private:
 
 	/** IMovieSceneTemplateGenerator overrides */
-	virtual void AddLegacyTrack(FMovieSceneEvaluationTrack&& InTrackTemplate, const UMovieSceneTrack& SourceTrack) override;
 	virtual void AddOwnedTrack(FMovieSceneEvaluationTrack&& InTrackTemplate, const UMovieSceneTrack& SourceTrack) override;
-	virtual void AddSharedTrack(FMovieSceneEvaluationTrack&& InTrackTemplate, FMovieSceneSharedDataId SharedId, const UMovieSceneTrack& SourceTrack) override;
-	virtual void AddExternalSegments(TRange<float> RootRange, TArrayView<const FMovieSceneEvaluationFieldSegmentPtr> SegmentPtrs, ESectionEvaluationFlags Flags) override;
-	virtual FMovieSceneSequenceTransform GetSequenceTransform(FMovieSceneSequenceIDRef InSequenceID) const override;
-	virtual void AddSubSequence(FMovieSceneSubSequenceData SequenceData, FMovieSceneSequenceIDRef ParentID, FMovieSceneSequenceID SequenceID) override;
-
-private:
-
-	/**
-	 * Sort predicate used to sort evaluation tracks in a given group
-	 */
-	bool SortPredicate(const FMovieSceneEvaluationFieldTrackPtr& InPtrA, const FMovieSceneEvaluationFieldTrackPtr& InPtrB, const TMap<FMovieSceneSequenceID, FMovieSceneEvaluationTemplate*>& Templates, IMovieSceneModule& MovieSceneModule);
-
-	/**
-	 * Attempt to find a track from its evaluation field pointer
-	 *
-	 * @param InPtr 			Ptr to the track to lookup
-	 * @return The track, or nullptr
-	 */
-	const FMovieSceneEvaluationTrack* LookupTrack(const FMovieSceneEvaluationFieldTrackPtr& InPtr, const TMap<FMovieSceneSequenceID, FMovieSceneEvaluationTemplate*>& Templates);
 
 private:
 
 	/** The source sequence we are generating a template for */
 	UMovieSceneSequence& SourceSequence;
+
 	/** The destination template we're generating into */
 	FMovieSceneEvaluationTemplate& Template;
 
-	/** Transient compiler args that are passed to each track that is compiled */
-	FMovieSceneTrackCompilerArgs TransientArgs;
-
 	/** Temporary set of signatures we've come across in the sequence we're compiling */
 	TSet<FGuid> CompiledSignatures;
-
-	/** Segment data used for evaluation field generation */
-	TArray<FMovieSceneSectionData> SegmentData;
-	/** Look up table whose indices match implementation indices in the above segments */
-	TArray<FMovieSceneEvaluationFieldSegmentPtr> TrackLUT;
-
-	/** Map of external segments added from sub sequences */
-	struct FExternalSegment
-	{
-		int32 Index;
-		ESectionEvaluationFlags Flags;
-	};
-	TMultiMap<FMovieSceneEvaluationFieldSegmentPtr, FExternalSegment> ExternalSegmentLookup;
-
-	/** Set of shared track keys that have been added (used to ensure we only add a shared track ID once) */
-	TSet<FSharedPersistentDataKey> AddedSharedTracks;
 };
diff --git a/Engine/Source/Runtime/MovieScene/Public/Compilation/MovieSceneSegmentCompiler.h b/Engine/Source/Runtime/MovieScene/Public/Compilation/MovieSceneSegmentCompiler.h
--- a/Engine/Source/Runtime/MovieScene/Public/Compilation/MovieSceneSegmentCompiler.h
+++ b/Engine/Source/Runtime/MovieScene/Public/Compilation/MovieSceneSegmentCompiler.h
@@ -2,53 +2,109 @@
 
 #pragma once
 
-#include "CoreMinimal.h"
+#include "CoreTypes.h"
+#include "CoreFwd.h"
+
+#include "Array.h"
+#include "InlineValue.h"
 #include "Evaluation/MovieSceneSegment.h"
-#include "MovieSceneSection.h"
-#include "MovieSceneBlendType.h"
 
-/** Data structure supplied to the segment compiler that represents a section range, evaluation data (including the section's index), blend type and a priority */
+class UMovieSceneSection;
+
+/**
+ * Data structure supplied to segment blenders that includes information about the section
+ */
 struct FMovieSceneSectionData
 {
-	MOVIESCENE_API FMovieSceneSectionData(const TRange<float>& InBounds, FSectionEvaluationData InEvalData, FOptionalMovieSceneBlendType InBlendType = FOptionalMovieSceneBlendType(), int32 InPriority = 0);
-
-	/** The time range in which this section is considered active */
-	TRange<float> Bounds;
-	/** Evaluation data with which the section is to be evaluated */
-	FSectionEvaluationData EvalData;
-	/** Optional blend type for the section */
-	FOptionalMovieSceneBlendType BlendType;
-	/** Priority for the data (ie OverlapPriority where this represents sections in a row, or row index in the case of rows in a track) */
-	int32 Priority;
+	/** Constructor */
+	FMovieSceneSectionData(const UMovieSceneSection* InSection, int32 InTemplateIndex, ESectionEvaluationFlags InFlags)
+		: Section(InSection), TemplateIndex(InTemplateIndex), Flags(InFlags)
+	{}
+
+	/** The section that is to be evaluated */
+	const UMovieSceneSection* Section;
+	/** The index of the template within the evaluation track */
+	int32 TemplateIndex;
+	/** Evaluation flags for the section */
+	ESectionEvaluationFlags Flags;
 };
 
-/** Structure that defines how the segment compiler should combine overlapping sections, insert empty space, sort evaluation order, and other arbitrary processing. */
-struct FMovieSceneSegmentCompilerRules
+/**
+ * Container supplied to segment blenders that includes information about all sections to be evaluated.
+ * Blenders are free to reorder and filter the elements as they desire.
+ * Implemented as a named type to afford future changes without changing call-sites signatures.
+ */
+struct FSegmentBlendData : TArray<FMovieSceneSectionData, TInlineAllocator<4>>
 {
-	/** Default constructor */
-	FMovieSceneSegmentCompilerRules()
+	/**
+	 * Add all the section data to the specified segment
+	 *
+	 * @param Segment The segment to populate with this data
+	 */
+	void AddToSegment(FMovieSceneSegment& Segment) const
 	{
-		bAllowEmptySegments = false;
+		for (const FMovieSceneSectionData& SectionData : *this)
+		{
+			Segment.Impls.Add(FSectionEvaluationData(SectionData.TemplateIndex, SectionData.Flags));
+		}
 	}
+};
 
-	virtual ~FMovieSceneSegmentCompilerRules() {}
+//~ trait to re-enable GetData on FSegmentBlendData
+template<> struct TIsContiguousContainer<FSegmentBlendData>
+{
+	enum { Value = true };
+};
+
+/**
+ * Structure that defines how to combine and sort overlapping sections on the same row.
+ */
+struct FMovieSceneTrackRowSegmentBlender
+{
+	virtual ~FMovieSceneTrackRowSegmentBlender() {}
 
 	/**
-	 * Public entry point to process an array of compiled segments using the specified SourceData
+	 * Blend the specified data by performing some specific processing such as sorting or filtering 
 	 *
-	 * @param Segments				(In/Out) Mutable array of compiled segments to be processed
-	 * @param SourceData			The source section data that is being compiled.
+	 * @param BlendData				Blend data to be modified to suit the implementation's will
 	 */
-	MOVIESCENE_API void ProcessSegments(TArray<FMovieSceneSegment>& Segments, const TArrayView<const FMovieSceneSectionData>& SourceData) const;
+	virtual void Blend(FSegmentBlendData& BlendData) const
+	{
+	}
+};
+
+/**
+ * Structure that defines how to combine and sort sections in the same track row after processing on a per-row basis.
+ */
+struct FMovieSceneTrackSegmentBlender
+{
+	/** Default constructor */
+	FMovieSceneTrackSegmentBlender()
+	{
+		bCanFillEmptySpace = false;
+		bAllowEmptySegments = false;
+	}
+
+	virtual ~FMovieSceneTrackSegmentBlender() {}
 
 	/**
-	 * Check whether the resulting segments be empty
+	 * Check whether the resulting segments be empty. When true, Empty segments will be allowed to be added to the resulting evaluation field.
+	 * @note: use when evaluation of a track should always be performed, even if there are no sections at the current time
 	 */
 	bool AllowEmptySegments() const
 	{
 		return bAllowEmptySegments;
 	}
 
+	/**
+	 * Check whether these rules can fill empty space between segments (such as to evaluate the nearest section)
+	 * @note: use in combination with AllowEmptySegments when evaluation of a track should always be performed, even if there are no sections at the current time
+	 */
+	bool CanFillEmptySpace() const
+	{
+		return bCanFillEmptySpace;
+	}
+
 	/**
 	 * Implementation function to insert empty space between two other segments or at the start/end
 	 *
@@ -64,136 +120,33 @@ struct FMovieSceneSegmentCompilerRules
 	}
 
 	/**
-	 * Blend the specified segment by performing some specific processing such as sorting or filtering 
-	 *
-	 * @param Segment				The segment to process. Segment.Impls[].ImplIndex relates to indices with SourceData
-	 * @param SourceData			The source section data that is being compiled.
-	 */
-	virtual void BlendSegment(FMovieSceneSegment& Segment, const TArrayView<const FMovieSceneSectionData>& SourceData) const
-	{
-	}
-
-	/**
-	 * Called after all segments have been calculated and blended for any additional processing
+	 * Blend the specified data by performing some specific processing such as sorting or filtering 
 	 *
-	 * @param Segments				The segments to process. Segments[].Impls[].ImplIndex relates to indices with SourceData
-	 * @param SourceData			The source section data that is being compiled.
+	 * @param BlendData				Blend data to be modified to suit the implementation's will
 	 */
-	virtual void PostProcessSegments(TArray<FMovieSceneSegment>& Segments, const TArrayView<const FMovieSceneSectionData>& SourceData) const
+	virtual void Blend(FSegmentBlendData& BlendData) const
 	{
 	}
 
 protected:
 
-	/** Private implementation function for inserting a segment at the specified index */
-	bool InsertSegment(TArray<FMovieSceneSegment>& Segments, int32 Index, const TArrayView<const FMovieSceneSectionData>& SourceData) const;
-
-	/** Whether we allow empty segments or not */
+	/** Whether we allow empty segments to be evaluated or not */
 	bool bAllowEmptySegments;
-};
 
-/** Enumeration specifying how impl indices should be specified in resulting FMovieSceneSegments */
-enum class EMovieSceneSegmentIndexSpace : uint8
-{
-	/** FSectionEvaluationData::ImplIndex should point to indices within the original source data */
-	SourceDataIndex,
-	/** FSectionEvaluationData::ImplIndex should match those specified within the original source data */
-	ActualImplIndex,
+	/** Whether we can fill empty space (false signifies InsertEmptySpace will never be called) */
+	bool bCanFillEmptySpace;
 };
 
-/** Movie scene segment compiler that takes an unordered, arbitrary array of section data, and produces an ordered array of segments that define which implementations occur at which time ranges */
-struct FMovieSceneSegmentCompiler
-{
-	/**
-	 * Compiles the specified source data into an array of segments
-	 *
-	 * @param SourceData 			Unordered data to compile. May contain duplicates or overlapping time ranges.
-	 * @param Rules 				(Optional) Compiler rulse structure to compile with
-	 * @param IndexSpace 			(Optional) How the resulting segments should relate to the source data (either indices into the data, or actual indices specified from the data)
-	 *
-	 * @return An ordered array of segments that define unique combinations of overlapping source data. Segments are guaranteed to not overlap, may be contiguous and may contain gaps depending on compiler rules
-	 */
-	MOVIESCENE_API TArray<FMovieSceneSegment> Compile(TArrayView<const FMovieSceneSectionData> Data, const FMovieSceneSegmentCompilerRules* Rules = nullptr, EMovieSceneSegmentIndexSpace IndexSpace = EMovieSceneSegmentIndexSpace::ActualImplIndex);
-
-private:
-
-	/**
-	 * Add any EvalData currently considered as overlapping to their relevant segments
-	 */
-	void CloseCompletedSegments();
-
-	/**
-	 * Find and index into OverlappingSections where the supplied eval data is considered to match
-	 *
-	 * @param InEvalData 		Evaluation data to find an index for
-	 * @return A valid index into OverlappingSections, or INDEX_NONE
-	 */
-	int32 FindOverlappingIndex(FSectionEvaluationData InEvalData) const;
 
-	/** Private structure that specifies where a particular evaluation starts or ends */
-	struct FBound
-	{
-		FBound(FSectionEvaluationData InEvalData, TRangeBound<float> InBound) : EvalData(InEvalData), Bound(InBound) {}
-
-		FSectionEvaluationData EvalData;
-		TRangeBound<float> Bound;
-	};
-
-	/** Current read indices for the lower and upper bounds */
-	int32 LowerReadIndex, UpperReadIndex;
-	/** Arrays of lower and upper bounds sorted ascending and descending respectively. */
-	TArray<FBound> LowerBounds, UpperBounds;
-	/** Ordered array of compiled segments */
-	TArray<FMovieSceneSegment> CompiledSegments;
-	/** Section evaluation data for sections that are being evaluated at the current time of processing */
-	TArray<FSectionEvaluationData, TInlineAllocator<16>> OverlappingSections;
-	/** Reference counts for any of the entries in OverlappingSections that map to the same object (necessary for when the same eval data is specified in multiple source data elements) */
-	TArray<uint32, TInlineAllocator<16>> OverlappingRefCounts;
-	/** The index space that we are generating for */
-	EMovieSceneSegmentIndexSpace IndexSpace;
-	/** The source section data */
-	TArrayView<const FMovieSceneSectionData> SourceData;
-};
+typedef TInlineValue<FMovieSceneTrackRowSegmentBlender, 16> FMovieSceneTrackRowSegmentBlenderPtr;
+typedef TInlineValue<FMovieSceneTrackSegmentBlender, 16> FMovieSceneTrackSegmentBlenderPtr;
 
-/** Resulting structure from FMovieSceneTrackCompiler::Compile */
-struct FMovieSceneTrackEvaluationField
-{
-	/** The compiled segments, ordered by time range, with no overlaps. May contain empty space */
-	TArray<FMovieSceneSegment> Segments;
-};
 
-/** Structure used for compiling a UMovieSceneTrack using specific rules for compiling its rows, and combining such rows */
-struct FMovieSceneTrackCompiler
+/** Deprecated type support */
+struct FMovieSceneSegmentCompilerRules : FMovieSceneTrackSegmentBlender
 {
-	/** A row in a track */
-	struct FRow
-	{
-		/** All the sections contained in the row */
-		TArray<FMovieSceneSectionData, TInlineAllocator<8>> Sections;
-
-		/** Compiler rules used to compile this row */
-		const FMovieSceneSegmentCompilerRules* CompileRules;
-	};
-
-	struct FRows
-	{
-		/** Construction from an unordered array of sections, and an optional compiler rules object for compiling each row */
-		MOVIESCENE_API FRows(const TArray<UMovieSceneSection*>& Sections, const FMovieSceneSegmentCompilerRules* CompileRules = nullptr);
-		
-		/** Array of rows */
-		TArray<FRow, TInlineAllocator<2>> Rows;
-	};
+	virtual void Blend(FSegmentBlendData& BlendData) const override;
 
-	/**
-	 * Compile the specified rows using the specified compiler rules for combining each row
-	 *
-	 * @param Rows 			The rows to compile
-	 * @param Rules 		Compiler rules that define how the rows should be blended together
-	 *
-	 * @return An ordered array of segments that define unique combinations of overlapping sections. Segments are guaranteed to not overlap, may be contiguous and may contain gaps depending on compiler rules
-	 */
-	MOVIESCENE_API FMovieSceneTrackEvaluationField Compile(TArrayView<FRow> Rows, const FMovieSceneSegmentCompilerRules* Rules = nullptr);
+protected:
+	virtual void BlendSegment(FMovieSceneSegment& Segment, const TArrayView<const FMovieSceneSectionData>& SourceData) const = 0;
 };
-
-
-
diff --git a/Engine/Source/Runtime/MovieScene/Public/Evaluation/Blending/MovieSceneMultiChannelBlending.h b/Engine/Source/Runtime/MovieScene/Public/Evaluation/Blending/MovieSceneMultiChannelBlending.h
--- a/Engine/Source/Runtime/MovieScene/Public/Evaluation/Blending/MovieSceneMultiChannelBlending.h
+++ b/Engine/Source/Runtime/MovieScene/Public/Evaluation/Blending/MovieSceneMultiChannelBlending.h
@@ -91,6 +91,12 @@ namespace MovieScene
 			return Channels[Index];
 		}
 
+		/** Access a value with a default */
+		T Get(uint8 Index, T Default) const
+		{
+			return IsSet(Index) ? Channels[Index] : Default;
+		}
+
 		/** Check if this value is empty */
 		bool IsEmpty() const
 		{
diff --git a/Engine/Source/Runtime/MovieScene/Public/Evaluation/MovieSceneEvalTemplate.h b/Engine/Source/Runtime/MovieScene/Public/Evaluation/MovieSceneEvalTemplate.h
--- a/Engine/Source/Runtime/MovieScene/Public/Evaluation/MovieSceneEvalTemplate.h
+++ b/Engine/Source/Runtime/MovieScene/Public/Evaluation/MovieSceneEvalTemplate.h
@@ -146,7 +146,7 @@ public:
 	 *
 	 * @param SourceSection 		The source section
 	 */
-	void SetSourceSection(UMovieSceneSection* InSourceSection)
+	void SetSourceSection(const UMovieSceneSection* InSourceSection)
 	{
 		SourceSection = InSourceSection;
 	}
@@ -156,7 +156,7 @@ public:
 	 *
 	 * @return The source section from which this template originated
 	 */
-	UMovieSceneSection* GetSourceSection() const
+	const UMovieSceneSection* GetSourceSection() const
 	{
 		return SourceSection;
 	}
@@ -182,7 +182,7 @@ protected:
 
 	/** The section from which this template originates */
 	UPROPERTY()
-	UMovieSceneSection* SourceSection;
+	const UMovieSceneSection* SourceSection;
 };
 
 /**
diff --git a/Engine/Source/Runtime/MovieScene/Public/Evaluation/MovieSceneEvaluationField.h b/Engine/Source/Runtime/MovieScene/Public/Evaluation/MovieSceneEvaluationField.h
--- a/Engine/Source/Runtime/MovieScene/Public/Evaluation/MovieSceneEvaluationField.h
+++ b/Engine/Source/Runtime/MovieScene/Public/Evaluation/MovieSceneEvaluationField.h
@@ -6,9 +6,14 @@
 #include "UObject/ObjectMacros.h"
 #include "MovieSceneSequenceID.h"
 #include "MovieSceneEvaluationKey.h"
+#include "MovieSceneSegment.h"
 #include "Evaluation/MovieSceneTrackIdentifier.h"
 #include "MovieSceneEvaluationField.generated.h"
 
+struct FMovieSceneSequenceHierarchy;
+
+class UMovieSceneSequence;
+
 /** A pointer to a track held within an evaluation template */
 USTRUCT()
 struct FMovieSceneEvaluationFieldTrackPtr
@@ -67,9 +72,9 @@ struct FMovieSceneEvaluationFieldSegmentPtr : public FMovieSceneEvaluationFieldT
 	/**
 	 * Construction from a sequence ID, and the index of the track within that sequence's track list
 	 */
-	FMovieSceneEvaluationFieldSegmentPtr(FMovieSceneSequenceIDRef InSequenceID, FMovieSceneTrackIdentifier InTrackIdentifier, int32 InSegmentIndex)
+	FMovieSceneEvaluationFieldSegmentPtr(FMovieSceneSequenceIDRef InSequenceID, FMovieSceneTrackIdentifier InTrackIdentifier, FMovieSceneSegmentIdentifier InSegmentID)
 		: FMovieSceneEvaluationFieldTrackPtr(InSequenceID, InTrackIdentifier)
-		, SegmentIndex(InSegmentIndex)
+		, SegmentID(InSegmentID)
 	{}
 
 	/**
@@ -77,7 +82,7 @@ struct FMovieSceneEvaluationFieldSegmentPtr : public FMovieSceneEvaluationFieldT
 	 */
 	friend bool operator==(FMovieSceneEvaluationFieldSegmentPtr A, FMovieSceneEvaluationFieldSegmentPtr B)
 	{
-		return A.SegmentIndex == B.SegmentIndex && A.TrackIdentifier == B.TrackIdentifier && A.SequenceID == B.SequenceID;
+		return A.SegmentID == B.SegmentID && A.TrackIdentifier == B.TrackIdentifier && A.SequenceID == B.SequenceID;
 	}
 
 	/**
@@ -85,12 +90,12 @@ struct FMovieSceneEvaluationFieldSegmentPtr : public FMovieSceneEvaluationFieldT
 	 */
 	friend uint32 GetTypeHash(FMovieSceneEvaluationFieldSegmentPtr LHS)
 	{
-		return HashCombine(LHS.SegmentIndex, GetTypeHash(static_cast<FMovieSceneEvaluationFieldTrackPtr&>(LHS)));
+		return HashCombine(GetTypeHash(LHS.SegmentID), GetTypeHash(static_cast<FMovieSceneEvaluationFieldTrackPtr&>(LHS)));
 	}
 
-	/** The index of the segment within the track (see FMovieSceneEvaluationTrack::Segments) */
+	/** The identifier of the segment within the track (see FMovieSceneEvaluationTrack::Segments) */
 	UPROPERTY()
-	int32 SegmentIndex;
+	FMovieSceneSegmentIdentifier SegmentID;
 };
 
 /** Lookup table index for a group of evaluation templates */
@@ -152,19 +157,6 @@ struct FMovieSceneEvaluationMetaData
 {
 	GENERATED_BODY()
 
-	FMovieSceneEvaluationMetaData() = default;
-
-	FMovieSceneEvaluationMetaData(const FMovieSceneEvaluationMetaData&) = default;
-	FMovieSceneEvaluationMetaData& operator=(const FMovieSceneEvaluationMetaData&) = default;
-
-#if PLATFORM_COMPILER_HAS_DEFAULTED_FUNCTIONS
-	FMovieSceneEvaluationMetaData(FMovieSceneEvaluationMetaData&&) = default;
-	FMovieSceneEvaluationMetaData& operator=(FMovieSceneEvaluationMetaData&&) = default;
-#else
-	FMovieSceneEvaluationMetaData(FMovieSceneEvaluationMetaData&& RHS) : ActiveSequences(MoveTemp(RHS.ActiveSequences)), ActiveEntities(MoveTemp(RHS.ActiveEntities)) {}
-	FMovieSceneEvaluationMetaData& operator=(FMovieSceneEvaluationMetaData&& RHS) { ActiveSequences = MoveTemp(RHS.ActiveSequences); ActiveEntities = MoveTemp(RHS.ActiveEntities); return *this; }
-#endif
-
 	/**
 	 * Reset this meta-data
 	 */
@@ -174,13 +166,6 @@ struct FMovieSceneEvaluationMetaData
 		ActiveEntities.Reset();
 	}
 
-	/**
-	 * Remap this meta-data onto a different parent ID
-	 *
-	 * @param OverrideRootID		The parent ID to remap entities onto
-	 */
-	void RemapSequenceIDsForRoot(FMovieSceneSequenceID OverrideRootID);
-
 	/**
 	 * Diff the active sequences this frame, with the specified previous frame's meta-data
 	 *
@@ -199,6 +184,17 @@ struct FMovieSceneEvaluationMetaData
 	 */
 	void DiffEntities(const FMovieSceneEvaluationMetaData& LastFrame, TArray<FMovieSceneOrderedEvaluationKey>* NewKeys, TArray<FMovieSceneOrderedEvaluationKey>* ExpiredKeys) const;
 
+	/**
+	 * Check whether this meta-data entry is still up-to-date
+	 *
+	 * @param RootSequence				The sequence that corresponds to this meta-data's root sequence
+	 * @param RootHierarchy				The hierarchy that corresponds to this meta-data's root sequence
+	 * @param OutSubRangeToInvalidate	(Optional) A range to fill with a range to invalidate in the sequence's evaluation field (in root space)
+	 *
+	 * @return true if the meta-data needs re-generating, false otherwise
+	 */
+	bool IsDirty(UMovieSceneSequence& RootSequence, const FMovieSceneSequenceHierarchy& RootHierarchy, TRange<float>* OutSubRangeToInvalidate = nullptr) const;
+
 	/** Array of sequences that are active in this time range. */
 	UPROPERTY()
 	TArray<FMovieSceneSequenceID> ActiveSequences;
@@ -206,6 +202,10 @@ struct FMovieSceneEvaluationMetaData
 	/** Array of entities (tracks and/or sections) that are active in this time range. */
 	UPROPERTY()
 	TArray<FMovieSceneOrderedEvaluationKey> ActiveEntities;
+
+	/** Map of sub sequence IDs to signatures that this meta data was generated with (not including root. */
+	UPROPERTY()
+	TMap<FMovieSceneSequenceID, FGuid> SubSequenceSignatures;
 };
 
 /**
@@ -232,6 +232,84 @@ struct FMovieSceneEvaluationField
 	 */
 	MOVIESCENE_API TRange<int32> OverlapRange(TRange<float> Range) const;
 
+	/**
+	 * Invalidate a range in this field
+	 *
+	 * @param Range			The range to overlap with our field
+	 * @return A range of indices into Ranges and Groups that overlap with the requested range
+	 */
+	MOVIESCENE_API void Invalidate(TRange<float> Range);
+
+	/**
+	 * Insert a new range into this field
+	 *
+	 * @param InRange		The range to add
+	 * @param InGroup		The group defining what should happen at this time
+	 * @param InMetaData	The meta-data defining efficient access to what happens in this frame
+	 * @return The index the entries were inserted at
+	 */
+	MOVIESCENE_API int32 Insert(TRange<float> InRange, FMovieSceneEvaluationGroup&& InGroup, FMovieSceneEvaluationMetaData&& InMetaData);
+
+	/**
+	 * Add the specified data to this field, assuming the specified range lies after any other entries
+	 *
+	 * @param InRange		The range to add
+	 * @param InGroup		The group defining what should happen at this time
+	 * @param InMetaData	The meta-data defining efficient access to what happens in this frame
+	 */
+	MOVIESCENE_API void Add(TRange<float> InRange, FMovieSceneEvaluationGroup&& InGroup, FMovieSceneEvaluationMetaData&& InMetaData);
+
+	/**
+	 * Access this field's signature
+	 */
+	const FGuid& GetSignature() const
+	{
+		return Signature;
+	}
+
+	/**
+	 * Access this field's size
+	 */
+	int32 Size() const
+	{
+		return Ranges.Num();
+	}
+
+	/**
+	 * Lookup a valid range by index
+	 * @param Index 	The valid index within the ranges to lookup
+	 * @return The range
+	 */
+	const TRange<float>& GetRange(int32 Index) const
+	{
+		return Ranges[Index];
+	}
+
+	/**
+	 * Lookup a valid evaluation group by entry index
+	 * @param Index 	The valid index within the evaluation group array to lookup
+	 * @return The group
+	 */
+	const FMovieSceneEvaluationGroup& GetGroup(int32 Index) const
+	{
+		return Groups[Index];
+	}
+
+	/**
+	 * Lookup valid meta-data by entry index
+	 * @param Index 	The valid index within the meta-data array to lookup
+	 * @return The meta-data
+	 */
+	const FMovieSceneEvaluationMetaData& GetMetaData(int32 Index) const
+	{
+		return MetaData[Index];
+	}
+
+private:
+	/** Signature that uniquely identifies any state this field can be in - regenerated on mutation */
+	UPROPERTY()
+	FGuid Signature;
+
 	/** Ranges stored separately for fast (cache efficient) lookup. Each index has a corresponding entry in FMovieSceneEvaluationField::Groups. */
 	UPROPERTY()
 	TArray<FFloatRange> Ranges;
diff --git a/Engine/Source/Runtime/MovieScene/Public/Evaluation/MovieSceneEvaluationKey.h b/Engine/Source/Runtime/MovieScene/Public/Evaluation/MovieSceneEvaluationKey.h
--- a/Engine/Source/Runtime/MovieScene/Public/Evaluation/MovieSceneEvaluationKey.h
+++ b/Engine/Source/Runtime/MovieScene/Public/Evaluation/MovieSceneEvaluationKey.h
@@ -19,16 +19,16 @@ struct FMovieSceneEvaluationKey
 	FMovieSceneEvaluationKey()
 		: SequenceID(MovieSceneSequenceID::Invalid)
 		, TrackIdentifier(FMovieSceneTrackIdentifier::Invalid())
-		, SectionIdentifier(-1)
+		, SectionIndex(-1)
 	{}
 
 	/**
 	 * User construction
 	 */
-	FMovieSceneEvaluationKey(FMovieSceneSequenceIDRef InSequenceID, FMovieSceneTrackIdentifier InTrackIdentifier, uint32 InSectionIdentifier = uint32(-1))
+	FMovieSceneEvaluationKey(FMovieSceneSequenceIDRef InSequenceID, FMovieSceneTrackIdentifier InTrackIdentifier, uint32 InSectionIndex = uint32(-1))
 		: SequenceID(InSequenceID)
 		, TrackIdentifier(InTrackIdentifier)
-		, SectionIdentifier(InSectionIdentifier)
+		, SectionIndex(InSectionIndex)
 	{}
 
 	/**
@@ -45,7 +45,7 @@ struct FMovieSceneEvaluationKey
 	FMovieSceneEvaluationKey AsSection(uint32 InSectionIdentifier) const
 	{
 		FMovieSceneEvaluationKey NewKey = *this;
-		NewKey.SectionIdentifier = InSectionIdentifier;
+		NewKey.SectionIndex = InSectionIdentifier;
 		return NewKey;
 	}
 
@@ -55,13 +55,13 @@ struct FMovieSceneEvaluationKey
 	FMovieSceneEvaluationKey AsTrack() const
 	{
 		FMovieSceneEvaluationKey NewKey = *this;
-		NewKey.SectionIdentifier = uint32(-1);
+		NewKey.SectionIndex = uint32(-1);
 		return NewKey;
 	}
 
 	friend bool operator==(const FMovieSceneEvaluationKey& A, const FMovieSceneEvaluationKey& B)
 	{
-		return A.TrackIdentifier == B.TrackIdentifier && A.SequenceID == B.SequenceID && A.SectionIdentifier == B.SectionIdentifier;
+		return A.TrackIdentifier == B.TrackIdentifier && A.SequenceID == B.SequenceID && A.SectionIndex == B.SectionIndex;
 	}
 
 	friend bool operator<(const FMovieSceneEvaluationKey& A, const FMovieSceneEvaluationKey& B)
@@ -80,13 +80,13 @@ struct FMovieSceneEvaluationKey
 		}
 		else
 		{
-			return A.TrackIdentifier == B.TrackIdentifier && A.SectionIdentifier < B.SectionIdentifier;
+			return A.TrackIdentifier == B.TrackIdentifier && A.SectionIndex < B.SectionIndex;
 		}
 	}
 
 	friend uint32 GetTypeHash(const FMovieSceneEvaluationKey& In)
 	{
-		return GetTypeHash(In.SequenceID) ^ (~GetTypeHash(In.TrackIdentifier)) ^ In.SectionIdentifier;
+		return GetTypeHash(In.SequenceID) ^ (~GetTypeHash(In.TrackIdentifier)) ^ In.SectionIndex;
 	}
 
 	/** Custom serialized to reduce memory footprint */
@@ -94,7 +94,7 @@ struct FMovieSceneEvaluationKey
 	{
 		Ar << SequenceID;
 		Ar << TrackIdentifier;
-		Ar << SectionIdentifier;
+		Ar << SectionIndex;
 		return true;
 	}
 
@@ -112,9 +112,9 @@ struct FMovieSceneEvaluationKey
 	UPROPERTY()
 	FMovieSceneTrackIdentifier TrackIdentifier;
 
-	/** ID of the section this key relates to (or -1 where this key relates to a track) */
+	/** Index of the section template within the track this key relates to (or -1 where this key relates to a track) */
 	UPROPERTY()
-	uint32 SectionIdentifier;
+	uint32 SectionIndex;
 };
 
 template<>
diff --git a/Engine/Source/Runtime/MovieScene/Public/Evaluation/MovieSceneEvaluationOperand.h b/Engine/Source/Runtime/MovieScene/Public/Evaluation/MovieSceneEvaluationOperand.h
--- a/Engine/Source/Runtime/MovieScene/Public/Evaluation/MovieSceneEvaluationOperand.h
+++ b/Engine/Source/Runtime/MovieScene/Public/Evaluation/MovieSceneEvaluationOperand.h
@@ -5,11 +5,16 @@
 #include "MovieSceneSequenceID.h"
 #include "Guid.h"
 
+#include "MovieSceneEvaluationOperand.generated.h"
+
 /**
  * Structure that describes an object that is to be animated. Used as an abstraction of the actual objects bound to object bindings.
  */
+USTRUCT()
 struct FMovieSceneEvaluationOperand
 {
+	GENERATED_BODY()
+
 	/**
 	 * Default Construction
 	 */
@@ -46,8 +51,10 @@ struct FMovieSceneEvaluationOperand
 	}
 
 	/** A GUID relating to either a posessable, or a spawnable binding */
+	UPROPERTY()
 	FGuid ObjectBindingID;
 
 	/** The ID of the sequence within which the object binding resides */
+	UPROPERTY()
 	FMovieSceneSequenceID SequenceID;
 };
\ No newline at end of file
