Copying //UE4/Dev-Sequencer to //UE4/Dev-Main (Source: //UE4/Dev-Sequencer @ 3800792)

diff --git a/Engine/Source/Runtime/MovieScene/Public/Evaluation/MovieSceneEvaluationTemplate.h b/Engine/Source/Runtime/MovieScene/Public/Evaluation/MovieSceneEvaluationTemplate.h
--- a/Engine/Source/Runtime/MovieScene/Public/Evaluation/MovieSceneEvaluationTemplate.h
+++ b/Engine/Source/Runtime/MovieScene/Public/Evaluation/MovieSceneEvaluationTemplate.h
@@ -13,21 +13,13 @@
 #include "Evaluation/MovieSceneEvaluationField.h"
 #include "Containers/ArrayView.h"
 #include "Evaluation/MovieSceneEvaluationTrack.h"
+#include "Evaluation/MovieSceneEvaluationTree.h"
 #include "Evaluation/MovieSceneSequenceHierarchy.h"
-#include "Compilation/IMovieSceneTemplateGenerator.h"
-#include "Evaluation/MovieSceneSequenceTemplateStore.h"
 #include "MovieSceneEvaluationTemplate.generated.h"
 
 class UMovieSceneSequence;
+class UMovieSceneSubSection;
 
-USTRUCT()
-struct FMovieSceneTrackIdentifiers
-{
-	GENERATED_BODY()
-
-	UPROPERTY()
-	TArray<FMovieSceneTrackIdentifier> Data;
-};
 
 
 USTRUCT()
@@ -37,25 +29,100 @@ struct FMovieSceneTemplateGenerationLedger
 
 public:
 
-	TArrayView<const FMovieSceneTrackIdentifier> FindTracks(const FGuid& InSignature) const;
+	/**
+	 * Lookup a track identifier by its originating signature
+	 */
+	MOVIESCENE_API FMovieSceneTrackIdentifier FindTrack(const FGuid& InSignature) const;
+
+	/**
+	 * Add a new track for the specified signature. Signature must not have already been used
+	 */
+	MOVIESCENE_API void AddTrack(const FGuid& InSignature, FMovieSceneTrackIdentifier Identifier);
 
-	void AddTrack(const FGuid& InSignature, FMovieSceneTrackIdentifier Identifier);
+	/**
+	 * Check whether we already have a sub section for the specified signature
+	 */
+	bool ContainsSubSection(const FGuid& InSignature)
+	{
+		return SubSectionRanges.Contains(InSignature);
+	}
 
 public:
 
 	UPROPERTY()
 	FMovieSceneTrackIdentifier LastTrackIdentifier;
 
-	/** Map of track identifiers to number of references within th template (generally 1, maybe >1 for shared tracks) */
+	/** Map of track signature to array of track identifiers that it created */
 	UPROPERTY()
-	TMap<FMovieSceneTrackIdentifier, int32> TrackReferenceCounts;
+	TMap<FGuid, FMovieSceneTrackIdentifier> TrackSignatureToTrackIdentifier;
 
-	/** Map of track signature to array of track identifiers that it created */
+	/** Map of sub section ranges that exist in a template */
 	UPROPERTY()
-	TMap<FGuid, FMovieSceneTrackIdentifiers> TrackSignatureToTrackIdentifier;
+	TMap<FGuid, FFloatRange> SubSectionRanges;
 };
 template<> struct TStructOpsTypeTraits<FMovieSceneTemplateGenerationLedger> : public TStructOpsTypeTraitsBase2<FMovieSceneTemplateGenerationLedger> { enum { WithCopy = true }; };
 
+/** Custom serialized track field data that allows efficient lookup of each track contained within this template for a given time */
+USTRUCT()
+struct FMovieSceneTrackFieldData
+{
+	GENERATED_BODY()
+
+	bool Serialize(FArchive& Ar)
+	{
+		Ar << Field;
+		return true;
+	}
+
+	TMovieSceneEvaluationTree<FMovieSceneTrackIdentifier> Field;
+};
+template<> struct TStructOpsTypeTraits<FMovieSceneTrackFieldData> : public TStructOpsTypeTraitsBase2<FMovieSceneTrackFieldData> { enum { WithSerializer = true }; };
+
+/** Data that represents a single sub-section */
+USTRUCT()
+struct FMovieSceneSubSectionData
+{
+	GENERATED_BODY()
+
+	FMovieSceneSubSectionData() {}
+
+	FMovieSceneSubSectionData(UMovieSceneSubSection& InSubSection, const FGuid& InObjectBindingId, ESectionEvaluationFlags InFlags);
+
+	friend FArchive& operator<<(FArchive& Ar, FMovieSceneSubSectionData& In)
+	{
+		return Ar << In.Section << In.ObjectBindingId << In.Flags;
+	}
+
+	/** The sub section itself */
+	UPROPERTY()
+	TWeakObjectPtr<UMovieSceneSubSection> Section;
+
+	/** The object binding that the sub section belongs to (usually zero) */
+	UPROPERTY()
+	FGuid ObjectBindingId;
+
+	/** Evaluation flags for the section */
+	UPROPERTY()
+	ESectionEvaluationFlags Flags;
+};
+
+/** Custom serialized track field data that allows efficient lookup of each sub section contained within this template for a given time */
+USTRUCT()
+struct FMovieSceneSubSectionFieldData
+{
+	GENERATED_BODY()
+
+	bool Serialize(FArchive& Ar)
+	{
+		Ar << Field;
+		return true;
+	}
+
+	TMovieSceneEvaluationTree<FMovieSceneSubSectionData> Field;
+};
+template<> struct TStructOpsTypeTraits<FMovieSceneSubSectionFieldData> : public TStructOpsTypeTraitsBase2<FMovieSceneSubSectionFieldData> { enum { WithSerializer = true }; };
+
+
 /**
  * Template that is used for efficient runtime evaluation of a movie scene sequence. Potentially serialized into the asset.
  */
@@ -69,59 +136,23 @@ struct FMovieSceneEvaluationTemplate
 	 */
 	FMovieSceneEvaluationTemplate()
 	{
-		bHasLegacyTrackInstances = 0;
 		bKeepStaleTracks = 0;
 	}
 
-	/**
-	 * Move construction/assignment
-	 */
-#if PLATFORM_COMPILER_HAS_DEFAULTED_FUNCTIONS
-	FMovieSceneEvaluationTemplate(FMovieSceneEvaluationTemplate&&) = default;
-	FMovieSceneEvaluationTemplate& operator=(FMovieSceneEvaluationTemplate&&) = default;
-#else
-	FMovieSceneEvaluationTemplate(FMovieSceneEvaluationTemplate&& RHS)
-	{
-		*this = MoveTemp(RHS);
-	}
-	FMovieSceneEvaluationTemplate& operator=(FMovieSceneEvaluationTemplate&& RHS)
-	{
-		Tracks = MoveTemp(RHS.Tracks);
-		StaleTracks = MoveTemp(RHS.StaleTracks);
-		EvaluationField = MoveTemp(RHS.EvaluationField);
-		Hierarchy = MoveTemp(RHS.Hierarchy);
-		Ledger = MoveTemp(RHS.Ledger);
-		bHasLegacyTrackInstances = RHS.bHasLegacyTrackInstances;
-		bKeepStaleTracks = RHS.bKeepStaleTracks;
-		return *this;
-	}
-#endif
-
-	/**
-	 * Non-copyable
-	 */
-	FMovieSceneEvaluationTemplate(const FMovieSceneEvaluationTemplate&) = default;
-	FMovieSceneEvaluationTemplate& operator=(const FMovieSceneEvaluationTemplate&) = default;
-
 public:
 
-	/**
-	 * Reset this evaluation template, whilst ensuring that track identifiers do not get reused
-	 */
-	MOVIESCENE_API void ResetGeneratedData();
-
 	/**
 	 * Attempt to locate a track with the specified identifier
 	 */
 	FMovieSceneEvaluationTrack* FindTrack(FMovieSceneTrackIdentifier Identifier)
 	{
 		// Fast, most common path
-		if (FMovieSceneEvaluationTrack* Track = Tracks.Find(Identifier.Value))
+		if (FMovieSceneEvaluationTrack* Track = Tracks.Find(Identifier))
 		{
 			return Track;
 		}
 
-		return bKeepStaleTracks ? StaleTracks.Find(Identifier.Value) : nullptr;
+		return bKeepStaleTracks ? StaleTracks.Find(Identifier) : nullptr;
 	}
 
 	/**
@@ -130,12 +161,28 @@ public:
 	const FMovieSceneEvaluationTrack* FindTrack(FMovieSceneTrackIdentifier Identifier) const
 	{
 		// Fast, most common path
-		if (const FMovieSceneEvaluationTrack* Track = Tracks.Find(Identifier.Value))
+		if (const FMovieSceneEvaluationTrack* Track = Tracks.Find(Identifier))
 		{
 			return Track;
 		}
 
-		return bKeepStaleTracks ? StaleTracks.Find(Identifier.Value) : nullptr;
+		return bKeepStaleTracks ? StaleTracks.Find(Identifier) : nullptr;
+	}
+
+	/**
+	 * Find a track within this template that relates to the specified signature
+	 */
+	FMovieSceneEvaluationTrack* FindTrack(const FGuid& InSignature)
+	{
+		return FindTrack(TemplateLedger.FindTrack(InSignature));
+	}
+
+	/**
+	 * Find a track within this template that relates to the specified signature
+	 */
+	const FMovieSceneEvaluationTrack* FindTrack(const FGuid& InSignature) const
+	{
+		return FindTrack(TemplateLedger.FindTrack(InSignature));
 	}
 
 	/**
@@ -143,14 +190,27 @@ public:
 	 */
 	bool IsTrackStale(FMovieSceneTrackIdentifier Identifier) const
 	{
-		return bKeepStaleTracks ? StaleTracks.Contains(Identifier.Value) : false;
+		return bKeepStaleTracks ? StaleTracks.Contains(Identifier) : false;
 	}
 
+	/**
+	 * Add a new sub section
+	 */
+	MOVIESCENE_API void AddSubSectionRange(UMovieSceneSubSection& InSubSection, const FGuid& InObjectBindingId, const TRange<float>& InRange, ESectionEvaluationFlags InFlags);
+
 	/**
 	 * Add a new track for the specified identifier
 	 */
 	MOVIESCENE_API FMovieSceneTrackIdentifier AddTrack(const FGuid& InSignature, FMovieSceneEvaluationTrack&& InTrack);
 
+	/**
+	 * Define the structural lookup for the specified track identifier, optionally invalidating any overlapping areas in the evaluation field
+	 *
+	 * @param TrackIdentifier 				The identifier for the track
+	 * @param bInvalidateEvaluationField 	Whether to invalidate any overlapping sections of the cached evaluation field
+	 */
+	MOVIESCENE_API void DefineTrackStructure(FMovieSceneTrackIdentifier TrackIdentifier, bool bInvalidateEvaluationField);
+
 	/**
 	 * Remove any tracks that correspond to the specified signature
 	 */
@@ -168,11 +228,6 @@ public:
 	 */
 	MOVIESCENE_API TMap<FMovieSceneTrackIdentifier, FMovieSceneEvaluationTrack>& GetTracks();
 
-	/**
-	 * Find tracks within this template that relate to the specified signature
-	 */
-	MOVIESCENE_API TArrayView<const FMovieSceneTrackIdentifier> FindTracks(const FGuid& InSignature) const;
-
 	/**
 	 * Called after this template has been serialized in some way
 	 */
@@ -196,14 +251,34 @@ public:
 		return TemplateLedger;
 	}
 
+	/**
+	 * Remove any data within this template that does not reside in the specified set of signatures
+	 */
+	void RemoveStaleData(const TSet<FGuid>& ActiveSignatures);
+
+	/**
+	 * Reset this template's field data and sub section cache (keeps tracks alive)
+	 */
+	void ResetFieldData();
+
+	/**
+	 * Access this template's track field
+	 */
+	const TMovieSceneEvaluationTree<FMovieSceneTrackIdentifier>& GetTrackField() const;
+
+	/**
+	 * Access this template's sub section field
+	 */
+	const TMovieSceneEvaluationTree<FMovieSceneSubSectionData>& GetSubSectionField() const;
+
 private:
 
 	/** Map of evaluation tracks from identifier to track */
 	UPROPERTY()
-	TMap<uint32, FMovieSceneEvaluationTrack> Tracks;
+	TMap<FMovieSceneTrackIdentifier, FMovieSceneEvaluationTrack> Tracks;
 
-	/** Transient map of stale tracks. Only populated dureing regeneration where bKeepStaleTracks is true */
-	TMap<uint32, FMovieSceneEvaluationTrack> StaleTracks;
+	/** Transient map of stale tracks. Only populated during regeneration where bKeepStaleTracks is true */
+	TMap<FMovieSceneTrackIdentifier, FMovieSceneEvaluationTrack> StaleTracks;
 
 public:
 
@@ -215,82 +290,24 @@ public:
 	UPROPERTY()
 	FMovieSceneSequenceHierarchy Hierarchy;
 
-private:
-
 	UPROPERTY()
-	FMovieSceneTemplateGenerationLedger TemplateLedger;
-
-public:
+	FGuid SequenceSignature;
 
-	/** When set, this template contains legacy track instances that require the initialization of a legacy sequence instance */
-	UPROPERTY()
-	uint32 bHasLegacyTrackInstances : 1;
+private:
 
-	/** Primarily used in editor to keep stale tracks around during template regeneration to ensure we can call OnEndEvaluation on them. */
 	UPROPERTY()
-	uint32 bKeepStaleTracks : 1;
-};
-template<> struct TStructOpsTypeTraits<FMovieSceneEvaluationTemplate> : public TStructOpsTypeTraitsBase2<FMovieSceneEvaluationTemplate> { enum { WithPostSerialize = true }; };
-
-USTRUCT()
-struct FMovieSceneSequenceCachedSignature
-{
-	GENERATED_BODY()
-
-	FMovieSceneSequenceCachedSignature() {}
-	MOVIESCENE_API FMovieSceneSequenceCachedSignature(UMovieSceneSequence& InSequence);
+	FMovieSceneTemplateGenerationLedger TemplateLedger;
 
 	UPROPERTY()
-	TWeakObjectPtr<UMovieSceneSequence> Sequence;
+	FMovieSceneTrackFieldData TrackFieldData;
 
 	UPROPERTY()
-	FGuid CachedSignature;
-};
-
-USTRUCT()
-struct FCachedMovieSceneEvaluationTemplate : public FMovieSceneEvaluationTemplate
-{
-	GENERATED_BODY()
-
-#if WITH_EDITORONLY_DATA
-
-	FCachedMovieSceneEvaluationTemplate()
-		: Origin(nullptr)
-	{
-		bKeepStaleTracks = 1;
-	}
+	FMovieSceneSubSectionFieldData SubSectionFieldData;
 
-	FCachedMovieSceneEvaluationTemplate(UMovieSceneSequence& InSequence)
-		: Origin(nullptr)
-	{
-		bKeepStaleTracks = 1;
-		Initialize(InSequence);
-	}
-
-	MOVIESCENE_API void Initialize(UMovieSceneSequence& InSequence, FMovieSceneSequenceTemplateStore* InOrigin = nullptr);
-
-	MOVIESCENE_API void Regenerate();
-
-	MOVIESCENE_API void Regenerate(const FMovieSceneTrackCompilationParams& NewParams);
-
-	MOVIESCENE_API void ForceRegenerate(const FMovieSceneTrackCompilationParams& NewParams);
-
-	MOVIESCENE_API bool IsOutOfDate(const FMovieSceneTrackCompilationParams& NewParams) const;
-
-private:
-
-	void RegenerateImpl(const FMovieSceneTrackCompilationParams& NewParams);
-
-	/** Transient data */
-	TWeakObjectPtr<UMovieSceneSequence> SourceSequence;
-	FMovieSceneSequenceTemplateStore* Origin;
-
-	UPROPERTY()
-	FMovieSceneTrackCompilationParams CachedCompilationParams;
+public:
 
+	/** Primarily used in editor to keep stale tracks around during template regeneration to ensure we can call OnEndEvaluation on them. */
 	UPROPERTY()
-	TArray<FMovieSceneSequenceCachedSignature> CachedSignatures;
-#endif
+	uint32 bKeepStaleTracks : 1;
 };
-
-template<> struct TStructOpsTypeTraits<FCachedMovieSceneEvaluationTemplate> : public TStructOpsTypeTraitsBase2<FCachedMovieSceneEvaluationTemplate> { enum { WithCopy = false, WithPostSerialize = true }; };
+template<> struct TStructOpsTypeTraits<FMovieSceneEvaluationTemplate> : public TStructOpsTypeTraitsBase2<FMovieSceneEvaluationTemplate> { enum { WithPostSerialize = true }; };
\ No newline at end of file
diff --git a/Engine/Source/Runtime/MovieScene/Public/Evaluation/MovieSceneEvaluationTemplateInstance.h b/Engine/Source/Runtime/MovieScene/Public/Evaluation/MovieSceneEvaluationTemplateInstance.h
--- a/Engine/Source/Runtime/MovieScene/Public/Evaluation/MovieSceneEvaluationTemplateInstance.h
+++ b/Engine/Source/Runtime/MovieScene/Public/Evaluation/MovieSceneEvaluationTemplateInstance.h
@@ -5,67 +5,80 @@
 #include "CoreMinimal.h"
 #include "MovieSceneSequenceID.h"
 #include "Evaluation/MovieSceneSequenceTransform.h"
-#include "Evaluation/MovieSceneEvaluationState.h"
 #include "Evaluation/MovieScenePlayback.h"
 #include "Evaluation/MovieSceneSequenceTemplateStore.h"
 #include "Evaluation/MovieSceneEvaluationTemplate.h"
 #include "Evaluation/MovieSceneExecutionTokens.h"
+#include "Evaluation/MovieSceneRootOverridePath.h"
 
 class UMovieSceneSequence;
 struct FDelayedPreAnimatedStateRestore;
+struct FCompileOnTheFlyData;
 
 /**
  * An instance of an evaluation template. Fast to initialize and evaluate.
  */
 struct FMovieSceneEvaluationTemplateInstance
 {
-	FMovieSceneEvaluationTemplateInstance() : Template(nullptr), HierarchicalBias(0), SequenceID(MovieSceneSequenceID::Invalid) {}
+	FMovieSceneEvaluationTemplateInstance();
 
-	/**
-	 * Constructor for root instances
-	 *
-	 * @param InSequence			The source master sequence we're playing back at the root level
-	 * @param InTemplate			The template to use when playing back the specified sequence. A pointer to which will be held within this instance.
-	 */
-	FMovieSceneEvaluationTemplateInstance(UMovieSceneSequence& InSequence, const FMovieSceneEvaluationTemplate& InTemplate);
+	FMovieSceneEvaluationTemplateInstance(UMovieSceneSequence& InSequence, FMovieSceneEvaluationTemplate* InTemplate);
 
-	/**
-	 * Constructor for sub sequence instances. SubData must point to a valid sequence
-	 *
-	 * @param SubData				Information about the sub sequence we're instantiating
-	 * @param InTemplate			The template to use when playing back the sub sequence. A pointer to which will be held within this instance.
-	 * @param InSequenceID			The unique identifier within the master sequence
-	 */
-	FMovieSceneEvaluationTemplateInstance(const FMovieSceneSubSequenceData& SubData, const FMovieSceneEvaluationTemplate& InTemplate, FMovieSceneSequenceIDRef InSequenceID);
+	FMovieSceneEvaluationTemplateInstance(const FMovieSceneSubSequenceData* InSubData, IMovieSceneSequenceTemplateStore& TemplateStore);
 
-	FMovieSceneEvaluationTemplateInstance(const FMovieSceneEvaluationTemplateInstance&) = delete;
-	FMovieSceneEvaluationTemplateInstance& operator=(const FMovieSceneEvaluationTemplateInstance&) = delete;
+	bool IsValid() const
+	{
+		return Sequence && Template;
+	}
 
-	FMovieSceneEvaluationTemplateInstance(FMovieSceneEvaluationTemplateInstance&&) = default;
-	FMovieSceneEvaluationTemplateInstance& operator=(FMovieSceneEvaluationTemplateInstance&&) = default;
+	/** The sequence this applies to */
+	UMovieSceneSequence* Sequence;
 
-public:
+	/** Pointer to the evaluation template we're evaluating */
+	FMovieSceneEvaluationTemplate* Template;
 
-	/** The sequence that the sub section references */
-	TWeakObjectPtr<UMovieSceneSequence> Sequence;
+	/** Sub Data or nullptr */
+	const FMovieSceneSubSequenceData* SubData;
+};
 
-	/** Transform that transforms a given time from the sequences outer space, to its authored space. */
-	FMovieSceneSequenceTransform RootToSequenceTransform;
+struct FMovieSceneEvaluationTemplateInstanceContainer
+{
+	FMovieSceneEvaluationTemplateInstanceContainer()
+		: RootID(MovieSceneSequenceID::Invalid)
+	{}
 
-	/** Pointer to the evaluation template we're evaluating */
-	const FMovieSceneEvaluationTemplate* Template;
+	void Reset()
+	{
+		RootID = MovieSceneSequenceID::Invalid;
+		RootInstance = FMovieSceneEvaluationTemplateInstance();
+		ResetSubInstances();
+	}
 
-	/** Pre and Post roll ranges in the inner-sequence's time space */
-	TRange<float> PreRollRange, PostRollRange;
+	void ResetSubInstances()
+	{
+		SubInstances.Reset();
+	}
 
-	/** The hierarchical bias for this template instance */
-	int32 HierarchicalBias;
+	void Add(FMovieSceneSequenceID InID, const FMovieSceneEvaluationTemplateInstance& InInstance)
+	{
+		check(InID != RootID);
+		SubInstances.Add(InID, InInstance);
+	}
 
-	/** ID of the sequence within the master sequence */
-	FMovieSceneSequenceID SequenceID;
+	const FMovieSceneEvaluationTemplateInstance& GetChecked(FMovieSceneSequenceID InID) const
+	{
+		return InID == RootID ? RootInstance : SubInstances.FindChecked(InID);
+	}
 
-	/** Legacy sequence instance required for old track instances */
-	TSharedPtr<FMovieSceneSequenceInstance> LegacySequenceInstance;
+	FMovieSceneEvaluationTemplateInstance* Find(FMovieSceneSequenceID InID)
+	{
+		return InID == RootID ? &RootInstance : SubInstances.Find(InID);
+	}
+
+	/** Sequence ID that is used to evaluate from for the current frame */
+	FMovieSceneSequenceID RootID;
+	FMovieSceneEvaluationTemplateInstance RootInstance;
+	TMap<FMovieSceneSequenceID, FMovieSceneEvaluationTemplateInstance> SubInstances;
 };
 
 /**
@@ -88,7 +101,7 @@ public:
 	 */
 	bool IsValid() const
 	{
-		return RootSequence.Get() && RootInstance.Template;
+		return RootSequence.Get() && RootTemplate;
 	}
 
 	/**
@@ -106,7 +119,7 @@ public:
 	 * @param Player					The player responsible for playback
 	 * @param TemplateStore				Template store responsible for supplying templates for a given sequence
 	 */
-	MOVIESCENE_API void Initialize(UMovieSceneSequence& RootSequence, IMovieScenePlayer& Player, TSharedRef<FMovieSceneSequenceTemplateStore> TemplateStore);
+	MOVIESCENE_API void Initialize(UMovieSceneSequence& RootSequence, IMovieScenePlayer& Player, TSharedRef<IMovieSceneSequenceTemplateStore> TemplateStore);
 
 	/**
 	 * Evaluate this sequence
@@ -124,61 +137,45 @@ public:
 	 */
 	MOVIESCENE_API void Finish(IMovieScenePlayer& Player);
 
-public:
-
-	class FOnUpdated : public TBaseMulticastDelegate<void>
-	{
-		friend FMovieSceneRootEvaluationTemplateInstance;
-	};
-
 	/**
-	 * Event that is triggered when this sequence has been updated
+	 * Check whether the evaluation template is dirty based on the last evaluated frame's meta-data
 	 */
-	FOnUpdated& OnUpdated() { return OnUpdatedEvent; }
+	MOVIESCENE_API bool IsDirty() const;
 
 public:
 
 	/**
-	 * Retrieve a specific template instance given a sequence ID
+	 * Attempt to locate the underlying sequence given a sequence ID
 	 *
-	 * @param SequenceID			The sequence ID to find a template for
-	 * @return Pointer to a template instance, or nullptr if the ID was not found
+	 * @param SequenceID 			ID of the sequence to locate
+	 * @return The sequence, or nullptr if the ID was not found
 	 */
-	FORCEINLINE const FMovieSceneEvaluationTemplateInstance* GetInstance(FMovieSceneSequenceIDRef SequenceID) const
-	{
-		return SequenceID == MovieSceneSequenceID::Root ? &RootInstance : SubInstances.Find(SequenceID);
-	}
+	MOVIESCENE_API UMovieSceneSequence* GetSequence(FMovieSceneSequenceIDRef SequenceID) const;
 
 	/**
-	 * Access the underlying map of sub template instances
+	 * Attempt to locate a template corresponding to the specified Sequence ID
+	 *
+	 * @param SequenceID 			ID of the sequence template to locate
+	 * @return The template, or nullptr if the ID is invalid, or the template has not been compiled
 	 */
-	const TMap<FMovieSceneSequenceID, FMovieSceneEvaluationTemplateInstance>& GetSubInstances() const
-	{
-		return SubInstances;
-	}
+	MOVIESCENE_API FMovieSceneEvaluationTemplate* FindTemplate(FMovieSceneSequenceIDRef SequenceID);
 
 	/**
-	 * Attempt to locate the underlying sequence given a sequence ID
-	 *
-	 * @param SequenceID 			ID of the sequence to locate
-	 * @return The sequence, or nullptr if the ID was not found
+	 * Access the master sequence's hierarchy data
 	 */
-	UMovieSceneSequence* GetSequence(FMovieSceneSequenceIDRef SequenceID) const
+	const FMovieSceneSequenceHierarchy& GetHierarchy() const
 	{
-		if (const FMovieSceneEvaluationTemplateInstance* Instance = GetInstance(SequenceID))
-		{
-			return Instance->Sequence.Get();
-		}
-		return nullptr;
+		check(RootTemplate);
+		return RootTemplate->Hierarchy;
 	}
 
-
 	/**
 	 * Access the master sequence's hierarchy data
 	 */
-	const FMovieSceneSequenceHierarchy& GetHierarchy() const
+	FMovieSceneSequenceHierarchy& GetHierarchy()
 	{
-		return RootInstance.Template->Hierarchy;
+		check(RootTemplate);
+		return RootTemplate->Hierarchy;
 	}
 
 	/**
@@ -189,14 +186,6 @@ public:
 		return ThisFrameMetaData;
 	}
 
-	/**
-	 * Check whether this instance is dirty (one of its referenced templates has changed)
-	 */
-	bool IsDirty() const
-	{
-		return bIsDirty;
-	}
-
 	/**
 	 * Copy any actuators from this template instance into the specified accumulator
 	 *
@@ -207,91 +196,45 @@ public:
 private:
 
 	/**
-	 * Process entities that are newly evaluated, and those that are no longer being evaluated
-	 */
-	void CallSetupTearDown(IMovieScenePlayer& Player, FDelayedPreAnimatedStateRestore* DelayedRestore = nullptr);
-
-	/**
-	 * Evaluate a particular group of a segment
-	 */
-	void EvaluateGroup(const FMovieSceneEvaluationGroup& Group, const FMovieSceneContext& Context, IMovieScenePlayer& Player);
-
-	/**
-	 * Remap the specified sequence ID based on the currently evaluating sequence path, to the Root
+	 * Setup the current frame by finding or generating the necessary evaluation group and meta-data
 	 *
-	 * @param SequenceID			The sequence ID to find a template for
-	 * @return Pointer to a template instance, or nullptr if the ID was not found
+	 * @param Context				Evaluation context containing the time (or range) to evaluate
 	 */
-	FORCEINLINE_DEBUGGABLE FMovieSceneSequenceID GetSequenceIdForRoot(FMovieSceneSequenceID SequenceID) const
-	{
-		if (!ReverseOverrideRootPath.Num())
-		{
-			return SequenceID;
-		}
-
-		for (FMovieSceneSequenceID Parent : ReverseOverrideRootPath)
-		{
-			SequenceID = SequenceID.AccumulateParentID(Parent);
-		}
-		return SequenceID;
-	}
+	const FMovieSceneEvaluationGroup* SetupFrame(IMovieScenePlayer& Player, FMovieSceneSequenceID InOverrideRootID, FMovieSceneContext Context);
 
 	/**
-	 * Retrieve a specific template instance given a sequence ID. Assumes it exists.
-	 *
-	 * @param SequenceID			The sequence ID to find a template for
-	 * @return The template instance
+	 * Process entities that are newly evaluated, and those that are no longer being evaluated
 	 */
-	FORCEINLINE_DEBUGGABLE const FMovieSceneEvaluationTemplateInstance& GetInstanceChecked(FMovieSceneSequenceIDRef SequenceID) const
-	{
-		return SequenceID == MovieSceneSequenceID::Root ? RootInstance : SubInstances.FindChecked(SequenceID);
-	}
+	void CallSetupTearDown(IMovieScenePlayer& Player, FDelayedPreAnimatedStateRestore* DelayedRestore = nullptr);
 
 	/**
-	 * Attempt to retrieve a specific template instance given a sequence ID
-	 *
-	 * @param SequenceID			The sequence ID to find a template for
-	 * @return Pointer to a template instance, or nullptr if the ID was not found
+	 * Evaluate a particular group of a segment
 	 */
-	FORCEINLINE_DEBUGGABLE const FMovieSceneEvaluationTemplateInstance* FindInstance(FMovieSceneSequenceIDRef SequenceID) const
-	{
-		return SequenceID == MovieSceneSequenceID::Root ? &RootInstance : SubInstances.Find(SequenceID);
-	}
+	void EvaluateGroup(const FMovieSceneEvaluationGroup& Group, const FMovieSceneContext& Context, IMovieScenePlayer& Player);
 
 private:
 
-	/** Called whenever any of the sequences contained within this instance is changed */
-	void OnSequenceChanged();
-
-	/** Reset this template instance */
-	void Reset();
+	void RecreateInstances(const FMovieSceneEvaluationTemplateInstance& RootOverrideInstance, FMovieSceneSequenceID InOverrideRootID, IMovieScenePlayer& Player);
 
 private:
 
 	TWeakObjectPtr<UMovieSceneSequence> RootSequence;
 
-	FMovieSceneEvaluationTemplateInstance RootInstance;
+	FMovieSceneEvaluationTemplate* RootTemplate;
+
+	FMovieSceneEvaluationTemplateInstanceContainer TransientInstances;
 
 	/** Cache of everything that was evaluated last frame */
 	FMovieSceneEvaluationMetaData LastFrameMetaData;
 	/** Cache of everything that is evaluated this frame */
 	FMovieSceneEvaluationMetaData ThisFrameMetaData;
 
-	/** Map of all sub instances, arranged by sequence ID */
-	TMap<FMovieSceneSequenceID, FMovieSceneEvaluationTemplateInstance> SubInstances;
-
 	/** Template store responsible for supplying templates for a given sequence */
-	TSharedPtr<FMovieSceneSequenceTemplateStore> TemplateStore;
+	TSharedPtr<IMovieSceneSequenceTemplateStore> TemplateStore;
 
-	/** A reverse path of deterministic sequence IDs required to accumulate from local -> root */
-	TArray<FMovieSceneSequenceID, TInlineAllocator<8>> ReverseOverrideRootPath;
+	/** Override path that is used to remap inner sequence IDs to the root space when evaluating with a root override */
+	FMovieSceneRootOverridePath RootOverridePath;
 
 	/** Execution tokens that are used to apply animated state */
 	FMovieSceneExecutionTokens ExecutionTokens;
-
-	/** True when any of our templates are out of date, and need reinitializing */
-	bool bIsDirty;
-
-	/** Event that is triggered on update */
-	FOnUpdated OnUpdatedEvent;
 };
diff --git a/Engine/Source/Runtime/MovieScene/Public/Evaluation/MovieSceneEvaluationTrack.h b/Engine/Source/Runtime/MovieScene/Public/Evaluation/MovieSceneEvaluationTrack.h
--- a/Engine/Source/Runtime/MovieScene/Public/Evaluation/MovieSceneEvaluationTrack.h
+++ b/Engine/Source/Runtime/MovieScene/Public/Evaluation/MovieSceneEvaluationTrack.h
@@ -10,10 +10,15 @@
 #include "Evaluation/MovieSceneSegment.h"
 #include "Evaluation/MovieScenePlayback.h"
 #include "Evaluation/MovieSceneEvalTemplate.h"
+#include "Compilation/MovieSceneSegmentCompiler.h"
 #include "Evaluation/MovieSceneTrackImplementation.h"
+#include "MovieSceneEvaluationTree.h"
 #include "MovieSceneEvaluationTrack.generated.h"
 
 struct FMovieSceneInterrogationData;
+struct FMovieSceneEvaluationTrackSegments;
+
+class UMovieSceneTrack;
 
 /** Enumeration to determine how a track should be evaluated */
 UENUM()
@@ -26,6 +31,79 @@ enum class EEvaluationMethod : uint8
 	Swept,
 };
 
+/** Custom serialized tree of section evaluation data for a track */
+USTRUCT()
+struct FSectionEvaluationDataTree
+{
+	GENERATED_BODY()
+
+	bool Serialize(FArchive& Ar)
+	{
+		Ar << Tree;
+		return true;
+	}
+
+	TMovieSceneEvaluationTree<FSectionEvaluationData> Tree;
+};
+template<> struct TStructOpsTypeTraits<FSectionEvaluationDataTree> : public TStructOpsTypeTraitsBase2<FSectionEvaluationDataTree> { enum { WithSerializer = true }; };
+
+/** Structure that references a sorted array of segments by indirect identifiers */
+USTRUCT()
+struct FMovieSceneEvaluationTrackSegments
+{
+	GENERATED_BODY()
+
+	/** Index operator that allows lookup by ID */
+	FORCEINLINE const FMovieSceneSegment& operator[](FMovieSceneSegmentIdentifier ID) const
+	{
+		return SortedSegments[SegmentIdentifierToIndex[ID.GetIndex()]];
+	}
+
+	/** Index operator that allows lookup by ID */
+	FORCEINLINE FMovieSceneSegment& operator[](FMovieSceneSegmentIdentifier ID)
+	{
+		return SortedSegments[SegmentIdentifierToIndex[ID.GetIndex()]];
+	}
+
+	/** Access the sorted index of the specified valid identifier */
+	FORCEINLINE int32 GetSortedIndex(FMovieSceneSegmentIdentifier ID) const
+	{
+		return SegmentIdentifierToIndex[ID.GetIndex()];
+	}
+
+	/** Access the sorted array of segments */
+	TArrayView<FMovieSceneSegment> GetSorted()
+	{
+		return SortedSegments;
+	}
+
+	/** Access the sorted array of segments */
+	TArrayView<const FMovieSceneSegment> GetSorted() const
+	{
+		return SortedSegments;
+	}
+
+	/** Reset this container, invalidating any previously produced identifiers */
+	void Reset(int32 ExpectedCapacity)
+	{
+		SegmentIdentifierToIndex.Reset(ExpectedCapacity);
+		SortedSegments.Reset(ExpectedCapacity);
+	}
+
+	/** Add a new segment to the container */
+	FMovieSceneSegmentIdentifier Add(FMovieSceneSegment&& In);
+
+private:
+
+	/** Array of indices into SortedSegments where each FMovieSceneSegmentIdentifier represents and index into SegmentIdentifierToIndex. Never shuffled until the container is reset. */
+	UPROPERTY()
+	TArray<int32> SegmentIdentifierToIndex;
+
+	/** Array of segmented ranges contained within the track. */
+	UPROPERTY()
+	TArray<FMovieSceneSegment> SortedSegments;
+};
+
 /**
  * Evaluation track that is stored within an evaluation template for a sequence.
  * Contains user-defined evaluation templates, and an optional track implementation
@@ -48,28 +126,8 @@ struct FMovieSceneEvaluationTrack
 	FMovieSceneEvaluationTrack& operator=(const FMovieSceneEvaluationTrack&) = default;
 
 	/** Move construction/assignment */
-#if PLATFORM_COMPILER_HAS_DEFAULTED_FUNCTIONS
 	FMovieSceneEvaluationTrack(FMovieSceneEvaluationTrack&&) = default;
 	FMovieSceneEvaluationTrack& operator=(FMovieSceneEvaluationTrack&&) = default;
-#else
-	FMovieSceneEvaluationTrack(FMovieSceneEvaluationTrack&& RHS)
-	{
-		*this = MoveTemp(RHS);
-	}
-	FMovieSceneEvaluationTrack& operator=(FMovieSceneEvaluationTrack&& RHS)
-	{
-		ObjectBindingID = RHS.ObjectBindingID;
-		EvaluationPriority = RHS.EvaluationPriority;
-		EvaluationMethod = RHS.EvaluationMethod;
-		Segments = MoveTemp(RHS.Segments);
-		ChildTemplates = MoveTemp(RHS.ChildTemplates);
-		TrackTemplate = MoveTemp(RHS.TrackTemplate);
-		EvaluationGroup = MoveTemp(RHS.EvaluationGroup);
-		bEvaluateInPreroll = RHS.bEvaluateInPreroll;
-		bEvaluateInPostroll = RHS.bEvaluateInPostroll;
-		return *this;
-	}
-#endif
 
 public:
 
@@ -82,19 +140,19 @@ public:
 	}
 
 	/**
-	 * Const iteration of this track's segments
+	 * Get the segment from the given segment index
 	 */
-	FORCEINLINE TArrayView<const FMovieSceneSegment> GetSegments() const
+	FORCEINLINE const FMovieSceneSegment& GetSegment(FMovieSceneSegmentIdentifier ID) const
 	{
-		return TArrayView<const FMovieSceneSegment>(Segments.GetData(), Segments.Num());
+		return Segments[ID];
 	}
 
 	/**
 	 * Get the segment from the given segment index
 	 */
-	FORCEINLINE const FMovieSceneSegment& GetSegment(int32 SegmentIndex) const
+	FORCEINLINE TArrayView<const FMovieSceneSegment> GetSortedSegments() const
 	{
-		return Segments[SegmentIndex];
+		return Segments.GetSorted();
 	}
 
 	/**
@@ -225,46 +283,46 @@ public:
 	/**
 	 * Called to initialize the specified segment index
 	 * 
-	 * @param SegmentIndex			The segment we are evaluating
+	 * @param SegmentID				The segment we are evaluating
 	 * @param Operand				Operand that relates to the thing we will animate
 	 * @param Context				Current sequence context
 	 * @param PersistentData		Persistent data store
 	 * @param Player				The player that is responsible for playing back this template
 	 */
-	void Initialize(int32 SegmentIndex, const FMovieSceneEvaluationOperand& Operand, const FMovieSceneContext& Context, FPersistentEvaluationData& PersistentData, IMovieScenePlayer& Player) const;
+	void Initialize(FMovieSceneSegmentIdentifier SegmentID, const FMovieSceneEvaluationOperand& Operand, const FMovieSceneContext& Context, FPersistentEvaluationData& PersistentData, IMovieScenePlayer& Player) const;
 
 	/**
 	 * Called to evaluate the specified segment index
 	 * 
-	 * @param SegmentIndex			The segment we are evaluating
+	 * @param SegmentID				The segment we are evaluating
 	 * @param Operand				Operand that relates to the thing we will animate
 	 * @param Context				Current sequence context
 	 * @param PersistentData		Persistent data store
 	 * @param ExecutionTokens		Token stack on which to add tokens that will be executed later
 	 */
-	void Evaluate(int32 SegmentIndex, const FMovieSceneEvaluationOperand& Operand, const FMovieSceneContext& Context, const FPersistentEvaluationData& PersistentData, FMovieSceneExecutionTokens& ExecutionTokens) const;
+	void Evaluate(FMovieSceneSegmentIdentifier SegmentID, const FMovieSceneEvaluationOperand& Operand, const FMovieSceneContext& Context, const FPersistentEvaluationData& PersistentData, FMovieSceneExecutionTokens& ExecutionTokens) const;
 
 	/**
 	 * Default implementation of initialization of child templates for the specified segment
 	 * 
-	 * @param SegmentIndex			The segment we are evaluating
+	 * @param SegmentID				The segment we are evaluating
 	 * @param Operand				Operand that relates to the thing we will animate
 	 * @param Context				Current sequence context
 	 * @param PersistentData		Persistent data store
 	 * @param Player				The player that is responsible for playing back this template
 	 */
-	MOVIESCENE_API void DefaultInitialize(int32 SegmentIndex, const FMovieSceneEvaluationOperand& Operand, FMovieSceneContext Context, FPersistentEvaluationData& PersistentData, IMovieScenePlayer& Player) const;
+	MOVIESCENE_API void DefaultInitialize(FMovieSceneSegmentIdentifier SegmentID, const FMovieSceneEvaluationOperand& Operand, FMovieSceneContext Context, FPersistentEvaluationData& PersistentData, IMovieScenePlayer& Player) const;
 
 	/**
 	 * Default implementation of evaluation of child templates for the specified segment
 	 * 
-	 * @param SegmentIndex			The segment we are evaluating
+	 * @param SegmentID				The segment we are evaluating
 	 * @param Operand				Operand that relates to the thing we will animate
 	 * @param Context				Current sequence context
 	 * @param PersistentData		Persistent data store
 	 * @param ExecutionTokens		Token stack on which to add tokens that will be executed later
 	 */
-	MOVIESCENE_API void DefaultEvaluate(int32 SegmentIndex, const FMovieSceneEvaluationOperand& Operand, const FMovieSceneContext& Context, const FPersistentEvaluationData& PersistentData, FMovieSceneExecutionTokens& ExecutionTokens) const;
+	MOVIESCENE_API void DefaultEvaluate(FMovieSceneSegmentIdentifier SegmentID, const FMovieSceneEvaluationOperand& Operand, const FMovieSceneContext& Context, const FPersistentEvaluationData& PersistentData, FMovieSceneExecutionTokens& ExecutionTokens) const;
 
 	/**
 	 * Interrogate this template for its output. Should not have any side effects.
@@ -273,19 +331,19 @@ public:
 	 * @param Container				Container to populate with the desired output from this track
 	 * @param BindingOverride		Optional binding to specify the object that is being animated by this track
 	 */
-	MOVIESCENE_API  void Interrogate(const FMovieSceneContext& Context, FMovieSceneInterrogationData& Container, UObject* BindingOverride = nullptr) const;
+	MOVIESCENE_API  void Interrogate(const FMovieSceneContext& Context, FMovieSceneInterrogationData& Container, UObject* BindingOverride = nullptr);
 
 private:
 
 	/**
 	 * Implementation function for static evaluation
 	 */
-	void EvaluateStatic(int32 SegmentIndex, const FMovieSceneEvaluationOperand& Operand, FMovieSceneContext Context, const FPersistentEvaluationData& PersistentData, FMovieSceneExecutionTokens& ExecutionTokens) const;
+	void EvaluateStatic(FMovieSceneSegmentIdentifier SegmentID, const FMovieSceneEvaluationOperand& Operand, FMovieSceneContext Context, const FPersistentEvaluationData& PersistentData, FMovieSceneExecutionTokens& ExecutionTokens) const;
 
 	/**
 	 * Implementation function for swept evaluation
 	 */
-	void EvaluateSwept(int32 SegmentIndex, const FMovieSceneEvaluationOperand& Operand, const FMovieSceneContext& Context, const FPersistentEvaluationData& PersistentData, FMovieSceneExecutionTokens& ExecutionTokens) const;
+	void EvaluateSwept(FMovieSceneSegmentIdentifier SegmentID, const FMovieSceneEvaluationOperand& Operand, const FMovieSceneContext& Context, const FPersistentEvaluationData& PersistentData, FMovieSceneExecutionTokens& ExecutionTokens) const;
 
 public:
 
@@ -305,12 +363,99 @@ public:
 	MOVIESCENE_API int32 AddChildTemplate(FMovieSceneEvalTemplatePtr&& InTemplate);
 
 	/**
-	 * Assign the specified segments to this track.
+	 * Find the first segment that overlaps the specified range
 	 *
-	 * @param InSegments		The segments that define this track
+	 * @param InLocalRange 			The range in this track's space to overlap
+	 * @return A (potentially invalid) segment identifier for the first segment whose range overlaps the specified range
 	 */
-	MOVIESCENE_API void SetSegments(TArray<FMovieSceneSegment>&& InSegments);
-	
+	FMovieSceneSegmentIdentifier FindFirstSegment(TRange<float> InLocalRange);
+
+	/**
+	 * Find or compile a segment for the specified time
+	 *
+	 * @param InTime 				The time to lookup or compile for
+	 * @return A (potentially invalid) segment identifier for the segment whose range overlaps the specified time
+	 */
+	FMovieSceneSegmentIdentifier GetSegmentFromTime(float InTime);
+
+	/**
+	 * Find or compile a segment for the specified iterator
+	 *
+	 * @param Iterator 				An iterator that should be used to find or compile time ranges. The iterator's current range is used to lookup segments
+	 * @return A (potentially invalid) segment identifier for the segment whose range overlaps the specified iterator's time range
+	 */
+	FMovieSceneSegmentIdentifier GetSegmentFromIterator(FMovieSceneEvaluationTreeRangeIterator Iterator);
+
+	/**
+	 * Find or compile all segments that overlap the specified range
+	 *
+	 * @param InLocalRange 			The range in this track's space to overlap
+	 * @return An (potentially empty) array of segment identifiers that overlap the specified range
+	 */
+	TArray<FMovieSceneSegmentIdentifier> GetSegmentsInRange(TRange<float> InLocalRange);
+
+	/**
+	 * Get the smallest range of unique FSectionEvaluationData combinations that overlaps the specified lower bound
+	 *
+	 * @param InLowerBound 			The lower bound from which to start looking for a time range
+	 * @return The smallest time range of unique evaluation data entries that encompasses the specified lower bound
+	 */
+	TRange<float> GetUniqueRangeFromLowerBound(TRangeBound<float> InLowerBound) const;
+
+	/**
+	 * Set the source track from which this track originates
+	 */
+	void SetSourceTrack(const UMovieSceneTrack* InSourceTrack)
+	{
+		SourceTrack = InSourceTrack;
+	}
+
+	/**
+	 * Get the source track from which this track originates
+	 */
+	const UMovieSceneTrack* GetSourceTrack() const
+	{
+		return SourceTrack;
+	}
+
+	/**
+	 * Add section evaluation data to the evaluation tree for this track
+	 *
+	 * @param Range 		The range that the specified section data should apply to
+	 * @param EvalData 		The actual evaluation data
+	 */
+	void AddTreeData(TRange<float> Range, FSectionEvaluationData EvalData)
+	{
+		EvaluationTree.Tree.Add(Range, EvalData);
+	}
+
+	/**
+	 * Add section evaluation data to the evaluation tree for this track, ensuring no duplciates
+	 *
+	 * @param Range 		The range that the specified section data should apply to
+	 * @param EvalData 		The actual evaluation data
+	 */
+	void AddUniqueTreeData(TRange<float> Range, FSectionEvaluationData EvalData)
+	{
+		EvaluationTree.Tree.AddUnique(Range, EvalData);
+	}
+
+	/**
+	 * Iterate all this track's unique time ranges
+	 */
+	FMovieSceneEvaluationTreeRangeIterator Iterate() const
+	{
+		return FMovieSceneEvaluationTreeRangeIterator(EvaluationTree.Tree);
+	}
+
+	/**
+	 * Iterate all this track's unique time ranges
+	 */
+	TMovieSceneEvaluationTreeDataIterator<FSectionEvaluationData> GetData(FMovieSceneEvaluationTreeNodeHandle TreeNodeHandle) const
+	{
+		return EvaluationTree.Tree.GetAllData(TreeNodeHandle);
+	}
+
 	/**
 	 * Assign a track implementation template to this track
 	 * @note Track implementations are evaluated once per frame before any segments.
@@ -343,10 +488,12 @@ private:
 	void ValidateSegments();
 
 	/**
-	 * Locate the segment that resides at the specified time
-	 * @return A segment index, or INDEX_NONE
+	 * Compile a segment for the specified iterator node
+	 *
+	 * @param iterator 			An iterator that should be used to compile segments.
+	 * @return A segment identifier for the newly created segment (or invalid if compilation was unsuccessful)
 	 */
-	int32 FindSegmentIndex(float InTime) const;
+	FMovieSceneSegmentIdentifier CompileSegment(FMovieSceneEvaluationTreeRangeIterator Iterator);
 
 public:
 
@@ -396,8 +543,16 @@ private:
 
 	/** Array of segmented ranges contained within the track. */
 	UPROPERTY()
-	TArray<FMovieSceneSegment> Segments;
-	
+	FMovieSceneEvaluationTrackSegments Segments;
+
+	/** The movie scene track that created this evaluation track. */
+	UPROPERTY()
+	const UMovieSceneTrack* SourceTrack;
+
+	/** Evaluation tree specifying what happens at any given time. */
+	UPROPERTY()
+	FSectionEvaluationDataTree EvaluationTree;
+
 	/** Domain-specific evaluation templates (normally 1 per section) */
 	UPROPERTY()
 	TArray<FMovieSceneEvalTemplatePtr> ChildTemplates;
@@ -420,3 +575,16 @@ private:
 };
 
 template<> struct TStructOpsTypeTraits<FMovieSceneEvaluationTrack> : public TStructOpsTypeTraitsBase2<FMovieSceneEvaluationTrack> { enum { WithPostSerialize = true, WithCopy = false }; };
+
+#if WITH_DEV_AUTOMATION_TESTS
+
+struct FScopedOverrideTrackSegmentBlender
+{
+	MOVIESCENE_API FScopedOverrideTrackSegmentBlender(FMovieSceneTrackSegmentBlenderPtr&& InTrackSegmentBlender);
+	MOVIESCENE_API ~FScopedOverrideTrackSegmentBlender();
+
+	FScopedOverrideTrackSegmentBlender(const FScopedOverrideTrackSegmentBlender&) = delete;
+	FScopedOverrideTrackSegmentBlender& operator=(const FScopedOverrideTrackSegmentBlender&) = delete;
+};
+
+#endif
\ No newline at end of file
diff --git a/Engine/Source/Runtime/MovieScene/Public/Evaluation/MovieScenePropertyTemplate.h b/Engine/Source/Runtime/MovieScene/Public/Evaluation/MovieScenePropertyTemplate.h
--- a/Engine/Source/Runtime/MovieScene/Public/Evaluation/MovieScenePropertyTemplate.h
+++ b/Engine/Source/Runtime/MovieScene/Public/Evaluation/MovieScenePropertyTemplate.h
@@ -245,10 +245,13 @@ struct TPropertyActuator : TMovieSceneBlendingActuator<PropertyType>
 
 	virtual void Actuate(UObject* InObject, typename TCallTraits<PropertyType>::ParamType InFinalValue, const TBlendableTokenStack<PropertyType>& OriginalStack, const FMovieSceneContext& Context, FPersistentEvaluationData& PersistentData, IMovieScenePlayer& Player) override
 	{
-		check(InObject);
+		ensureMsgf(InObject, TEXT("Attempting to evaluate a Property track with a null object."));
 
-		OriginalStack.SavePreAnimatedState(Player, *InObject, PropertyData.PropertyID, PropertyTemplate::FTokenProducer<PropertyType>(*PropertyData.PropertyBindings));
-		PropertyData.PropertyBindings->CallFunction<PropertyType>(*InObject, InFinalValue);
+		if (InObject)
+		{
+			OriginalStack.SavePreAnimatedState(Player, *InObject, PropertyData.PropertyID, PropertyTemplate::FTokenProducer<PropertyType>(*PropertyData.PropertyBindings));
+			PropertyData.PropertyBindings->CallFunction<PropertyType>(*InObject, InFinalValue);
+		}
 	}
 };
 
diff --git a/Engine/Source/Runtime/MovieScene/Public/Evaluation/MovieSceneRootOverridePath.h b/Engine/Source/Runtime/MovieScene/Public/Evaluation/MovieSceneRootOverridePath.h
new file mode 100644
--- /dev/null
+++ b/Engine/Source/Runtime/MovieScene/Public/Evaluation/MovieSceneRootOverridePath.h
@@ -0,0 +1,71 @@
+// Copyright 1998-2017 Epic Games, Inc. All Rights Reserved.
+
+#pragma once
+
+#include "MovieSceneSequenceID.h"
+
+struct FMovieSceneSequenceHierarchy;
+
+/**
+ * A path of unaccumulated sequence IDs ordered from child->parent->grandparent that is used to generate unique sequenceIDs for inner sequences
+ * Optimized for Remap rather than Push/Pop by keeping sequence IDs child-parent order (the order they are required for remapping)
+ */
+struct MOVIESCENE_API FMovieSceneRootOverridePath
+{
+	/**
+	 * Remap the specified sequence ID based on the currently evaluating sequence path, to the Root
+	 *
+	 * @param SequenceID			The sequence ID to find a template for
+	 * @return Pointer to a template instance, or nullptr if the ID was not found
+	 */
+	FORCEINLINE_DEBUGGABLE FMovieSceneSequenceID Remap(FMovieSceneSequenceID SequenceID) const
+	{
+		if (LIKELY(!ReverseOverrideRootPath.Num()))
+		{
+			return SequenceID;
+		}
+
+		for (FMovieSceneSequenceID Parent : ReverseOverrideRootPath)
+		{
+			SequenceID = SequenceID.AccumulateParentID(Parent);
+		}
+		return SequenceID;
+	}
+
+	/**
+	 * Reset this path to its default state (pointing to the root sequence)
+	 */
+	void Reset();
+
+	/**
+	 * Set up this path from a specific sequence ID that points to a particular sequence in the specified hierarchy
+	 *
+	 * @param LeafID 			ID of the child-most sequence to include in this path
+	 * @param RootHierarchy 	Hierarchy to get unaccumulated sequence IDs from
+	 */
+	void Set(FMovieSceneSequenceID LeafID, const FMovieSceneSequenceHierarchy& RootHierarchy);
+
+	/**
+	 * Push a new child sequence ID into this path
+	 *
+	 * @param InUnaccumulatedSequenceID		The sequece ID of the child sequence within its outer parent
+	 */
+	void Push(FMovieSceneSequenceID InUnaccumulatedSequenceID)
+	{
+		ReverseOverrideRootPath.Insert(InUnaccumulatedSequenceID, 0);
+	}
+
+
+	/**
+	 * Pop the child-most sequence ID from this path
+	 */
+	void Pop()
+	{
+		ReverseOverrideRootPath.RemoveAt(0, 1, false);
+	}
+
+private:
+
+	/** A reverse path of deterministic sequence IDs required to accumulate from local -> root */
+	TArray<FMovieSceneSequenceID, TInlineAllocator<8>> ReverseOverrideRootPath;
+};
\ No newline at end of file
diff --git a/Engine/Source/Runtime/MovieScene/Public/Evaluation/MovieSceneSegment.h b/Engine/Source/Runtime/MovieScene/Public/Evaluation/MovieSceneSegment.h
--- a/Engine/Source/Runtime/MovieScene/Public/Evaluation/MovieSceneSegment.h
+++ b/Engine/Source/Runtime/MovieScene/Public/Evaluation/MovieSceneSegment.h
@@ -6,8 +6,12 @@
 #include "UObject/ObjectMacros.h"
 #include "UObject/Class.h"
 #include "Containers/ArrayView.h"
+#include "SequencerObjectVersion.h"
+
 #include "MovieSceneSegment.generated.h"
 
+struct FMovieSceneEvaluationTrackSegments;
+
 /** Enumeration specifying how to evaluate a particular section when inside a segment */
 UENUM()
 enum class ESectionEvaluationFlags : uint8
@@ -21,6 +25,55 @@ enum class ESectionEvaluationFlags : uint8
 };
 ENUM_CLASS_FLAGS(ESectionEvaluationFlags);
 
+/** A unique identifier for a segment within a FMovieSceneEvaluationTrackSegments container */
+USTRUCT()
+struct FMovieSceneSegmentIdentifier
+{
+	GENERATED_BODY()
+
+	FMovieSceneSegmentIdentifier() : IdentifierIndex(INDEX_NONE) {}
+
+	friend bool operator==(FMovieSceneSegmentIdentifier A, FMovieSceneSegmentIdentifier B) { return A.IdentifierIndex == B.IdentifierIndex; }
+	friend bool operator!=(FMovieSceneSegmentIdentifier A, FMovieSceneSegmentIdentifier B) { return A.IdentifierIndex != B.IdentifierIndex; }
+	friend uint32 GetTypeHash(FMovieSceneSegmentIdentifier LHS) { return GetTypeHash(LHS.IdentifierIndex); }
+
+	bool IsValid() const
+	{
+		return IdentifierIndex != INDEX_NONE;
+	}
+
+	/** Custom serializer as this is just a type-safe int */
+	bool Serialize(FArchive& Ar)
+	{
+		Ar << IdentifierIndex;
+		return true;
+	}
+
+	friend FArchive& operator<<(FArchive& Ar, FMovieSceneSegmentIdentifier ID)
+	{
+		ID.Serialize(Ar);
+		return Ar;
+	}
+
+private:
+
+	int32 GetIndex() const
+	{
+		return IdentifierIndex;
+	}
+	explicit FMovieSceneSegmentIdentifier(int32 InIdentifier) : IdentifierIndex(InIdentifier) {}
+
+	friend FMovieSceneEvaluationTrackSegments;
+
+	UPROPERTY()
+	int32 IdentifierIndex;
+};
+
+template<> struct TStructOpsTypeTraits<FMovieSceneSegmentIdentifier> : public TStructOpsTypeTraitsBase2<FMovieSceneSegmentIdentifier>
+{
+	enum { WithSerializer = true, WithIdenticalViaEquality = true };
+};
+
 /**
  * Evaluation data that specifies information about what to evaluate for a given template
  */
@@ -63,6 +116,12 @@ struct FSectionEvaluationData
 	/** Check if this is a postroll eval */
 	FORCEINLINE bool IsPostRoll() const { return (Flags & ESectionEvaluationFlags::PostRoll) != ESectionEvaluationFlags::None; }
 
+	friend FArchive& operator<<(FArchive& Ar, FSectionEvaluationData& Data)
+	{
+		FSectionEvaluationData::StaticStruct()->SerializeItem(Ar, &Data, nullptr);
+		return Ar;
+	}
+
 	/** The implementation index we should evaluate (index into FMovieSceneEvaluationTrack::ChildTemplates) */
 	UPROPERTY()
 	int32 ImplIndex;
@@ -89,10 +148,12 @@ struct FMovieSceneSegment
 
 	FMovieSceneSegment(const TRange<float>& InRange)
 		: Range(InRange)
+		, bAllowEmpty(false)
 	{}
 
 	FMovieSceneSegment(const TRange<float>& InRange, TArrayView<const FSectionEvaluationData> InApplicationImpls)
 		: Range(InRange)
+		, bAllowEmpty(true)
 	{
 		Impls.Reserve(InApplicationImpls.Num());
 		for (const FSectionEvaluationData& Impl : InApplicationImpls)
@@ -103,7 +164,22 @@ struct FMovieSceneSegment
 
 	friend bool operator==(const FMovieSceneSegment& A, const FMovieSceneSegment& B)
 	{
-		return A.Range == B.Range && A.Impls == B.Impls;
+		return A.Range == B.Range && A.ID == B.ID && A.bAllowEmpty == B.bAllowEmpty && A.Impls == B.Impls;
+	}
+
+	bool IsValid() const
+	{
+		return bAllowEmpty || Impls.Num() != 0;
+	}
+
+	bool CombineWith(const FMovieSceneSegment& OtherSegment)
+	{
+		if (Range.Adjoins(OtherSegment.Range) && Impls == OtherSegment.Impls && bAllowEmpty == OtherSegment.bAllowEmpty)
+		{
+			Range = TRange<float>::Hull(OtherSegment.Range, Range);
+			return true;
+		}
+		return false;
 	}
 
 	/** Custom serializer to accomodate the inline allocator on our array */
@@ -111,6 +187,13 @@ struct FMovieSceneSegment
 	{
 		Ar << Range;
 
+		Ar.UsingCustomVersion(FSequencerObjectVersion::GUID);
+		if (Ar.CustomVer(FSequencerObjectVersion::GUID) >= FSequencerObjectVersion::EvaluationTree)
+		{
+			Ar << ID;
+			Ar << bAllowEmpty;
+		}
+
 		int32 NumStructs = Impls.Num();
 		Ar << NumStructs;
 
@@ -119,7 +202,7 @@ struct FMovieSceneSegment
 			for (int32 Index = 0; Index < NumStructs; ++Index)
 			{
 				FSectionEvaluationData Data;
-				FSectionEvaluationData::StaticStruct()->SerializeItem(Ar, &Data, nullptr);
+				Ar << Data;
 				Impls.Add(Data);
 			}
 		}
@@ -127,7 +210,7 @@ struct FMovieSceneSegment
 		{
 			for (FSectionEvaluationData& Data : Impls)
 			{
-				FSectionEvaluationData::StaticStruct()->SerializeItem(Ar, &Data, nullptr);
+				Ar << Data;
 			}
 		}
 		return true;
@@ -136,6 +219,11 @@ struct FMovieSceneSegment
 	/** The segment's range */
 	FFloatRange Range;
 
+	FMovieSceneSegmentIdentifier ID;
+
+	/** Whether this segment has been generated yet or not */
+	bool bAllowEmpty;
+
 	/** Array of implementations that reside at the segment's range */
 	TArray<FSectionEvaluationData, TInlineAllocator<4>> Impls;
 };
@@ -143,4 +231,4 @@ struct FMovieSceneSegment
 template<> struct TStructOpsTypeTraits<FMovieSceneSegment> : public TStructOpsTypeTraitsBase2<FMovieSceneSegment>
 {
 	enum { WithSerializer = true, WithCopy = true, WithIdenticalViaEquality = true };
-};
+};
\ No newline at end of file
diff --git a/Engine/Source/Runtime/MovieScene/Public/Evaluation/MovieSceneSequenceHierarchy.h b/Engine/Source/Runtime/MovieScene/Public/Evaluation/MovieSceneSequenceHierarchy.h
--- a/Engine/Source/Runtime/MovieScene/Public/Evaluation/MovieSceneSequenceHierarchy.h
+++ b/Engine/Source/Runtime/MovieScene/Public/Evaluation/MovieSceneSequenceHierarchy.h
@@ -6,10 +6,14 @@
 #include "UObject/ObjectMacros.h"
 #include "Evaluation/MovieSceneSequenceTransform.h"
 #include "Evaluation/MovieSceneSectionParameters.h"
+#include "SoftObjectPath.h"
 #include "MovieSceneSequenceID.h"
+#include "MovieSceneSequenceInstanceData.h"
+#include "ArrayView.h"
 #include "MovieSceneSequenceHierarchy.generated.h"
 
 class UMovieSceneSequence;
+class UMovieSceneSubSection;
 struct FMovieSceneSequenceID;
 
 /**
@@ -23,51 +27,33 @@ struct FMovieSceneSubSequenceData
 	/**
 	 * Default constructor for serialization
 	 */
-	FMovieSceneSubSequenceData()
-		: Sequence(nullptr)
-	{}
+	MOVIESCENE_API FMovieSceneSubSequenceData();
 
 	/**
 	 * Construction from a movie scene sequence, and a sub section name, and its valid play range
 	 */
-	FMovieSceneSubSequenceData(UMovieSceneSequence& InSequence, FMovieSceneSequenceID InDeterministicSequenceID
-#if WITH_EDITORONLY_DATA
-		, FName InSectionPath
-		, TRange<float> InValidPlayRange
-#endif
-		)
-		: Sequence(&InSequence)
-		, SequenceKeyObject(nullptr)
-		, DeterministicSequenceID(InDeterministicSequenceID)
-		, PreRollRange(TRange<float>::Empty())
-		, PostRollRange(TRange<float>::Empty())
-		, HierarchicalBias(0)
-#if WITH_EDITORONLY_DATA
-		, SectionPath(InSectionPath)
-		, ValidPlayRange(InValidPlayRange)
-#endif
-	{}
+	MOVIESCENE_API FMovieSceneSubSequenceData(const UMovieSceneSubSection& InSubSection);
 
-	/** The sequence that the sub section references */
-	UPROPERTY()
-	UMovieSceneSequence* Sequence;
+	/**
+	 * Get this sub sequence's sequence asset
+	 */
+	MOVIESCENE_API UMovieSceneSequence* GetSequence() const;
 
-	/** The key object that the sub section uses. Usually either the sequence or the section. */
-	UPROPERTY()
-	const UObject* SequenceKeyObject;
+	/** The sequence that the sub section references */
+	UPROPERTY(meta=(AllowedClasses="MovieSceneSequence"))
+	FSoftObjectPath Sequence;
 
 	/** Transform that transforms a given time from the sequences outer space, to its authored space. */
 	UPROPERTY()
 	FMovieSceneSequenceTransform RootToSequenceTransform;
 
-	/** Cached signature of the evaluation template */
-	UPROPERTY()
-	FGuid SourceSequenceSignature;
-
-	/** This sequence's deterministic sequence ID. Used in editor to reduce the risk of collisions on recompilation */ 
+	/** This sequence's deterministic sequence ID. Used in editor to reduce the risk of collisions on recompilation. */ 
 	UPROPERTY()
 	FMovieSceneSequenceID DeterministicSequenceID;
 
+	/** This sub sequence's playback range according to its parent sub section. Clamped recursively during template generation */
+	UPROPERTY()
+	FFloatRange PlayRange;
 
 	/** The sequence preroll range considering the start offset */
 	UPROPERTY()
@@ -81,16 +67,21 @@ struct FMovieSceneSubSequenceData
 	UPROPERTY()
 	int32 HierarchicalBias;
 
+	/** Instance data that should be used for any tracks contained immediately within this sub sequence */
+	UPROPERTY()
+	FMovieSceneSequenceInstanceDataPtr InstanceData;
+
 #if WITH_EDITORONLY_DATA
 
 	/** This sequence's path within its movie scene */
 	UPROPERTY()
 	FName SectionPath;
-
-	/** This sub sequence's valid bounds according to its parent sub section. Clamped recursively during template generation */
-	UPROPERTY()
-	FFloatRange ValidPlayRange;
 #endif
+
+private:
+
+	/** Cached version of the sequence to avoid resolving it every time */
+	mutable TWeakObjectPtr<UMovieSceneSequence> CachedSequence;
 };
 
 /**
@@ -133,7 +124,7 @@ struct FMovieSceneSequenceHierarchy
 
 	FMovieSceneSequenceHierarchy()
 	{
-		Hierarchy.Add(MovieSceneSequenceID::Root.GetInternalValue(), FMovieSceneSequenceHierarchyNode(MovieSceneSequenceID::Invalid));
+		Hierarchy.Add(MovieSceneSequenceID::Root, FMovieSceneSequenceHierarchyNode(MovieSceneSequenceID::Invalid));
 	}
 
 	/**
@@ -144,7 +135,7 @@ struct FMovieSceneSequenceHierarchy
 	 */
 	const FMovieSceneSequenceHierarchyNode* FindNode(FMovieSceneSequenceIDRef SequenceID) const
 	{
-		return Hierarchy.Find(SequenceID.GetInternalValue());
+		return Hierarchy.Find(SequenceID);
 	}
 
 	/**
@@ -155,7 +146,7 @@ struct FMovieSceneSequenceHierarchy
 	 */
 	FMovieSceneSequenceHierarchyNode* FindNode(FMovieSceneSequenceIDRef SequenceID)
 	{
-		return Hierarchy.Find(SequenceID.GetInternalValue());
+		return Hierarchy.Find(SequenceID);
 	}
 
 	/**
@@ -166,7 +157,7 @@ struct FMovieSceneSequenceHierarchy
 	 */
 	const FMovieSceneSubSequenceData* FindSubData(FMovieSceneSequenceIDRef SequenceID) const
 	{
-		return SubSequences.Find(SequenceID.GetInternalValue());
+		return SubSequences.Find(SequenceID);
 	}
 
 	/**
@@ -177,7 +168,7 @@ struct FMovieSceneSequenceHierarchy
 	 */
 	FMovieSceneSubSequenceData* FindSubData(FMovieSceneSequenceIDRef SequenceID)
 	{
-		return SubSequences.Find(SequenceID.GetInternalValue());
+		return SubSequences.Find(SequenceID);
 	}
 
 	/**
@@ -187,20 +178,9 @@ struct FMovieSceneSequenceHierarchy
 	 * @param ThisSequenceID 		The sequence ID of the sequence the data relates to
 	 * @param ParentID 				The parent ID of this sequence data
 	 */
-	void Add(const FMovieSceneSubSequenceData& Data, FMovieSceneSequenceIDRef ThisSequenceID, FMovieSceneSequenceIDRef ParentID)
-	{
-		SubSequences.Add(ThisSequenceID.GetInternalValue(), Data);
-
-		FMovieSceneSequenceHierarchyNode Node(ParentID);
-		Hierarchy.Add(ThisSequenceID.GetInternalValue(), Node);
+	void Add(const FMovieSceneSubSequenceData& Data, FMovieSceneSequenceIDRef ThisSequenceID, FMovieSceneSequenceIDRef ParentID);
 
-		if (ParentID != MovieSceneSequenceID::Invalid)
-		{
-			FMovieSceneSequenceHierarchyNode* Parent = Hierarchy.Find(ParentID.GetInternalValue());
-			check(Parent);
-			Parent->Children.Add(ThisSequenceID);
-		}
-	}
+	void Remove(TArrayView<const FMovieSceneSequenceID> SequenceIDs);
 
 	/** Access to all the subsequence data */
 	const TMap<FMovieSceneSequenceID, FMovieSceneSubSequenceData>& AllSubSequenceData() const
@@ -212,9 +192,9 @@ private:
 
 	/** Map of all (recursive) sub sequences found in this template, keyed on sequence ID */
 	UPROPERTY()
-	TMap<uint32, FMovieSceneSubSequenceData> SubSequences;
+	TMap<FMovieSceneSequenceID, FMovieSceneSubSequenceData> SubSequences;
 
 	/** Structural information describing the structure of the sequence */
 	UPROPERTY()
-	TMap<uint32, FMovieSceneSequenceHierarchyNode> Hierarchy;
+	TMap<FMovieSceneSequenceID, FMovieSceneSequenceHierarchyNode> Hierarchy;
 };
diff --git a/Engine/Source/Runtime/MovieScene/Public/Evaluation/MovieSceneSequenceTemplateStore.h b/Engine/Source/Runtime/MovieScene/Public/Evaluation/MovieSceneSequenceTemplateStore.h
--- a/Engine/Source/Runtime/MovieScene/Public/Evaluation/MovieSceneSequenceTemplateStore.h
+++ b/Engine/Source/Runtime/MovieScene/Public/Evaluation/MovieSceneSequenceTemplateStore.h
@@ -2,37 +2,26 @@
 
 #pragma once
 
-#include "CoreMinimal.h"
-#include "UObject/ObjectKey.h"
+#include "UniquePtr.h"
+#include "Evaluation/MovieSceneEvaluationTemplate.h"
 
 class UMovieSceneSequence;
-struct FMovieSceneEvaluationTemplate;
 
-struct FMovieSceneSequenceTemplateStore
+struct IMovieSceneSequenceTemplateStore
 {
-	FMovieSceneSequenceTemplateStore()
-#if WITH_EDITORONLY_DATA
-		: bTemplatesAreVolatile(true)
-#else
-		// In cooked builds we assume the evaluation template can never change, and is thus always up-to-date
-		: bTemplatesAreVolatile(false)
-#endif
-	{
-	}
+	virtual ~IMovieSceneSequenceTemplateStore() {}
 
-	virtual ~FMovieSceneSequenceTemplateStore() {}
-	
-	bool AreTemplatesVolatile() const
-	{
-		return bTemplatesAreVolatile;
-	}
-
-	MOVIESCENE_API virtual FMovieSceneEvaluationTemplate& GetCompiledTemplate(UMovieSceneSequence& Sequence);
-
-	MOVIESCENE_API virtual FMovieSceneEvaluationTemplate& GetCompiledTemplate(UMovieSceneSequence& Sequence, FObjectKey InSequenceKey);
-
-protected:
+	MOVIESCENE_API virtual FMovieSceneEvaluationTemplate& AccessTemplate(UMovieSceneSequence& Sequence) = 0;
+};
 
-	/** True where templates are liable to change, and we should listen for changes */
-	bool bTemplatesAreVolatile;
+struct FMovieSceneSequencePrecompiledTemplateStore : IMovieSceneSequenceTemplateStore
+{
+	MOVIESCENE_API virtual FMovieSceneEvaluationTemplate& AccessTemplate(UMovieSceneSequence& Sequence);
 };
+
+struct FMovieSceneSequenceTemporaryTemplateStore : IMovieSceneSequenceTemplateStore
+{
+	MOVIESCENE_API virtual FMovieSceneEvaluationTemplate& AccessTemplate(UMovieSceneSequence& Sequence);
+private:
+	TMap<FObjectKey, FMovieSceneEvaluationTemplate> Templates;
+};
\ No newline at end of file
diff --git a/Engine/Source/Runtime/MovieScene/Public/Evaluation/MovieSceneTrackIdentifier.h b/Engine/Source/Runtime/MovieScene/Public/Evaluation/MovieSceneTrackIdentifier.h
--- a/Engine/Source/Runtime/MovieScene/Public/Evaluation/MovieSceneTrackIdentifier.h
+++ b/Engine/Source/Runtime/MovieScene/Public/Evaluation/MovieSceneTrackIdentifier.h
@@ -18,6 +18,11 @@ struct FMovieSceneTrackIdentifier
 
 	static FMovieSceneTrackIdentifier Invalid() { return FMovieSceneTrackIdentifier(); }
 
+	explicit operator bool() const
+	{
+		return Value != -1;
+	}
+
 	FMovieSceneTrackIdentifier& operator++()
 	{
 		++Value;
@@ -52,6 +57,7 @@ struct FMovieSceneTrackIdentifier
 	/** Custom serialized to reduce memory footprint */
 	bool Serialize(FArchive& Ar)
 	{
+		Ar.UsingCustomVersion(FEditorObjectVersion::GUID);
 		if (Ar.CustomVer(FEditorObjectVersion::GUID) < FEditorObjectVersion::MovieSceneMetaDataSerialization)
 		{
 			return false;
diff --git a/Engine/Source/Runtime/MovieScene/Public/Evaluation/MovieSceneTrackImplementation.h b/Engine/Source/Runtime/MovieScene/Public/Evaluation/MovieSceneTrackImplementation.h
--- a/Engine/Source/Runtime/MovieScene/Public/Evaluation/MovieSceneTrackImplementation.h
+++ b/Engine/Source/Runtime/MovieScene/Public/Evaluation/MovieSceneTrackImplementation.h
@@ -7,6 +7,7 @@
 #include "UObject/Class.h"
 #include "Misc/InlineValue.h"
 #include "Evaluation/MovieSceneEvalTemplateBase.h"
+#include "MovieSceneSegment.h"
 #include "MovieSceneTrackImplementation.generated.h"
 
 
@@ -15,7 +16,6 @@ struct FMovieSceneEvaluationTrack;
 struct FMovieSceneExecutionTokens;
 struct FMovieSceneInterrogationData;
 
-
 /**
  * Structure that allows the implementation of setup/teardown/initialization/evaluation logic at the track level.
  */
@@ -54,13 +54,13 @@ public:
 	 *							Only called if EnableOverrides(CustomInitializeFlag) has been called (see SetupOverrides).
 	 *
 	 * @param Track				The parent evaluation track that has knowledge of all child tracks and segments
-	 * @param SegmentIndex		The index of the segment at the current time
+	 * @param SegmentID			The identifier of the segment at the current time
 	 * @param Operand			Unique handle to the operand on which we are to operate. May represent multiple objects. Resolve through IMovieScenePlayer::FindBoundObjects(Operand)
 	 * @param Context			Evaluation context specifying the current evaluation time, sub sequence transform and other relevant information.
 	 * @param PersistentData	Persistent data store which can be used to store arbitrary data pertaining to the current template that may be required in Evaluate(Swept)
 	 * @param Player			The movie scene player currently playing back this sequence
 	 */
-	virtual void Initialize(const FMovieSceneEvaluationTrack& Track, int32 SegmentIndex, const FMovieSceneEvaluationOperand& Operand, const FMovieSceneContext& Context, FPersistentEvaluationData& PersistentData, IMovieScenePlayer& Player) const
+	virtual void Initialize(const FMovieSceneEvaluationTrack& Track, FMovieSceneSegmentIdentifier SegmentID, const FMovieSceneEvaluationOperand& Operand, const FMovieSceneContext& Context, FPersistentEvaluationData& PersistentData, IMovieScenePlayer& Player) const
 	{
 		ensureMsgf(false, TEXT("FMovieSceneTrackImplementation::Initialize has not been implemented. Did you erroneously call EnableOverrides(CustomInitializeFlag)?"));
 	}
@@ -72,13 +72,13 @@ public:
 	 *							Only called if EnableOverrides(CustomEvaluateFlag) has been called (see SetupOverrides).
 	 *
 	 * @param Track				The parent evaluation track that has knowledge of all child tracks and segments
-	 * @param SegmentIndex		The index of the segment at the current time
+	 * @param SegmentID			The identifier of the segment at the current time
 	 * @param Operand			Unique handle to the operand on which we are to operate. Only to be used as a reference, or forwarded throgh to an execution token.
 	 * @param Context			Evaluation context specifying the current evaluation time, sub sequence transform and other relevant information.
 	 * @param PersistentData	Persistent data store which can be used to access arbitrary data pertaining to the current template that should have been set up in initialize.
 	 * @param ExecutionTokens	Stack of execution tokens that will be used to apply animated state to the environment at a later time.
 	 */
-	virtual void Evaluate(const FMovieSceneEvaluationTrack& Track, int32 SegmentIndex, const FMovieSceneEvaluationOperand& Operand, const FMovieSceneContext& Context, const FPersistentEvaluationData& PersistentData, FMovieSceneExecutionTokens& ExecutionTokens) const
+	virtual void Evaluate(const FMovieSceneEvaluationTrack& Track, FMovieSceneSegmentIdentifier SegmentID, const FMovieSceneEvaluationOperand& Operand, const FMovieSceneContext& Context, const FPersistentEvaluationData& PersistentData, FMovieSceneExecutionTokens& ExecutionTokens) const
 	{
 		ensureMsgf(false, TEXT("FMovieSceneTrackImplementation::Evaluate has not been implemented. Did you erroneously call EnableOverrides(CustomEvaluateFlag)?"));
 	}
diff --git a/Engine/Source/Runtime/MovieScene/Public/Evaluation/PersistentEvaluationData.h b/Engine/Source/Runtime/MovieScene/Public/Evaluation/PersistentEvaluationData.h
--- a/Engine/Source/Runtime/MovieScene/Public/Evaluation/PersistentEvaluationData.h
+++ b/Engine/Source/Runtime/MovieScene/Public/Evaluation/PersistentEvaluationData.h
@@ -7,6 +7,9 @@
 #include "MovieSceneFwd.h"
 #include "Evaluation/MovieSceneEvaluationKey.h"
 #include "Evaluation/MovieSceneEvaluationOperand.h"
+#include "Evaluation/MovieSceneSequenceInstanceData.h"
+
+class IMovieScenePlayer;
 
 /**
  * Unique identifier for shared persistent data entries (see FSharedPersistentDataKey)
@@ -84,10 +87,7 @@ struct FPersistentEvaluationData
 	/**
 	 * Proxy constructor from 2 externally owned maps for entity, and shared data
 	 */
-	FPersistentEvaluationData(TMap<FMovieSceneEvaluationKey, TUniquePtr<IPersistentEvaluationData>>& InEntityData, TMap<FSharedPersistentDataKey, TUniquePtr<IPersistentEvaluationData>>& InSharedData)
-		: EntityData(InEntityData)
-		, SharedData(InSharedData)
-	{}
+	MOVIESCENE_API FPersistentEvaluationData(IMovieScenePlayer& InPlayer);
 
 	FPersistentEvaluationData(const FPersistentEvaluationData&) = delete;
 	FPersistentEvaluationData& operator=(const FPersistentEvaluationData&) = delete;
@@ -115,6 +115,24 @@ public:
 	template<typename T>	   T*	FindSectionData() const	{ return Find<T>(SectionKey); }
 	void 							ResetSectionData()		{ Reset(SectionKey); }
 
+	/**
+	 * Get the raw instance data for the current sequence
+	 */
+	MOVIESCENE_API const FMovieSceneSequenceInstanceData* GetInstanceData() const;
+
+	/**
+	 * Find the current sequence's instance data as the templated type, provided its type matches
+	 */
+	template<typename T> const T*   FindInstanceData() const
+	{
+		const FMovieSceneSequenceInstanceData* InstanceData = GetInstanceData();
+		if (InstanceData && (&InstanceData->GetScriptStruct() == T::StaticStruct()))
+		{
+			return static_cast<const T*>(InstanceData);
+		}
+		return nullptr;
+	}
+
 public:
 	
 	/**
@@ -285,6 +303,9 @@ private:
 
 private:
 	
+	/** The movie scene player */
+	IMovieScenePlayer& Player;
+
 	/** Persistent data that's associated with a template entity (such as a track or a section) */
 	TMap<FMovieSceneEvaluationKey, TUniquePtr<IPersistentEvaluationData>>& EntityData;
 
diff --git a/Engine/Source/Runtime/MovieScene/Public/IMovieSceneModule.h b/Engine/Source/Runtime/MovieScene/Public/IMovieSceneModule.h
--- a/Engine/Source/Runtime/MovieScene/Public/IMovieSceneModule.h
+++ b/Engine/Source/Runtime/MovieScene/Public/IMovieSceneModule.h
@@ -68,4 +68,10 @@ public:
 	 * Find group parameters for a specific evaluation group
 	 */
 	virtual FMovieSceneEvaluationGroupParameters GetEvaluationGroupParameters(FName GroupName) const = 0;
+
+	/**
+	 * Get this module ptr as a weak ptr.
+	 * @note: resulting weak ptr should not be used to hold persistent strong references
+	 */
+	virtual TWeakPtr<IMovieSceneModule> GetWeakPtr() = 0;
 };
diff --git a/Engine/Source/Runtime/MovieScene/Public/IMovieSceneObjectSpawner.h b/Engine/Source/Runtime/MovieScene/Public/IMovieSceneObjectSpawner.h
--- a/Engine/Source/Runtime/MovieScene/Public/IMovieSceneObjectSpawner.h
+++ b/Engine/Source/Runtime/MovieScene/Public/IMovieSceneObjectSpawner.h
@@ -12,6 +12,7 @@ class UMovieScene;
 class IMovieSceneObjectSpawner;
 class ISequencer;
 class USequencerSettings;
+class UActorFactory;
 struct FTransformData;
 
 /** A delegate which will create an object spawner */
@@ -72,9 +73,10 @@ public:
 	 *
 	 * @param SourceObject		The source object to create the spawnable from
 	 * @param OwnerMovieScene	The owner movie scene that this spawnable type should reside in
+	 * @param ActorFactory      Optional actor factory to use to create spawnable type
 	 * @return the new spawnable type, or error
 	 */
-	virtual TValueOrError<FNewSpawnable, FText> CreateNewSpawnableType(UObject& SourceObject, UMovieScene& OwnerMovieScene) { return MakeError(NSLOCTEXT("IMovieSceneObjectSpawner", "NotSupported", "Not supported")); }
+	virtual TValueOrError<FNewSpawnable, FText> CreateNewSpawnableType(UObject& SourceObject, UMovieScene& OwnerMovieScene, UActorFactory* ActorFactory = nullptr) { return MakeError(NSLOCTEXT("IMovieSceneObjectSpawner", "NotSupported", "Not supported")); }
 
 	/**
 	 * Setup a new spawnable object with some default tracks and keys
diff --git a/Engine/Source/Runtime/MovieScene/Public/IMovieScenePlayer.h b/Engine/Source/Runtime/MovieScene/Public/IMovieScenePlayer.h
--- a/Engine/Source/Runtime/MovieScene/Public/IMovieScenePlayer.h
+++ b/Engine/Source/Runtime/MovieScene/Public/IMovieScenePlayer.h
@@ -59,6 +59,11 @@ public:
 	 */
 	virtual FMovieSceneRootEvaluationTemplateInstance& GetEvaluationTemplate() = 0;
 
+	/**
+	 * Whether this player can update the camera cut
+	 */
+	virtual bool CanUpdateCameraCut() const { return true; }
+
 	/**
 	 * Updates the perspective viewports with the actor to view through
 	 *
@@ -135,6 +140,11 @@ public:
 	 */
 	virtual UObject* GetPlaybackContext() const { return nullptr; }
 
+	/**
+	 * Access the global instance data object for this movie scene player
+	 */
+	virtual const UObject* GetInstanceData() const { return nullptr; }
+
 	/**
 	 * Access the event contexts for this movie scene player
 	 */
@@ -293,14 +303,6 @@ public:
 
 	/** Container that stores any per-animated state tokens  */
 	FMovieScenePreAnimatedState PreAnimatedState;
-
-public:
-
-	DEPRECATED(4.15, "Sequence Instances have been deprecated in favor of a template approach (see GetEvaluationTemplate).")
-	MOVIESCENE_API TSharedRef<FMovieSceneSequenceInstance> GetRootMovieSceneSequenceInstance();
-
-	DEPRECATED(4.15, "Please use FindBoundObjects directly.")
-	MOVIESCENE_API void GetRuntimeObjects(TSharedRef<FMovieSceneSequenceInstance> MovieSceneInstance, const FGuid& Guid, TArray<TWeakObjectPtr<>>& OutRuntimeObjects);
 	
 private:
 
diff --git a/Engine/Source/Runtime/MovieScene/Public/MovieScene.h b/Engine/Source/Runtime/MovieScene/Public/MovieScene.h
--- a/Engine/Source/Runtime/MovieScene/Public/MovieScene.h
+++ b/Engine/Source/Runtime/MovieScene/Public/MovieScene.h
@@ -106,6 +106,14 @@ public:
 	 */
 	FGuid AddSpawnable(const FString& Name, UObject& ObjectTemplate);
 
+	/*
+	 * Adds an existing spawnable to this movie scene.
+	 *
+	 * @param InNewSpawnable The posssesable to add.
+	 * @param InNewBinding The object binding to add.
+	 */
+	void AddSpawnable(const FMovieSceneSpawnable& InNewSpawnable, const FMovieSceneBinding& InNewBinding);
+
 	/**
 	 * Removes a spawnable from this movie scene.
 	 *
@@ -158,6 +166,14 @@ public:
 	 */
 	FGuid AddPossessable(const FString& Name, UClass* Class);
 
+	/*
+	 * Adds an existing possessable to this movie scene.
+	 *
+	 * @param InNewPossessable The posssesable to add.
+	 * @param InNewBinding The object binding to add.
+	 */
+	void AddPossessable(const FMovieScenePossessable& InNewPossessable, const FMovieSceneBinding& InNewBinding);
+
 	/**
 	 * Removes a possessable from this movie scene.
 	 *
@@ -170,14 +186,6 @@ public:
 	*/
 	bool ReplacePossessable(const FGuid& OldGuid, const FMovieScenePossessable& InNewPosessable);
 
-	DEPRECATED(4.15, "Please use ReplacePossessable(const FGuid&, const FMovieScenePossessable&) so that the possessable class gets updated correctly.")
-	bool ReplacePossessable(const FGuid& OldGuid, const FGuid& NewGuid, const FString& Name)
-	{
-		FMovieScenePossessable NewPossessable(Name, nullptr);
-		NewPossessable.SetGuid(NewGuid);
-		return ReplacePossessable(OldGuid, NewPossessable);
-	}
-
 	/**
 	 * Tries to locate a possessable in this MovieScene for the specified possessable GUID.
 	 *
diff --git a/Engine/Source/Runtime/MovieScene/Public/MovieSceneCommonHelpers.h b/Engine/Source/Runtime/MovieScene/Public/MovieSceneCommonHelpers.h
--- a/Engine/Source/Runtime/MovieScene/Public/MovieSceneCommonHelpers.h
+++ b/Engine/Source/Runtime/MovieScene/Public/MovieSceneCommonHelpers.h
@@ -99,6 +99,22 @@ public:
 	 * @param InInterpolation The interpolation to set
 	 */
 	static void SetKeyInterpolation(FRichCurve& InCurve, FKeyHandle InKeyHandle, EMovieSceneKeyInterpolation InKeyInterpolation);
+
+	/**
+	 * Sort predicate that sorts lower bounds of a range
+	 */
+	static bool SortLowerBounds(TRangeBound<float> A, TRangeBound<float> B)
+	{
+		return TRangeBound<float>::MinLower(A, B) == A && A != B;
+	}
+
+	/**
+	 * Sort predicate that sorts upper bounds of a range
+	 */
+	static bool SortUpperBounds(TRangeBound<float> A, TRangeBound<float> B)
+	{
+		return TRangeBound<float>::MinUpper(A, B) == A && A != B;
+	}
 };
 
 /**
diff --git a/Engine/Source/Runtime/MovieScene/Public/MovieSceneObjectBindingID.h b/Engine/Source/Runtime/MovieScene/Public/MovieSceneObjectBindingID.h
--- a/Engine/Source/Runtime/MovieScene/Public/MovieSceneObjectBindingID.h
+++ b/Engine/Source/Runtime/MovieScene/Public/MovieSceneObjectBindingID.h
@@ -40,6 +40,14 @@ struct FMovieSceneObjectBindingID
 	{
 	}
 
+	FMovieSceneObjectBindingID& operator=(const FMovieSceneObjectBindingID& InObjectBindingID)
+	{
+		Guid = InObjectBindingID.GetGuid();
+		SequenceID = int32(InObjectBindingID.GetSequenceID().GetInternalValue());
+		Space = InObjectBindingID.GetBindingSpace();
+		return *this;
+	}
+
 	/**
 	 * Check whether this object binding ID has been set to something valied
 	 * @note: does not imply that the ID resolves to a valid object
@@ -65,6 +73,14 @@ struct FMovieSceneObjectBindingID
 		return Guid;
 	}
 	
+	/**
+	 * Set the guid that identifies the object binding within the sequence
+	 */
+	void SetGuid(const FGuid& InGuid)
+	{
+		Guid = InGuid;
+	}
+	
 	/**
 	 * Access how this binding's sequence ID relates to the master sequence
 	 */
diff --git a/Engine/Source/Runtime/MovieScene/Public/MovieSceneSequence.h b/Engine/Source/Runtime/MovieScene/Public/MovieSceneSequence.h
--- a/Engine/Source/Runtime/MovieScene/Public/MovieSceneSequence.h
+++ b/Engine/Source/Runtime/MovieScene/Public/MovieSceneSequence.h
@@ -48,14 +48,6 @@ public:
 	 */
 	virtual bool CanPossessObject(UObject& Object, UObject* InPlaybackContext) const PURE_VIRTUAL(UMovieSceneSequence::CanPossessObject, return false;);
 
-	DEPRECATED(4.15, "Please implement LocateBoundObjects instead")
-	virtual UObject* FindPossessableObject(const FGuid& ObjectId, UObject* Context) const
-	{
-		TArray<UObject*, TInlineAllocator<1>> OutObjects;
-		LocateBoundObjects(ObjectId, Context, OutObjects);
-		return OutObjects.Num() ? OutObjects[0] : nullptr;
-	}
-
 	/**
 	 * Locate all the objects that correspond to the specified object ID, using the specified context
 	 *
@@ -79,9 +71,6 @@ public:
 		return OutObjects;
 	}
 
-	DEPRECATED(4.15, "Please use IMovieScenePlayer::FindObjectId or FindPossessableObjectId(UObject&, UObject*) instead.")
-	virtual FGuid FindPossessableObjectId(UObject& Object) const PURE_VIRTUAL(UMovieSceneSequence::FindPossessableObjectId, return FGuid(); );
-
 	/**
 	 * Attempt to find the guid relating to the specified object
 	 *
@@ -159,16 +148,9 @@ public:
 	MOVIESCENE_API virtual void PostDuplicate(bool bDuplicateForPIE) override;
 #endif
 
-	MOVIESCENE_API virtual void GenerateEvaluationTemplate(FMovieSceneEvaluationTemplate& Template, const FMovieSceneTrackCompilationParams& Params, FMovieSceneSequenceTemplateStore& Store);
-
 	UPROPERTY()
-	FCachedMovieSceneEvaluationTemplate EvaluationTemplate;
+	FMovieSceneEvaluationTemplate PrecompiledEvaluationTemplate;
 
-	UPROPERTY()
-	FMovieSceneTrackCompilationParams TemplateParameters;
-
-	UPROPERTY()
-	TMap<UObject*, FCachedMovieSceneEvaluationTemplate> InstancedSubSequenceEvaluationTemplates;
 
 	/* The default completion mode for this movie scene when a section's completion mode is set to project default */
 	UPROPERTY(config)
diff --git a/Engine/Source/Runtime/MovieScene/Public/MovieSceneSequencePlayer.h b/Engine/Source/Runtime/MovieScene/Public/MovieSceneSequencePlayer.h
--- a/Engine/Source/Runtime/MovieScene/Public/MovieSceneSequencePlayer.h
+++ b/Engine/Source/Runtime/MovieScene/Public/MovieSceneSequencePlayer.h
@@ -29,6 +29,8 @@ struct FMovieSceneSequencePlaybackSettings
 		, bDisableLookAtInput(false)
 		, bHidePlayer(false)
 		, bHideHud(false)
+		, bDisableCameraCuts(false)
+		, InstanceData(nullptr)
 	{ }
 
 	GENERATED_BODY()
@@ -69,6 +71,14 @@ struct FMovieSceneSequencePlaybackSettings
 	UPROPERTY(EditAnywhere, BlueprintReadWrite, Category="Cinematic")
 	bool bHideHud;
 
+	/** Disable camera cuts */
+	UPROPERTY(EditAnywhere, BlueprintReadWrite, Category="Cinematic")
+	bool bDisableCameraCuts;
+
+	/** An object that can implement specific instance overrides for the sequence */
+	UPROPERTY(BlueprintReadWrite, Category="Cinematic")
+	UObject* InstanceData;
+
 	/** Interface that defines overridden bindings for this sequence */
 	UPROPERTY()
 	TScriptInterface<IMovieSceneBindingOverridesInterface> BindingOverrides;
@@ -194,6 +204,14 @@ public:
 	UFUNCTION(BlueprintCallable, Category="Game|Cinematic")
 	float GetPlaybackEnd() const { return EndTime; }
 	
+	/** Set whether to disable camera cuts */
+	UFUNCTION(BlueprintCallable, Category="Game|Cinematic")
+	void SetDisableCameraCuts(bool bInDisableCameraCuts) { PlaybackSettings.bDisableCameraCuts = bInDisableCameraCuts; }
+
+	/** Set whether to disable camera cuts */
+	UFUNCTION(BlueprintCallable, Category="Game|Cinematic")
+	bool GetDisableCameraCuts() { return PlaybackSettings.bDisableCameraCuts; }
+
 	/** An event that is broadcast each time this level sequence player is updated */
 	DECLARE_EVENT_ThreeParams( UMovieSceneSequencePlayer, FOnMovieSceneSequencePlayerUpdated, const UMovieSceneSequencePlayer&, float /*current time*/, float /*previous time*/ );
 	FOnMovieSceneSequencePlayerUpdated& OnSequenceUpdated() const { return OnMovieSceneSequencePlayerUpdate; }
@@ -266,9 +284,11 @@ protected:
 	virtual void SetPlaybackStatus(EMovieScenePlayerStatus::Type InPlaybackStatus) override {}
 	virtual void SetViewportSettings(const TMap<FViewportClient*, EMovieSceneViewportParams>& ViewportParamsMap) override {}
 	virtual void GetViewportSettings(TMap<FViewportClient*, EMovieSceneViewportParams>& ViewportParamsMap) const override {}
+	virtual bool CanUpdateCameraCut() const override { return !PlaybackSettings.bDisableCameraCuts; }
 	virtual void UpdateCameraCut(UObject* CameraObject, UObject* UnlockIfCameraObject, bool bJumpCut) override {}
 	virtual void ResolveBoundObjects(const FGuid& InBindingId, FMovieSceneSequenceID SequenceID, UMovieSceneSequence& Sequence, UObject* ResolutionContext, TArray<UObject*, TInlineAllocator<1>>& OutObjects) const override;
 	virtual const IMovieSceneBindingOverridesInterface* GetBindingOverrides() const override { return PlaybackSettings.BindingOverrides ? &*PlaybackSettings.BindingOverrides : nullptr; }
+	virtual const UObject* GetInstanceData() const override { return PlaybackSettings.InstanceData; }
 
 	//~ UObject interface
 	virtual void BeginDestroy() override;
@@ -325,7 +345,7 @@ protected:
 
 	enum class ELatentAction
 	{
-		Stop, Pause
+		Stop, Pause, SetPlaybackPosition
 	};
 
 	/** Set of latent actions that are to be performed when the sequence has finished evaluating this frame */
@@ -350,4 +370,7 @@ private:
 
 	/** The maximum tick rate prior to playing (used for overriding delta time during playback). */
 	TOptional<double> OldMaxTickRate;
+
+	/** Playback position to use when invoking latent SetPlaybackPosition */
+	float LatentPlaybackPosition;
 };
diff --git a/Engine/Source/Runtime/MovieScene/Public/MovieSceneSpawnRegister.h b/Engine/Source/Runtime/MovieScene/Public/MovieSceneSpawnRegister.h
--- a/Engine/Source/Runtime/MovieScene/Public/MovieSceneSpawnRegister.h
+++ b/Engine/Source/Runtime/MovieScene/Public/MovieSceneSpawnRegister.h
@@ -110,9 +110,10 @@ public:
 	 *
 	 * @param SourceObject		The source object to create the spawnable from
 	 * @param OwnerMovieScene	The owner movie scene that this spawnable type should reside in
+	 * @param ActorFactory      Optional actor factory to use to create spawnable type
 	 * @return the new spawnable type, or error
 	 */
-	virtual TValueOrError<FNewSpawnable, FText> CreateNewSpawnableType(UObject& SourceObject, UMovieScene& OwnerMovieScene) { return MakeError(NSLOCTEXT("SpawnRegister", "NotSupported", "Not supported")); }
+	virtual TValueOrError<FNewSpawnable, FText> CreateNewSpawnableType(UObject& SourceObject, UMovieScene& OwnerMovieScene, UActorFactory* ActorFactory = nullptr) { return MakeError(NSLOCTEXT("SpawnRegister", "NotSupported", "Not supported")); }
 
 	/**
 	 * Called to save the default state of the specified spawnable
diff --git a/Engine/Source/Runtime/MovieScene/Public/MovieSceneSpawnable.h b/Engine/Source/Runtime/MovieScene/Public/MovieSceneSpawnable.h
--- a/Engine/Source/Runtime/MovieScene/Public/MovieSceneSpawnable.h
+++ b/Engine/Source/Runtime/MovieScene/Public/MovieSceneSpawnable.h
@@ -184,6 +184,10 @@ public:
 		Ownership = InOwnership;
 	}
 
+	/** Optional spawn transform */
+	UPROPERTY(EditAnywhere, AdvancedDisplay, Category = Actor)
+	FTransform SpawnTransform;
+
 private:
 
 	/** Unique identifier of the spawnable object. */
diff --git a/Engine/Source/Runtime/MovieScene/Public/MovieSceneTrack.h b/Engine/Source/Runtime/MovieScene/Public/MovieSceneTrack.h
--- a/Engine/Source/Runtime/MovieScene/Public/MovieSceneTrack.h
+++ b/Engine/Source/Runtime/MovieScene/Public/MovieSceneTrack.h
@@ -10,13 +10,18 @@
 #include "MovieSceneSignedObject.h"
 #include "MovieSceneSection.h"
 #include "Misc/InlineValue.h"
+#include "Compilation/MovieSceneSegmentCompiler.h"
 #include "MovieSceneTrack.generated.h"
 
 struct FMovieSceneEvaluationTrack;
-struct FMovieSceneSegmentCompilerRules;
-struct FMovieSceneSequenceTemplateStore;
+struct FMovieSceneTrackSegmentBlender;
+struct FMovieSceneTrackRowSegmentBlender;
 struct IMovieSceneTemplateGenerator;
 
+PRAGMA_DISABLE_DEPRECATION_WARNINGS
+	typedef TInlineValue<FMovieSceneSegmentCompilerRules> FMovieSceneDeprecatedCompilerRulesPtr;
+PRAGMA_ENABLE_DEPRECATION_WARNINGS
+
 /** Flags used to perform cook-time optimization of movie scene data */
 enum class ECookOptimizationFlags
 {
@@ -29,49 +34,14 @@ enum class ECookOptimizationFlags
 };
 ENUM_CLASS_FLAGS(ECookOptimizationFlags)
 
-/** Movie scene compilation parameters. Serialized items contribute to a compiled template's cached hash */
-USTRUCT()
-struct FMovieSceneTrackCompilationParams
-{
-	GENERATED_BODY()
-
-	FMovieSceneTrackCompilationParams()
-		: bForEditorPreview(false)
-		, bDuringBlueprintCompile(false)
-	{
-	}
-
-	friend bool operator!=(const FMovieSceneTrackCompilationParams& A, const FMovieSceneTrackCompilationParams& B)
-	{
-		return !(A == B);
-	}
-
-	friend bool operator==(const FMovieSceneTrackCompilationParams& A, const FMovieSceneTrackCompilationParams& B)
-	{
-		return A.bForEditorPreview == B.bForEditorPreview && A.bDuringBlueprintCompile == B.bDuringBlueprintCompile;
-	}
-
-	/** Whether we're generating for an editor preview, or for efficient runtime evaluation */
-	UPROPERTY()
-	bool bForEditorPreview;
-
-	/** Whether we're generating during a blueprint compile. As such, UObject types may not have been fully loaded. */
-	UPROPERTY()
-	bool bDuringBlueprintCompile;
-};
-
 /** Track compiler arguments */
 struct FMovieSceneTrackCompilerArgs
 {
-	FMovieSceneTrackCompilerArgs(IMovieSceneTemplateGenerator& InGenerator, FMovieSceneSequenceTemplateStore& InSubSequenceStore)
+	FMovieSceneTrackCompilerArgs(IMovieSceneTemplateGenerator& InGenerator)
 		: Generator(InGenerator)
-		, SubSequenceStore(InSubSequenceStore)
 	{
 	}
 
-	/** Compilation parameters */
-	FMovieSceneTrackCompilationParams Params;
-
 	/** The object binding ID that this track belongs to. */
 	FGuid ObjectBindingId;
 
@@ -80,8 +50,15 @@ struct FMovieSceneTrackCompilerArgs
 	/** The generator responsible for generating the template */
 	IMovieSceneTemplateGenerator& Generator;
 
-	/** Store that describes how to find sub sequence templates */
-	FMovieSceneSequenceTemplateStore& SubSequenceStore;
+	struct FMovieSceneSequenceTemplateStore {};
+
+	DEPRECATED(4.19, "Template store is no longer supplied as part of track compilation")
+	FMovieSceneSequenceTemplateStore SubSequenceStore;
+
+	struct FMovieSceneTrackCompilationParams { 	bool bForEditorPreview, bDuringBlueprintCompile; };
+
+	DEPRECATED(4.19, "Template store is no longer supplied as part of track compilation")
+	FMovieSceneTrackCompilationParams Params;
 };
 
 /** Generic evaluation options for any track */
@@ -169,13 +146,19 @@ public:
 	 * Get compiler rules to use when compiling sections that overlap on the same row.
 	 * These define how to deal with overlapping sections and empty space on a row
 	 */
-	MOVIESCENE_API virtual TInlineValue<FMovieSceneSegmentCompilerRules> GetRowCompilerRules() const;
+	MOVIESCENE_API virtual FMovieSceneTrackRowSegmentBlenderPtr GetRowSegmentBlender() const;
 
 	/** 
 	 * Get compiler rules to use when compiling sections that overlap on different rows.
 	 * These define how to deal with overlapping sections and empty space at the track level
 	 */
-	MOVIESCENE_API virtual TInlineValue<FMovieSceneSegmentCompilerRules> GetTrackCompilerRules() const;
+	MOVIESCENE_API virtual FMovieSceneTrackSegmentBlenderPtr GetTrackSegmentBlender() const;
+
+	DEPRECATED(4.19, "Please override GetRowSegmentBlender() instead.")
+	MOVIESCENE_API virtual FMovieSceneDeprecatedCompilerRulesPtr GetRowCompilerRules() const;
+
+	DEPRECATED(4.19, "Please override GetTrackSegmentBlender() instead.")
+	MOVIESCENE_API virtual FMovieSceneDeprecatedCompilerRulesPtr GetTrackCompilerRules() const;
 
 	/**
 	 * Generate a template for this track
@@ -208,13 +191,6 @@ protected:
 	 */
 	virtual void PostCompile(FMovieSceneEvaluationTrack& Track, const FMovieSceneTrackCompilerArgs& Args) const {}
 
-public:
-
-	/**
-	 * Internal function to create a legacy track instance
-	 */
-	TSharedPtr<class IMovieSceneTrackInstance> CreateLegacyInstance() const;
-
 protected:
 
 	/**
@@ -234,12 +210,6 @@ protected:
 	 */
 	MOVIESCENE_API EMovieSceneCompileResult Compile(FMovieSceneEvaluationTrack& Track, const FMovieSceneTrackCompilerArgs& Args) const;
 
-	/**
-	 * Deprecated method of creating track instances
-	 */
-	DEPRECATED(4.15, "Create Instance has been deprecated. Please provide an evaluation template through CreateTemplateForSection instead.")
-	virtual TSharedPtr<class IMovieSceneTrackInstance> CreateInstance() { return nullptr; }
-
 protected:
 
 	//~ UObject interface
diff --git a/Engine/Source/Runtime/MovieScene/Public/Sections/MovieSceneSubSection.h b/Engine/Source/Runtime/MovieScene/Public/Sections/MovieSceneSubSection.h
new file mode 100644
--- /dev/null
+++ b/Engine/Source/Runtime/MovieScene/Public/Sections/MovieSceneSubSection.h
@@ -0,0 +1,183 @@
+// Copyright 1998-2017 Epic Games, Inc. All Rights Reserved.
+
+#pragma once
+
+#include "CoreMinimal.h"
+#include "UObject/ObjectMacros.h"
+#include "Engine/EngineTypes.h"
+#include "GameFramework/Actor.h"
+#include "Curves/KeyHandle.h"
+#include "MovieSceneSequenceID.h"
+#include "MovieSceneSection.h"
+#include "Evaluation/MovieSceneSectionParameters.h"
+#include "Evaluation/MovieSceneSequenceHierarchy.h"
+#include "Evaluation/MovieSceneEvaluationOperand.h"
+#include "MovieSceneSubSection.generated.h"
+
+class UMovieSceneSequence;
+struct FMovieSceneEvaluationTemplate;
+struct FMovieSceneTrackCompilerArgs;
+
+DECLARE_DELEGATE_OneParam(FOnSequenceChanged, UMovieSceneSequence* /*Sequence*/);
+
+struct FSubSequenceInstanceDataParams
+{
+	/** The ID of the sequence instance that is being generated */
+	FMovieSceneSequenceID InstanceSequenceID;
+
+	/** The object binding ID in which the section to be generated resides */
+	FMovieSceneEvaluationOperand Operand;
+};
+
+/**
+ * Implements a section in sub-sequence tracks.
+ */
+UCLASS()
+class MOVIESCENE_API UMovieSceneSubSection
+	: public UMovieSceneSection
+{
+	GENERATED_BODY()
+
+public:
+
+	/** Default constructor. */
+	UMovieSceneSubSection();
+
+	/**
+	 * Get the sequence that is assigned to this section.
+	 *
+	 * @return The sequence.
+	 * @see SetSequence
+	 */
+	UMovieSceneSequence* GetSequence() const;
+
+	/**
+	 * Get the path name to this sub section from the outer moviescene
+	 */
+	FString GetPathNameInMovieScene() const;
+
+	/**
+	 * Get this sub section's sequence ID
+	 */
+	FMovieSceneSequenceID GetSequenceID() const;
+
+	/** Generate subsequence data */
+	virtual FMovieSceneSubSequenceData GenerateSubSequenceData(const FSubSequenceInstanceDataParams& Params) const;
+
+public:
+
+	/**
+	 * Sets the sequence played by this section.
+	 *
+	 * @param Sequence The sequence to play.
+	 * @see GetSequence
+	 */
+	void SetSequence(UMovieSceneSequence* Sequence);
+
+	/** Prime this section as the one and only recording section */
+	void SetAsRecording(bool bRecord);
+
+	/** Get the section we are recording to */
+	static UMovieSceneSubSection* GetRecordingSection();
+
+	/** Get the actor we are targeting for recording */
+	static AActor* GetActorToRecord();
+
+	/** Check if we are primed for recording */
+	static bool IsSetAsRecording();
+
+	virtual void PostLoad() override;
+
+#if WITH_EDITOR
+	virtual void PostEditChangeProperty(FPropertyChangedEvent& PropertyChangedEvent) override;
+
+	/** Delegate to fire when our sequence is changed in the property editor */
+	FOnSequenceChanged& OnSequenceChanged() { return OnSequenceChangedDelegate; }
+#endif
+
+	/** Get the name of the sequence we are going to try to record to */
+	const FString& GetTargetSequenceName() const
+	{
+		return TargetSequenceName;
+	}
+
+	/** Set the name of the sequence we are going to try to record to */
+	void SetTargetSequenceName(const FString& Name)
+	{
+		TargetSequenceName = Name;
+	}
+
+	/** Get the path of the sequence we are going to try to record to */
+	const FString& GetTargetPathToRecordTo() const
+	{
+		return TargetPathToRecordTo.Path;
+	}
+
+	/** Set the path of the sequence we are going to try to record to */
+	void SetTargetPathToRecordTo(const FString& Path)
+	{
+		TargetPathToRecordTo.Path = Path;
+	}
+
+	/** Set the target actor to record */
+	void SetActorToRecord(AActor* InActorToRecord)
+	{
+		ActorToRecord = InActorToRecord;
+	}
+
+public:
+
+	//~ UMovieSceneSection interface
+
+	virtual UMovieSceneSection* SplitSection( float SplitTime ) override;
+	virtual void TrimSection( float TrimTime, bool bTrimLeft ) override;
+	virtual TOptional<float> GetOffsetTime() const override { return TOptional<float>(Parameters.StartOffset); }
+	virtual TOptional<float> GetKeyTime( FKeyHandle KeyHandle ) const override { return TOptional<float>(); }
+	virtual void SetKeyTime( FKeyHandle KeyHandle, float Time ) override { }
+
+public:
+
+	UPROPERTY(EditAnywhere, Category="General", meta=(ShowOnlyInnerProperties))
+	FMovieSceneSectionParameters Parameters;
+
+private:
+
+	UPROPERTY()
+	float StartOffset_DEPRECATED;
+
+	UPROPERTY()
+	float TimeScale_DEPRECATED;
+
+	UPROPERTY()
+	float PrerollTime_DEPRECATED;
+
+protected:
+
+	/**
+	 * Movie scene being played by this section.
+	 *
+	 * @todo Sequencer: Should this be lazy loaded?
+	 */
+	UPROPERTY(EditAnywhere, Category="Sequence")
+	UMovieSceneSequence* SubSequence;
+
+	/** Target actor to record */
+	UPROPERTY(EditAnywhere, Category="Sequence Recording")
+	TLazyObjectPtr<AActor> ActorToRecord;
+
+	/** Target name of sequence to try to record to (will record automatically to another if this already exists) */
+	UPROPERTY(EditAnywhere, Category="Sequence Recording")
+	FString TargetSequenceName;
+
+	/** Target path of sequence to record to */
+	UPROPERTY(EditAnywhere, Category="Sequence Recording", meta=(ContentDir))
+	FDirectoryPath TargetPathToRecordTo;
+
+	/** Keep track of our constructed recordings */
+	static TWeakObjectPtr<UMovieSceneSubSection> TheRecordingSection;
+
+#if WITH_EDITOR
+	/** Delegate to fire when our sequence is changed in the property editor */
+	FOnSequenceChanged OnSequenceChangedDelegate;
+#endif
+};
diff --git a/Engine/Source/Runtime/MovieScene/Public/Tracks/MovieSceneSubTrack.h b/Engine/Source/Runtime/MovieScene/Public/Tracks/MovieSceneSubTrack.h
new file mode 100644
--- /dev/null
+++ b/Engine/Source/Runtime/MovieScene/Public/Tracks/MovieSceneSubTrack.h
@@ -0,0 +1,76 @@
+// Copyright 1998-2017 Epic Games, Inc. All Rights Reserved.
+
+#pragma once
+
+#include "CoreMinimal.h"
+#include "UObject/ObjectMacros.h"
+#include "Misc/InlineValue.h"
+#include "MovieSceneNameableTrack.h"
+#include "MovieSceneSubTrack.generated.h"
+
+class UMovieSceneSequence;
+class UMovieSceneSubSection;
+struct FMovieSceneSegmentCompilerRules;
+
+/**
+ * A track that holds sub-sequences within a larger sequence.
+ */
+UCLASS()
+class MOVIESCENE_API UMovieSceneSubTrack
+	: public UMovieSceneNameableTrack
+{
+	GENERATED_BODY()
+
+public:
+
+	UMovieSceneSubTrack( const FObjectInitializer& ObjectInitializer );
+
+	/**
+	 * Adds a movie scene section at the requested time.
+	 *
+	 * @param Sequence The sequence to add
+	 * @param StartTime The time to add the section at
+	 * @param Duration The duration of the section
+	 * @param bInsertSequence Whether or not to insert the sequence and push existing sequences out
+	 * @return The newly created sub section
+	 */
+	virtual UMovieSceneSubSection* AddSequence(UMovieSceneSequence* Sequence, float StartTime, float Duration, const bool& bInsertSequence = false);
+
+	/**
+	 * Check whether this track contains the given sequence.
+	 *
+	 * @param Sequence The sequence to find.
+	 * @param Recursively Whether to search for the sequence in sub-sequences.
+	 * @return true if the sequence is in this track, false otherwise.
+	 */
+	bool ContainsSequence(const UMovieSceneSequence& Sequence, bool Recursively = false) const;
+
+	/**
+	 * Add a new sequence to record
+	 */
+	virtual UMovieSceneSubSection* AddSequenceToRecord();
+
+public:
+
+	// UMovieSceneTrack interface
+
+	virtual void AddSection(UMovieSceneSection& Section) override;
+	virtual UMovieSceneSection* CreateNewSection() override;
+	virtual const TArray<UMovieSceneSection*>& GetAllSections() const override;
+	virtual TRange<float> GetSectionBoundaries() const override;
+	virtual bool HasSection(const UMovieSceneSection& Section) const override;
+	virtual bool IsEmpty() const override;
+	virtual void RemoveAllAnimationData() override;
+	virtual void RemoveSection(UMovieSceneSection& Section) override;
+	virtual bool SupportsMultipleRows() const override;
+
+#if WITH_EDITORONLY_DATA
+	virtual FText GetDefaultDisplayName() const override;
+#endif
+
+protected:
+
+	/** All movie scene sections. */
+	UPROPERTY()
+	TArray<UMovieSceneSection*> Sections;
+};
diff --git a/Engine/Source/Runtime/MovieSceneCapture/Private/AutomatedLevelSequenceCapture.cpp b/Engine/Source/Runtime/MovieSceneCapture/Private/AutomatedLevelSequenceCapture.cpp
--- a/Engine/Source/Runtime/MovieSceneCapture/Private/AutomatedLevelSequenceCapture.cpp
+++ b/Engine/Source/Runtime/MovieSceneCapture/Private/AutomatedLevelSequenceCapture.cpp
@@ -49,7 +49,7 @@ void UAutomatedLevelSequenceCapture::SetLevelSequenceAsset(FString AssetPath)
 
 void UAutomatedLevelSequenceCapture::AddFormatMappings(TMap<FString, FStringFormatArg>& OutFormatMappings, const FFrameMetrics& FrameMetrics) const
 {
-	OutFormatMappings.Add(TEXT("shot"), CachedState.CurrentShotName.ToString());
+	OutFormatMappings.Add(TEXT("shot"), CachedState.CurrentShotName);
 
 	const int32 FrameNumber = FMath::RoundToInt(CachedState.CurrentShotLocalTime * CachedState.Settings.FrameRate);
 	OutFormatMappings.Add(TEXT("shot_frame"), FString::Printf(TEXT("%0*d"), Settings.ZeroPadFrameNumbers, FrameNumber));
@@ -542,9 +542,9 @@ void UAutomatedLevelSequenceCapture::SequenceUpdated(const UMovieSceneSequencePl
 		if (Actor && Actor->SequencePlayer)
 		{
 			// If this is a new shot, set the state to shot warm up and pause on this frame until warmed up			
-			bool bHasMultipleShots = !PreviousState.CurrentShotName.IdenticalTo(PreviousState.MasterName);
+			bool bHasMultipleShots = PreviousState.CurrentShotName != PreviousState.MasterName;
 			bool bNewShot = bHasMultipleShots && PreviousState.ShotID != CachedState.ShotID;
-
+			
 			if (bNewShot && Actor->SequencePlayer->IsPlaying() && DelayBeforeWarmUp > 0)
 			{
 				CaptureState = ELevelSequenceCaptureState::Paused;
diff --git a/Engine/Source/Runtime/MovieSceneCapture/Private/MovieSceneCaptureHelpers.cpp b/Engine/Source/Runtime/MovieSceneCapture/Private/MovieSceneCaptureHelpers.cpp
--- a/Engine/Source/Runtime/MovieSceneCapture/Private/MovieSceneCaptureHelpers.cpp
+++ b/Engine/Source/Runtime/MovieSceneCapture/Private/MovieSceneCaptureHelpers.cpp
@@ -465,7 +465,7 @@ bool MovieSceneCaptureHelpers::ExportEDL(const UMovieScene* InMovieScene, float
 					continue;
 				}
 
-				FString ShotName = CinematicShotSection->GetShotDisplayName().ToString();
+				FString ShotName = CinematicShotSection->GetShotDisplayName();
 				FString ShotPath = CinematicShotSection->GetSequence()->GetMovieScene()->GetOuter()->GetPathName();
 
 				float HandleFrameTime = (float)InHandleFrames / (float)InFrameRate;
