Sequencer: Added missing reference to segment identifier serialization operator

diff --git a/Engine/Source/Runtime/MovieScene/Public/Evaluation/MovieSceneSegment.h b/Engine/Source/Runtime/MovieScene/Public/Evaluation/MovieSceneSegment.h
--- a/Engine/Source/Runtime/MovieScene/Public/Evaluation/MovieSceneSegment.h
+++ b/Engine/Source/Runtime/MovieScene/Public/Evaluation/MovieSceneSegment.h
@@ -49,7 +49,7 @@ struct FMovieSceneSegmentIdentifier
 		return true;
 	}
 
-	friend FArchive& operator<<(FArchive& Ar, FMovieSceneSegmentIdentifier ID)
+	friend FArchive& operator<<(FArchive& Ar, FMovieSceneSegmentIdentifier& ID)
 	{
 		ID.Serialize(Ar);
 		return Ar;
