Copying //UE4/Dev-Sequencer to //UE4/Dev-Main (Source: //UE4/Dev-Sequencer @ 3946692)

diff --git a/Engine/Source/Runtime/MediaAssets/Private/Assets/MediaTexture.cpp b/Engine/Source/Runtime/MediaAssets/Private/Assets/MediaTexture.cpp
--- a/Engine/Source/Runtime/MediaAssets/Private/Assets/MediaTexture.cpp
+++ b/Engine/Source/Runtime/MediaAssets/Private/Assets/MediaTexture.cpp
@@ -102,6 +102,7 @@ int32 UMediaTexture::GetWidth() const
 void UMediaTexture::SetMediaPlayer(UMediaPlayer* NewMediaPlayer)
 {
 	CurrentPlayer = NewMediaPlayer;
+	UpdateQueue();
 }
 
 
@@ -271,19 +272,12 @@ void UMediaTexture::TickResource(FTimespan Timecode)
 		return;
 	}
 
-	const FGuid LastGuid = CurrentGuid;
+	const FGuid PreviousGuid = CurrentGuid;
 
 	// media player bookkeeping
 	if (CurrentPlayer.IsValid())
 	{
-		const FGuid PlayerGuid = CurrentPlayer->GetGuid();
-
-		if (CurrentGuid != PlayerGuid)
-		{
-			SampleQueue = MakeShared<FMediaTextureSampleQueue, ESPMode::ThreadSafe>();
-			CurrentPlayer->GetPlayerFacade()->AddVideoSampleSink(SampleQueue.ToSharedRef());
-			CurrentGuid = PlayerGuid;
-		}
+		UpdateQueue();
 	}
 	else if (CurrentGuid != DefaultGuid)
 	{
@@ -317,14 +311,15 @@ void UMediaTexture::TickResource(FTimespan Timecode)
 		RenderParams.Rate = CurrentPlayer->GetRate();
 		RenderParams.Time = CurrentPlayer->GetTime();
 	}
-	else if (!AutoClear && LastGuid.IsValid())
+	else if (!AutoClear && (CurrentGuid == PreviousGuid))
 	{
 		return; // retain last frame
 	}
 
+	RenderParams.CanClear = AutoClear;
 	RenderParams.ClearColor = ClearColor;
-	RenderParams.LastGuid = LastGuid;
-	RenderParams.TextureGuid = CurrentGuid;
+	RenderParams.PreviousGuid = PreviousGuid;
+	RenderParams.CurrentGuid = CurrentGuid;
 	RenderParams.SrgbOutput = SRGB;
 
 	// redraw texture resource on render thread
@@ -335,3 +330,23 @@ void UMediaTexture::TickResource(FTimespan Timecode)
 			ResourceParam->Render(RenderParamsParam);
 		});
 }
+
+
+void UMediaTexture::UpdateQueue()
+{
+	if (CurrentPlayer.IsValid())
+	{
+		const FGuid PlayerGuid = CurrentPlayer->GetGuid();
+
+		if (CurrentGuid != PlayerGuid)
+		{
+			SampleQueue = MakeShared<FMediaTextureSampleQueue, ESPMode::ThreadSafe>();
+			CurrentPlayer->GetPlayerFacade()->AddVideoSampleSink(SampleQueue.ToSharedRef());
+			CurrentGuid = PlayerGuid;
+		}
+	}
+	else
+	{
+		SampleQueue.Reset();
+	}
+}
diff --git a/Engine/Source/Runtime/MediaAssets/Private/MediaAssetsPrivate.h b/Engine/Source/Runtime/MediaAssets/Private/MediaAssetsPrivate.h
--- a/Engine/Source/Runtime/MediaAssets/Private/MediaAssetsPrivate.h
+++ b/Engine/Source/Runtime/MediaAssets/Private/MediaAssetsPrivate.h
@@ -3,6 +3,7 @@
 #pragma once
 
 #include "Logging/LogMacros.h"
+#include "Stats/Stats.h"
 
 
 /** Declares a log category for this module. */
diff --git a/Engine/Source/Runtime/MediaAssets/Private/Misc/MediaTextureResource.cpp b/Engine/Source/Runtime/MediaAssets/Private/Misc/MediaTextureResource.cpp
--- a/Engine/Source/Runtime/MediaAssets/Private/Misc/MediaTextureResource.cpp
+++ b/Engine/Source/Runtime/MediaAssets/Private/Misc/MediaTextureResource.cpp
@@ -26,6 +26,11 @@
 #define MEDIATEXTURERESOURCE_TRACE_RENDER 0
 
 
+/** Time spent in media player facade closing media. */
+DECLARE_CYCLE_STAT(TEXT("MediaAssets MediaTextureResource Render"), STAT_MediaAssets_MediaTextureResourceRender, STATGROUP_Media);
+
+
+
 /* Local helpers
  *****************************************************************************/
 
@@ -145,6 +150,8 @@ void FMediaTextureResource::Render(const FRenderParams& Params)
 {
 	check(IsInRenderingThread());
 
+	SCOPE_CYCLE_COUNTER(STAT_MediaAssets_MediaTextureResourceRender);
+
 	FLinearColor Rotation(1, 0, 0, 1);
 	FLinearColor Offset(0, 0, 0, 0);
 
@@ -202,26 +209,30 @@ void FMediaTextureResource::Render(const FRenderParams& Params)
 			Offset = Sample->GetOffset();
 		}
 	}
-	else if (!Cleared || (Params.ClearColor != CurrentClearColor))
+	else if (Params.CanClear)
 	{
-		#if MEDIATEXTURERESOURCE_TRACE_RENDER
-			UE_LOG(LogMediaAssets, VeryVerbose, TEXT("TextureResource %p: Clearing texture at time %s"), this, *Params.Time.ToString());
-		#endif
+		if (!Cleared || (Params.ClearColor != CurrentClearColor))
+		{
+			#if MEDIATEXTURERESOURCE_TRACE_RENDER
+				UE_LOG(LogMediaAssets, VeryVerbose, TEXT("TextureResource %p: Clearing texture at time %s"), this, *Params.Time.ToString());
+			#endif
 
-		ClearTexture(Params.ClearColor, Params.SrgbOutput);
+			ClearTexture(Params.ClearColor, Params.SrgbOutput);
+		}
 	}
 
+	// update external texture registration
 	if (!GSupportsImageExternal)
 	{
-		if (Params.TextureGuid.IsValid())
+		if (Params.CurrentGuid.IsValid())
 		{
 			FTextureRHIRef VideoTexture = (FTextureRHIRef)Owner.TextureReference.TextureReferenceRHI;
-			FExternalTextureRegistry::Get().RegisterExternalTexture(Params.TextureGuid, VideoTexture, SamplerStateRHI, Rotation, Offset);
+			FExternalTextureRegistry::Get().RegisterExternalTexture(Params.CurrentGuid, VideoTexture, SamplerStateRHI, Rotation, Offset);
 		}
 
-		if (Params.LastGuid.IsValid() && (Params.LastGuid != Params.TextureGuid))
+		if (Params.PreviousGuid.IsValid() && (Params.PreviousGuid != Params.CurrentGuid))
 		{
-			FExternalTextureRegistry::Get().UnregisterExternalTexture(Params.LastGuid);
+			FExternalTextureRegistry::Get().UnregisterExternalTexture(Params.PreviousGuid);
 		}
 	}
 }
diff --git a/Engine/Source/Runtime/MediaAssets/Private/Misc/MediaTextureResource.h b/Engine/Source/Runtime/MediaAssets/Private/Misc/MediaTextureResource.h
--- a/Engine/Source/Runtime/MediaAssets/Private/Misc/MediaTextureResource.h
+++ b/Engine/Source/Runtime/MediaAssets/Private/Misc/MediaTextureResource.h
@@ -50,11 +50,17 @@ public:
 	/** Parameters for the Render method. */
 	struct FRenderParams
 	{
+		/** Whether the texture can be cleared. */
+		bool CanClear;
+
 		/** The clear color to use when clearing the texture. */
 		FLinearColor ClearColor;
 
-		/** The external texture GUID at the previously rendered frame. */
-		FGuid LastGuid;
+		/** The texture's current external texture GUID. */
+		FGuid CurrentGuid;
+
+		/** The texture's previously used external texture GUID. */
+		FGuid PreviousGuid;
 
 		/** The player's play rate. */
 		float Rate;
@@ -65,9 +71,6 @@ public:
 		/** Whether output should be in sRGB color space. */
 		bool SrgbOutput;
 
-		/** Guid associated with the texture. */
-		FGuid TextureGuid;
-
 		/** The time of the video frame to render (in player's clock). */
 		FTimespan Time;
 	};
diff --git a/Engine/Source/Runtime/MediaAssets/Public/MediaSoundComponent.h b/Engine/Source/Runtime/MediaAssets/Public/MediaSoundComponent.h
--- a/Engine/Source/Runtime/MediaAssets/Public/MediaSoundComponent.h
+++ b/Engine/Source/Runtime/MediaAssets/Public/MediaSoundComponent.h
@@ -8,6 +8,7 @@
 #include "HAL/CriticalSection.h"
 #include "MediaSampleQueue.h"
 #include "Misc/Timespan.h"
+#include "Templates/Atomic.h"
 #include "Templates/SharedPointer.h"
 #include "UObject/ObjectMacros.h"
 #include "UObject/ScriptMacros.h"
@@ -77,10 +78,20 @@ public:
 	UFUNCTION(BlueprintCallable, Category="Media|MediaSoundComponent", meta=(DisplayName="Get Attenuation Settings To Apply", ScriptName="GetAttenuationSettingsToApply"))
 	bool BP_GetAttenuationSettingsToApply(FSoundAttenuationSettings& OutAttenuationSettings);
 
+	/**
+	 * Get the media player that provides the audio samples.
+	 *
+	 * @return The component's media player, or nullptr if not set.
+	 * @see SetMediaPlayer
+	 */
+	UFUNCTION(BlueprintCallable, Category="Media|MediaSoundComponent")
+	UMediaPlayer* GetMediaPlayer() const;
+
 	/**
 	 * Set the media player that provides the audio samples.
 	 *
 	 * @param NewMediaPlayer The player to set.
+	 * @see GetMediaPlayer
 	 */
 	UFUNCTION(BlueprintCallable, Category="Media|MediaSoundComponent")
 	void SetMediaPlayer(UMediaPlayer* NewMediaPlayer);
@@ -151,10 +162,10 @@ protected:
 private:
 
 	/** The player's current play rate (cached for use on audio thread). */
-	float CachedRate;
+	TAtomic<float> CachedRate;
 
 	/** The player's current time (cached for use on audio thread). */
-	FTimespan CachedTime;
+	TAtomic<FTimespan> CachedTime;
 
 	/** Critical section for synchronizing access to PlayerFacadePtr. */
 	FCriticalSection CriticalSection;
diff --git a/Engine/Source/Runtime/MediaAssets/Public/MediaTexture.h b/Engine/Source/Runtime/MediaAssets/Public/MediaTexture.h
--- a/Engine/Source/Runtime/MediaAssets/Public/MediaTexture.h
+++ b/Engine/Source/Runtime/MediaAssets/Public/MediaTexture.h
@@ -137,6 +137,9 @@ protected:
 	 */
 	void TickResource(FTimespan Timecode);
 
+	/** Update the video sample queue, if necessary. */
+	void UpdateQueue();
+
 protected:
 
 	/**
diff --git a/Engine/Source/Runtime/MediaUtils/Private/MediaPlayerFacade.cpp b/Engine/Source/Runtime/MediaUtils/Private/MediaPlayerFacade.cpp
--- a/Engine/Source/Runtime/MediaUtils/Private/MediaPlayerFacade.cpp
+++ b/Engine/Source/Runtime/MediaUtils/Private/MediaPlayerFacade.cpp
@@ -25,6 +25,29 @@
 #define MEDIAPLAYERFACADE_TRACE_SINKOVERFLOWS 0
 
 
+/** Time spent in media player facade closing media. */
+DECLARE_CYCLE_STAT(TEXT("MediaUtils MediaPlayerFacade Close"), STAT_MediaUtils_FacadeClose, STATGROUP_Media);
+
+/** Time spent in media player facade opening media. */
+DECLARE_CYCLE_STAT(TEXT("MediaUtils MediaPlayerFacade Open"), STAT_MediaUtils_FacadeOpen, STATGROUP_Media);
+
+/** Time spent in media player facade event processing. */
+DECLARE_CYCLE_STAT(TEXT("MediaUtils MediaPlayerFacade ProcessEvent"), STAT_MediaUtils_FacadeProcessEvent, STATGROUP_Media);
+
+/** Time spent in media player facade fetch tick. */
+DECLARE_CYCLE_STAT(TEXT("MediaUtils MediaPlayerFacade TickFetch"), STAT_MediaUtils_FacadeTickFetch, STATGROUP_Media);
+
+/** Time spent in media player facade input tick. */
+DECLARE_CYCLE_STAT(TEXT("MediaUtils MediaPlayerFacade TickInput"), STAT_MediaUtils_FacadeTickInput, STATGROUP_Media);
+
+/** Time spent in media player facade output tick. */
+DECLARE_CYCLE_STAT(TEXT("MediaUtils MediaPlayerFacade TickOutput"), STAT_MediaUtils_FacadeTickOutput, STATGROUP_Media);
+
+/** Time spent in media player facade high frequency tick. */
+DECLARE_CYCLE_STAT(TEXT("MediaUtils MediaPlayerFacade TickTickable"), STAT_MediaUtils_FacadeTickTickable, STATGROUP_Media);
+
+
+
 /* Local helpers
 *****************************************************************************/
 
@@ -152,6 +175,8 @@ bool FMediaPlayerFacade::CanSeek() const
 
 void FMediaPlayerFacade::Close()
 {
+	SCOPE_CYCLE_COUNTER(STAT_MediaUtils_FacadeClose);
+
 	if (CurrentUrl.IsEmpty())
 	{
 		return;
@@ -388,6 +413,8 @@ bool FMediaPlayerFacade::IsReady() const
 
 bool FMediaPlayerFacade::Open(const FString& Url, const IMediaOptions* Options)
 {
+	SCOPE_CYCLE_COUNTER(STAT_MediaUtils_FacadeOpen);
+
 	if (IsRunningDedicatedServer())
 	{
 		return false;
@@ -715,6 +742,8 @@ bool FMediaPlayerFacade::GetVideoTrackFormat(int32 TrackIndex, int32 FormatIndex
 
 void FMediaPlayerFacade::ProcessEvent(EMediaEvent Event)
 {
+	SCOPE_CYCLE_COUNTER(STAT_MediaUtils_FacadeProcessEvent);
+
 	if (Event == EMediaEvent::TracksChanged)
 	{
 		SelectDefaultTracks();
@@ -773,6 +802,8 @@ void FMediaPlayerFacade::SelectDefaultTracks()
 
 void FMediaPlayerFacade::TickFetch(FTimespan DeltaTime, FTimespan Timecode)
 {
+	SCOPE_CYCLE_COUNTER(STAT_MediaUtils_FacadeTickFetch);
+
 	// let the player generate samples & process events
 	if (Player.IsValid())
 	{
@@ -833,6 +864,8 @@ void FMediaPlayerFacade::TickFetch(FTimespan DeltaTime, FTimespan Timecode)
 
 void FMediaPlayerFacade::TickInput(FTimespan DeltaTime, FTimespan Timecode)
 {
+	SCOPE_CYCLE_COUNTER(STAT_MediaUtils_FacadeTickInput);
+
 	if (Player.IsValid())
 	{
 		Player->TickInput(DeltaTime, Timecode);
@@ -842,6 +875,8 @@ void FMediaPlayerFacade::TickInput(FTimespan DeltaTime, FTimespan Timecode)
 
 void FMediaPlayerFacade::TickOutput(FTimespan DeltaTime, FTimespan /*Timecode*/)
 {
+	SCOPE_CYCLE_COUNTER(STAT_MediaUtils_FacadeTickOutput);
+
 	if (!Player.IsValid())
 	{
 		return;
@@ -857,6 +892,8 @@ void FMediaPlayerFacade::TickOutput(FTimespan DeltaTime, FTimespan /*Timecode*/)
 
 void FMediaPlayerFacade::TickTickable()
 {
+	SCOPE_CYCLE_COUNTER(STAT_MediaUtils_FacadeTickTickable);
+
 	if (LastRate == 0.0f)
 	{
 		return;
diff --git a/Engine/Source/Runtime/MediaUtils/Private/MediaSamples.cpp b/Engine/Source/Runtime/MediaUtils/Private/MediaSamples.cpp
--- a/Engine/Source/Runtime/MediaUtils/Private/MediaSamples.cpp
+++ b/Engine/Source/Runtime/MediaUtils/Private/MediaSamples.cpp
@@ -8,121 +8,60 @@
 #include "IMediaTextureSample.h"
 
 
-/* IMediaSamples interface
+/* Local helpers
 *****************************************************************************/
 
-bool FMediaSamples::FetchAudio(TRange<FTimespan> TimeRange, TSharedPtr<IMediaAudioSample, ESPMode::ThreadSafe>& OutSample)
+template<typename SampleType>
+bool FetchSample(TMediaSampleQueue<SampleType>& SampleQueue, TRange<FTimespan> TimeRange, TSharedPtr<SampleType, ESPMode::ThreadSafe>& OutSample)
 {
-	TSharedPtr<IMediaAudioSample, ESPMode::ThreadSafe> Sample;
-
-	if (!AudioSampleQueue.Peek(Sample))
+	if (!SampleQueue.Peek(OutSample))
 	{
 		return false;
 	}
 
-	const FTimespan SampleTime = Sample->GetTime();
+	const FTimespan SampleTime = OutSample->GetTime();
 
-	if (!TimeRange.Overlaps(TRange<FTimespan>(SampleTime, SampleTime + Sample->GetDuration())))
+	if (!TimeRange.Overlaps(TRange<FTimespan>(SampleTime, SampleTime + OutSample->GetDuration())))
 	{
 		return false;
 	}
 
-	AudioSampleQueue.Pop();
-	OutSample = Sample;
+	SampleQueue.Pop();
 
 	return true;
 }
 
 
-bool FMediaSamples::FetchCaption(TRange<FTimespan> TimeRange, TSharedPtr<IMediaOverlaySample, ESPMode::ThreadSafe>& OutSample)
-{
-	TSharedPtr<IMediaOverlaySample, ESPMode::ThreadSafe> Sample;
-
-	if (!CaptionSampleQueue.Peek(Sample))
-	{
-		return false;
-	}
-
-	const FTimespan SampleTime = Sample->GetTime();
+/* IMediaSamples interface
+*****************************************************************************/
 
-	if (!TimeRange.Overlaps(TRange<FTimespan>(SampleTime, SampleTime + Sample->GetDuration())))
-	{
-		return false;
-	}
+bool FMediaSamples::FetchAudio(TRange<FTimespan> TimeRange, TSharedPtr<IMediaAudioSample, ESPMode::ThreadSafe>& OutSample)
+{
+	return FetchSample(AudioSampleQueue, TimeRange, OutSample);
+}
 
-	CaptionSampleQueue.Pop();
-	OutSample = Sample;
 
-	return true;
+bool FMediaSamples::FetchCaption(TRange<FTimespan> TimeRange, TSharedPtr<IMediaOverlaySample, ESPMode::ThreadSafe>& OutSample)
+{
+	return FetchSample(CaptionSampleQueue, TimeRange, OutSample);
 }
 
 
 bool FMediaSamples::FetchMetadata(TRange<FTimespan> TimeRange, TSharedPtr<IMediaBinarySample, ESPMode::ThreadSafe>& OutSample)
 {
-	TSharedPtr<IMediaBinarySample, ESPMode::ThreadSafe> Sample;
-
-	if (!MetadataSampleQueue.Peek(Sample))
-	{
-		return false;
-	}
-
-	const FTimespan SampleTime = Sample->GetTime();
-
-	if (!TimeRange.Overlaps(TRange<FTimespan>(SampleTime, SampleTime + Sample->GetDuration())))
-	{
-		return false;
-	}
-
-	MetadataSampleQueue.Pop();
-	OutSample = Sample;
-
-	return true;
+	return FetchSample(MetadataSampleQueue, TimeRange, OutSample);
 }
 
 
 bool FMediaSamples::FetchSubtitle(TRange<FTimespan> TimeRange, TSharedPtr<IMediaOverlaySample, ESPMode::ThreadSafe>& OutSample)
 {
-	TSharedPtr<IMediaOverlaySample, ESPMode::ThreadSafe> Sample;
-
-	if (!SubtitleSampleQueue.Peek(Sample))
-	{
-		return false;
-	}
-
-	const FTimespan SampleTime = Sample->GetTime();
-
-	if (!TimeRange.Overlaps(TRange<FTimespan>(SampleTime, SampleTime + Sample->GetDuration())))
-	{
-		return false;
-	}
-
-	SubtitleSampleQueue.Pop();
-	OutSample = Sample;
-
-	return true;
+	return FetchSample(SubtitleSampleQueue, TimeRange, OutSample);
 }
 
 
 bool FMediaSamples::FetchVideo(TRange<FTimespan> TimeRange, TSharedPtr<IMediaTextureSample, ESPMode::ThreadSafe>& OutSample)
 {
-	TSharedPtr<IMediaTextureSample, ESPMode::ThreadSafe> Sample;
-
-	if (!VideoSampleQueue.Peek(Sample))
-	{
-		return false;
-	}
-
-	const FTimespan SampleTime = Sample->GetTime();
-
-	if (!TimeRange.Overlaps(TRange<FTimespan>(SampleTime, SampleTime + Sample->GetDuration())))
-	{
-		return false;
-	}
-
-	VideoSampleQueue.Pop();
-	OutSample = Sample;
-
-	return true;
+	return FetchSample(VideoSampleQueue, TimeRange, OutSample);
 }
 
 
diff --git a/Engine/Source/Runtime/MediaUtils/Private/MediaUtilsPrivate.h b/Engine/Source/Runtime/MediaUtils/Private/MediaUtilsPrivate.h
--- a/Engine/Source/Runtime/MediaUtils/Private/MediaUtilsPrivate.h
+++ b/Engine/Source/Runtime/MediaUtils/Private/MediaUtilsPrivate.h
@@ -3,6 +3,7 @@
 #pragma once
 
 #include "CoreMinimal.h"
+#include "Stats/Stats.h"
 
 
 /** Declares a log category for this module. */
diff --git a/Engine/Source/Runtime/Messaging/Private/Bus/MessageRouter.cpp b/Engine/Source/Runtime/Messaging/Private/Bus/MessageRouter.cpp
--- a/Engine/Source/Runtime/Messaging/Private/Bus/MessageRouter.cpp
+++ b/Engine/Source/Runtime/Messaging/Private/Bus/MessageRouter.cpp
@@ -31,6 +31,12 @@ FMessageRouter::~FMessageRouter()
 /* FRunnable interface
  *****************************************************************************/
 
+FSingleThreadRunnable* FMessageRouter::GetSingleThreadInterface()
+{
+	return this;
+}
+
+
 bool FMessageRouter::Init()
 {
 	return true;
@@ -45,14 +51,7 @@ uint32 FMessageRouter::Run()
 	{
 		if (WorkEvent->Wait(CalculateWaitTime()))
 		{
-			CurrentTime = FDateTime::UtcNow();
-			CommandDelegate Command;
-
-			while (Commands.Dequeue(Command))
-			{
-				Command.Execute();
-			}
-
+			ProcessCommands();
 			WorkEvent->Reset();
 		}
 
@@ -207,6 +206,18 @@ void FMessageRouter::FilterSubscriptions(
 }
 
 
+void FMessageRouter::ProcessCommands()
+{
+	CurrentTime = FDateTime::UtcNow();
+	CommandDelegate Command;
+
+	while (Commands.Dequeue(Command))
+	{
+		Command.Execute();
+	}
+}
+
+
 void FMessageRouter::ProcessDelayedMessages()
 {
 	FDelayedMessage DelayedMessage;
@@ -219,6 +230,16 @@ void FMessageRouter::ProcessDelayedMessages()
 }
 
 
+/* FSingleThreadRunnable interface
+ *****************************************************************************/
+
+void FMessageRouter::Tick()
+{
+	ProcessCommands();
+	ProcessDelayedMessages();
+}
+
+
 /* FMessageRouter callbacks
  *****************************************************************************/
 
diff --git a/Engine/Source/Runtime/Messaging/Private/Bus/MessageRouter.h b/Engine/Source/Runtime/Messaging/Private/Bus/MessageRouter.h
--- a/Engine/Source/Runtime/Messaging/Private/Bus/MessageRouter.h
+++ b/Engine/Source/Runtime/Messaging/Private/Bus/MessageRouter.h
@@ -3,13 +3,15 @@
 #pragma once
 
 #include "CoreMinimal.h"
-#include "IMessageContext.h"
 #include "Containers/Queue.h"
-#include "IMessageTracer.h"
-#include "Bus/MessageTracer.h"
 #include "HAL/Runnable.h"
+#include "Misc/SingleThreadRunnable.h"
 #include "Templates/Atomic.h"
 
+#include "IMessageContext.h"
+#include "IMessageTracer.h"
+#include "Bus/MessageTracer.h"
+
 class IMessageInterceptor;
 class IMessageReceiver;
 class IMessageSubscription;
@@ -19,6 +21,7 @@ class IMessageSubscription;
  */
 class FMessageRouter
 	: public FRunnable
+	, private FSingleThreadRunnable
 {
 	DECLARE_DELEGATE(CommandDelegate)
 
@@ -121,6 +124,7 @@ public:
 
 	//~ FRunnable interface
 
+	virtual FSingleThreadRunnable* GetSingleThreadInterface() override;
 	virtual bool Init() override;
 	virtual uint32 Run() override;
 	virtual void Stop() override;
@@ -173,9 +177,26 @@ protected:
 	 */
 	void DispatchMessage(const TSharedRef<IMessageContext, ESPMode::ThreadSafe>& Message);
 
-	/** Processes all delayed messages. */
+	/**
+	 * Process all queued commands.
+	 *
+	 * @see ProcessDelayedMessages
+	 */
+	void ProcessCommands();
+
+	/**
+	 * Processes all delayed messages.
+	 *
+	 * @see ProcessCommands
+	 */
 	void ProcessDelayedMessages();
 
+protected:
+
+	//~ FSingleThreadRunnable interface
+
+	virtual void Tick() override;
+
 private:
 
 	/** Structure for delayed messages. */
diff --git a/Engine/Source/Runtime/Messaging/Private/Bus/MessageTracer.cpp b/Engine/Source/Runtime/Messaging/Private/Bus/MessageTracer.cpp
--- a/Engine/Source/Runtime/Messaging/Private/Bus/MessageTracer.cpp
+++ b/Engine/Source/Runtime/Messaging/Private/Bus/MessageTracer.cpp
@@ -1,8 +1,8 @@
 // Copyright 1998-2018 Epic Games, Inc. All Rights Reserved.
 
 #include "Bus/MessageTracer.h"
-#include "HAL/PlatformProcess.h"
 #include "Containers/Ticker.h"
+#include "HAL/PlatformProcess.h"
 #include "IMessageInterceptor.h"
 #include "IMessageReceiver.h"
 #include "IMessageTracerBreakpoint.h"
@@ -490,17 +490,20 @@ void FMessageTracer::ResetMessages()
 
 bool FMessageTracer::ShouldBreak(const TSharedRef<IMessageContext, ESPMode::ThreadSafe>& Context) const
 {
-	if (Breaking)
+	if (FPlatformProcess::SupportsMultithreading())
 	{
-		return true;
-	}
-
-	for (const auto& Breakpoint : Breakpoints)
-	{
-		if (Breakpoint->IsEnabled() && Breakpoint->ShouldBreak(Context))
+		if (Breaking)
 		{
 			return true;
 		}
+
+		for (const auto& Breakpoint : Breakpoints)
+		{
+			if (Breakpoint->IsEnabled() && Breakpoint->ShouldBreak(Context))
+			{
+				return true;
+			}
+		}
 	}
 
 	return false;
diff --git a/Engine/Source/Runtime/MovieScene/MovieScene.Build.cs b/Engine/Source/Runtime/MovieScene/MovieScene.Build.cs
--- a/Engine/Source/Runtime/MovieScene/MovieScene.Build.cs
+++ b/Engine/Source/Runtime/MovieScene/MovieScene.Build.cs
@@ -12,7 +12,8 @@ public class MovieScene : ModuleRules
 			new string[] {
 				"Core",
 				"CoreUObject",
-				"Engine"
+				"Engine",
+				"TimeManagement"
 			}
 		);
 	}
diff --git a/Engine/Source/Runtime/MovieScene/Private/Channels/MovieSceneBoolChannel.cpp b/Engine/Source/Runtime/MovieScene/Private/Channels/MovieSceneBoolChannel.cpp
new file mode 100644
--- /dev/null
+++ b/Engine/Source/Runtime/MovieScene/Private/Channels/MovieSceneBoolChannel.cpp
@@ -0,0 +1,64 @@
+// Copyright 1998-2018 Epic Games, Inc. All Rights Reserved.
+
+#include "Channels/MovieSceneBoolChannel.h"
+#include "Channels/MovieSceneChannelProxy.h"
+#include "Curves/IntegralCurve.h"
+#include "FrameRate.h"
+#include "MovieSceneFwd.h"
+#include "MovieSceneFrameMigration.h"
+
+uint32 FMovieSceneBoolChannel::GetChannelID()
+{
+	static uint32 ID = FMovieSceneChannelEntry::RegisterNewID();
+	return ID;
+}
+
+bool FMovieSceneBoolChannel::SerializeFromMismatchedTag(const FPropertyTag& Tag, FArchive& Ar)
+{
+	static const FName IntegralCurveName("IntegralCurve");
+	if (Tag.Type == NAME_StructProperty && Tag.StructName == IntegralCurveName)
+	{
+		FIntegralCurve IntegralCurve;
+		FIntegralCurve::StaticStruct()->SerializeItem(Ar, &IntegralCurve, nullptr);
+
+		if (IntegralCurve.GetDefaultValue() != MAX_int32)
+		{
+			bHasDefaultValue = true;
+			DefaultValue = IntegralCurve.GetDefaultValue() != 0;
+		}
+
+		Times.Reserve(IntegralCurve.GetNumKeys());
+		Values.Reserve(IntegralCurve.GetNumKeys());
+
+		FFrameRate LegacyFrameRate = GetLegacyConversionFrameRate();
+
+		int32 Index = 0;
+		for (auto It = IntegralCurve.GetKeyIterator(); It; ++It)
+		{
+			FFrameNumber KeyTime = UpgradeLegacyMovieSceneTime(nullptr, LegacyFrameRate, It->Time);
+
+			bool Val = It->Value != 0;
+			ConvertInsertAndSort<bool>(Index++, KeyTime, Val, Times, Values);
+		}
+		return true;
+	}
+
+	return false;
+}
+
+bool FMovieSceneBoolChannel::Evaluate(FFrameTime InTime, bool& OutValue) const
+{
+	if (Times.Num())
+	{
+		const int32 Index = FMath::Max(0, Algo::UpperBound(Times, InTime.FrameNumber)-1);
+		OutValue = Values[Index];
+		return true;
+	}
+	else if (bHasDefaultValue)
+	{
+		OutValue = DefaultValue;
+		return true;
+	}
+
+	return false;
+}
\ No newline at end of file
diff --git a/Engine/Source/Runtime/MovieScene/Private/Channels/MovieSceneByteChannel.cpp b/Engine/Source/Runtime/MovieScene/Private/Channels/MovieSceneByteChannel.cpp
new file mode 100644
--- /dev/null
+++ b/Engine/Source/Runtime/MovieScene/Private/Channels/MovieSceneByteChannel.cpp
@@ -0,0 +1,66 @@
+// Copyright 1998-2018 Epic Games, Inc. All Rights Reserved.
+
+#include "Channels/MovieSceneByteChannel.h"
+#include "Channels/MovieSceneChannelProxy.h"
+#include "Curves/IntegralCurve.h"
+#include "MovieSceneFwd.h"
+#include "FrameRate.h"
+#include "MovieSceneFrameMigration.h"
+
+uint32 FMovieSceneByteChannel::GetChannelID()
+{
+	static uint32 ID = FMovieSceneChannelEntry::RegisterNewID();
+	return ID;
+}
+
+bool FMovieSceneByteChannel::SerializeFromMismatchedTag(const FPropertyTag& Tag, FArchive& Ar)
+{
+	static const FName IntegralCurveName("IntegralCurve");
+	if (Tag.Type == NAME_StructProperty && Tag.StructName == IntegralCurveName)
+	{
+		FIntegralCurve IntegralCurve;
+		FIntegralCurve::StaticStruct()->SerializeItem(Ar, &IntegralCurve, nullptr);
+
+		if (IntegralCurve.GetDefaultValue() != MAX_int32)
+		{
+			bHasDefaultValue = true;
+			// We cast rather than clamp here as the old integer curve used to wrap around
+			DefaultValue = static_cast<uint8>(IntegralCurve.GetDefaultValue());
+		}
+
+		Times.Reserve(IntegralCurve.GetNumKeys());
+		Values.Reserve(IntegralCurve.GetNumKeys());
+
+		FFrameRate LegacyFrameRate = GetLegacyConversionFrameRate();
+
+		int32 Index = 0;
+		for (auto It = IntegralCurve.GetKeyIterator(); It; ++It)
+		{
+			FFrameNumber KeyTime = UpgradeLegacyMovieSceneTime(nullptr, LegacyFrameRate, It->Time);
+
+			// We cast rather than clamp here as the old integer curve used to wrap around
+			uint8 Val = static_cast<uint8>(It->Value);
+			ConvertInsertAndSort<uint8>(Index++, KeyTime, Val, Times, Values);
+		}
+		return true;
+	}
+
+	return false;
+}
+
+bool FMovieSceneByteChannel::Evaluate(FFrameTime InTime, uint8& OutValue) const
+{
+	if (Times.Num())
+	{
+		const int32 Index = FMath::Max(0, Algo::UpperBound(Times, InTime.FrameNumber)-1);
+		OutValue = Values[Index];
+		return true;
+	}
+	else if (bHasDefaultValue)
+	{
+		OutValue = DefaultValue;
+		return true;
+	}
+
+	return false;
+}
\ No newline at end of file
diff --git a/Engine/Source/Runtime/MovieScene/Private/Channels/MovieSceneChannel.cpp b/Engine/Source/Runtime/MovieScene/Private/Channels/MovieSceneChannel.cpp
new file mode 100644
--- /dev/null
+++ b/Engine/Source/Runtime/MovieScene/Private/Channels/MovieSceneChannel.cpp
@@ -0,0 +1,163 @@
+// Copyright 1998-2018 Epic Games, Inc. All Rights Reserved.
+
+#include "Channels/MovieSceneChannel.h"
+
+
+namespace MovieScene
+{
+	void EvaluateTime(TArrayView<const FFrameNumber> InTimes, FFrameTime InTime, int32& OutIndex1, int32& OutIndex2)
+	{
+		const int32 Index2 = Algo::UpperBound(InTimes, InTime.FrameNumber);
+		const int32 Index1 = Index2 - 1;
+
+		OutIndex1 = Index1 >= 0            ? Index1 : INDEX_NONE;
+		OutIndex2 = Index2 < InTimes.Num() ? Index2 : INDEX_NONE;
+	}
+
+	void EvaluateTime(TArrayView<const FFrameNumber> InTimes, FFrameTime InTime, int32& OutIndex1, int32& OutIndex2, float& OutInterp)
+	{
+		const int32 Index2 = Algo::UpperBound(InTimes, InTime.FrameNumber);
+		const int32 Index1 = Index2 - 1;
+
+		OutIndex1 = Index1 >= 0            ? Index1 : INDEX_NONE;
+		OutIndex2 = Index2 < InTimes.Num() ? Index2 : INDEX_NONE;
+
+		if (Index1 >= 0 && Index2 < InTimes.Num())
+		{
+			// Stay in integer space as long as possible
+			const int32 Time1 = InTimes[Index1].Value, Time2 = InTimes[Index2].Value;
+			const int32 Difference = Time2 - Time1;
+
+			OutInterp = (InTime.FrameNumber.Value - Time1) / float(Difference) + InTime.GetSubFrame() / Difference;
+		}
+		else
+		{
+			OutInterp = 0.f;
+		}
+	}
+
+	void FindRange(TArrayView<const FFrameNumber> InTimes, FFrameNumber PredicateTime, FFrameNumber InTolerance, int32 MaxNum, int32& OutMin, int32& OutMax)
+	{
+		const int32 LowerBound = Algo::LowerBound(InTimes, PredicateTime);
+
+		int32 MinIndex = LowerBound, MaxIndex = LowerBound;
+		int32 FwdIndex = LowerBound, BwdIndex = LowerBound-1;
+
+		for (;MaxIndex - MinIndex < MaxNum;)
+		{
+			const bool bConsiderFwdIndex = FwdIndex < InTimes.Num() && FMath::Abs(InTimes[FwdIndex] - PredicateTime) <= InTolerance;
+			const bool bConsiderBwdIndex = BwdIndex >= 0            && FMath::Abs(InTimes[BwdIndex] - PredicateTime) <= InTolerance;
+
+			if (bConsiderFwdIndex && bConsiderBwdIndex)
+			{
+				const FFrameNumber FwdDifference = FMath::Abs(PredicateTime - InTimes[FwdIndex]);
+				const FFrameNumber BwdDifference = FMath::Abs(PredicateTime - InTimes[BwdIndex]);
+
+				if (FwdDifference < BwdDifference)
+				{
+					MaxIndex = ++FwdIndex;
+				}
+				else
+				{
+					MinIndex = BwdIndex--;
+				}
+			}
+			else if (bConsiderFwdIndex)
+			{
+				MaxIndex = ++FwdIndex;
+				// Stop considering backwards
+				BwdIndex = INDEX_NONE;
+			}
+			else if (bConsiderBwdIndex)
+			{
+				MinIndex = BwdIndex--;
+				// Stop considering forwards
+				FwdIndex = InTimes.Num();
+			}
+			else
+			{
+				break;
+			}
+		}
+
+		OutMin = MinIndex;
+		OutMax = MaxIndex;
+	}
+} // namespace MovieScene
+
+FMovieSceneChannel::FMovieSceneChannel(TArray<FFrameNumber>* InTimes, FKeyHandleLookupTable* InKeyHandles)
+	: Times(InTimes), KeyHandles(InKeyHandles)
+{}
+
+FKeyHandle FMovieSceneChannel::GetHandle(int32 Index)
+{
+	check(Times->IsValidIndex(Index));
+	ensureMsgf(KeyHandles, TEXT("This channel does not contain key handles"));
+	return KeyHandles ? KeyHandles->FindOrAddKeyHandle(Index) : FKeyHandle();
+}
+
+int32 FMovieSceneChannel::GetIndex(FKeyHandle Handle)
+{
+	ensureMsgf(KeyHandles, TEXT("This channel does not contain key handles"));
+	return KeyHandles ? KeyHandles->GetIndex(Handle) : INDEX_NONE;
+}
+
+int32 FMovieSceneChannel::FindKey(FFrameNumber InTime, int32 InTolerance)
+{
+	int32 MinIndex = 0, MaxIndex = 0;
+	MovieScene::FindRange(*Times, InTime, InTolerance, 1, MinIndex, MaxIndex);
+	if (Times->IsValidIndex(MinIndex) && FMath::Abs((*Times)[MinIndex] - InTime) <= InTolerance)
+	{
+		return MinIndex;
+	}
+
+	return INDEX_NONE;
+}
+
+void FMovieSceneChannel::FindKeys(FFrameNumber InTime, int32 MaxNum, int32& OutMinIndex, int32& OutMaxIndex, int32 InTolerance)
+{
+	MovieScene::FindRange(*Times, InTime, InTolerance, MaxNum, OutMinIndex, OutMaxIndex);
+}
+
+int32 FMovieSceneChannel::AddKeyInternal(FFrameNumber InTime)
+{
+	const int32 InsertIndex = Algo::UpperBound(*Times, InTime);
+
+	Times->Insert(InTime, InsertIndex);
+	if (KeyHandles)
+	{
+		KeyHandles->AllocateHandle(InsertIndex);
+	}
+	return InsertIndex;
+}
+
+int32 FMovieSceneChannel::MoveKeyInternal(int32 KeyIndex, FFrameNumber InNewTime)
+{
+	check(Times->IsValidIndex(KeyIndex));
+
+	int32 NewIndex = Algo::LowerBound(*Times, InNewTime);
+	if (NewIndex < KeyIndex || NewIndex > KeyIndex+1)
+	{
+		// If we're inserting after this key, decrement the new index since we will remove this key
+		if (NewIndex > KeyIndex)
+		{
+			--NewIndex;
+		}
+
+		// We have to remove the key and re-add it in the right place
+		// This could probably be done better by just shuffling up/down the items that need to move, without ever changing the size of the array
+		Times->RemoveAt(KeyIndex, 1, false);
+		Times->Insert(InNewTime, NewIndex);
+
+		if (KeyHandles)
+		{
+			KeyHandles->MoveHandle(KeyIndex, NewIndex);
+		}
+		return NewIndex;
+	}
+	else
+	{
+		(*Times)[KeyIndex] = InNewTime;
+		return KeyIndex;
+	}
+}
\ No newline at end of file
diff --git a/Engine/Source/Runtime/MovieScene/Private/Channels/MovieSceneChannelEditorData.cpp b/Engine/Source/Runtime/MovieScene/Private/Channels/MovieSceneChannelEditorData.cpp
new file mode 100644
--- /dev/null
+++ b/Engine/Source/Runtime/MovieScene/Private/Channels/MovieSceneChannelEditorData.cpp
@@ -0,0 +1,45 @@
+// Copyright 1998-2018 Epic Games, Inc. All Rights Reserved.
+
+#include "Channels/MovieSceneChannelEditorData.h"
+
+#if WITH_EDITOR
+
+const FText FCommonChannelData::ChannelX = NSLOCTEXT("MovieSceneChannels", "ChannelX", "X");
+const FText FCommonChannelData::ChannelY = NSLOCTEXT("MovieSceneChannels", "ChannelY", "Y");
+const FText FCommonChannelData::ChannelZ = NSLOCTEXT("MovieSceneChannels", "ChannelZ", "Z");
+const FText FCommonChannelData::ChannelW = NSLOCTEXT("MovieSceneChannels", "ChannelW", "W");
+
+const FText FCommonChannelData::ChannelR = NSLOCTEXT("MovieSceneChannels", "ChannelR", "R");
+const FText FCommonChannelData::ChannelG = NSLOCTEXT("MovieSceneChannels", "ChannelG", "G");
+const FText FCommonChannelData::ChannelB = NSLOCTEXT("MovieSceneChannels", "ChannelB", "B");
+const FText FCommonChannelData::ChannelA = NSLOCTEXT("MovieSceneChannels", "ChannelA", "A");
+
+const FLinearColor FCommonChannelData::RedChannelColor(0.7f, 0.0f, 0.0f, 0.5f);
+const FLinearColor FCommonChannelData::GreenChannelColor(0.0f, 0.7f, 0.0f, 0.5f);
+const FLinearColor FCommonChannelData::BlueChannelColor(0.0f, 0.0f, 0.7f, 0.5f);
+
+
+FMovieSceneChannelEditorData::FMovieSceneChannelEditorData()
+	: bEnabled(true)
+	, bCanCollapseToTrack(true)
+	, SortOrder(0)
+	, Name(NAME_None)
+{}
+
+FMovieSceneChannelEditorData::FMovieSceneChannelEditorData(FName InName, FText InDisplayText, FText InGroup)
+	: bEnabled(true)
+	, bCanCollapseToTrack(true)
+	, SortOrder(0)
+	, Name(InName)
+	, DisplayText(InDisplayText)
+	, Group(InGroup)
+{}
+
+void FMovieSceneChannelEditorData::SetIdentifiers(FName InName, FText InDisplayText, FText InGroup)
+{
+	Group = InGroup;
+	Name = InName;
+	DisplayText = InDisplayText;
+}
+
+#endif	// WITH_EDITOR
\ No newline at end of file
diff --git a/Engine/Source/Runtime/MovieScene/Private/Channels/MovieSceneChannelProxy.cpp b/Engine/Source/Runtime/MovieScene/Private/Channels/MovieSceneChannelProxy.cpp
new file mode 100644
--- /dev/null
+++ b/Engine/Source/Runtime/MovieScene/Private/Channels/MovieSceneChannelProxy.cpp
@@ -0,0 +1,45 @@
+// Copyright 1998-2018 Epic Games, Inc. All Rights Reserved.
+
+#include "Channels/MovieSceneChannelProxy.h"
+#include "Algo/BinarySearch.h"
+
+
+uint32 FMovieSceneChannelEntry::RegisterNewID()
+{
+	static FThreadSafeCounter CurrentID = 0;
+	check(CurrentID.GetValue() < TNumericLimits<int32>::Max());
+	return CurrentID.Increment();
+}
+
+const FMovieSceneChannelEntry* FMovieSceneChannelProxy::FindEntry(uint32 ChannelTypeID) const
+{
+	const int32 ChannelTypeIndex = Algo::BinarySearchBy(Entries, ChannelTypeID, &FMovieSceneChannelEntry::ChannelID);
+
+	if (ChannelTypeIndex != INDEX_NONE)
+	{
+		return &Entries[ChannelTypeIndex];
+	}
+
+	return nullptr;
+}
+
+int32 FMovieSceneChannelProxy::FindIndex(uint32 ChannelTypeID, void* ChannelPtr) const
+{
+	const FMovieSceneChannelEntry* FoundEntry = FindEntry(ChannelTypeID);
+	if (FoundEntry)
+	{
+		return FoundEntry->GetChannels().IndexOfByKey(ChannelPtr);
+	}
+
+	return INDEX_NONE;
+}
+
+void* FMovieSceneChannelProxy::GetChannel(uint32 ChannelTypeID, int32 ChannelIndex) const
+{
+	if (const FMovieSceneChannelEntry* Entry = FindEntry(ChannelTypeID))
+	{
+		TArrayView<void* const> Channels = Entry->GetChannels();
+		return Channels.IsValidIndex(ChannelIndex) ? Channels[ChannelIndex] : nullptr;
+	}
+	return nullptr;
+}
\ No newline at end of file
diff --git a/Engine/Source/Runtime/MovieScene/Private/Channels/MovieSceneFloatChannel.cpp b/Engine/Source/Runtime/MovieScene/Private/Channels/MovieSceneFloatChannel.cpp
new file mode 100644
--- /dev/null
+++ b/Engine/Source/Runtime/MovieScene/Private/Channels/MovieSceneFloatChannel.cpp
@@ -0,0 +1,623 @@
+// Copyright 1998-2018 Epic Games, Inc. All Rights Reserved.
+
+#include "Channels/MovieSceneFloatChannel.h"
+#include "MovieSceneFwd.h"
+#include "Channels/MovieSceneChannelProxy.h"
+#include "MovieSceneFrameMigration.h"
+
+
+uint32 FMovieSceneFloatChannel::GetChannelID()
+{
+	static uint32 ID = FMovieSceneChannelEntry::RegisterNewID();
+	return ID;
+}
+
+bool FMovieSceneFloatChannel::SerializeFromMismatchedTag(const FPropertyTag& Tag, FArchive& Ar)
+{
+	static const FName RichCurveName("RichCurve");
+	if (Tag.Type == NAME_StructProperty && Tag.StructName == RichCurveName)
+	{
+		FRichCurve RichCurve;
+		FRichCurve::StaticStruct()->SerializeItem(Ar, &RichCurve, nullptr);
+
+		if (RichCurve.GetDefaultValue() != MAX_flt)
+		{
+			bHasDefaultValue = true;
+			DefaultValue = RichCurve.GetDefaultValue();
+		}
+
+		PreInfinityExtrap = RichCurve.PreInfinityExtrap;
+		PostInfinityExtrap = RichCurve.PostInfinityExtrap;
+
+		Times.Reserve(RichCurve.GetNumKeys());
+		Values.Reserve(RichCurve.GetNumKeys());
+
+		const FFrameRate LegacyFrameRate = GetLegacyConversionFrameRate();
+		const float      Interval        = LegacyFrameRate.AsInterval();
+
+		int32 Index = 0;
+		for (auto It = RichCurve.GetKeyIterator(); It; ++It)
+		{
+			const FRichCurveKey& Key = *It;
+
+			FFrameNumber KeyTime = UpgradeLegacyMovieSceneTime(nullptr, LegacyFrameRate, It->Time);
+
+			FMovieSceneFloatValue NewValue;
+			NewValue.Value = Key.Value;
+			NewValue.InterpMode  = Key.InterpMode;
+			NewValue.TangentMode = Key.TangentMode;
+			NewValue.Tangent.ArriveTangent = Key.ArriveTangent * Interval;
+			NewValue.Tangent.LeaveTangent  = Key.LeaveTangent  * Interval;
+			ConvertInsertAndSort<FMovieSceneFloatValue>(Index++, KeyTime, NewValue, Times, Values);
+		}
+		return true;
+	}
+
+	return false;
+}
+
+int32 FMovieSceneFloatChannel::InsertKeyInternal(FFrameNumber InTime)
+{
+	const int32 InsertIndex = Algo::UpperBound(Times, InTime);
+
+	Times.Insert(InTime, InsertIndex);
+	Values.Insert(FMovieSceneFloatValue(), InsertIndex);
+
+	KeyHandles.AllocateHandle(InsertIndex);
+
+	return InsertIndex;
+}
+
+int32 FMovieSceneFloatChannel::AddConstantKey(FFrameNumber InTime, float InValue)
+{
+	const int32 Index = InsertKeyInternal(InTime);
+
+	FMovieSceneFloatValue& Value = Values[Index];
+	Value.Value = InValue;
+	Value.InterpMode = RCIM_Constant;
+
+	AutoSetTangents();
+
+	return Index;
+}
+
+int32 FMovieSceneFloatChannel::AddLinearKey(FFrameNumber InTime, float InValue)
+{
+	const int32 Index = InsertKeyInternal(InTime);
+
+	FMovieSceneFloatValue& Value = Values[Index];
+	Value.Value = InValue;
+	Value.InterpMode = RCIM_Linear;
+
+	AutoSetTangents();
+
+	return Index;
+}
+
+int32 FMovieSceneFloatChannel::AddCubicKey(FFrameNumber InTime, float InValue, ERichCurveTangentMode TangentMode, const FMovieSceneTangentData& Tangent)
+{
+	const int32 Index = InsertKeyInternal(InTime);
+
+	FMovieSceneFloatValue& Value = Values[Index];
+	Value.Value = InValue;
+	Value.InterpMode = RCIM_Cubic;
+	Value.TangentMode = TangentMode;
+	Value.Tangent = Tangent;
+
+	AutoSetTangents();
+
+	return Index;
+}
+
+/** Util to find float value on bezier defined by 4 control points */
+float BezierInterp(float P0, float P1, float P2, float P3, float Alpha)
+{
+	const float P01   = FMath::Lerp(P0,   P1,   Alpha);
+	const float P12   = FMath::Lerp(P1,   P2,   Alpha);
+	const float P23   = FMath::Lerp(P2,   P3,   Alpha);
+	const float P012  = FMath::Lerp(P01,  P12,  Alpha);
+	const float P123  = FMath::Lerp(P12,  P23,  Alpha);
+	const float P0123 = FMath::Lerp(P012, P123, Alpha);
+
+	return P0123;
+}
+
+struct FCycleParams
+{
+	FFrameTime Time;
+	int32 CycleCount;
+	float ValueOffset;
+
+	FCycleParams(FFrameTime InTime)
+		: Time(InTime)
+		, CycleCount(0)
+		, ValueOffset(0.f)
+	{}
+
+	FORCEINLINE void ComputePreValueOffset(float FirstValue, float LastValue)
+	{
+		ValueOffset = (FirstValue-LastValue) * CycleCount;
+	}
+	FORCEINLINE void ComputePostValueOffset(float FirstValue, float LastValue)
+	{
+		ValueOffset = (LastValue-FirstValue) * CycleCount;
+	}
+	FORCEINLINE void Oscillate(int32 MinFrame, int32 MaxFrame)
+	{
+		if (CycleCount % 2 == 1)
+		{
+			Time = MinFrame + (FFrameTime(MaxFrame) - Time);
+		}
+	}
+};
+
+FCycleParams CycleTime(FFrameNumber MinFrame, FFrameNumber MaxFrame, FFrameTime InTime)
+{
+	FCycleParams Params(InTime);
+	
+	const int32 Duration = MaxFrame.Value - MinFrame.Value;
+	if (Duration == 0)
+	{
+		Params.Time = MaxFrame;
+		Params.CycleCount = 0;
+	}
+	else if (InTime < MinFrame)
+	{
+		const int32 CycleCount = ((MaxFrame - InTime) / Duration).FloorToFrame().Value;
+
+		Params.Time = InTime + FFrameTime(Duration)*CycleCount;
+		Params.CycleCount = CycleCount;
+	}
+	else if (InTime > MaxFrame)
+	{
+		const int32 CycleCount = ((InTime - MinFrame) / Duration).FloorToFrame().Value;
+
+		Params.Time = InTime - FFrameTime(Duration)*CycleCount;
+		Params.CycleCount = CycleCount;
+	}
+
+	return Params;
+}
+
+bool FMovieSceneFloatChannel::EvaluateExtrapolation(FFrameTime InTime, float& OutValue) const
+{
+	// If the time is outside of the curve, deal with extrapolation
+	if (InTime < Times[0])
+	{
+		if (PreInfinityExtrap == RCCE_None)
+		{
+			return false;
+		}
+
+		if (PreInfinityExtrap == RCCE_Constant)
+		{
+			OutValue = Values[0].Value;
+			return true;
+		}
+
+		if (PreInfinityExtrap == RCCE_Linear)
+		{
+			const FMovieSceneFloatValue FirstValue = Values[0];
+
+			if (FirstValue.InterpMode == RCIM_Constant)
+			{
+				OutValue = FirstValue.Value;
+			}
+			else if(FirstValue.InterpMode == RCIM_Cubic)
+			{
+				FFrameTime Delta = FFrameTime(Times[0]) - InTime;
+				OutValue = FirstValue.Value - Delta.AsDecimal() * FirstValue.Tangent.ArriveTangent;
+			}
+			else if(FirstValue.InterpMode == RCIM_Linear)
+			{
+				const int32 InterpStartFrame = Times[1].Value;
+				const int32 DeltaFrame       = InterpStartFrame - Times[0].Value;
+				if (DeltaFrame == 0)
+				{
+					OutValue = FirstValue.Value;
+				}
+				else
+				{
+					OutValue = FMath::Lerp(Values[1].Value, FirstValue.Value, (InterpStartFrame - InTime.AsDecimal())/DeltaFrame);
+				}
+			}
+			return true;
+		}
+	}
+	else if (InTime > Times.Last())
+	{
+		if (PostInfinityExtrap == RCCE_None)
+		{
+			return false;
+		}
+
+		if (PostInfinityExtrap == RCCE_Constant)
+		{
+			OutValue = Values.Last().Value;
+			return true;
+		}
+
+		if (PostInfinityExtrap == RCCE_Linear)
+		{
+			const FMovieSceneFloatValue LastValue = Values.Last();
+
+			if (LastValue.InterpMode == RCIM_Constant)
+			{
+				OutValue = LastValue.Value;
+			}
+			else if(LastValue.InterpMode == RCIM_Cubic)
+			{
+				FFrameTime Delta = InTime - Times.Last();
+				OutValue = LastValue.Value + Delta.AsDecimal() * LastValue.Tangent.LeaveTangent;
+			}
+			else if(LastValue.InterpMode == RCIM_Linear)
+			{
+				const int32 NumKeys          = Times.Num();
+				const int32 InterpStartFrame = Times[NumKeys-2].Value;
+				const int32 DeltaFrame       = Times.Last().Value-InterpStartFrame;
+
+				if (DeltaFrame == 0)
+				{
+					OutValue = LastValue.Value;
+				}
+				else
+				{
+					OutValue = FMath::Lerp(Values[NumKeys-2].Value, LastValue.Value, (InTime.AsDecimal() - InterpStartFrame)/DeltaFrame);
+				}
+			}
+			return true;
+		}
+	}
+
+	return false;
+}
+
+bool FMovieSceneFloatChannel::Evaluate(FFrameTime InTime, float& OutValue) const
+{
+	const int32 NumKeys = Times.Num();
+
+	// No keys means default value, or nothing
+	if (NumKeys == 0)
+	{
+		if (bHasDefaultValue)
+		{
+			OutValue = DefaultValue;
+			return true;
+		}
+		return false;
+	}
+
+	// For single keys, we can only ever return that value
+	if (NumKeys == 1)
+	{
+		OutValue = Values[0].Value;
+		return true;
+	}
+
+	// Evaluate with extrapolation if we're outside the bounds of the curve
+	if (EvaluateExtrapolation(InTime, OutValue))
+	{
+		return true;
+	}
+
+	const FFrameNumber MinFrame = Times[0];
+	const FFrameNumber MaxFrame = Times.Last();
+
+	// Compute the cycled time
+	FCycleParams Params = CycleTime(MinFrame, MaxFrame, InTime);
+
+	// Deal with offset cycles and oscillation
+	if (InTime < FFrameTime(MinFrame))
+	{
+		switch (PreInfinityExtrap)
+		{
+		case RCCE_CycleWithOffset: Params.ComputePreValueOffset(Values[0].Value, Values[NumKeys-1].Value); break;
+		case RCCE_Oscillate:       Params.Oscillate(MinFrame.Value, MaxFrame.Value);                       break;
+		}
+	}
+	else if (InTime > FFrameTime(MaxFrame))
+	{
+		switch (PostInfinityExtrap)
+		{
+		case RCCE_CycleWithOffset: Params.ComputePostValueOffset(Values[0].Value, Values[NumKeys-1].Value); break;
+		case RCCE_Oscillate:       Params.Oscillate(MinFrame.Value, MaxFrame.Value);                        break;
+		}
+	}
+
+	if (!ensureMsgf(Params.Time.FrameNumber >= MinFrame && Params.Time.FrameNumber <= MaxFrame, TEXT("Invalid time computed for float channel evaluation")))
+	{
+		return false;
+	}
+
+	// Evaluate the curve data
+	float Interp = 0.f;
+	int32 Index1 = INDEX_NONE, Index2 = INDEX_NONE;
+	MovieScene::EvaluateTime(Times, Params.Time, Index1, Index2, Interp);
+
+	if (Index1 == INDEX_NONE)
+	{
+		OutValue = Params.ValueOffset + Values[Index2].Value;
+	}
+	else if (Index2 == INDEX_NONE)
+	{
+		OutValue = Params.ValueOffset + Values[Index1].Value;
+	}
+	else
+	{
+		FMovieSceneFloatValue Key1 = Values[Index1];
+		FMovieSceneFloatValue Key2 = Values[Index2];
+
+		switch (Key1.InterpMode)
+		{
+		case RCIM_Cubic:
+		{
+			const float OneThird = 1.0f / 3.0f;
+			const int32 Diff = Times[Index2].Value - Times[Index1].Value;
+			const float P0 = Key1.Value;
+			const float P1 = P0 + (Key1.Tangent.LeaveTangent * Diff * OneThird);
+			const float P3 = Key2.Value;
+			const float P2 = P3 - (Key2.Tangent.ArriveTangent * Diff * OneThird);
+
+			OutValue = Params.ValueOffset + BezierInterp(P0, P1, P2, P3, Interp);
+			break;
+		}
+
+		case RCIM_Linear:
+			OutValue = Params.ValueOffset + FMath::Lerp(Key1.Value, Key2.Value, Interp);
+			break;
+
+		default:
+			OutValue = Params.ValueOffset + Key1.Value;
+			break;
+		}
+	}
+
+	return true;
+}
+
+void FMovieSceneFloatChannel::AutoSetTangents(float Tension)
+{
+	if (Values.Num() < 2)
+	{
+		return;
+	}
+
+	{
+		FMovieSceneFloatValue& FirstValue = Values[0];
+		if (FirstValue.InterpMode == RCIM_Cubic && FirstValue.TangentMode == RCTM_Auto)
+		{
+			FirstValue.Tangent.LeaveTangent = FirstValue.Tangent.ArriveTangent = 0.0f;
+		}
+	}
+
+	{
+		FMovieSceneFloatValue& LastValue = Values.Last();
+		if (LastValue.InterpMode == RCIM_Cubic && LastValue.TangentMode == RCTM_Auto)
+		{
+			LastValue.Tangent.LeaveTangent = LastValue.Tangent.ArriveTangent = 0.0f;
+		}
+	}
+
+	for (int32 Index = 1; Index < Values.Num() - 1; ++Index)
+	{
+		FMovieSceneFloatValue  PrevKey = Values[Index-1];
+		FMovieSceneFloatValue& ThisKey = Values[Index  ];
+
+		if (ThisKey.InterpMode == RCIM_Cubic && ThisKey.TangentMode == RCTM_Auto)
+		{
+			FMovieSceneFloatValue& NextKey = Values[Index+1];
+
+			const float PrevToNextTimeDiff = FMath::Max<double>(KINDA_SMALL_NUMBER, Times[Index+1].Value - Times[Index-1].Value);
+
+			float NewTangent = 0.f;
+			AutoCalcTangent(PrevKey.Value, ThisKey.Value, NextKey.Value, Tension, NewTangent);
+			NewTangent /= PrevToNextTimeDiff;
+
+			// In 'auto' mode, arrive and leave tangents are always the same
+			ThisKey.Tangent.LeaveTangent = ThisKey.Tangent.ArriveTangent = NewTangent;
+		}
+		else if (PrevKey.InterpMode == RCIM_Constant || ThisKey.InterpMode == RCIM_Constant)
+		{
+			if (PrevKey.InterpMode != RCIM_Cubic)
+			{
+				ThisKey.Tangent.ArriveTangent = 0.f;
+			}
+
+			ThisKey.Tangent.LeaveTangent  = 0.0f;
+		}
+	}
+}
+
+void FMovieSceneFloatChannel::PopulateCurvePoints(double StartTimeSeconds, double EndTimeSeconds, double TimeThreshold, float ValueThreshold, FFrameRate FrameResolution, TArray<TTuple<double, double>>& InOutPoints) const
+{
+	const FFrameNumber StartFrame = (StartTimeSeconds * FrameResolution).FloorToFrame();
+	const FFrameNumber EndFrame   = (EndTimeSeconds   * FrameResolution).CeilToFrame();
+
+	const int32 StartingIndex = Algo::UpperBound(Times, StartFrame);
+	const int32 EndingIndex   = Algo::LowerBound(Times, EndFrame);
+
+	// Add the lower bound of the visible space
+	float EvaluatedValue = 0.f;
+	if (Evaluate(StartFrame, EvaluatedValue))
+	{
+		InOutPoints.Add(MakeTuple(StartFrame / FrameResolution, double(EvaluatedValue)));
+	}
+
+	// Add all keys in-between
+	for (int32 KeyIndex = StartingIndex; KeyIndex < EndingIndex; ++KeyIndex)
+	{
+		InOutPoints.Add(MakeTuple(Times[KeyIndex] / FrameResolution, double(Values[KeyIndex].Value)));
+	}
+
+	// Add the upper bound of the visible space
+	if (Evaluate(EndFrame, EvaluatedValue))
+	{
+		InOutPoints.Add(MakeTuple(EndFrame / FrameResolution, double(EvaluatedValue)));
+	}
+
+	int32 OldSize = InOutPoints.Num();
+	do
+	{
+		OldSize = InOutPoints.Num();
+		RefineCurvePoints(FrameResolution, TimeThreshold, ValueThreshold, InOutPoints);
+	}
+	while(OldSize != InOutPoints.Num());
+}
+
+void FMovieSceneFloatChannel::RefineCurvePoints(FFrameRate FrameResolution, double TimeThreshold, float ValueThreshold, TArray<TTuple<double, double>>& InOutPoints) const
+{
+	const float InterpTimes[] = { 0.25f, 0.5f, 0.6f };
+
+	for (int32 Index = 0; Index < InOutPoints.Num() - 1; ++Index)
+	{
+		TTuple<double, double> Lower = InOutPoints[Index];
+		TTuple<double, double> Upper = InOutPoints[Index + 1];
+
+		if ((Upper.Get<0>() - Lower.Get<0>()) >= TimeThreshold)
+		{
+			bool bSegmentIsLinear = true;
+
+			TTuple<double, double> Evaluated[ARRAY_COUNT(InterpTimes)];
+
+			for (int32 InterpIndex = 0; InterpIndex < ARRAY_COUNT(InterpTimes); ++InterpIndex)
+			{
+				double& EvalTime  = Evaluated[InterpIndex].Get<0>();
+
+				EvalTime = FMath::Lerp(Lower.Get<0>(), Upper.Get<0>(), InterpTimes[InterpIndex]);
+
+				float Value = 0.f;
+				Evaluate(EvalTime * FrameResolution, Value);
+
+				const float LinearValue = FMath::Lerp(Lower.Get<1>(), Upper.Get<1>(), InterpTimes[InterpIndex]);
+				if (bSegmentIsLinear)
+				{
+					bSegmentIsLinear = FMath::IsNearlyEqual(Value, LinearValue, ValueThreshold);
+				}
+
+				Evaluated[InterpIndex].Get<1>() = Value;
+			}
+
+			if (!bSegmentIsLinear)
+			{
+				// Add the point
+				InOutPoints.Insert(Evaluated, ARRAY_COUNT(Evaluated), Index+1);
+				--Index;
+			}
+		}
+	}
+}
+
+FKeyHandle AddKeyToChannel(FMovieSceneFloatChannel* Channel, FFrameNumber InFrameNumber, float InValue, EMovieSceneKeyInterpolation Interpolation)
+{
+	auto ChannelInterface = Channel->GetInterface();
+	int32 ExistingIndex = ChannelInterface.FindKey(InFrameNumber);
+	if (ExistingIndex != INDEX_NONE)
+	{
+		FMovieSceneFloatValue& Value = ChannelInterface.GetValues()[ExistingIndex];
+		Value.Value = InValue;
+		switch (Interpolation)
+		{
+		case EMovieSceneKeyInterpolation::User:
+			Value.InterpMode = RCIM_Cubic;
+			Value.TangentMode = RCTM_User;
+			break;
+
+		case EMovieSceneKeyInterpolation::Break:
+			Value.InterpMode = RCIM_Cubic;
+			Value.TangentMode = RCTM_Break;
+			break;
+
+		case EMovieSceneKeyInterpolation::Linear:
+			Value.InterpMode = RCIM_Linear;
+			Value.TangentMode = RCTM_Auto;
+			break;
+
+		case EMovieSceneKeyInterpolation::Constant:
+			Value.InterpMode = RCIM_Constant;
+			Value.TangentMode = RCTM_Auto;
+			break;
+
+		default:
+			Value.InterpMode = RCIM_Cubic;
+			Value.TangentMode = RCTM_Auto;
+			break;
+		}
+	}
+	else switch (Interpolation)
+	{
+		case EMovieSceneKeyInterpolation::Auto:     ExistingIndex = Channel->AddCubicKey(InFrameNumber, InValue, RCTM_Auto);  break;
+		case EMovieSceneKeyInterpolation::User:     ExistingIndex = Channel->AddCubicKey(InFrameNumber, InValue, RCTM_User);  break;
+		case EMovieSceneKeyInterpolation::Break:    ExistingIndex = Channel->AddCubicKey(InFrameNumber, InValue, RCTM_Break); break;
+		case EMovieSceneKeyInterpolation::Linear:   ExistingIndex = Channel->AddLinearKey(InFrameNumber, InValue);            break;
+		case EMovieSceneKeyInterpolation::Constant: ExistingIndex = Channel->AddConstantKey(InFrameNumber, InValue);          break;
+	}
+
+	return Channel->GetInterface().GetHandle(ExistingIndex);
+}
+
+void Optimize(FMovieSceneFloatChannel* InChannel, const FKeyDataOptimizationParams& Params)
+{
+	auto ChannelInterface = InChannel->GetInterface();
+	if (ChannelInterface.GetTimes().Num() > 1)
+	{
+		int32 StartIndex = 0;
+		int32 EndIndex = 0;
+
+		{
+			TArrayView<const FFrameNumber> Times = ChannelInterface.GetTimes();
+			StartIndex = Params.Range.GetLowerBound().IsClosed() ? Algo::LowerBound(Times, Params.Range.GetLowerBoundValue()) : 0;
+			EndIndex   = Params.Range.GetUpperBound().IsClosed() ? Algo::UpperBound(Times, Params.Range.GetUpperBoundValue()) : Times.Num();
+		}
+
+		for (int32 Index = StartIndex; Index < EndIndex; ++Index)
+		{
+			// Reget times and values as they may be reallocated
+			FFrameNumber Time                   = ChannelInterface.GetTimes()[Index];
+			FMovieSceneFloatValue OriginalValue = ChannelInterface.GetValues()[Index];
+
+			// If the channel evaluates the same with this key removed, we can leave it out
+			ChannelInterface.RemoveKey(Index);
+			if (ValueExistsAtTime(InChannel, Time, OriginalValue.Value, Params.Tolerance))
+			{
+				Index--;
+			}
+			else
+			{
+				ChannelInterface.AddKey(Time, OriginalValue);
+			}
+		}
+
+		if (Params.bAutoSetInterpolation)
+		{
+			InChannel->AutoSetTangents();
+		}
+	}
+}
+
+void Dilate(FMovieSceneFloatChannel* InChannel, FFrameNumber Origin, float DilationFactor)
+{
+	TArrayView<FFrameNumber> Times = InChannel->GetInterface().GetTimes();
+	for (FFrameNumber& Time : Times)
+	{
+		Time = Origin + FFrameNumber(FMath::FloorToInt((Time - Origin).Value * DilationFactor));
+	}
+	InChannel->AutoSetTangents();
+}
+
+void ChangeFrameResolution(FMovieSceneFloatChannel* InChannel, FFrameRate SourceRate, FFrameRate DestinationRate)
+{
+	TArrayView<FFrameNumber>          Times  = InChannel->GetInterface().GetTimes();
+	TArrayView<FMovieSceneFloatValue> Values = InChannel->GetInterface().GetValues();
+
+	check(Times.Num() == Values.Num());
+
+	float IntervalFactor = DestinationRate.AsInterval() / SourceRate.AsInterval();
+	for (int32 Index = 0; Index < Times.Num(); ++Index)
+	{
+		Times[Index] = ConvertFrameTime(Times[Index], SourceRate, DestinationRate).RoundToFrame();
+
+		FMovieSceneFloatValue& Value = Values[Index];
+		Value.Tangent.ArriveTangent *= IntervalFactor;
+		Value.Tangent.LeaveTangent  *= IntervalFactor;
+	}
+}
\ No newline at end of file
diff --git a/Engine/Source/Runtime/MovieScene/Private/Channels/MovieSceneIntegerChannel.cpp b/Engine/Source/Runtime/MovieScene/Private/Channels/MovieSceneIntegerChannel.cpp
new file mode 100644
--- /dev/null
+++ b/Engine/Source/Runtime/MovieScene/Private/Channels/MovieSceneIntegerChannel.cpp
@@ -0,0 +1,64 @@
+// Copyright 1998-2018 Epic Games, Inc. All Rights Reserved.
+
+#include "Channels/MovieSceneIntegerChannel.h"
+#include "Channels/MovieSceneChannelProxy.h"
+#include "Curves/IntegralCurve.h"
+#include "FrameRate.h"
+#include "MovieSceneFwd.h"
+#include "MovieSceneFrameMigration.h"
+
+uint32 FMovieSceneIntegerChannel::GetChannelID()
+{
+	static uint32 ID = FMovieSceneChannelEntry::RegisterNewID();
+	return ID;
+}
+
+bool FMovieSceneIntegerChannel::SerializeFromMismatchedTag(const FPropertyTag& Tag, FArchive& Ar)
+{
+	static const FName IntegralCurveName("IntegralCurve");
+	if (Tag.Type == NAME_StructProperty && Tag.StructName == IntegralCurveName)
+	{
+		FIntegralCurve IntegralCurve;
+		FIntegralCurve::StaticStruct()->SerializeItem(Ar, &IntegralCurve, nullptr);
+
+		if (IntegralCurve.GetDefaultValue() != MAX_int32)
+		{
+			bHasDefaultValue = true;
+			DefaultValue = IntegralCurve.GetDefaultValue();
+		}
+
+		Times.Reserve(IntegralCurve.GetNumKeys());
+		Values.Reserve(IntegralCurve.GetNumKeys());
+
+		FFrameRate LegacyFrameRate = GetLegacyConversionFrameRate();
+
+		int32 Index = 0;
+		for (auto It = IntegralCurve.GetKeyIterator(); It; ++It)
+		{
+			FFrameNumber KeyTime = UpgradeLegacyMovieSceneTime(nullptr, LegacyFrameRate, It->Time);
+
+			int32 Val = It->Value;
+			ConvertInsertAndSort<int32>(Index++, KeyTime, Val, Times, Values);
+		}
+		return true;
+	}
+
+	return false;
+}
+
+bool FMovieSceneIntegerChannel::Evaluate(FFrameTime InTime, int32& OutValue) const
+{
+	if (Times.Num())
+	{
+		const int32 Index = FMath::Max(0, Algo::UpperBound(Times, InTime.FrameNumber)-1);
+		OutValue = Values[Index];
+		return true;
+	}
+	else if (bHasDefaultValue)
+	{
+		OutValue = DefaultValue;
+		return true;
+	}
+
+	return false;
+}
\ No newline at end of file
diff --git a/Engine/Source/Runtime/MovieScene/Private/Compilation/MovieSceneCompiler.cpp b/Engine/Source/Runtime/MovieScene/Private/Compilation/MovieSceneCompiler.cpp
--- a/Engine/Source/Runtime/MovieScene/Private/Compilation/MovieSceneCompiler.cpp
+++ b/Engine/Source/Runtime/MovieScene/Private/Compilation/MovieSceneCompiler.cpp
@@ -15,12 +15,12 @@
 /** Parameter structure used for gathering entities for a given time or range */
 struct FGatherParameters
 {
-	FGatherParameters(FMovieSceneRootOverridePath& InRootPath, FMovieSceneSequenceHierarchy& InRootHierarchy, IMovieSceneSequenceTemplateStore& InTemplateStore, TRange<float> InCompileRange)
+	FGatherParameters(FMovieSceneRootOverridePath& InRootPath, FMovieSceneSequenceHierarchy& InRootHierarchy, IMovieSceneSequenceTemplateStore& InTemplateStore, TRange<FFrameNumber> InCompileRange)
 		: RootPath(InRootPath)
 		, RootHierarchy(InRootHierarchy)
 		, TemplateStore(InTemplateStore)
 		, RootCompileRange(InCompileRange)
-		, RootClampRange(TRange<float>::All())
+		, RootClampRange(TRange<FFrameNumber>::All())
 		, LocalCompileRange(RootCompileRange)
 		, LocalClampRange(RootClampRange)
 		, Flags(ESectionEvaluationFlags::None)
@@ -40,7 +40,7 @@ struct FGatherParameters
 		return SubParams;
 	}
 
-	void SetClampRange(TRange<float> InNewRootClampRange)
+	void SetClampRange(TRange<FFrameNumber> InNewRootClampRange)
 	{
 		RootClampRange  = InNewRootClampRange;
 		LocalClampRange = InNewRootClampRange * RootToSequenceTransform;
@@ -54,14 +54,14 @@ struct FGatherParameters
 	IMovieSceneSequenceTemplateStore& TemplateStore;
 
 	/** The range that is being compiled in the root's time-space */
-	TRange<float> RootCompileRange;
+	TRange<FFrameNumber> RootCompileRange;
 	/** A range to clamp compilation to in the root's time-space */
-	TRange<float> RootClampRange;
+	TRange<FFrameNumber> RootClampRange;
 
 	/** The range that is being compiled in the current sequence's time-space */
-	TRange<float> LocalCompileRange;
+	TRange<FFrameNumber> LocalCompileRange;
 	/** A range to clamp compilation to in the current sequence's time-space */
-	TRange<float> LocalClampRange;
+	TRange<FFrameNumber> LocalClampRange;
 
 	/** Evaluation flags for the current sequence */
 	ESectionEvaluationFlags Flags;
@@ -94,7 +94,7 @@ struct FCompileOnTheFlyData
 struct FMovieSceneGatheredCompilerData
 {
 	/** Intersection of any empty space that overlaps the currently evaluating time range */
-	TRange<float> EmptyOverlappingRange;
+	TRange<FFrameNumber> EmptyOverlappingRange;
 	/** Tree of tracks to evaluate */
 	TMovieSceneEvaluationTree<FCompileOnTheFlyData> Tracks;
 	/** Tree of active sequences */
@@ -125,7 +125,7 @@ void FMovieSceneCompiler::Compile(UMovieSceneSequence& InCompileSequence, IMovie
 
 	// Gather everything that happens, recursively
 	FMovieSceneGatheredCompilerData GatherData;
-	FGatherParameters GatherParams(RootPath, CompileTemplate.Hierarchy, InTemplateStore, TRange<float>::All());
+	FGatherParameters GatherParams(RootPath, CompileTemplate.Hierarchy, InTemplateStore, TRange<FFrameNumber>::All());
 	GatherCompileOnTheFlyData(InCompileSequence, GatherParams, GatherData);
 
 	// Wipe the current evaluation field for the template
@@ -156,17 +156,17 @@ void FMovieSceneCompiler::Compile(UMovieSceneSequence& InCompileSequence, IMovie
 
 		// Compute meta data for this segment
 		TMovieSceneEvaluationTreeDataIterator<FMovieSceneSequenceID> SubSequences = GatherData.Sequences.GetAllData(GatherData.Sequences.IterateFromLowerBound(It.Range().GetLowerBound()).Node());
-		PopulateMetaData(Result, CompileTemplate.Hierarchy, CompileData, SubSequences);
+		PopulateMetaData(Result, CompileTemplate.Hierarchy, InTemplateStore, CompileData, SubSequences);
 
 		CompileTemplate.EvaluationField.Add(Result.Range, MoveTemp(Result.Group), MoveTemp(Result.MetaData));
 	}
 }
 
-TOptional<FCompiledGroupResult> FMovieSceneCompiler::CompileTime(float InGlobalTime, UMovieSceneSequence& InCompileSequence, IMovieSceneSequenceTemplateStore& InTemplateStore)
+TOptional<FCompiledGroupResult> FMovieSceneCompiler::CompileTime(FFrameNumber InGlobalTime, UMovieSceneSequence& InCompileSequence, IMovieSceneSequenceTemplateStore& InTemplateStore)
 {
 	FMovieSceneEvaluationTemplate& CompileTemplate = InTemplateStore.AccessTemplate(InCompileSequence);
 
-	TRange<float> CompileRange = TRange<float>::Inclusive(InGlobalTime, InGlobalTime);
+	TRange<FFrameNumber> CompileRange = TRange<FFrameNumber>::Inclusive(InGlobalTime, InGlobalTime);
 
 	// Pass down a mutable path to the gather functions
 	FMovieSceneRootOverridePath RootPath;
@@ -179,8 +179,8 @@ TOptional<FCompiledGroupResult> FMovieSceneCompiler::CompileTime(float InGlobalT
 	FMovieSceneEvaluationTreeRangeIterator NodeAtTime = GatherData.Tracks.IterateFromTime(InGlobalTime);
 	if (ensure(NodeAtTime))
 	{
-		TRange<float> CompiledRange = TRange<float>::Intersection(
-			TRange<float>::Intersection(
+		TRange<FFrameNumber> CompiledRange = TRange<FFrameNumber>::Intersection(
+			TRange<FFrameNumber>::Intersection(
 				GatherData.EmptyOverlappingRange,
 				GatherData.Sequences.IterateFromLowerBound(CompileRange.GetLowerBound()).Range()
 			),
@@ -188,7 +188,7 @@ TOptional<FCompiledGroupResult> FMovieSceneCompiler::CompileTime(float InGlobalT
 
 		// Always include the global time in the compiled range - this is necessary when floating point rounding
 		// introduced during sub sequence transformations shift the range to just outside the current time
-		CompiledRange = TRange<float>::Hull(CompiledRange, CompileRange);
+		CompiledRange = TRange<FFrameNumber>::Hull(CompiledRange, CompileRange);
 
 		TArray<FCompileOnTheFlyData> SortedCompileData;
 		for (const FCompileOnTheFlyData& TrackData : GatherData.Tracks.GetAllData(NodeAtTime.Node()))
@@ -210,7 +210,7 @@ TOptional<FCompiledGroupResult> FMovieSceneCompiler::CompileTime(float InGlobalT
 
 		// Compute meta data for this segment
 		TMovieSceneEvaluationTreeDataIterator<FMovieSceneSequenceID> SubSequences = GatherData.Sequences.GetAllData(GatherData.Sequences.IterateFromLowerBound(CompileRange.GetLowerBound()).Node());
-		PopulateMetaData(Result, CompileTemplate.Hierarchy, SortedCompileData, SubSequences);
+		PopulateMetaData(Result, CompileTemplate.Hierarchy, InTemplateStore, SortedCompileData, SubSequences);
 
 		return Result;
 	}
@@ -307,7 +307,7 @@ void FMovieSceneCompiler::GatherCompileOnTheFlyData(UMovieSceneSequence& InSeque
 		}
 	}
 
-	TRange<float> CompileClampIntersection = TRange<float>::Intersection(Params.LocalCompileRange, Params.LocalClampRange);
+	TRange<FFrameNumber> CompileClampIntersection = TRange<FFrameNumber>::Intersection(Params.LocalCompileRange, Params.LocalClampRange);
 
 	// Iterate sub section ranges that overlap with the compile range
 	FGatherParameters SubSectionGatherParams = Params;
@@ -318,11 +318,11 @@ void FMovieSceneCompiler::GatherCompileOnTheFlyData(UMovieSceneSequence& InSeque
 	FMovieSceneEvaluationTreeRangeIterator SubSectionIt(SubSectionField.IterateFromLowerBound(CompileClampIntersection.GetLowerBound()));
 
 	// Intersect the unique range in the tree with the current overlapping empty range to constrict the resulting compile range in the case where this is a gap between sub sections
-	OutData.EmptyOverlappingRange = TRange<float>::Intersection(OutData.EmptyOverlappingRange, SubSectionIt.Range() * Params.RootToSequenceTransform.Inverse());
+	OutData.EmptyOverlappingRange = TRange<FFrameNumber>::Intersection(OutData.EmptyOverlappingRange, SubSectionIt.Range() * Params.RootToSequenceTransform.Inverse());
 
 	for ( ; SubSectionIt && SubSectionIt.Range().Overlaps(CompileClampIntersection); ++SubSectionIt)
 	{
-		TRange<float> ThisSegmentRange = SubSectionIt.Range() * Params.RootToSequenceTransform.Inverse();
+		TRange<FFrameNumber> ThisSegmentRange = SubSectionIt.Range() * Params.RootToSequenceTransform.Inverse();
 		if (ThisSegmentRange.IsEmpty())
 		{
 			continue;
@@ -408,8 +408,8 @@ const FMovieSceneSubSequenceData* FMovieSceneCompiler::GetOrCreateSubSequenceDat
 	const FMovieSceneSubSequenceData* ParentSubData = ParentSequenceID != MovieSceneSequenceID::Root ? InOutHierarchy.FindSubData(ParentSequenceID) : nullptr;
 	if (ParentSubData)
 	{
-		TRange<float> ParentPlayRangeChildSpace = ParentSubData->PlayRange * NewSubData.RootToSequenceTransform;
-		NewSubData.PlayRange = TRange<float>::Intersection(ParentPlayRangeChildSpace, NewSubData.PlayRange);
+		TRange<FFrameNumber> ParentPlayRangeChildSpace = ParentSubData->PlayRange.Value * NewSubData.RootToSequenceTransform;
+		NewSubData.PlayRange = TRange<FFrameNumber>::Intersection(ParentPlayRangeChildSpace, NewSubData.PlayRange.Value);
 
 		// Accumulate parent transform
 		NewSubData.RootToSequenceTransform = NewSubData.RootToSequenceTransform * ParentSubData->RootToSequenceTransform;
@@ -438,10 +438,10 @@ void FMovieSceneCompiler::GatherCompileDataForTrack(FMovieSceneEvaluationTrack&
 	if (SegmentIDs.Num() == 0)
 	{
 		// No segment at this time, so just report the time range of the empty space.
-		TRange<float> EmptyTrackSpace = Track.GetUniqueRangeFromLowerBound(Params.LocalCompileRange.GetLowerBound());
-		TRange<float> ClampedEmptyTrackSpaceRoot = TRange<float>::Intersection(Params.LocalClampRange, EmptyTrackSpace) * SequenceToRootTransform;
+		TRange<FFrameNumber> EmptyTrackSpace = Track.GetUniqueRangeFromLowerBound(Params.LocalCompileRange.GetLowerBound());
+		TRange<FFrameNumber> ClampedEmptyTrackSpaceRoot = TRange<FFrameNumber>::Intersection(Params.LocalClampRange, EmptyTrackSpace) * SequenceToRootTransform;
 
-		OutData.EmptyOverlappingRange = TRange<float>::Intersection(ClampedEmptyTrackSpaceRoot, OutData.EmptyOverlappingRange);
+		OutData.EmptyOverlappingRange = TRange<FFrameNumber>::Intersection(ClampedEmptyTrackSpaceRoot, OutData.EmptyOverlappingRange);
 	}
 	else for (FMovieSceneSegmentIdentifier SegmentID : SegmentIDs)
 	{
@@ -455,7 +455,7 @@ void FMovieSceneCompiler::GatherCompileDataForTrack(FMovieSceneEvaluationTrack&
 		Data.Track = &Track;
 		Data.bRequiresInit = ThisSegment.Impls.ContainsByPredicate(RequiresInit);
 
-		TRange<float> IntersectionRange = TRange<float>::Intersection(Params.LocalClampRange, ThisSegment.Range) * SequenceToRootTransform;
+		TRange<FFrameNumber> IntersectionRange = TRange<FFrameNumber>::Intersection(Params.LocalClampRange, ThisSegment.Range) * SequenceToRootTransform;
 		if (!IntersectionRange.IsEmpty())
 		{
 			OutData.Tracks.Add(IntersectionRange, Data);
@@ -463,7 +463,7 @@ void FMovieSceneCompiler::GatherCompileDataForTrack(FMovieSceneEvaluationTrack&
 	}
 }
 
-void FMovieSceneCompiler::PopulateMetaData(FCompiledGroupResult& OutResult, const FMovieSceneSequenceHierarchy& RootHierarchy, const TArray<FCompileOnTheFlyData>& SortedCompileData, TMovieSceneEvaluationTreeDataIterator<FMovieSceneSequenceID> SubSequences)
+void FMovieSceneCompiler::PopulateMetaData(FCompiledGroupResult& OutResult, const FMovieSceneSequenceHierarchy& RootHierarchy, IMovieSceneSequenceTemplateStore& TemplateStore, const TArray<FCompileOnTheFlyData>& SortedCompileData, TMovieSceneEvaluationTreeDataIterator<FMovieSceneSequenceID> SubSequences)
 {
 	OutResult.MetaData.Reset();
 
@@ -523,8 +523,10 @@ void FMovieSceneCompiler::PopulateMetaData(FCompiledGroupResult& OutResult, cons
 
 			UMovieSceneSequence* Sequence = SubData->GetSequence();
 
+			FGuid TemplateID = Sequence ? TemplateStore.AccessTemplate(*Sequence).TemplateSignature : FGuid();
+
 			OutResult.MetaData.ActiveSequences.Add(SequenceID);
-			OutResult.MetaData.SubSequenceSignatures.Add(SequenceID, Sequence ? Sequence->GetSignature() : FGuid());
+			OutResult.MetaData.SubTemplateSignatures.Add(SequenceID, TemplateID);
 		}
 
 		OutResult.MetaData.ActiveSequences.Sort();
diff --git a/Engine/Source/Runtime/MovieScene/Private/Compilation/MovieSceneEvaluationTemplateGenerator.cpp b/Engine/Source/Runtime/MovieScene/Private/Compilation/MovieSceneEvaluationTemplateGenerator.cpp
--- a/Engine/Source/Runtime/MovieScene/Private/Compilation/MovieSceneEvaluationTemplateGenerator.cpp
+++ b/Engine/Source/Runtime/MovieScene/Private/Compilation/MovieSceneEvaluationTemplateGenerator.cpp
@@ -53,6 +53,7 @@ void FMovieSceneEvaluationTemplateGenerator::Generate()
 	Template.RemoveStaleData(CompiledSignatures);
 
 	Template.SequenceSignature = SourceSequence.GetSignature();
+	Template.TemplateSignature = FGuid::NewGuid();
 }
 
 void FMovieSceneEvaluationTemplateGenerator::ProcessTrack(const UMovieSceneTrack& Track, const FGuid& ObjectBindingId)
@@ -102,7 +103,7 @@ void FMovieSceneEvaluationTemplateGenerator::ProcessSubTrack(const UMovieSceneSu
 	for (int32 SectionIndex = 0; SectionIndex < AllSections.Num(); ++SectionIndex)
 	{
 		UMovieSceneSubSection* SubSection = CastChecked<UMovieSceneSubSection>(AllSections[SectionIndex]);
-		if (!SubSection || !SubSection->IsActive())
+		if (!SubSection || !SubSection->IsActive() || SubSection->GetRange().IsEmpty())
 		{
 			continue;
 		}
@@ -111,7 +112,7 @@ void FMovieSceneEvaluationTemplateGenerator::ProcessSubTrack(const UMovieSceneSu
 
 		if (bRequiresSegmentBlending)
 		{
-			const TRange<float> SectionRange = SubSection->GetRange();
+			const TRange<FFrameNumber> SectionRange = SubSection->GetRange();
 
 			// Process the actual section range
 			if (bRequiresSegmentBlending)
@@ -124,9 +125,9 @@ void FMovieSceneEvaluationTemplateGenerator::ProcessSubTrack(const UMovieSceneSu
 			}
 
 			// Process the section preroll range
-			if (!SectionRange.GetLowerBound().IsOpen() && SubSection->GetPreRollTime() > 0)
+			if (!SectionRange.GetLowerBound().IsOpen() && SubSection->GetPreRollFrames() > 0)
 			{
-				TRange<float> PreRollRange(SectionRange.GetLowerBoundValue() - SubSection->GetPreRollTime(), TRangeBound<float>::FlipInclusion(SectionRange.GetLowerBoundValue()));
+				TRange<FFrameNumber> PreRollRange(SectionRange.GetLowerBoundValue() - SubSection->GetPreRollFrames(), TRangeBound<FFrameNumber>::FlipInclusion(SectionRange.GetLowerBoundValue()));
 
 				if (bRequiresSegmentBlending)
 				{
@@ -139,9 +140,9 @@ void FMovieSceneEvaluationTemplateGenerator::ProcessSubTrack(const UMovieSceneSu
 			}
 
 			// Process the section postroll range
-			if (!SectionRange.GetUpperBound().IsOpen() && SubSection->GetPostRollTime() > 0)
+			if (!SectionRange.GetUpperBound().IsOpen() && SubSection->GetPostRollFrames() > 0)
 			{
-				TRange<float> PostRollRange(TRangeBound<float>::FlipInclusion(SectionRange.GetUpperBoundValue()), SectionRange.GetUpperBoundValue() + SubSection->GetPostRollTime());
+				TRange<FFrameNumber> PostRollRange(TRangeBound<FFrameNumber>::FlipInclusion(SectionRange.GetUpperBoundValue()), SectionRange.GetUpperBoundValue() + SubSection->GetPostRollFrames());
 
 				if (bRequiresSegmentBlending)
 				{
diff --git a/Engine/Source/Runtime/MovieScene/Private/Compilation/MovieSceneSegmentCompiler.cpp b/Engine/Source/Runtime/MovieScene/Private/Compilation/MovieSceneSegmentCompiler.cpp
--- a/Engine/Source/Runtime/MovieScene/Private/Compilation/MovieSceneSegmentCompiler.cpp
+++ b/Engine/Source/Runtime/MovieScene/Private/Compilation/MovieSceneSegmentCompiler.cpp
@@ -27,16 +27,5 @@ PRAGMA_ENABLE_DEPRECATION_WARNINGS
 
 bool FMovieSceneAdditiveCameraTrackBlender::SortByStartTime(const FMovieSceneSectionData& A, const FMovieSceneSectionData& B)
 {
-	if (A.Section->IsInfinite())
-	{
-		return true;
-	}
-	else if (B.Section->IsInfinite())
-	{
-		return false;
-	}
-	else
-	{
-		return A.Section->GetStartTime() < B.Section->GetStartTime();
-	}
+	return TRangeBound<FFrameNumber>::MinLower(A.Section->GetRange().GetLowerBound(), B.Section->GetRange().GetLowerBound()) == A.Section->GetRange().GetLowerBound();
 }
diff --git a/Engine/Source/Runtime/MovieScene/Private/Evaluation/MovieSceneEvalTemplate.cpp b/Engine/Source/Runtime/MovieScene/Private/Evaluation/MovieSceneEvalTemplate.cpp
--- a/Engine/Source/Runtime/MovieScene/Private/Evaluation/MovieSceneEvalTemplate.cpp
+++ b/Engine/Source/Runtime/MovieScene/Private/Evaluation/MovieSceneEvalTemplate.cpp
@@ -3,12 +3,13 @@
 #include "Evaluation/MovieSceneEvalTemplate.h"
 #include "Evaluation/MovieSceneEvalTemplateSerializer.h"
 
-float FMovieSceneEvalTemplate::EvaluateEasing(float CurrentTime) const
+float FMovieSceneEvalTemplate::EvaluateEasing(FFrameTime CurrentTime) const
 {
 	return SourceSection ? SourceSection->EvaluateEasing(CurrentTime) : 1.f;
 }
 
 bool FMovieSceneEvalTemplatePtr::Serialize(FArchive& Ar)
 {
-	return SerializeInlineValue(*this, Ar);
+	bool bShouldWarn = !WITH_EDITORONLY_DATA;
+	return SerializeInlineValue(*this, Ar, bShouldWarn);
 }
diff --git a/Engine/Source/Runtime/MovieScene/Private/Evaluation/MovieSceneEvalTemplateSerializer.h b/Engine/Source/Runtime/MovieScene/Private/Evaluation/MovieSceneEvalTemplateSerializer.h
--- a/Engine/Source/Runtime/MovieScene/Private/Evaluation/MovieSceneEvalTemplateSerializer.h
+++ b/Engine/Source/Runtime/MovieScene/Private/Evaluation/MovieSceneEvalTemplateSerializer.h
@@ -5,12 +5,13 @@
 #include "CoreMinimal.h"
 #include "MovieSceneFwd.h"
 #include "Evaluation/MovieSceneEvaluationCustomVersion.h"
+#include "Evaluation/MovieSceneEvalTemplateBase.h"
 
 struct FMovieSceneEmptyStruct;
 
 /** Serialize an inline value that has a GetScriptStruct method */
 template<typename T, uint8 N>
-bool SerializeInlineValue(TInlineValue<T, N>& Impl, FArchive& Ar)
+bool SerializeInlineValue(TInlineValue<T, N>& Impl, FArchive& Ar, bool bWarnOnError)
 {
 	Ar.UsingCustomVersion(FMovieSceneEvaluationCustomVersion::GUID);
 	
@@ -27,12 +28,14 @@ bool SerializeInlineValue(TInlineValue<T, N>& Impl, FArchive& Ar)
 
 		// Find the script struct of the thing that was serialized
 		UScriptStruct* Struct = FindObject<UScriptStruct>(nullptr, *TypeName);
-		if (!Struct || Struct == T::StaticStruct())
+		if (!Struct)
 		{
-#if !WITH_EDITORONLY_DATA
-			// We only throw this warning in cooked builds as that is the only place where this deserialized data matters
-			UE_LOG(LogMovieScene, Warning, TEXT("Unknown or invalid type (%s) found in serialized data. This will no longer work. Please recompile template data."), *TypeName);
-#endif
+			if (bWarnOnError)
+			{
+				// We only throw this warning in cooked builds as that is the only place where this deserialized data matters
+				UE_LOG(LogMovieScene, Warning, TEXT("Unknown or invalid type (%s) found in serialized data. This will no longer work."), *TypeName);
+			}
+
 			// If it wasn't found, just deserialize an empty struct instead, and set ourselves to default
 			FMovieSceneEmptyStruct Empty;
 			FMovieSceneEmptyStruct::StaticStruct()->SerializeItem(Ar, &Empty, nullptr);
diff --git a/Engine/Source/Runtime/MovieScene/Private/Evaluation/MovieSceneEvaluationField.cpp b/Engine/Source/Runtime/MovieScene/Private/Evaluation/MovieSceneEvaluationField.cpp
--- a/Engine/Source/Runtime/MovieScene/Private/Evaluation/MovieSceneEvaluationField.cpp
+++ b/Engine/Source/Runtime/MovieScene/Private/Evaluation/MovieSceneEvaluationField.cpp
@@ -2,18 +2,19 @@
 
 #include "Evaluation/MovieSceneEvaluationField.h"
 #include "Evaluation/MovieSceneEvaluationTemplateInstance.h"
+#include "Evaluation/MovieSceneSequenceTemplateStore.h"
 #include "MovieSceneCommonHelpers.h"
 #include "Algo/Sort.h"
 
 #include "MovieSceneSequence.h"
 
-int32 FMovieSceneEvaluationField::GetSegmentFromTime(float Time) const
+int32 FMovieSceneEvaluationField::GetSegmentFromTime(FFrameNumber Time) const
 {
 	// Linear search
 	// @todo: accelerated search based on the last evaluated index?
 	for (int32 Index = 0; Index < Ranges.Num(); ++Index)
 	{
-		if (Ranges[Index].Contains(Time))
+		if (Ranges[Index].Value.Contains(Time))
 		{
 			return Index;
 		}
@@ -22,12 +23,12 @@ int32 FMovieSceneEvaluationField::GetSegmentFromTime(float Time) const
 	return INDEX_NONE;
 }
 
-TRange<int32> FMovieSceneEvaluationField::OverlapRange(TRange<float> Range) const
+TRange<int32> FMovieSceneEvaluationField::OverlapRange(TRange<FFrameNumber> Range) const
 {
 	int32 StartIndex = 0, Num = 0;
 	for (int32 Index = 0; Index < Ranges.Num(); ++Index)
 	{
-		if (Ranges[Index].Overlaps(Range))
+		if (Ranges[Index].Value.Overlaps(Range))
 		{
 			if (Num == 0)
 			{
@@ -44,7 +45,7 @@ TRange<int32> FMovieSceneEvaluationField::OverlapRange(TRange<float> Range) cons
 	return Num != 0 ? TRange<int32>(StartIndex, StartIndex + Num) : TRange<int32>::Empty();
 }
 
-void FMovieSceneEvaluationField::Invalidate(TRange<float> Range)
+void FMovieSceneEvaluationField::Invalidate(TRange<FFrameNumber> Range)
 {
 	TRange<int32> OverlappingRange = OverlapRange(Range);
 	if (!OverlappingRange.IsEmpty())
@@ -59,25 +60,26 @@ void FMovieSceneEvaluationField::Invalidate(TRange<float> Range)
 	}
 }
 
-int32 FMovieSceneEvaluationField::Insert(TRange<float> InRange, FMovieSceneEvaluationGroup&& InGroup, FMovieSceneEvaluationMetaData&& InMetaData)
+int32 FMovieSceneEvaluationField::Insert(FFrameNumber InsertTime, TRange<FFrameNumber> InRange, FMovieSceneEvaluationGroup&& InGroup, FMovieSceneEvaluationMetaData&& InMetaData)
 {
-	const int32 InsertIndex = Algo::UpperBoundBy(Ranges, InRange.GetLowerBound(), &TRange<float>::GetLowerBound, MovieSceneHelpers::SortLowerBounds);
+	const int32 InsertIndex = Algo::UpperBoundBy(Ranges, TRangeBound<FFrameNumber>(InsertTime), &FMovieSceneFrameRange::GetLowerBound, MovieSceneHelpers::SortLowerBounds);
 
-	// While we have floating point representations of ranges, there are a number of situations that can lead to the check firing:
-	// 		1. Consider ranges of (0,9.16666603f) and (9.16666698f,20) offset in an inner sequence by 98.8333511f:
-	//			Such transformation would result in the upper and lower bounds rounding to the same inclusive value (108.000015f) making them overlap in the parent space.
-	//		2. When compiling a segment at a particular time, that time transformed into an inner sequence could result in a compiled range that doesn't actually
-	//			include the global time when transformed back into the master space. In this scenario we always inflate the compiled range to include the global time, but this shouldn't be necessary
-	//			Such logic can lead to us attempting to add overlapping ranges into the field
+	// Intersect the supplied range with the allowable space between adjacent existing ranges
+	TRange<FFrameNumber> InsertSpace(
+		Ranges.IsValidIndex(InsertIndex-1) ? TRangeBound<FFrameNumber>::FlipInclusion(Ranges[InsertIndex-1].GetUpperBound()) : TRangeBound<FFrameNumber>::Open(),
+		Ranges.IsValidIndex(InsertIndex  ) ? TRangeBound<FFrameNumber>::FlipInclusion(Ranges[InsertIndex  ].GetLowerBound()) : TRangeBound<FFrameNumber>::Open()
+		);
+
+	InRange = TRange<FFrameNumber>::Intersection(InRange, InsertSpace);
 
 	// @todo: Remove this code and enforce the check below outright when we have proper time representation
-	if (Ranges.IsValidIndex(InsertIndex  ) && Ranges[InsertIndex  ].Overlaps(InRange))
+	if (Ranges.IsValidIndex(InsertIndex  ) && Ranges[InsertIndex  ].Value.Overlaps(InRange))
 	{
-		InRange = TRange<float>(InRange.GetLowerBound(), TRangeBound<float>::FlipInclusion(Ranges[InsertIndex].GetLowerBound()));
+		InRange = TRange<FFrameNumber>(InRange.GetLowerBound(), TRangeBound<FFrameNumber>::FlipInclusion(Ranges[InsertIndex].Value.GetLowerBound()));
 	}
-	if (Ranges.IsValidIndex(InsertIndex-1) && Ranges[InsertIndex-1].Overlaps(InRange))
+	if (Ranges.IsValidIndex(InsertIndex-1) && Ranges[InsertIndex-1].Value.Overlaps(InRange))
 	{
-		InRange = TRange<float>(TRangeBound<float>::FlipInclusion(Ranges[InsertIndex-1].GetUpperBound()), InRange.GetUpperBound());
+		InRange = TRange<FFrameNumber>(TRangeBound<FFrameNumber>::FlipInclusion(Ranges[InsertIndex-1].Value.GetUpperBound()), InRange.GetUpperBound());
 	}
 
 	if (!ensure(!InRange.IsEmpty()))
@@ -86,8 +88,8 @@ int32 FMovieSceneEvaluationField::Insert(TRange<float> InRange, FMovieSceneEvalu
 	}
 
 	const bool bOverlapping = 
-		(Ranges.IsValidIndex(InsertIndex  ) && Ranges[InsertIndex  ].Overlaps(InRange)) ||
-		(Ranges.IsValidIndex(InsertIndex-1) && Ranges[InsertIndex-1].Overlaps(InRange));
+		(Ranges.IsValidIndex(InsertIndex  ) && Ranges[InsertIndex  ].Value.Overlaps(InRange)) ||
+		(Ranges.IsValidIndex(InsertIndex-1) && Ranges[InsertIndex-1].Value.Overlaps(InRange));
 
 	if (!ensureAlwaysMsgf(!bOverlapping, TEXT("Attempting to insert an overlapping range into the evaluation field.")))
 	{
@@ -105,9 +107,9 @@ int32 FMovieSceneEvaluationField::Insert(TRange<float> InRange, FMovieSceneEvalu
 	return InsertIndex;
 }
 
-void FMovieSceneEvaluationField::Add(TRange<float> InRange, FMovieSceneEvaluationGroup&& InGroup, FMovieSceneEvaluationMetaData&& InMetaData)
+void FMovieSceneEvaluationField::Add(TRange<FFrameNumber> InRange, FMovieSceneEvaluationGroup&& InGroup, FMovieSceneEvaluationMetaData&& InMetaData)
 {
-	if (ensureAlwaysMsgf(!Ranges.Num() || !Ranges.Last().Overlaps(InRange), TEXT("Attempting to add overlapping ranges to sequence evaluation field.")))
+	if (ensureAlwaysMsgf(!Ranges.Num() || !Ranges.Last().Value.Overlaps(InRange), TEXT("Attempting to add overlapping ranges to sequence evaluation field.")))
 	{
 		Ranges.Add(InRange);
 		MetaData.Add(MoveTemp(InMetaData));
@@ -248,24 +250,32 @@ void FMovieSceneEvaluationMetaData::DiffEntities(const FMovieSceneEvaluationMeta
 	}
 }
 
-bool FMovieSceneEvaluationMetaData::IsDirty(UMovieSceneSequence& RootSequence, const FMovieSceneSequenceHierarchy& RootHierarchy, TRange<float>* OutSubRangeToInvalidate) const
+bool FMovieSceneEvaluationMetaData::IsDirty(UMovieSceneSequence& RootSequence, const FMovieSceneSequenceHierarchy& RootHierarchy, IMovieSceneSequenceTemplateStore& TemplateStore, TRange<FFrameNumber>* OutSubRangeToInvalidate) const
 {
 	bool bDirty = false;
 
-	for (const TTuple<FMovieSceneSequenceID, FGuid>& Pair : SubSequenceSignatures)
+	for (const TTuple<FMovieSceneSequenceID, FGuid>& Pair : SubTemplateSignatures)
 	{
 		// Sequence IDs at this point are relative to the root override template
 		const FMovieSceneSubSequenceData* SubData = RootHierarchy.FindSubData(Pair.Key);
 		UMovieSceneSequence* SubSequence = SubData ? SubData->GetSequence() : nullptr;
 
-		if (!SubSequence || SubSequence->GetSignature() != Pair.Value)
+		bool bThisSequenceIsDirty = true;
+		if (SubSequence)
+		{
+			FMovieSceneEvaluationTemplate& Template = TemplateStore.AccessTemplate(*SubSequence);
+
+			bThisSequenceIsDirty = Template.TemplateSignature != Pair.Value || Template.SequenceSignature != SubSequence->GetSignature();
+		}
+
+		if (bThisSequenceIsDirty)
 		{
 			bDirty = true;
 
 			if (OutSubRangeToInvalidate)
 			{
-				TRange<float> DirtyRange = SubData ? TRange<float>::Hull(TRange<float>::Hull(SubData->PreRollRange, SubData->PlayRange), SubData->PostRollRange) * SubData->RootToSequenceTransform.Inverse() : TRange<float>::All();
-				*OutSubRangeToInvalidate = TRange<float>::Hull(*OutSubRangeToInvalidate, DirtyRange);
+				TRange<FFrameNumber> DirtyRange = SubData ? TRange<FFrameNumber>::Hull(TRange<FFrameNumber>::Hull(SubData->PreRollRange.Value, SubData->PlayRange.Value), SubData->PostRollRange.Value) * SubData->RootToSequenceTransform.Inverse() : TRange<FFrameNumber>::All();
+				*OutSubRangeToInvalidate = TRange<FFrameNumber>::Hull(*OutSubRangeToInvalidate, DirtyRange);
 			}
 		}
 	}
diff --git a/Engine/Source/Runtime/MovieScene/Private/Evaluation/MovieSceneEvaluationTemplate.cpp b/Engine/Source/Runtime/MovieScene/Private/Evaluation/MovieSceneEvaluationTemplate.cpp
--- a/Engine/Source/Runtime/MovieScene/Private/Evaluation/MovieSceneEvaluationTemplate.cpp
+++ b/Engine/Source/Runtime/MovieScene/Private/Evaluation/MovieSceneEvaluationTemplate.cpp
@@ -51,11 +51,10 @@ const TMovieSceneEvaluationTree<FMovieSceneSubSectionData>& FMovieSceneEvaluatio
 	return SubSectionFieldData.Field;
 }
 
-void FMovieSceneEvaluationTemplate::AddSubSectionRange(UMovieSceneSubSection& InSubSection, const FGuid& InObjectBindingId, const TRange<float>& InRange, ESectionEvaluationFlags InFlags)
+void FMovieSceneEvaluationTemplate::AddSubSectionRange(UMovieSceneSubSection& InSubSection, const FGuid& InObjectBindingId, const TRange<FFrameNumber>& InRange, ESectionEvaluationFlags InFlags)
 {
-	if (!ensure(!InSubSection.IsInfinite()))
+	if (!ensure(InRange.GetLowerBound().IsClosed() && InRange.GetUpperBound().IsClosed()))
 	{
-		// @todo: Infinite sub sections??
 		return;
 	}
 
@@ -75,7 +74,9 @@ void FMovieSceneEvaluationTemplate::AddSubSectionRange(UMovieSceneSubSection& In
 			Hierarchy.Add(NewSubData, SubSequenceID, MovieSceneSequenceID::Root);
 		}
 
-		TRange<float> EntireSectionRange(InSubSection.GetStartTime() - InSubSection.GetPreRollTime(), TRangeBound<float>::Inclusive(InSubSection.GetEndTime() + InSubSection.GetPostRollTime()));
+		TRange<FFrameNumber> EntireSectionRange = InSubSection.GetRange();
+		EntireSectionRange.SetLowerBoundValue(EntireSectionRange.GetLowerBoundValue() - InSubSection.GetPreRollFrames());
+		EntireSectionRange.SetUpperBoundValue(EntireSectionRange.GetUpperBoundValue() + InSubSection.GetPostRollFrames());
 
 		// Add the section to the ledger
 		TemplateLedger.SubSectionRanges.Add(InSubSection.GetSignature(), EntireSectionRange);
@@ -113,7 +114,7 @@ void FMovieSceneEvaluationTemplate::DefineTrackStructure(FMovieSceneTrackIdentif
 	if (bAddEmptySpace && bInvalidateEvaluationField)
 	{
 		// Optimization - when tracks can add empty space, we just invalidate the entire field once
-		EvaluationField.Invalidate(TRange<float>::All());
+		EvaluationField.Invalidate(TRange<FFrameNumber>::All());
 		bInvalidateEvaluationField = false;
 	}
 
@@ -185,9 +186,9 @@ void FMovieSceneEvaluationTemplate::RemoveStaleData(const TSet<FGuid>& ActiveSig
 
 	// Remove stale sub sections
 	{
-		TArray<TTuple<FGuid, FFloatRange>> SubSectionsToRemove;
+		TArray<TTuple<FGuid, FMovieSceneFrameRange>> SubSectionsToRemove;
 
-		for (const TTuple<FGuid, FFloatRange>& Pair : TemplateLedger.SubSectionRanges)
+		for (const TTuple<FGuid, FMovieSceneFrameRange>& Pair : TemplateLedger.SubSectionRanges)
 		{
 			if (!ActiveSignatures.Contains(Pair.Key))
 			{
@@ -195,10 +196,10 @@ void FMovieSceneEvaluationTemplate::RemoveStaleData(const TSet<FGuid>& ActiveSig
 			}
 		}
 
-		for (const TTuple<FGuid, FFloatRange>& Pair : SubSectionsToRemove)
+		for (const TTuple<FGuid, FMovieSceneFrameRange>& Pair : SubSectionsToRemove)
 		{
 			TemplateLedger.SubSectionRanges.Remove(Pair.Key);
-			EvaluationField.Invalidate(Pair.Value);
+			EvaluationField.Invalidate(Pair.Value.Value);
 		}
 	}
 }
diff --git a/Engine/Source/Runtime/MovieScene/Private/Evaluation/MovieSceneEvaluationTemplateInstance.cpp b/Engine/Source/Runtime/MovieScene/Private/Evaluation/MovieSceneEvaluationTemplateInstance.cpp
--- a/Engine/Source/Runtime/MovieScene/Private/Evaluation/MovieSceneEvaluationTemplateInstance.cpp
+++ b/Engine/Source/Runtime/MovieScene/Private/Evaluation/MovieSceneEvaluationTemplateInstance.cpp
@@ -195,15 +195,15 @@ const FMovieSceneEvaluationGroup* FMovieSceneRootEvaluationTemplateInstance::Set
 	}
 
 	// First off, attempt to find the evaluation group in the existing evaluation field data from the template
-	int32 TemplateFieldIndex = RootOverrideInstance.Template->EvaluationField.GetSegmentFromTime(Context.GetTime());
+	int32 TemplateFieldIndex = RootOverrideInstance.Template->EvaluationField.GetSegmentFromTime(Context.GetTime().FrameNumber);
 
 	if (TemplateFieldIndex != INDEX_NONE)
 	{
 		const FMovieSceneEvaluationMetaData& FieldMetaData = RootOverrideInstance.Template->EvaluationField.GetMetaData(TemplateFieldIndex);
 
 		// Verify that this field entry is still valid (all its cached signatures are still the same)
-		TRange<float> InvalidatedSubSequenceRange = TRange<float>::Empty();
-		if (FieldMetaData.IsDirty(*RootOverrideInstance.Sequence, RootOverrideInstance.Template->Hierarchy, &InvalidatedSubSequenceRange))
+		TRange<FFrameNumber> InvalidatedSubSequenceRange = TRange<FFrameNumber>::Empty();
+		if (FieldMetaData.IsDirty(*RootOverrideInstance.Sequence, RootOverrideInstance.Template->Hierarchy, *TemplateStore, &InvalidatedSubSequenceRange))
 		{
 			TemplateFieldIndex = INDEX_NONE;
 
@@ -222,16 +222,17 @@ const FMovieSceneEvaluationGroup* FMovieSceneRootEvaluationTemplateInstance::Set
 		if (bFullCompile)
 		{
 			FMovieSceneCompiler::Compile(*RootOverrideInstance.Sequence, *TemplateStore);
-			TemplateFieldIndex = RootOverrideInstance.Template->EvaluationField.GetSegmentFromTime(Context.GetTime());
+			TemplateFieldIndex = RootOverrideInstance.Template->EvaluationField.GetSegmentFromTime(Context.GetTime().FrameNumber);
 		}
 		else
 		{
-			TOptional<FCompiledGroupResult> CompileResult = FMovieSceneCompiler::CompileTime(Context.GetTime(), *RootOverrideInstance.Sequence, *TemplateStore);
+			TOptional<FCompiledGroupResult> CompileResult = FMovieSceneCompiler::CompileTime(Context.GetTime().FrameNumber, *RootOverrideInstance.Sequence, *TemplateStore);
 
 			if (CompileResult.IsSet())
 			{
-				TRange<float> FieldRange = CompileResult->Range;
+				TRange<FFrameNumber> FieldRange = CompileResult->Range;
 				TemplateFieldIndex = RootOverrideInstance.Template->EvaluationField.Insert(
+					Context.GetTime().FrameNumber,
 					FieldRange,
 					MoveTemp(CompileResult->Group),
 					MoveTemp(CompileResult->MetaData)
@@ -292,10 +293,10 @@ void FMovieSceneRootEvaluationTemplateInstance::EvaluateGroup(const FMovieSceneE
 				SubContext = Context;
 				if (Instance.SubData)
 				{
-					SubContext = Context.Transform(Instance.SubData->RootToSequenceTransform);
+					SubContext = Context.Transform(Instance.SubData->RootToSequenceTransform, Instance.SubData->FrameResolution);
 
 					// Hittest against the sequence's pre and postroll ranges
-					SubContext.ReportOuterSectionRanges(Instance.SubData->PreRollRange, Instance.SubData->PostRollRange);
+					SubContext.ReportOuterSectionRanges(Instance.SubData->PreRollRange.Value, Instance.SubData->PostRollRange.Value);
 					SubContext.SetHierarchicalBias(Instance.SubData->HierarchicalBias);
 				}
 
@@ -334,10 +335,10 @@ void FMovieSceneRootEvaluationTemplateInstance::EvaluateGroup(const FMovieSceneE
 				SubContext = Context;
 				if (Instance.SubData)
 				{
-					SubContext = Context.Transform(Instance.SubData->RootToSequenceTransform);
+					SubContext = Context.Transform(Instance.SubData->RootToSequenceTransform, Instance.SubData->FrameResolution);
 
 					// Hittest against the sequence's pre and postroll ranges
-					SubContext.ReportOuterSectionRanges(Instance.SubData->PreRollRange, Instance.SubData->PostRollRange);
+					SubContext.ReportOuterSectionRanges(Instance.SubData->PreRollRange.Value, Instance.SubData->PostRollRange.Value);
 					SubContext.SetHierarchicalBias(Instance.SubData->HierarchicalBias);
 				}
 
@@ -452,7 +453,12 @@ bool FMovieSceneRootEvaluationTemplateInstance::IsDirty() const
 {
 	if (TransientInstances.RootInstance.IsValid())
 	{
-		return LastFrameMetaData.IsDirty(*TransientInstances.RootInstance.Sequence, TransientInstances.RootInstance.Template->Hierarchy);
+		if (TransientInstances.RootInstance.Template->SequenceSignature != TransientInstances.RootInstance.Sequence->GetSignature())
+		{
+			return true;
+		}
+
+		return LastFrameMetaData.IsDirty(*TransientInstances.RootInstance.Sequence, TransientInstances.RootInstance.Template->Hierarchy, *TemplateStore);
 	}
 
 	return true;
diff --git a/Engine/Source/Runtime/MovieScene/Private/Evaluation/MovieSceneEvaluationTrack.cpp b/Engine/Source/Runtime/MovieScene/Private/Evaluation/MovieSceneEvaluationTrack.cpp
--- a/Engine/Source/Runtime/MovieScene/Private/Evaluation/MovieSceneEvaluationTrack.cpp
+++ b/Engine/Source/Runtime/MovieScene/Private/Evaluation/MovieSceneEvaluationTrack.cpp
@@ -161,7 +161,7 @@ void FMovieSceneEvaluationTrack::DefineAsSingleTemplate(FMovieSceneEvalTemplateP
 	ChildTemplates.Add(MoveTemp(InTemplate));
 
 	FSectionEvaluationData EvalData(0);
-	Segments.Add(FMovieSceneSegment(TRange<float>::All(), TArrayView<FSectionEvaluationData>(&EvalData, 1)));
+	Segments.Add(FMovieSceneSegment(TRange<FFrameNumber>::All(), TArrayView<FSectionEvaluationData>(&EvalData, 1)));
 }
 
 int32 FMovieSceneEvaluationTrack::AddChildTemplate(FMovieSceneEvalTemplatePtr&& InTemplate)
@@ -259,7 +259,7 @@ FMovieSceneSegmentIdentifier FMovieSceneEvaluationTrack::CompileSegment(FMovieSc
 	return Segments.Add(MoveTemp(NewSegment));
 }
 
-FMovieSceneSegmentIdentifier FMovieSceneEvaluationTrack::FindFirstSegment(TRange<float> InLocalRange)
+FMovieSceneSegmentIdentifier FMovieSceneEvaluationTrack::FindFirstSegment(TRange<FFrameNumber> InLocalRange)
 {
 	// Binary search for an existing compiled segment in the sorted array
 	auto GetLowerBound = [](const FMovieSceneSegment& In) {
@@ -281,7 +281,7 @@ FMovieSceneSegmentIdentifier FMovieSceneEvaluationTrack::FindFirstSegment(TRange
 	return FMovieSceneSegmentIdentifier();
 }
 
-TArray<FMovieSceneSegmentIdentifier> FMovieSceneEvaluationTrack::GetSegmentsInRange(TRange<float> InLocalRange)
+TArray<FMovieSceneSegmentIdentifier> FMovieSceneEvaluationTrack::GetSegmentsInRange(TRange<FFrameNumber> InLocalRange)
 {
 	TArray<FMovieSceneSegmentIdentifier> SegmentsInRange;
 
@@ -302,7 +302,7 @@ TArray<FMovieSceneSegmentIdentifier> FMovieSceneEvaluationTrack::GetSegmentsInRa
 					break;
 				}
 
-				InLocalRange = TRange<float>(TRangeBound<float>::FlipInclusion(NodeAtTime.Range().GetUpperBound()), InLocalRange.GetUpperBound());
+				InLocalRange.SetLowerBound(TRangeBound<FFrameNumber>::FlipInclusion(NodeAtTime.Range().GetUpperBound()));
 				continue;
 			}
 		}
@@ -313,15 +313,15 @@ TArray<FMovieSceneSegmentIdentifier> FMovieSceneEvaluationTrack::GetSegmentsInRa
 			break;
 		}
 
-		InLocalRange = TRange<float>(TRangeBound<float>::FlipInclusion(Segments[SegmentID].Range.GetUpperBound()), InLocalRange.GetUpperBound());
+		InLocalRange.SetLowerBound(TRangeBound<FFrameNumber>::FlipInclusion(Segments[SegmentID].Range.GetUpperBound()));
 	}
 
 	return SegmentsInRange;
 }
 
-FMovieSceneSegmentIdentifier FMovieSceneEvaluationTrack::GetSegmentFromTime(float InTime)
+FMovieSceneSegmentIdentifier FMovieSceneEvaluationTrack::GetSegmentFromTime(FFrameNumber InTime)
 {
-	FMovieSceneSegmentIdentifier Existing = FindFirstSegment(TRange<float>::Inclusive(InTime, InTime));
+	FMovieSceneSegmentIdentifier Existing = FindFirstSegment(TRange<FFrameNumber>::Inclusive(InTime, InTime));
 	if (Existing.IsValid())
 	{
 		return Existing;
@@ -346,7 +346,7 @@ FMovieSceneSegmentIdentifier FMovieSceneEvaluationTrack::GetSegmentFromIterator(
 	return CompileSegment(Iterator);
 }
 
-TRange<float> FMovieSceneEvaluationTrack::GetUniqueRangeFromLowerBound(TRangeBound<float> InLowerBound) const
+TRange<FFrameNumber> FMovieSceneEvaluationTrack::GetUniqueRangeFromLowerBound(TRangeBound<FFrameNumber> InLowerBound) const
 {
 	return EvaluationTree.Tree.IterateFromLowerBound(InLowerBound).Range();
 }
@@ -459,9 +459,10 @@ void FMovieSceneEvaluationTrack::EvaluateStatic(FMovieSceneSegmentIdentifier Seg
 
 namespace
 {
-	bool IntersectSegmentRanges(const FMovieSceneSegment& Segment, TRange<float> TraversedRange, TMap<int32, TRange<float>>& ImplToAccumulatedRange)
+
+	bool IntersectSegmentRanges(const FMovieSceneSegment& Segment, TRange<FFrameNumber> TraversedRange, TMap<int32, TRange<FFrameNumber>>& ImplToAccumulatedRange)
 	{
-		TRange<float> Intersection = TRange<float>::Intersection(Segment.Range, TraversedRange);
+		TRange<FFrameNumber> Intersection = TRange<FFrameNumber>::Intersection(Segment.Range, TraversedRange);
 		if (Intersection.IsEmpty())
 		{
 			return false;
@@ -469,21 +470,21 @@ namespace
 
 		for (const FSectionEvaluationData& EvalData : Segment.Impls)
 		{
-			TRange<float>* AccumulatedRange = ImplToAccumulatedRange.Find(EvalData.ImplIndex);
+			TRange<FFrameNumber>* AccumulatedRange = ImplToAccumulatedRange.Find(EvalData.ImplIndex);
 			if (!AccumulatedRange)
 			{
 				ImplToAccumulatedRange.Add(EvalData.ImplIndex, Intersection);
 			}
 			else
 			{
-				*AccumulatedRange = TRange<float>::Hull(*AccumulatedRange, Intersection);
+				*AccumulatedRange = TRange<FFrameNumber>::Hull(*AccumulatedRange, Intersection);
 			}
 		}
 
 		return true;
 	}
 
-	void GatherSweptSegments(TRange<float> TraversedRange, int32 CurrentSegmentIndex, TArrayView<const FMovieSceneSegment> Segments, TMap<int32, TRange<float>>& ImplToAccumulatedRange)
+	void GatherSweptSegments(TRange<FFrameNumber> TraversedRange, int32 CurrentSegmentIndex, TArrayView<const FMovieSceneSegment> Segments, TMap<int32, TRange<FFrameNumber>>& ImplToAccumulatedRange)
 	{
 		// Search backwards from the current segment for any segments intersecting the traversed range
 		{
@@ -505,7 +506,7 @@ namespace
 void FMovieSceneEvaluationTrack::EvaluateSwept(FMovieSceneSegmentIdentifier SegmentID, const FMovieSceneEvaluationOperand& Operand, const FMovieSceneContext& Context, const FPersistentEvaluationData& PersistentData, FMovieSceneExecutionTokens& ExecutionTokens) const
 {
 	// Accumulate the relevant ranges that each section intersects with the evaluated range
-	TMap<int32, TRange<float>> ImplToAccumulatedRange;
+	TMap<int32, TRange<FFrameNumber>> ImplToAccumulatedRange;
 
 	int32 SortedIndex = Segments.GetSortedIndex(SegmentID);
 	if (SortedIndex == INDEX_NONE)
@@ -513,12 +514,12 @@ void FMovieSceneEvaluationTrack::EvaluateSwept(FMovieSceneSegmentIdentifier Segm
 		return;
 	}
 
-	GatherSweptSegments(Context.GetRange(), SortedIndex, Segments.GetSorted(), ImplToAccumulatedRange);
+	GatherSweptSegments(Context.GetFrameNumberRange(), SortedIndex, Segments.GetSorted(), ImplToAccumulatedRange);
 
 	for (auto& Pair : ImplToAccumulatedRange)
 	{
 		const int32 SectionIndex = Pair.Key;
-		TRange<float> EvaluationRange = Pair.Value;
+		TRange<FFrameTime> ClampRange = FMovieSceneEvaluationRange::NumberRangeToTimeRange(Pair.Value);
 		const FMovieSceneEvalTemplate& Template = GetChildTemplate(SectionIndex);
 
 		PersistentData.DeriveSectionKey(SectionIndex);
@@ -527,7 +528,7 @@ void FMovieSceneEvaluationTrack::EvaluateSwept(FMovieSceneSegmentIdentifier Segm
 		
 		Template.EvaluateSwept(
 			Operand,
-			Context.Clamp(EvaluationRange),
+			Context.Clamp(ClampRange),
 			PersistentData,
 			ExecutionTokens);
 	}
@@ -540,7 +541,7 @@ void FMovieSceneEvaluationTrack::Interrogate(const FMovieSceneContext& Context,
 		return;
 	}
 
-	FMovieSceneSegmentIdentifier SegmentID = GetSegmentFromTime(Context.GetTime());
+	FMovieSceneSegmentIdentifier SegmentID = GetSegmentFromTime(Context.GetTime().FrameNumber);
 	if (!SegmentID.IsValid())
 	{
 		return;
@@ -556,9 +557,9 @@ void FMovieSceneEvaluationTrack::Interrogate(const FMovieSceneContext& Context,
 	else
 	{
 		// Accumulate the relevant ranges that each section intersects with the evaluated range
-		TMap<int32, TRange<float>> ImplToAccumulatedRange;
+		TMap<int32, TRange<FFrameNumber>> ImplToAccumulatedRange;
 
-		GatherSweptSegments(Context.GetRange(), Segments.GetSortedIndex(SegmentID), Segments.GetSorted(), ImplToAccumulatedRange);
+		GatherSweptSegments(Context.GetFrameNumberRange(), Segments.GetSortedIndex(SegmentID), Segments.GetSorted(), ImplToAccumulatedRange);
 		for (auto& Pair : ImplToAccumulatedRange)
 		{
 			GetChildTemplate(Pair.Key).Interrogate(Context, Pair.Value, Container, BindingOverride);
