Copying //UE4/Dev-Sequencer to //UE4/Dev-Main (Source: //UE4/Dev-Sequencer @ 3946692)

diff --git a/Engine/Source/Runtime/MovieScene/Private/Evaluation/MovieSceneEvaluationTree.cpp b/Engine/Source/Runtime/MovieScene/Private/Evaluation/MovieSceneEvaluationTree.cpp
--- a/Engine/Source/Runtime/MovieScene/Private/Evaluation/MovieSceneEvaluationTree.cpp
+++ b/Engine/Source/Runtime/MovieScene/Private/Evaluation/MovieSceneEvaluationTree.cpp
@@ -13,7 +13,7 @@
 
 
 FMovieSceneEvaluationTreeRangeIterator::FMovieSceneEvaluationTreeRangeIterator(const FMovieSceneEvaluationTree& InTree)
-	: CurrentRange(TRange<float>::All()), CurrentNodeHandle(FMovieSceneEvaluationTreeNodeHandle::Root()), Tree(InTree)
+	: CurrentRange(TRange<FFrameNumber>::All()), CurrentNodeHandle(FMovieSceneEvaluationTreeNodeHandle::Root()), Tree(InTree)
 {
 	// Compute the starting range by inspecting the front of the tree
 	const FMovieSceneEvaluationTreeNode* CurrentNode = &Tree.GetRootNode();
@@ -36,18 +36,18 @@ FMovieSceneEvaluationTreeRangeIterator::FMovieSceneEvaluationTreeRangeIterator(c
 	}
 
 	// The upper bound is either the current node's upper bound (if it has no children), or the inverse of its first child's lower bound
-	TRangeBound<float> TrailingBound = Children.Num() ? TRangeBound<float>::FlipInclusion(Children[0].Range.GetLowerBound()) : CurrentNode->Range.GetUpperBound();
-	CurrentRange = TRange<float>(TRangeBound<float>::Open(), TrailingBound);
+	TRangeBound<FFrameNumber> TrailingBound = Children.Num() ? TRangeBound<FFrameNumber>::FlipInclusion(Children[0].Range.GetLowerBound()) : CurrentNode->Range.GetUpperBound();
+	CurrentRange = TRange<FFrameNumber>(TRangeBound<FFrameNumber>::Open(), TrailingBound);
 }
 
-FMovieSceneEvaluationTreeRangeIterator::FMovieSceneEvaluationTreeRangeIterator(const FMovieSceneEvaluationTree& InTree, TRangeBound<float> StartingBound)
+FMovieSceneEvaluationTreeRangeIterator::FMovieSceneEvaluationTreeRangeIterator(const FMovieSceneEvaluationTree& InTree, TRangeBound<FFrameNumber> StartingBound)
 	: CurrentNodeHandle(FMovieSceneEvaluationTreeNodeHandle::Root()), Tree(InTree)
 {
 	auto GetLowerBound = [](const FMovieSceneEvaluationTreeNode& In) {
 		return In.Range.GetLowerBound();
 	};
 
-	TRange<float> CompareRange(StartingBound, TRangeBound<float>::Open());
+	TRange<FFrameNumber> CompareRange(StartingBound, TRangeBound<FFrameNumber>::Open());
 
 	for (;;)
 	{
@@ -75,9 +75,9 @@ FMovieSceneEvaluationTreeRangeIterator::FMovieSceneEvaluationTreeRangeIterator(c
 		else
 		{
 			// We're at some empty space between children
-			CurrentRange = TRange<float>(
-				Children.IsValidIndex(ChildIndex-1) ? TRangeBound<float>::FlipInclusion(Children[ChildIndex-1].Range.GetUpperBound()) : CurrentRange.GetLowerBound(),
-				Children.IsValidIndex(ChildIndex) ? TRangeBound<float>::FlipInclusion(Children[ChildIndex].Range.GetLowerBound()) : CurrentRange.GetUpperBound()
+			CurrentRange = TRange<FFrameNumber>(
+				Children.IsValidIndex(ChildIndex-1) ? TRangeBound<FFrameNumber>::FlipInclusion(Children[ChildIndex-1].Range.GetUpperBound()) : CurrentRange.GetLowerBound(),
+				Children.IsValidIndex(ChildIndex)   ? TRangeBound<FFrameNumber>::FlipInclusion(Children[ChildIndex].Range.GetLowerBound())   : CurrentRange.GetUpperBound()
 			);
 
 			break;
@@ -103,7 +103,7 @@ void FMovieSceneEvaluationTreeRangeIterator::Iter(bool bForwards)
 		return;
 	}
 
-	TRangeBound<float> NewLeadingBound = CurrentRange.IsEmpty() ? TRangeBound<float>::Open() : TRangeBound<float>::FlipInclusion(GetTrailingBound(bForwards, CurrentRange));
+	TRangeBound<FFrameNumber> NewLeadingBound = CurrentRange.IsEmpty() ? TRangeBound<FFrameNumber>::Open() : TRangeBound<FFrameNumber>::FlipInclusion(GetTrailingBound(bForwards, CurrentRange));
 
 	// Iterate into children when possible (where the leading bound matches the current new leading bound)
 	FMovieSceneEvaluationTreeNodeHandle NextChildHandle = FindNextChild(CurrentNodeHandle, NewLeadingBound, bForwards);
@@ -115,20 +115,20 @@ void FMovieSceneEvaluationTreeRangeIterator::Iter(bool bForwards)
 
 	// The new traililng bound is either the trailing bound of the whole node, or the inverse of the leading bound of the next child (which will only be valid if it's not the same as the leading bound).
 	// The latter means this is space in-between child nodes
-	TRangeBound<float> NewTrailingBound = GetTrailingBound(bForwards, Tree.GetNode(CurrentNodeHandle).Range);
+	TRangeBound<FFrameNumber> NewTrailingBound = GetTrailingBound(bForwards, Tree.GetNode(CurrentNodeHandle).Range);
 	if (NextChildHandle.IsValid())
 	{
 		const FMovieSceneEvaluationTreeNode& NextChildNode = Tree.GetNode(NextChildHandle);
 		if (!CompareBound(bForwards, NextChildNode.Range, NewLeadingBound))
 		{
-			NewTrailingBound = TRangeBound<float>::FlipInclusion(GetLeadingBound(bForwards, NextChildNode.Range));
+			NewTrailingBound = TRangeBound<FFrameNumber>::FlipInclusion(GetLeadingBound(bForwards, NextChildNode.Range));
 		}
 	}
 
-	CurrentRange = TRange<float>(bForwards ? NewLeadingBound : NewTrailingBound, bForwards ? NewTrailingBound : NewLeadingBound);
+	CurrentRange = TRange<FFrameNumber>(bForwards ? NewLeadingBound : NewTrailingBound, bForwards ? NewTrailingBound : NewLeadingBound);
 }
 
-FMovieSceneEvaluationTreeNodeHandle FMovieSceneEvaluationTreeRangeIterator::FindNextChild(FMovieSceneEvaluationTreeNodeHandle ParentNodeHandle, TRangeBound<float> PredicateBound, bool bForwards)
+FMovieSceneEvaluationTreeNodeHandle FMovieSceneEvaluationTreeRangeIterator::FindNextChild(FMovieSceneEvaluationTreeNodeHandle ParentNodeHandle, TRangeBound<FFrameNumber> PredicateBound, bool bForwards)
 {
 	auto GetLowerBound = [](const FMovieSceneEvaluationTreeNode& In) {
 		return In.Range.GetLowerBound();
@@ -166,17 +166,17 @@ TArrayView<FMovieSceneEvaluationTreeNode> FMovieSceneEvaluationTree::GetChildren
 	return InNode.ChildrenID.IsValid() ? ChildNodes.Get(InNode.ChildrenID) : TArrayView<FMovieSceneEvaluationTreeNode>();
 }
 
-FMovieSceneEvaluationTreeRangeIterator FMovieSceneEvaluationTree::IterateFromTime(float InTime) const
+FMovieSceneEvaluationTreeRangeIterator FMovieSceneEvaluationTree::IterateFromTime(FFrameNumber InTime) const
 {
-	return IterateFromLowerBound(TRangeBound<float>::Inclusive(InTime));
+	return IterateFromLowerBound(TRangeBound<FFrameNumber>::Inclusive(InTime));
 }
 
-FMovieSceneEvaluationTreeRangeIterator FMovieSceneEvaluationTree::IterateFromLowerBound(TRangeBound<float> InStartingLowerBound) const
+FMovieSceneEvaluationTreeRangeIterator FMovieSceneEvaluationTree::IterateFromLowerBound(TRangeBound<FFrameNumber> InStartingLowerBound) const
 {
 	return FMovieSceneEvaluationTreeRangeIterator(*this, InStartingLowerBound);
 }
 
-void FMovieSceneEvaluationTree::InsertNewChild(TRange<float> InEffectiveRange, const IMovieSceneEvaluationTreeNodeOperator& InOperator, int32 InsertIndex, FMovieSceneEvaluationTreeNodeHandle InParent)
+void FMovieSceneEvaluationTree::InsertNewChild(TRange<FFrameNumber> InEffectiveRange, const IMovieSceneEvaluationTreeNodeOperator& InOperator, int32 InsertIndex, FMovieSceneEvaluationTreeNodeHandle InParent)
 {
 	// Check if this parent node has any 
 	FEvaluationTreeEntryHandle ChildID = GetNode(InParent).ChildrenID;
@@ -205,7 +205,7 @@ void FMovieSceneEvaluationTree::InsertNewChild(TRange<float> InEffectiveRange, c
 	ChildNodes.Insert(ChildID, InsertIndex, MoveTemp(NewChildNode));
 }
 
-void FMovieSceneEvaluationTree::AddTimeRange(TRange<float> InTimeRange)
+void FMovieSceneEvaluationTree::AddTimeRange(TRange<FFrameNumber> InTimeRange)
 {
 	struct FNullOp : IMovieSceneEvaluationTreeNodeOperator
 	{
@@ -215,7 +215,7 @@ void FMovieSceneEvaluationTree::AddTimeRange(TRange<float> InTimeRange)
 	AddTimeRange(InTimeRange, FNullOp(), FMovieSceneEvaluationTreeNodeHandle::Root());
 }
 
-void FMovieSceneEvaluationTree::AddTimeRange(TRange<float> InTimeRange, const IMovieSceneEvaluationTreeNodeOperator& InOperator, FMovieSceneEvaluationTreeNodeHandle InParent)
+void FMovieSceneEvaluationTree::AddTimeRange(TRange<FFrameNumber> InTimeRange, const IMovieSceneEvaluationTreeNodeOperator& InOperator, FMovieSceneEvaluationTreeNodeHandle InParent)
 {
 	// Take a temporary copy of the node as the container may be reallocated in this function
 	FMovieSceneEvaluationTreeNode ThisNode = GetNode(InParent);
@@ -233,13 +233,13 @@ void FMovieSceneEvaluationTree::AddTimeRange(TRange<float> InTimeRange, const IM
 	else if (!ThisNode.ChildrenID.IsValid())
 	{
 		// We have no child nodes yet, so just add the intersection as a child at the start
-		TRange<float> Intersection = TRange<float>::Intersection(ThisNode.Range, InTimeRange);
+		TRange<FFrameNumber> Intersection = TRange<FFrameNumber>::Intersection(ThisNode.Range, InTimeRange);
 		InsertNewChild(Intersection, InOperator, 0, InParent);
 	}
 	else
 	{
 		// It must be a sub-division of this range, find the intersection and add new children as necessary
-		TRangeBound<float> LastBound = ThisNode.Range.GetLowerBound();
+		TRangeBound<FFrameNumber> LastBound = ThisNode.Range.GetLowerBound();
 
 		// Iterate all existing children, adding new children for gaps that the new section intersects
 		int32 ChildIndex = 0;
@@ -249,8 +249,8 @@ void FMovieSceneEvaluationTree::AddTimeRange(TRange<float> InTimeRange, const IM
 
 			if (!ChildNode.Range.GetLowerBound().IsOpen())
 			{
-				TRange<float> PrecedingSpace(LastBound, TRangeBound<float>::FlipInclusion(ChildNode.Range.GetLowerBound()));
-				TRange<float> PrecedingIntersection = TRange<float>::Intersection(PrecedingSpace, InTimeRange);
+				TRange<FFrameNumber> PrecedingSpace(LastBound, TRangeBound<FFrameNumber>::FlipInclusion(ChildNode.Range.GetLowerBound()));
+				TRange<FFrameNumber> PrecedingIntersection = TRange<FFrameNumber>::Intersection(PrecedingSpace, InTimeRange);
 
 				if (!PrecedingIntersection.IsEmpty())
 				{
@@ -267,15 +267,15 @@ void FMovieSceneEvaluationTree::AddTimeRange(TRange<float> InTimeRange, const IM
 				// Find the node again since InsertNewChild may have re-allocated the nodes
 				AddTimeRange(InTimeRange, InOperator, FMovieSceneEvaluationTreeNodeHandle(ThisNode.ChildrenID, ChildIndex));
 			}
-			LastBound = TRangeBound<float>::FlipInclusion(ChildNode.Range.GetUpperBound());
+			LastBound = TRangeBound<FFrameNumber>::FlipInclusion(ChildNode.Range.GetUpperBound());
 			ChildIndex++;
 		}
 
 		// If the last child didn't fully fill this parent node range, add a new one for this new section
 		if (!LastBound.IsOpen())
 		{
-			TRange<float> ProceedingSpace(LastBound, ThisNode.Range.GetUpperBound());
-			TRange<float> ProceedingIntersection = TRange<float>::Intersection(ProceedingSpace, InTimeRange);
+			TRange<FFrameNumber> ProceedingSpace(LastBound, ThisNode.Range.GetUpperBound());
+			TRange<FFrameNumber> ProceedingIntersection = TRange<FFrameNumber>::Intersection(ProceedingSpace, InTimeRange);
 
 			if (!ProceedingIntersection.IsEmpty())
 			{
diff --git a/Engine/Source/Runtime/MovieScene/Private/Evaluation/MovieScenePlayback.cpp b/Engine/Source/Runtime/MovieScene/Private/Evaluation/MovieScenePlayback.cpp
--- a/Engine/Source/Runtime/MovieScene/Private/Evaluation/MovieScenePlayback.cpp
+++ b/Engine/Source/Runtime/MovieScene/Private/Evaluation/MovieScenePlayback.cpp
@@ -6,84 +6,195 @@
 
 namespace 
 {
-	TRange<float> CalculateEvaluationRange(float CurrentTime, float PreviousTime, bool bInclusivePreviousTime)
+	TRange<FFrameTime> CalculateEvaluationRange(FFrameTime CurrentTime, FFrameTime PreviousTime, bool bInclusivePreviousTime)
 	{
 		if (CurrentTime == PreviousTime)
 		{
-			return TRange<float>(CurrentTime);
+			return TRange<FFrameTime>(CurrentTime);
 		}
 		else if (CurrentTime < PreviousTime)
 		{
-			return TRange<float>(
-				TRangeBound<float>::Inclusive(CurrentTime),
+			return TRange<FFrameTime>(
+				TRangeBound<FFrameTime>::Inclusive(CurrentTime),
 				bInclusivePreviousTime
-					? TRangeBound<float>::Inclusive(PreviousTime)
-					: TRangeBound<float>::Exclusive(PreviousTime)
+					? TRangeBound<FFrameTime>::Inclusive(PreviousTime)
+					: TRangeBound<FFrameTime>::Exclusive(PreviousTime)
 				);
 		}
 
-		return TRange<float>(
+		return TRange<FFrameTime>(
 			bInclusivePreviousTime
-				? TRangeBound<float>::Inclusive(PreviousTime)
-				: TRangeBound<float>::Exclusive(PreviousTime)
-			, TRangeBound<float>::Inclusive(CurrentTime)
+				? TRangeBound<FFrameTime>::Inclusive(PreviousTime)
+				: TRangeBound<FFrameTime>::Exclusive(PreviousTime)
+			, TRangeBound<FFrameTime>::Inclusive(CurrentTime)
 			);
 	}
 }
 
-FMovieSceneEvaluationRange::FMovieSceneEvaluationRange(float InTime)
+FMovieSceneEvaluationRange::FMovieSceneEvaluationRange(FFrameTime InTime, FFrameRate InFrameRate)
 	: EvaluationRange(InTime)
+	, CurrentFrameRate(InFrameRate)
 	, Direction(EPlayDirection::Forwards)
-	, TimeOverride(TNumericLimits<float>::Lowest())
+	, TimeOverride(FFrameNumber(TNumericLimits<int32>::Lowest()))
 {
 }
 
-FMovieSceneEvaluationRange::FMovieSceneEvaluationRange(TRange<float> InRange, EPlayDirection InDirection)
+FMovieSceneEvaluationRange::FMovieSceneEvaluationRange(TRange<FFrameTime> InRange, FFrameRate InFrameRate, EPlayDirection InDirection)
 	: EvaluationRange(InRange)
+	, CurrentFrameRate(InFrameRate)
 	, Direction(InDirection)
-	, TimeOverride(TNumericLimits<float>::Lowest())
+	, TimeOverride(FFrameNumber(TNumericLimits<int32>::Lowest()))
 {
 }
 
-FMovieSceneEvaluationRange::FMovieSceneEvaluationRange(float InCurrentTime, float InPreviousTime, bool bInclusivePreviousTime)
+FMovieSceneEvaluationRange::FMovieSceneEvaluationRange(FFrameTime InCurrentTime, FFrameTime InPreviousTime, FFrameRate InFrameRate, bool bInclusivePreviousTime)
 	: EvaluationRange(CalculateEvaluationRange(InCurrentTime, InPreviousTime, bInclusivePreviousTime))
-	, Direction((InCurrentTime - InPreviousTime >= 0) ? EPlayDirection::Forwards : EPlayDirection::Backwards)
-	, TimeOverride(TNumericLimits<float>::Lowest())
+	, CurrentFrameRate(InFrameRate)
+	, Direction((InCurrentTime - InPreviousTime >= FFrameTime()) ? EPlayDirection::Forwards : EPlayDirection::Backwards)
+	, TimeOverride(TNumericLimits<int32>::Lowest())
 {
 }
 
-void FMovieScenePlaybackPosition::Reset(float StartPos)
+TRange<FFrameNumber> FMovieSceneEvaluationRange::TimeRangeToNumberRange(const TRange<FFrameTime>& InFrameTimeRange)
 {
-	PreviousPosition = StartPos;
+	TRange<FFrameNumber> FrameNumberRange;
+
+	if (!InFrameTimeRange.GetLowerBound().IsOpen())
+	{
+		FFrameTime LowerTime = InFrameTimeRange.GetLowerBoundValue();
+		// If there is a sub frame on the start time, we're actually beyond that frame number, so it needs incrementing
+		if (LowerTime.GetSubFrame() != 0.f || InFrameTimeRange.GetLowerBound().IsExclusive())
+		{
+			LowerTime.FrameNumber = LowerTime.FrameNumber + 1;
+		}
+		FrameNumberRange.SetLowerBound(TRangeBound<FFrameNumber>::Inclusive(LowerTime.FrameNumber));
+	}
+
+	if (!InFrameTimeRange.GetUpperBound().IsOpen())
+	{
+		const FFrameTime UpperTime = InFrameTimeRange.GetUpperBoundValue();
+
+		FrameNumberRange.SetUpperBound(
+			InFrameTimeRange.GetUpperBound().IsExclusive()
+			? TRangeBound<FFrameNumber>::Exclusive(UpperTime.FrameNumber)
+			: TRangeBound<FFrameNumber>::Inclusive(UpperTime.FrameNumber)
+		);
+	}
+
+	return FrameNumberRange;
+}
+
+TRange<FFrameTime> FMovieSceneEvaluationRange::NumberRangeToTimeRange(const TRange<FFrameNumber>& InFrameNumberRange)
+{
+	TRange<FFrameTime> FrameTimeRange;
+
+	if (!InFrameNumberRange.GetLowerBound().IsOpen())
+	{
+		const FFrameNumber FrameNumber = InFrameNumberRange.GetLowerBoundValue();
+
+		FrameTimeRange.SetLowerBound(
+			InFrameNumberRange.GetLowerBound().IsExclusive()
+			? TRangeBound<FFrameTime>::Exclusive(FrameNumber)
+			: TRangeBound<FFrameTime>::Inclusive(FrameNumber)
+		);
+	}
+
+	if (!InFrameNumberRange.GetUpperBound().IsOpen())
+	{
+		const FFrameNumber FrameNumber = InFrameNumberRange.GetUpperBoundValue();
+
+		FrameTimeRange.SetUpperBound(
+			InFrameNumberRange.GetUpperBound().IsExclusive()
+			? TRangeBound<FFrameTime>::Exclusive(FrameNumber)
+			: TRangeBound<FFrameTime>::Inclusive(FrameNumber)
+		);
+	}
+
+	return FrameTimeRange;
+}
+
+void FMovieScenePlaybackPosition::CheckInvariants() const
+{
+	checkf(InputRate.IsValid() && OutputRate.IsValid(), TEXT("Invalid input or output rate. SetTimeBase must be called before any use of this class."))
+}
+
+void FMovieScenePlaybackPosition::SetTimeBase(FFrameRate NewInputRate, FFrameRate NewOutputRate, EMovieSceneEvaluationType NewEvaluationType)
+{
+	// Move the current position if necessary
+	if (InputRate.IsValid() && InputRate != NewInputRate)
+	{
+		FFrameTime NewPosition = ConvertFrameTime(CurrentPosition, InputRate, NewInputRate);
+		if (NewEvaluationType == EMovieSceneEvaluationType::FrameLocked)
+		{
+			NewPosition = NewPosition.FloorToFrame();
+		}
+
+		Reset(NewPosition);
+	}
+
+	InputRate      = NewInputRate;
+	OutputRate     = NewOutputRate;
+	EvaluationType = NewEvaluationType;
+}
+
+void FMovieScenePlaybackPosition::Reset(FFrameTime StartPos)
+{
+	CurrentPosition = StartPos;
 	PreviousPlayEvalPosition.Reset();
 	LastRange.Reset();
 }
 
-FMovieSceneEvaluationRange FMovieScenePlaybackPosition::JumpTo(float NewPosition, TOptional<float> FixedInterval)
+FMovieSceneEvaluationRange FMovieScenePlaybackPosition::GetCurrentPositionAsRange() const
 {
-	Reset(NewPosition);
+	CheckInvariants();
 
-	PreviousPosition = NewPosition;
-	float EvalPosition = FixedInterval.IsSet() ? UMovieScene::CalculateFixedFrameTime(NewPosition, FixedInterval.GetValue()) : NewPosition;
+	FFrameTime OutputPosition = ConvertFrameTime(CurrentPosition, InputRate, OutputRate);
+	return FMovieSceneEvaluationRange(OutputPosition, OutputRate);
+}
+
+FMovieSceneEvaluationRange FMovieScenePlaybackPosition::JumpTo(FFrameTime InputPosition)
+{
+	CheckInvariants();
+
+	PreviousPlayEvalPosition.Reset();
+
+	// Floor to the current frame number if running frame-locked
+	if (EvaluationType == EMovieSceneEvaluationType::FrameLocked)
+	{
+		InputPosition = InputPosition.FloorToFrame();
+	}
+
+	// Assign the cached input values
+	CurrentPosition = InputPosition;
 
-	FMovieSceneEvaluationRange Range(TRange<float>(EvalPosition), EPlayDirection::Forwards);
-	LastRange = Range;
-	return Range;
+	// Convert to output time-base
+	FFrameTime OutputPosition = ConvertFrameTime(InputPosition, InputRate, OutputRate);
+
+	LastRange = FMovieSceneEvaluationRange(TRange<FFrameTime>(OutputPosition), OutputRate, EPlayDirection::Forwards);
+	return LastRange.GetValue();
 }
 
-FMovieSceneEvaluationRange FMovieScenePlaybackPosition::PlayTo(float NewPosition, TOptional<float> FixedInterval)
+FMovieSceneEvaluationRange FMovieScenePlaybackPosition::PlayTo(FFrameTime InputPosition)
 {
-	float EvalPositionTo = FixedInterval.IsSet() ? UMovieScene::CalculateFixedFrameTime(NewPosition, FixedInterval.GetValue()) : NewPosition;
-	float EvalPositionFrom = PreviousPlayEvalPosition.Get(PreviousPosition.Get(NewPosition));
+	CheckInvariants();
 
-	FMovieSceneEvaluationRange Range(EvalPositionTo, EvalPositionFrom, !PreviousPlayEvalPosition.IsSet());
+	// Floor to the current frame number if running frame-locked
+	if (EvaluationType == EMovieSceneEvaluationType::FrameLocked)
+	{
+		InputPosition = InputPosition.FloorToFrame();
+	}
 
-	PreviousPosition = NewPosition;
-	PreviousPlayEvalPosition = EvalPositionTo;
+	// Convert to output time-base
+	FFrameTime InputEvalPositionFrom  = PreviousPlayEvalPosition.Get(CurrentPosition);
+	FFrameTime OutputEvalPositionFrom = ConvertFrameTime(InputEvalPositionFrom, InputRate, OutputRate);
+	FFrameTime OutputEvalPositionTo   = ConvertFrameTime(InputPosition, InputRate, OutputRate);
 
-	LastRange = Range;
+	// Assign the cached input values
+	CurrentPosition          = InputPosition;
+	PreviousPlayEvalPosition = InputPosition;
 
-	return Range;
+	LastRange = FMovieSceneEvaluationRange(OutputEvalPositionTo, OutputEvalPositionFrom, OutputRate, !PreviousPlayEvalPosition.IsSet());
+	return LastRange.GetValue();
 }
 
 TOptional<FMovieSceneEvaluationRange> FMovieScenePlaybackPosition::GetLastRange() const
diff --git a/Engine/Source/Runtime/MovieScene/Private/Evaluation/MovieSceneSequenceHierarchy.cpp b/Engine/Source/Runtime/MovieScene/Private/Evaluation/MovieSceneSequenceHierarchy.cpp
--- a/Engine/Source/Runtime/MovieScene/Private/Evaluation/MovieSceneSequenceHierarchy.cpp
+++ b/Engine/Source/Runtime/MovieScene/Private/Evaluation/MovieSceneSequenceHierarchy.cpp
@@ -4,6 +4,7 @@
 #include "Sections/MovieSceneSubSection.h"
 #include "MovieSceneSequence.h"
 #include "MovieScene.h"
+#include "MovieSceneTimeHelpers.h"
 
 FMovieSceneSubSequenceData::FMovieSceneSubSequenceData()
 	: Sequence(nullptr)
@@ -12,35 +13,39 @@ FMovieSceneSubSequenceData::FMovieSceneSubSequenceData()
 FMovieSceneSubSequenceData::FMovieSceneSubSequenceData(const UMovieSceneSubSection& InSubSection)
 	: Sequence(InSubSection.GetSequence())
 	, DeterministicSequenceID(InSubSection.GetSequenceID())
-	, PreRollRange(TRange<float>::Empty())
-	, PostRollRange(TRange<float>::Empty())
 	, HierarchicalBias(InSubSection.Parameters.HierarchicalBias)
 #if WITH_EDITORONLY_DATA
 	, SectionPath(*InSubSection.GetPathNameInMovieScene())
 #endif
 {
-	UMovieSceneSequence* SequencePtr = GetSequence();
+	PreRollRange.Value  = TRange<FFrameNumber>::Empty();
+	PostRollRange.Value = TRange<FFrameNumber>::Empty();
 
-	checkf(SequencePtr, TEXT("Attempting to construct sub sequence data with a null sequence."));
+	UMovieSceneSequence* SequencePtr   = GetSequence();
+	UMovieScene*         MovieScenePtr = SequencePtr ? SequencePtr->GetMovieScene() : nullptr;
 
-	const float StartTime = InSubSection.GetStartTime();
-	const float EndTime = InSubSection.GetEndTime();
+	checkf(MovieScenePtr, TEXT("Attempting to construct sub sequence data with a null sequence."));
 
-	RootToSequenceTransform = 
-		FMovieSceneSequenceTransform(SequencePtr->GetMovieScene()->GetPlaybackRange().GetLowerBoundValue() + InSubSection.Parameters.StartOffset) *		// Inner play offset
-		FMovieSceneSequenceTransform(0.f, InSubSection.Parameters.TimeScale) *		// Inner play rate
-		FMovieSceneSequenceTransform(-StartTime);					// Outer section start time
+	FrameResolution = MovieScenePtr->GetFrameResolution();
 
-	PlayRange = InSubSection.GetRange() * RootToSequenceTransform;
+	TRange<FFrameNumber> SubRange = InSubSection.GetRange();
+	checkf(SubRange.GetLowerBound().IsClosed() && SubRange.GetUpperBound().IsClosed(), TEXT("Use of open (infinite) bounds with sub sections is not supported."));
 
-	// Make sure pre/postroll ranges are in the inner sequence's time space
-	if (InSubSection.GetPreRollTime() > 0)
+	const FFrameNumber OuterStartTime = MovieScene::DiscreteInclusiveLower(SubRange);
+	const FFrameNumber OuterEndTime   = MovieScene::DiscreteExclusiveUpper(SubRange);
+
+	RootToSequenceTransform = InSubSection.OuterToInnerTransform();
+
+	PlayRange.Value = SubRange * RootToSequenceTransform;
+
+	// Make sure pre/postroll *ranges* are in the inner sequence's time space. Pre/PostRollFrames are in the outer sequence space.
+	if (InSubSection.GetPreRollFrames() > 0)
 	{
-		PreRollRange = TRange<float>(StartTime - InSubSection.GetPreRollTime(), TRangeBound<float>::Exclusive(StartTime)) * RootToSequenceTransform;
+		PreRollRange = MovieScene::MakeDiscreteRangeFromUpper( TRangeBound<FFrameNumber>::FlipInclusion(SubRange.GetLowerBound()), InSubSection.GetPreRollFrames() ) * RootToSequenceTransform;
 	}
-	if (InSubSection.GetPostRollTime() > 0)
+	if (InSubSection.GetPostRollFrames() > 0)
 	{
-		PostRollRange = TRange<float>(TRangeBound<float>::Exclusive(EndTime), EndTime + InSubSection.GetPostRollTime()) * RootToSequenceTransform;
+		PostRollRange = MovieScene::MakeDiscreteRangeFromLower( TRangeBound<FFrameNumber>::FlipInclusion(SubRange.GetUpperBound()), InSubSection.GetPostRollFrames() ) * RootToSequenceTransform;
 	}
 }
 
diff --git a/Engine/Source/Runtime/MovieScene/Private/Evaluation/MovieSceneTrackImplementation.cpp b/Engine/Source/Runtime/MovieScene/Private/Evaluation/MovieSceneTrackImplementation.cpp
--- a/Engine/Source/Runtime/MovieScene/Private/Evaluation/MovieSceneTrackImplementation.cpp
+++ b/Engine/Source/Runtime/MovieScene/Private/Evaluation/MovieSceneTrackImplementation.cpp
@@ -6,5 +6,6 @@
 
 bool FMovieSceneTrackImplementationPtr::Serialize(FArchive& Ar)
 {
-	return SerializeInlineValue(*this, Ar);
+	bool bShouldWarn = !WITH_EDITORONLY_DATA;
+	return SerializeInlineValue(*this, Ar, bShouldWarn);
 }
diff --git a/Engine/Source/Runtime/MovieScene/Private/MovieScene.cpp b/Engine/Source/Runtime/MovieScene/Private/MovieScene.cpp
--- a/Engine/Source/Runtime/MovieScene/Private/MovieScene.cpp
+++ b/Engine/Source/Runtime/MovieScene/Private/MovieScene.cpp
@@ -6,26 +6,52 @@
 #include "Evaluation/MovieSceneEvaluationCustomVersion.h"
 #include "Compilation/MovieSceneSegmentCompiler.h"
 #include "UObject/SequencerObjectVersion.h"
+#include "CommonFrameRates.h"
+
+
+TOptional<TRangeBound<FFrameNumber>> GetMaxUpperBound(const UMovieSceneTrack* Track)
+{
+	TOptional<TRangeBound<FFrameNumber>> MaxBound;
+
+	// Find the largest closed upper bound of all the track's sections
+	for (UMovieSceneSection* Section : Track->GetAllSections())
+	{
+		TRangeBound<FFrameNumber> SectionUpper = Section->GetRange().GetUpperBound();
+		if (SectionUpper.IsClosed())
+		{
+			MaxBound = TRangeBound<FFrameNumber>::MaxUpper(MaxBound.Get(SectionUpper), SectionUpper);
+		}
+	}
+
+	return MaxBound;
+}
 
 /* UMovieScene interface
  *****************************************************************************/
 
 UMovieScene::UMovieScene(const FObjectInitializer& ObjectInitializer)
 	: Super(ObjectInitializer)
-	, SelectionRange(FFloatRange::Empty())
-	, PlaybackRange(FFloatRange::Empty())
-#if WITH_EDITORONLY_DATA
-	, bPlaybackRangeLocked(false)
-#endif
-	, bForceFixedFrameIntervalPlayback(false)
-	, FixedFrameInterval(0.0f)
-	, InTime_DEPRECATED(FLT_MAX)
-	, OutTime_DEPRECATED(-FLT_MAX)
-	, StartTime_DEPRECATED(FLT_MAX)
-	, EndTime_DEPRECATED(-FLT_MAX)
 {
+	EvaluationType = EMovieSceneEvaluationType::WithSubFrames;
+	ClockSource = EUpdateClockSource::Tick;
+
+	if (!HasAnyFlags(RF_ClassDefaultObject) && GetLinkerCustomVersion(FSequencerObjectVersion::GUID) < FSequencerObjectVersion::FloatToIntConversion)
+	{
+		FrameResolution = GetLegacyConversionFrameRate();
+	}
+
 #if WITH_EDITORONLY_DATA
-	EditorData.WorkingRange = EditorData.ViewRange = TRange<float>::Empty();
+	bPlaybackRangeLocked = false;
+	PlaybackRange.MigrationDefault = FFloatRange::Empty();
+	EditorData.WorkingRange_DEPRECATED = EditorData.ViewRange_DEPRECATED = TRange<float>::Empty();
+
+	bForceFixedFrameIntervalPlayback_DEPRECATED = false;
+	FixedFrameInterval_DEPRECATED = 0.0f;
+
+	InTime_DEPRECATED    =  FLT_MAX;
+	OutTime_DEPRECATED   = -FLT_MAX;
+	StartTime_DEPRECATED =  FLT_MAX;
+	EndTime_DEPRECATED   = -FLT_MAX;
 #endif
 }
 
@@ -49,6 +75,54 @@ void UMovieScene::Serialize( FArchive& Ar )
 #endif // WITH_EDITOR
 
 	Super::Serialize(Ar);
+
+#if WITH_EDITORONLY_DATA
+	if (Ar.CustomVer(FSequencerObjectVersion::GUID) < FSequencerObjectVersion::FloatToIntConversion)
+	{
+		if (bForceFixedFrameIntervalPlayback_DEPRECATED)
+		{
+			EvaluationType = EMovieSceneEvaluationType::FrameLocked;
+		}
+
+		// Legacy fixed frame interval conversion to integer play rates
+		if      (FixedFrameInterval_DEPRECATED == 1 / 15.0f)  { PlayRate = FCommonFrameRates::FPS_15();  }
+		else if (FixedFrameInterval_DEPRECATED == 1 / 24.0f)  { PlayRate = FCommonFrameRates::FPS_24();  }
+		else if (FixedFrameInterval_DEPRECATED == 1 / 25.0f)  { PlayRate = FCommonFrameRates::FPS_25();  }
+		else if (FixedFrameInterval_DEPRECATED == 1 / 29.97)  { PlayRate = FCommonFrameRates::NTSC_30(); }
+		else if (FixedFrameInterval_DEPRECATED == 1 / 30.0f)  { PlayRate = FCommonFrameRates::FPS_30();  }
+		else if (FixedFrameInterval_DEPRECATED == 1 / 48.0f)  { PlayRate = FCommonFrameRates::FPS_48();  }
+		else if (FixedFrameInterval_DEPRECATED == 1 / 50.0f)  { PlayRate = FCommonFrameRates::FPS_50();  }
+		else if (FixedFrameInterval_DEPRECATED == 1 / 59.94)  { PlayRate = FCommonFrameRates::NTSC_60(); }
+		else if (FixedFrameInterval_DEPRECATED == 1 / 60.0f)  { PlayRate = FCommonFrameRates::FPS_60();  }
+		else if (FixedFrameInterval_DEPRECATED == 1 / 120.0f) { PlayRate = FCommonFrameRates::FPS_120(); }
+		else if (FixedFrameInterval_DEPRECATED != 0.f)
+		{
+			uint32 Numerator = FMath::RoundToInt(1.f / FixedFrameInterval_DEPRECATED);
+			PlayRate = FFrameRate(Numerator, 1);
+		}
+		else
+		{
+			// Sequences with 0 FixedFrameInterval used to be assigned a proper interval in SSequencer::OnSequenceInstanceActivated for some reason,
+			// But we don't have access to the relevant sequencer settings class here so we just have to make a hacky educated guess based on the class name.
+			UObject* Outer = GetOuter();
+			if (Outer && Outer->GetClass()->GetFName() == "WidgetAnimation")
+			{
+				// Widget animations defaulted to 0.05s
+				PlayRate = FFrameRate(20, 1);
+			}
+			else if (Outer && Outer->GetClass()->GetFName() == "ActorSequence")
+			{
+				// Actor sequences defaulted to 0.1s
+				PlayRate = FFrameRate(10, 1);
+			}
+			else
+			{
+				// Level sequences defaulted to 30fps - this is the fallback default for anything else
+				PlayRate = FFrameRate(30, 1);
+			}
+		}
+	}
+#endif
 }
 
 #if WITH_EDITOR
@@ -349,56 +423,66 @@ TArray<UMovieSceneFolder*>&  UMovieScene::GetRootFolders()
 }
 #endif
 
+void UMovieScene::SetPlaybackRange(FFrameNumber Start, int32 Duration, bool bAlwaysMarkDirty)
+{
+	// Inclusive lower, Exclusive upper bound
+	SetPlaybackRange(TRange<FFrameNumber>(Start, Start + Duration), bAlwaysMarkDirty);
+}
 
-void UMovieScene::SetPlaybackRange(float Start, float End, bool bAlwaysMarkDirty)
+void UMovieScene::SetPlaybackRange(const TRange<FFrameNumber>& NewRange, bool bAlwaysMarkDirty)
 {
-	if (ensure(End >= Start))
-	{
-		const auto NewRange = TRange<float>(Start, TRangeBound<float>::Inclusive(End));
+	check(NewRange.GetLowerBound().IsClosed() && NewRange.GetUpperBound().IsClosed());
 
-		if (PlaybackRange == NewRange)
-		{
-			return;
-		}
+	if (PlaybackRange.Value == NewRange)
+	{
+		return;
+	}
 
-		if (bAlwaysMarkDirty)
-		{
-			Modify();
-		}
+	if (bAlwaysMarkDirty)
+	{
+		Modify();
+	}
 
-		PlaybackRange = NewRange;
+	PlaybackRange.Value = NewRange;
 
 #if WITH_EDITORONLY_DATA
-		// Initialize the working and view range with a little bit more space
-		const float OutputViewSize = PlaybackRange.GetUpperBoundValue() - PlaybackRange.GetLowerBoundValue();
-		const float OutputChange = OutputViewSize * 0.1f;
+	// Update the working and view ranges to encompass the new range
+	const double RangeStartSeconds = NewRange.GetLowerBoundValue() / FrameResolution;
+	const double RangeEndSeconds   = NewRange.GetUpperBoundValue() / FrameResolution;
 
-		TRange<float> ExpandedPlaybackRange = TRange<float>(PlaybackRange.GetLowerBoundValue() - OutputChange, PlaybackRange.GetUpperBoundValue() + OutputChange);
+	// Initialize the working and view range with a little bit more space
+	const double OutputChange      = (RangeEndSeconds - RangeStartSeconds) * 0.1;
 
-		if (EditorData.WorkingRange.IsEmpty())
-		{
-			EditorData.WorkingRange = ExpandedPlaybackRange;
-		}
+	double ExpandedStart = RangeStartSeconds - OutputChange;
+	double ExpandedEnd   = RangeEndSeconds + OutputChange;
 
-		if (EditorData.ViewRange.IsEmpty())
-		{
-			EditorData.ViewRange = ExpandedPlaybackRange;
-		}
-#endif
+	if (EditorData.WorkStart >= EditorData.WorkEnd)
+	{
+		EditorData.WorkStart = ExpandedStart;
+		EditorData.WorkEnd   = ExpandedEnd;
 	}
+
+	if (EditorData.ViewStart >= EditorData.ViewEnd)
+	{
+		EditorData.ViewStart = ExpandedStart;
+		EditorData.ViewEnd   = ExpandedEnd;
+	}
+#endif
 }
 
 void UMovieScene::SetWorkingRange(float Start, float End)
 {
 #if WITH_EDITORONLY_DATA
-	EditorData.WorkingRange = TRange<float>(Start, End);
+	EditorData.WorkStart = Start;
+	EditorData.WorkEnd   = End;
 #endif
 }
 
 void UMovieScene::SetViewRange(float Start, float End)
 {
 #if WITH_EDITORONLY_DATA
-	EditorData.ViewRange = TRange<float>(Start, End);
+	EditorData.ViewStart = Start;
+	EditorData.ViewEnd   = End;
 #endif
 }
 
@@ -414,31 +498,6 @@ void UMovieScene::SetPlaybackRangeLocked(bool bLocked)
 }
 #endif
 
-bool UMovieScene::GetForceFixedFrameIntervalPlayback() const
-{
-	return bForceFixedFrameIntervalPlayback;
-}
-
-
-void UMovieScene::SetForceFixedFrameIntervalPlayback( bool bInForceFixedFrameIntervalPlayback )
-{
-	bForceFixedFrameIntervalPlayback = bInForceFixedFrameIntervalPlayback;
-}
-
-
-void UMovieScene::SetFixedFrameInterval( float InFixedFrameInterval )
-{
-	FixedFrameInterval = InFixedFrameInterval;
-}
-
-
-const float UMovieScene::FixedFrameIntervalEpsilon = .0001f;
-
-float UMovieScene::CalculateFixedFrameTime( float Time, float FixedFrameInterval )
-{
-	return ( FMath::RoundToInt( Time / FixedFrameInterval ) ) * FixedFrameInterval + FixedFrameIntervalEpsilon;
-}
-
 
 TArray<UMovieSceneSection*> UMovieScene::GetAllSections() const
 {
@@ -476,7 +535,7 @@ UMovieSceneTrack* UMovieScene::FindTrack(TSubclassOf<UMovieSceneTrack> TrackClas
 
 		for (const auto& Track : Binding.GetTracks())
 		{
-			if (Track->GetClass() == TrackClass)
+			if (TrackClass.GetDefaultObject() == nullptr ||  Track->GetClass() == TrackClass)
 			{
 				if (TrackName == NAME_None || Track->GetTrackName() == TrackName)
 				{
@@ -676,63 +735,98 @@ void UMovieScene::UpgradeTimeRanges()
 	// Level sequences defaulted to having a fixed play range.
 	// We now expose the playback range more visibly, but we need to upgrade the old data.
 
+	bool bFiniteRangeDefined = false;
+
+#if WITH_EDITORONLY_DATA
 	if (InTime_DEPRECATED != FLT_MAX && OutTime_DEPRECATED != -FLT_MAX)
 	{
 		// Finite range already defined in old data
-		PlaybackRange = TRange<float>(InTime_DEPRECATED, TRangeBound<float>::Inclusive(OutTime_DEPRECATED));
+		PlaybackRange.Value = TRange<FFrameNumber>(
+			FrameResolution.AsFrameNumber(InTime_DEPRECATED),
+			// Prefer exclusive upper bounds for playback ranges so we stop at the next frame
+			++FrameResolution.AsFrameNumber(OutTime_DEPRECATED)
+			);
+		bFiniteRangeDefined = true;
 	}
-	else if (PlaybackRange.IsEmpty())
+#endif
+	if (!bFiniteRangeDefined && PlaybackRange.Value.IsEmpty())
 	{
 		// No range specified, so automatically calculate one by determining the maximum upper bound of the sequence
 		// In this instance (UMG), playback always started at 0
-		float MaxBound = 0.f;
+		TRangeBound<FFrameNumber> MaxFrame = TRangeBound<FFrameNumber>::Exclusive(0);
 
-		for (const auto& Track : MasterTracks)
+		for (const UMovieSceneTrack* Track : MasterTracks)
 		{
-			auto Range = Track->GetSectionBoundaries();
-			if (Range.HasUpperBound())
+			TOptional<TRangeBound<FFrameNumber>> MaxUpper = GetMaxUpperBound(Track);
+			if (MaxUpper.IsSet())
 			{
-				MaxBound = FMath::Max(MaxBound, Range.GetUpperBoundValue());
+				MaxFrame = TRangeBound<FFrameNumber>::MaxUpper(MaxFrame, MaxUpper.GetValue());
 			}
 		}
 
-		for (const auto& Binding : ObjectBindings)
+		for (const FMovieSceneBinding& Binding : ObjectBindings)
 		{
-			auto Range = Binding.GetTimeRange();
-			if (Range.HasUpperBound())
+			for (UMovieSceneTrack* Track : Binding.GetTracks())
 			{
-				MaxBound = FMath::Max(MaxBound, Range.GetUpperBoundValue());
+				TOptional<TRangeBound<FFrameNumber>> MaxUpper = GetMaxUpperBound(Track);
+				if (MaxUpper.IsSet())
+				{
+					MaxFrame = TRangeBound<FFrameNumber>::MaxUpper(MaxFrame, MaxUpper.GetValue());
+				}
 			}
 		}
 
-		PlaybackRange = TRange<float>(0.f, TRangeBound<float>::Inclusive(MaxBound));
+		// Playback ranges should always have exclusive upper bounds
+		if (MaxFrame.IsInclusive())
+		{
+			MaxFrame = TRangeBound<FFrameNumber>::Exclusive(MaxFrame.GetValue() + 1);
+		}
+
+		PlaybackRange.Value = TRange<FFrameNumber>(TRangeBound<FFrameNumber>::Inclusive(0), MaxFrame);
 	}
-	else if (PlaybackRange.GetUpperBound().IsExclusive())
+	else if (PlaybackRange.Value.GetUpperBound().IsInclusive())
 	{
 		// playback ranges are now always inclusive
-		PlaybackRange = TRange<float>(PlaybackRange.GetLowerBound(), TRangeBound<float>::Inclusive(PlaybackRange.GetUpperBoundValue()));
+		PlaybackRange.Value = TRange<FFrameNumber>(PlaybackRange.Value.GetLowerBound(), TRangeBound<FFrameNumber>::Exclusive(PlaybackRange.Value.GetUpperBoundValue() + 1));
 	}
 
 	// PlaybackRange must always be defined to a finite range
-	if (!PlaybackRange.HasLowerBound() || !PlaybackRange.HasUpperBound() || PlaybackRange.IsDegenerate())
+	if (!PlaybackRange.Value.HasLowerBound() || !PlaybackRange.Value.HasUpperBound() || PlaybackRange.Value.IsDegenerate())
 	{
-		PlaybackRange = TRange<float>(0.f, 0.f);
+		PlaybackRange.Value = TRange<FFrameNumber>(FFrameNumber(0), FFrameNumber(0));
 	}
 
 #if WITH_EDITORONLY_DATA
+	if (GetLinkerCustomVersion(FSequencerObjectVersion::GUID) < FSequencerObjectVersion::FloatToIntConversion)
+	{
+		EditorData.ViewStart    = EditorData.ViewRange_DEPRECATED.GetLowerBoundValue();
+		EditorData.ViewEnd      = EditorData.ViewRange_DEPRECATED.GetUpperBoundValue();
+
+		EditorData.WorkStart    = EditorData.WorkingRange_DEPRECATED.GetLowerBoundValue();
+		EditorData.WorkEnd      = EditorData.WorkingRange_DEPRECATED.GetUpperBoundValue();
+	}
+
 	// Legacy upgrade for working range
 	if (StartTime_DEPRECATED != FLT_MAX && EndTime_DEPRECATED != -FLT_MAX)
 	{
-		EditorData.WorkingRange = TRange<float>(StartTime_DEPRECATED, EndTime_DEPRECATED);
+		EditorData.WorkStart = StartTime_DEPRECATED;
+		EditorData.WorkEnd   = EndTime_DEPRECATED;
 	}
-	else if (EditorData.WorkingRange.IsEmpty())
+	else if (EditorData.WorkStart >= EditorData.WorkEnd)
+	{
+		EditorData.WorkStart = PlaybackRange.Value.GetLowerBoundValue() / FrameResolution;
+		EditorData.WorkEnd   = PlaybackRange.Value.GetUpperBoundValue() / FrameResolution;
+	}
+
+	if (EditorData.ViewStart >= EditorData.ViewEnd)
 	{
-		EditorData.WorkingRange = PlaybackRange;
+		EditorData.ViewStart = PlaybackRange.Value.GetLowerBoundValue() / FrameResolution;
+		EditorData.ViewEnd   = PlaybackRange.Value.GetUpperBoundValue() / FrameResolution;
 	}
 
-	if (EditorData.ViewRange.IsEmpty())
+	if (SelectionRange.Value.GetLowerBound().IsOpen() || SelectionRange.Value.GetUpperBound().IsOpen())
 	{
-		EditorData.ViewRange = PlaybackRange;
+		SelectionRange.Value = TRange<FFrameNumber>::Empty();
 	}
 #endif
 }
diff --git a/Engine/Source/Runtime/MovieScene/Private/MovieSceneBinding.cpp b/Engine/Source/Runtime/MovieScene/Private/MovieSceneBinding.cpp
--- a/Engine/Source/Runtime/MovieScene/Private/MovieSceneBinding.cpp
+++ b/Engine/Source/Runtime/MovieScene/Private/MovieSceneBinding.cpp
@@ -7,21 +7,6 @@
 /* FMovieSceneBinding interface
  *****************************************************************************/
 
-TRange<float> FMovieSceneBinding::GetTimeRange() const
-{
-	TArray<TRange<float>> Bounds;
-
-	for (int32 TypeIndex = 0; TypeIndex < Tracks.Num(); ++TypeIndex)
-	{
-		if (Tracks[TypeIndex] != nullptr)
-		{
-			Bounds.Add(Tracks[TypeIndex]->GetSectionBoundaries());
-		}
-	}
-
-	return TRange<float>::Hull(Bounds);
-}
-
 void FMovieSceneBinding::AddTrack(UMovieSceneTrack& NewTrack)
 {
 	Tracks.Add(&NewTrack);
diff --git a/Engine/Source/Runtime/MovieScene/Private/MovieSceneClipboard.cpp b/Engine/Source/Runtime/MovieScene/Private/MovieSceneClipboard.cpp
--- a/Engine/Source/Runtime/MovieScene/Private/MovieSceneClipboard.cpp
+++ b/Engine/Source/Runtime/MovieScene/Private/MovieSceneClipboard.cpp
@@ -32,12 +32,12 @@ FMovieSceneClipboardKey& FMovieSceneClipboardKey::operator=(const FMovieSceneCli
 	return *this;
 }
 
-float FMovieSceneClipboardKey::GetTime() const
+FFrameNumber FMovieSceneClipboardKey::GetTime() const
 {
 	return Time;
 }
 
-void FMovieSceneClipboardKey::SetTime(float InTime)
+void FMovieSceneClipboardKey::SetTime(FFrameNumber InTime)
 {
 	Time = InTime;
 }
@@ -149,7 +149,7 @@ FMovieSceneClipboardEnvironment& FMovieSceneClipboard::GetEnvironment()
 	return Environment;
 }
 
-FMovieSceneClipboard FMovieSceneClipboardBuilder::Commit(TOptional<float> CopyRelativeTo)
+FMovieSceneClipboard FMovieSceneClipboardBuilder::Commit(TOptional<FFrameNumber> CopyRelativeTo)
 {
 	FMovieSceneClipboard Clipboard;
 
diff --git a/Engine/Source/Runtime/MovieScene/Private/MovieSceneCommonHelpers.cpp b/Engine/Source/Runtime/MovieScene/Private/MovieSceneCommonHelpers.cpp
--- a/Engine/Source/Runtime/MovieScene/Private/MovieSceneCommonHelpers.cpp
+++ b/Engine/Source/Runtime/MovieScene/Private/MovieSceneCommonHelpers.cpp
@@ -6,134 +6,80 @@
 #include "Camera/CameraComponent.h"
 #include "KeyParams.h"
 #include "MovieSceneSection.h"
+#include "Algo/Sort.h"
 
 
-TArray<UMovieSceneSection*> MovieSceneHelpers::GetAllTraversedSections( const TArray<UMovieSceneSection*>& Sections, float CurrentTime, float PreviousTime )
-{
-	TArray<UMovieSceneSection*> TraversedSections;
-
-	bool bPlayingBackwards = CurrentTime - PreviousTime < 0.0f;
-	float MaxTime = bPlayingBackwards ? PreviousTime : CurrentTime;
-	float MinTime = bPlayingBackwards ? CurrentTime : PreviousTime;
 
-	TRange<float> TraversedRange(MinTime, TRangeBound<float>::Inclusive(MaxTime));
-
-	for (int32 SectionIndex = 0; SectionIndex < Sections.Num(); ++SectionIndex)
+UMovieSceneSection* MovieSceneHelpers::FindSectionAtTime( const TArray<UMovieSceneSection*>& Sections, FFrameNumber Time )
+{
+	for( int32 SectionIndex = 0; SectionIndex < Sections.Num(); ++SectionIndex )
 	{
 		UMovieSceneSection* Section = Sections[SectionIndex];
-		if ((Section->GetStartTime() == CurrentTime) || TraversedRange.Overlaps(TRange<float>(Section->GetRange())))
+
+		//@todo sequencer: There can be multiple sections overlapping in time. Returning instantly does not account for that.
+		if( Section->IsTimeWithinSection( Time ) && Section->IsActive() )
 		{
-			TraversedSections.Add(Section);
+			return Section;
 		}
 	}
 
-	return TraversedSections;
+	return nullptr;
 }
 
-TArray<UMovieSceneSection*> MovieSceneHelpers::GetTraversedSections( const TArray<UMovieSceneSection*>& Sections, float CurrentTime, float PreviousTime )
-{
-	TArray<UMovieSceneSection*> TraversedSections = GetAllTraversedSections(Sections, CurrentTime, PreviousTime);
 
-	// Remove any overlaps that are underneath another
-	for (int32 RemoveAt = 0; RemoveAt < TraversedSections.Num(); )
+UMovieSceneSection* MovieSceneHelpers::FindNearestSectionAtTime( const TArray<UMovieSceneSection*>& Sections, FFrameNumber Time )
+{
+	TArray<UMovieSceneSection*> OverlappingSections, NonOverlappingSections;
+	for (UMovieSceneSection* Section : Sections)
 	{
-		UMovieSceneSection* Section = TraversedSections[RemoveAt];
-		
-		const bool bShouldRemove = TraversedSections.ContainsByPredicate([=](UMovieSceneSection* OtherSection){
-			if (Section->GetRowIndex() == OtherSection->GetRowIndex() &&
-				Section->GetRange().Overlaps(OtherSection->GetRange()) &&
-				Section->GetOverlapPriority() < OtherSection->GetOverlapPriority())
-			{
-				return true;
-			}
-			return false;
-		});
-		
-		if (bShouldRemove)
+		if (Section->GetRange().Contains(Time))
 		{
-			TraversedSections.RemoveAt(RemoveAt, 1, false);
+			OverlappingSections.Add(Section);
 		}
 		else
 		{
-			++RemoveAt;
+			NonOverlappingSections.Add(Section);
 		}
 	}
 
-	return TraversedSections;
-}
+	if (OverlappingSections.Num())
+	{
+		Algo::Sort(OverlappingSections, SortOverlappingSections);
+		return OverlappingSections[0];
+	}
 
-UMovieSceneSection* MovieSceneHelpers::FindSectionAtTime( const TArray<UMovieSceneSection*>& Sections, float Time )
-{
-	for( int32 SectionIndex = 0; SectionIndex < Sections.Num(); ++SectionIndex )
+	if (NonOverlappingSections.Num())
 	{
-		UMovieSceneSection* Section = Sections[SectionIndex];
+		Algo::SortBy(NonOverlappingSections, [](const UMovieSceneSection* A) { return A->GetRange().GetUpperBound(); }, SortUpperBounds);
 
-		//@todo sequencer: There can be multiple sections overlapping in time. Returning instantly does not account for that.
-		if( Section->IsTimeWithinSection( Time ) && Section->IsActive() )
+		const int32 PreviousIndex = Algo::UpperBoundBy(NonOverlappingSections, TRangeBound<FFrameNumber>(Time), [](const UMovieSceneSection* A){ return A->GetRange().GetUpperBound(); }, SortUpperBounds)-1;
+		if (NonOverlappingSections.IsValidIndex(PreviousIndex))
 		{
-			return Section;
+			return NonOverlappingSections[PreviousIndex];
+		}
+		else
+		{
+			Algo::SortBy(NonOverlappingSections, [](const UMovieSceneSection* A) { return A->GetRange().GetLowerBound(); }, SortLowerBounds);
+			return NonOverlappingSections[0];
 		}
 	}
 
 	return nullptr;
 }
 
-
-UMovieSceneSection* MovieSceneHelpers::FindNearestSectionAtTime( const TArray<UMovieSceneSection*>& Sections, float Time )
+bool MovieSceneHelpers::SortOverlappingSections(const UMovieSceneSection* A, const UMovieSceneSection* B)
 {
-	// Only update the section if the position is within the time span of the section
-	// Or if there are no sections at the time, the left closest section to the time
-	// Or in the case that Time is before all sections, take the section with the earliest start time
-	UMovieSceneSection* ClosestSection = nullptr;
-	float ClosestSectionTime = 0.f;
-	UMovieSceneSection* EarliestSection = nullptr;
-	float EarliestSectionTime = 0.f;
-	for( int32 SectionIndex = 0; SectionIndex < Sections.Num(); ++SectionIndex )
-	{
-		UMovieSceneSection* Section = Sections[SectionIndex];
-		checkSlow(Section);
-
-		if (Section->IsActive())
-		{
-			//@todo sequencer: There can be multiple sections overlapping in time. Returning instantly does not account for that.
-			if( Section->IsTimeWithinSection( Time ) )
-			{
-				return Section;
-			}
-
-			float EndTime = Section->GetEndTime();
-			if (EndTime < Time)
-			{
-				float ClosestTime = Time - EndTime;
-				if (!ClosestSection || ClosestTime < ClosestSectionTime)
-				{
-					ClosestSection = Section;
-					ClosestSectionTime = ClosestTime;
-				}
-			}
-
-			float StartTime = Section->GetStartTime();
-			if (!EarliestSection || StartTime < EarliestSectionTime)
-			{
-				EarliestSection = Section;
-				EarliestSectionTime = StartTime;
-			}
-		}
-	}
-
-	// if we get here, we are off of any section
-	// if ClosestSection, then we take the closest to left of this time
-	// else, we take the EarliestSection
-	// if that's nullptr, then there are no sections
-	return ClosestSection ? ClosestSection : EarliestSection;
+	return A->GetRowIndex() == B->GetRowIndex()
+		? A->GetOverlapPriority() < B->GetOverlapPriority()
+		: A->GetRowIndex() < B->GetRowIndex();
 }
 
-
 void MovieSceneHelpers::SortConsecutiveSections(TArray<UMovieSceneSection*>& Sections)
 {
 	Sections.Sort([](const UMovieSceneSection& A, const UMovieSceneSection& B)
 		{
-			return A.GetStartTime() < B.GetStartTime();
+			TRangeBound<FFrameNumber> LowerBoundA = A.GetRange().GetLowerBound();
+			return TRangeBound<FFrameNumber>::MinLower(LowerBoundA, B.GetRange().GetLowerBound()) == LowerBoundA;
 		}
 	);
 }
@@ -143,13 +89,22 @@ void MovieSceneHelpers::FixupConsecutiveSections(TArray<UMovieSceneSection*>& Se
 	// Find the previous section and extend it to take the place of the section being deleted
 	int32 SectionIndex = INDEX_NONE;
 
+	TRange<FFrameNumber> SectionRange = Section.GetRange();
+
 	if (Sections.Find(&Section, SectionIndex))
 	{
 		int32 PrevSectionIndex = SectionIndex - 1;
 		if( Sections.IsValidIndex( PrevSectionIndex ) )
 		{
 			// Extend the previous section
-			Sections[PrevSectionIndex]->SetEndTime( bDelete ? Section.GetEndTime() : Section.GetStartTime() );
+			if (bDelete)
+			{
+				Sections[PrevSectionIndex]->SetEndFrame(SectionRange.GetUpperBound());
+			}
+			else
+			{
+				Sections[PrevSectionIndex]->SetEndFrame(TRangeBound<FFrameNumber>::FlipInclusion(SectionRange.GetLowerBound()));
+			}
 		}
 
 		if( !bDelete )
@@ -158,7 +113,7 @@ void MovieSceneHelpers::FixupConsecutiveSections(TArray<UMovieSceneSection*>& Se
 			if(Sections.IsValidIndex(NextSectionIndex))
 			{
 				// Shift the next CameraCut's start time so that it starts when the new CameraCut ends
-				Sections[NextSectionIndex]->SetStartTime(Section.GetEndTime());
+				Sections[NextSectionIndex]->SetStartFrame(TRangeBound<FFrameNumber>::FlipInclusion(SectionRange.GetUpperBound()));
 			}
 		}
 	}
@@ -245,43 +200,6 @@ UCameraComponent* MovieSceneHelpers::CameraComponentFromRuntimeObject(UObject* R
 	return nullptr;
 }
 
-
-void MovieSceneHelpers::SetKeyInterpolation(FRichCurve& InCurve, FKeyHandle InKeyHandle, EMovieSceneKeyInterpolation InKeyInterpolation)
-{
-	switch (InKeyInterpolation)
-	{
-		case EMovieSceneKeyInterpolation::Auto:
-			InCurve.SetKeyInterpMode(InKeyHandle, RCIM_Cubic);
-			InCurve.SetKeyTangentMode(InKeyHandle, RCTM_Auto);
-			break;
-
-		case EMovieSceneKeyInterpolation::User:
-			InCurve.SetKeyInterpMode(InKeyHandle, RCIM_Cubic);
-			InCurve.SetKeyTangentMode(InKeyHandle, RCTM_User);
-			break;
-
-		case EMovieSceneKeyInterpolation::Break:
-			InCurve.SetKeyInterpMode(InKeyHandle, RCIM_Cubic);
-			InCurve.SetKeyTangentMode(InKeyHandle, RCTM_Break);
-			break;
-
-		case EMovieSceneKeyInterpolation::Linear:
-			InCurve.SetKeyInterpMode(InKeyHandle, RCIM_Linear);
-			InCurve.SetKeyTangentMode(InKeyHandle, RCTM_Auto);
-			break;
-
-		case EMovieSceneKeyInterpolation::Constant:
-			InCurve.SetKeyInterpMode(InKeyHandle, RCIM_Constant);
-			InCurve.SetKeyTangentMode(InKeyHandle, RCTM_Auto);
-			break;
-
-		default:
-			InCurve.SetKeyInterpMode(InKeyHandle, RCIM_Cubic);
-			InCurve.SetKeyTangentMode(InKeyHandle, RCTM_Auto);
-			break;
-	}
-}
-
 FTrackInstancePropertyBindings::FTrackInstancePropertyBindings( FName InPropertyName, const FString& InPropertyPath, const FName& InFunctionName, const FName& InNotifyFunctionName )
     : PropertyPath( InPropertyPath )
 	, NotifyFunctionName(InNotifyFunctionName)
diff --git a/Engine/Source/Runtime/MovieScene/Private/MovieSceneFolder.cpp b/Engine/Source/Runtime/MovieScene/Private/MovieSceneFolder.cpp
--- a/Engine/Source/Runtime/MovieScene/Private/MovieSceneFolder.cpp
+++ b/Engine/Source/Runtime/MovieScene/Private/MovieSceneFolder.cpp
@@ -6,6 +6,7 @@ UMovieSceneFolder::UMovieSceneFolder( const FObjectInitializer& ObjectInitialize
 	: Super( ObjectInitializer )
 #if WITH_EDITORONLY_DATA
 	, FolderColor(FColor::White)
+	, SortingOrder(-1)
 #endif
 {
 }
diff --git a/Engine/Source/Runtime/MovieScene/Private/MovieSceneModule.cpp b/Engine/Source/Runtime/MovieScene/Private/MovieSceneModule.cpp
--- a/Engine/Source/Runtime/MovieScene/Private/MovieSceneModule.cpp
+++ b/Engine/Source/Runtime/MovieScene/Private/MovieSceneModule.cpp
@@ -5,10 +5,75 @@
 #include "IMovieSceneModule.h"
 #include "MovieScene.h"
 #include "Compilation/IMovieSceneTemplateGenerator.h"
-
+#include "HAL/IConsoleManager.h"
+#include "Logging/MessageLog.h"
+#include "Misc/UObjectToken.h"
 
 DEFINE_LOG_CATEGORY(LogMovieScene);
 
+TAutoConsoleVariable<FString> CVarLegacyConversionFrameRate(
+	TEXT("MovieScene.LegacyConversionFrameRate"),
+	TEXT("60000fps"),
+	TEXT("Specifies default frame resolution for UMovieScene data saved before 4.20 (default: 60000fps). Examples: 60000 fps, 120/1 (120 fps), 30000/1001 (29.97), 0.01s (10ms)."),
+	ECVF_Default);
+
+struct FCachedLegacyConversionFrameRate
+{
+	FCachedLegacyConversionFrameRate()
+		: FrameRate(60000, 1)
+	{
+		DelegateHandle = IConsoleManager::Get().RegisterConsoleVariableSink_Handle(FConsoleCommandDelegate::CreateRaw(this, &FCachedLegacyConversionFrameRate::OnChanged));
+
+		OnChanged();
+	}
+	~FCachedLegacyConversionFrameRate()
+	{
+		IConsoleManager::Get().UnregisterConsoleVariableSink_Handle(DelegateHandle);
+	}
+
+	void OnChanged()
+	{
+		TryParseString(FrameRate, *CVarLegacyConversionFrameRate.GetValueOnGameThread());
+	}
+
+	FFrameRate FrameRate;
+	FConsoleVariableSinkHandle DelegateHandle;
+};
+
+FFrameRate GetLegacyConversionFrameRate()
+{
+	static FCachedLegacyConversionFrameRate CachedRate;
+	return CachedRate.FrameRate;
+}
+
+void EmitLegacyOutOfBoundsError(UObject* Object, FFrameRate InFrameRate, double InTime)
+{
+#if WITH_EDITOR
+	static const FName NAME_AssetCheck("AssetCheck");
+
+	FMessageLog AssetCheckLog(NAME_AssetCheck);
+
+	const FText Message = FText::Format(
+		NSLOCTEXT("MovieScene", "LegacyOutOfBoundsError", "Encountered time ({0} seconds) that is out of the supported range with a resolution of {1}fps. Saving this asset will cause loss of data. Please reduce MovieScene.LegacyConversionFrameRate and re-load this asset."),
+		InTime, InFrameRate.AsDecimal()
+	);
+	AssetCheckLog.Error()
+		->AddToken(FUObjectToken::Create(Object))
+		->AddToken(FTextToken::Create(Message));
+
+	AssetCheckLog.Open(EMessageSeverity::Warning);
+#endif
+}
+
+FFrameNumber UpgradeLegacyMovieSceneTime(UObject* ErrorContext, FFrameRate InFrameRate, double InTime)
+{
+	double ClampedKeyTime = FMath::Clamp(InTime, -InFrameRate.MaxSeconds(), InFrameRate.MaxSeconds());
+	if (InTime != ClampedKeyTime)
+	{
+		EmitLegacyOutOfBoundsError(ErrorContext, InFrameRate, InTime);
+	}
+	return InFrameRate.AsFrameNumber(ClampedKeyTime);
+}
 
 /**
  * MovieScene module implementation.
diff --git a/Engine/Source/Runtime/MovieScene/Private/MovieSceneSection.cpp b/Engine/Source/Runtime/MovieScene/Private/MovieSceneSection.cpp
--- a/Engine/Source/Runtime/MovieScene/Private/MovieSceneSection.cpp
+++ b/Engine/Source/Runtime/MovieScene/Private/MovieSceneSection.cpp
@@ -3,22 +3,30 @@
 #include "MovieSceneSection.h"
 #include "MovieSceneTrack.h"
 #include "MovieSceneCommonHelpers.h"
+#include "MovieSceneTimeHelpers.h"
 #include "Evaluation/MovieSceneEvalTemplate.h"
 #include "Generators/MovieSceneEasingCurves.h"
-
+#include "Channels/MovieSceneChannelProxy.h"
+#include "Containers/ArrayView.h"
+#include "UObject/SequencerObjectVersion.h"
 
 UMovieSceneSection::UMovieSceneSection(const FObjectInitializer& ObjectInitializer)
 	: Super( ObjectInitializer )
-	, StartTime(0.0f)
-	, EndTime(0.0f)
+	, PreRollFrames(0)
+	, PostRollFrames(0)
 	, RowIndex(0)
 	, OverlapPriority(0)
 	, bIsActive(true)
 	, bIsLocked(false)
-	, bIsInfinite(false)
-	, PreRollTime(0)
-	, PostRollTime(0)
+	, StartTime_DEPRECATED(0.f)
+	, EndTime_DEPRECATED(0.f)
+	, PreRollTime_DEPRECATED(0.f)
+	, PostRollTime_DEPRECATED(0.f)
+	, bIsInfinite_DEPRECATED(0)
+	, bSupportsInfiniteRange(false)
 {
+	SectionRange.Value = TRange<FFrameNumber>(0);
+
 	UMovieSceneBuiltInEasingFunction* DefaultEaseIn = ObjectInitializer.CreateDefaultSubobject<UMovieSceneBuiltInEasingFunction>(this, "EaseInFunction");
 	DefaultEaseIn->SetFlags(RF_Public); //@todo Need to be marked public. GLEO occurs when transform sections are added to actor sequence blueprints. Are these not being duplicated properly?
 	DefaultEaseIn->Type = EMovieSceneBuiltInEasing::CubicInOut;
@@ -40,9 +48,142 @@ void UMovieSceneSection::PostInitProperties()
 	}
 	
 	Super::PostInitProperties();
+
+	// Set up a default channel proxy if this class hasn't done so already
+	if (!HasAnyFlags(RF_ClassDefaultObject) && !ChannelProxy.IsValid())
+	{
+		ChannelProxy = MakeShared<FMovieSceneChannelProxy>();
+	}
+}
+
+void UMovieSceneSection::Serialize(FArchive& Ar)
+{
+	Super::Serialize(Ar);
+
+	Ar.UsingCustomVersion(FSequencerObjectVersion::GUID);
+
+	if (Ar.CustomVer(FSequencerObjectVersion::GUID) < FSequencerObjectVersion::FloatToIntConversion)
+	{
+		const FFrameRate LegacyFrameRate = GetLegacyConversionFrameRate();
+
+		if (bIsInfinite_DEPRECATED)
+		{
+			SectionRange = TRange<FFrameNumber>::All();
+		}
+		else
+		{
+			FFrameNumber StartFrame = UpgradeLegacyMovieSceneTime(this, LegacyFrameRate, StartTime_DEPRECATED);
+			FFrameNumber LastFrame  = UpgradeLegacyMovieSceneTime(this, LegacyFrameRate, EndTime_DEPRECATED);
+
+			// Exclusive upper bound so we want the upper bound to be exclusively the next frame after LastFrame
+			SectionRange = TRange<FFrameNumber>(StartFrame, LastFrame + 1);
+		}
+
+		// All these times are offsets from the start/end time so it's highly unlikely that they'll be out-of bounds
+		PreRollFrames                = LegacyFrameRate.AsFrameNumber(PreRollTime_DEPRECATED).Value;
+		PostRollFrames               = LegacyFrameRate.AsFrameNumber(PostRollTime_DEPRECATED).Value;
+#if WITH_EDITORONLY_DATA
+		Easing.AutoEaseInDuration    = (Easing.AutoEaseInTime_DEPRECATED * LegacyFrameRate).RoundToFrame().Value;
+		Easing.AutoEaseOutDuration   = (Easing.AutoEaseOutTime_DEPRECATED * LegacyFrameRate).RoundToFrame().Value;
+		Easing.ManualEaseInDuration  = (Easing.ManualEaseInTime_DEPRECATED * LegacyFrameRate).RoundToFrame().Value;
+		Easing.ManualEaseOutDuration = (Easing.ManualEaseOutTime_DEPRECATED * LegacyFrameRate).RoundToFrame().Value;
+#endif
+	}
+}
+
+void UMovieSceneSection::SetStartFrame(TRangeBound<FFrameNumber> NewStartFrame)
+{
+	if (TryModify())
+	{
+		bool bIsValidStartFrame = ensureMsgf(SectionRange.Value.GetUpperBound().IsOpen() || NewStartFrame.IsOpen() || SectionRange.Value.GetUpperBound().GetValue() >= NewStartFrame.GetValue(),
+			TEXT("Invalid start frame specified; will be clamped to current end frame."));
+
+		if (bIsValidStartFrame)
+		{
+			SectionRange.Value.SetLowerBound(NewStartFrame);
+		}
+		else
+		{
+			SectionRange.Value.SetLowerBound(TRangeBound<FFrameNumber>::FlipInclusion(SectionRange.Value.GetUpperBound()));
+		}
+	}
+}
+
+void UMovieSceneSection::SetEndFrame(TRangeBound<FFrameNumber> NewEndFrame)
+{
+	if (TryModify())
+	{
+		bool bIsValidEndFrame = ensureMsgf(SectionRange.Value.GetLowerBound().IsOpen() || NewEndFrame.IsOpen() || SectionRange.Value.GetLowerBound().GetValue() <= NewEndFrame.GetValue(),
+			TEXT("Invalid end frame specified; will be clamped to current start frame."));
+
+		if (bIsValidEndFrame)
+		{
+			SectionRange.Value.SetUpperBound(NewEndFrame);
+		}
+		else
+		{
+			SectionRange.Value.SetUpperBound(TRangeBound<FFrameNumber>::FlipInclusion(SectionRange.Value.GetLowerBound()));
+		}
+	}
+}
+
+FMovieSceneChannelProxy& UMovieSceneSection::GetChannelProxy() const
+{
+	FMovieSceneChannelProxy* Proxy = ChannelProxy.Get();
+	check(Proxy);
+	return *Proxy;
+}
+
+TSharedPtr<FStructOnScope> UMovieSceneSection::GetKeyStruct(TArrayView<const FKeyHandle> KeyHandles)
+{
+	return nullptr;
+}
+
+void UMovieSceneSection::MoveSection(FFrameNumber DeltaFrame)
+{
+	if (TryModify())
+	{
+		TRange<FFrameNumber> NewRange = SectionRange.Value;
+		if (SectionRange.Value.GetLowerBound().IsClosed())
+		{
+			SectionRange.Value.SetLowerBoundValue(SectionRange.Value.GetLowerBoundValue() + DeltaFrame);
+		}
+		if (SectionRange.Value.GetUpperBound().IsClosed())
+		{
+			SectionRange.Value.SetUpperBoundValue(SectionRange.Value.GetUpperBoundValue() + DeltaFrame);
+		}
+
+		if (ChannelProxy.IsValid())
+		{
+			for (const FMovieSceneChannelEntry& Entry : ChannelProxy->GetAllEntries())
+			{
+				Entry.GetBatchChannelInterface().Offset_Batch(Entry.GetChannels(), DeltaFrame);
+			}
+		}
+	}
 }
 
 
+TRange<FFrameNumber> UMovieSceneSection::ComputeEffectiveRange() const
+{
+	if (!SectionRange.Value.GetLowerBound().IsOpen() && !SectionRange.Value.GetUpperBound().IsOpen())
+	{
+		return GetRange();
+	}
+
+	TRange<FFrameNumber> EffectiveRange = TRange<FFrameNumber>::Empty();
+
+	if (ChannelProxy.IsValid())
+	{
+		for (const FMovieSceneChannelEntry& Entry : ChannelProxy->GetAllEntries())
+		{
+			EffectiveRange = TRange<FFrameNumber>::Hull(EffectiveRange, Entry.GetBatchChannelInterface().ComputeEffectiveRange_Batch(Entry.GetChannels()));
+		}
+	}
+
+	return TRange<FFrameNumber>::Intersection(EffectiveRange, SectionRange.Value);
+}
+
 FMovieSceneBlendTypeField UMovieSceneSection::GetSupportedBlendTypes() const
 {
 	UMovieSceneTrack* Track = GetTypedOuter<UMovieSceneTrack>();
@@ -71,7 +212,7 @@ void UMovieSceneSection::GetOverlappingSections(TArray<UMovieSceneSection*>& Out
 		return;
 	}
 
-	TRange<float> ThisRange = GetRange();
+	TRange<FFrameNumber> ThisRange = GetRange();
 	for (UMovieSceneSection* Section : Track->GetAllSections())
 	{
 		if (!Section || (!bIncludeThis && Section == this))
@@ -92,19 +233,38 @@ void UMovieSceneSection::GetOverlappingSections(TArray<UMovieSceneSection*>& Out
 }
 
 
-const UMovieSceneSection* UMovieSceneSection::OverlapsWithSections(const TArray<UMovieSceneSection*>& Sections, int32 TrackDelta, float TimeDelta) const
+const UMovieSceneSection* UMovieSceneSection::OverlapsWithSections(const TArray<UMovieSceneSection*>& Sections, int32 TrackDelta, int32 TimeDelta) const
 {
 	// Check overlaps with exclusive ranges so that sections can butt up against each other
 	int32 NewTrackIndex = RowIndex + TrackDelta;
-	TRange<float> NewSectionRange = TRange<float>(TRange<float>::BoundsType::Exclusive(StartTime + TimeDelta), TRange<float>::BoundsType::Exclusive(EndTime + TimeDelta));
+
+	// @todo: sequencer-timecode: is this correct? it seems like we should just use the section's ranges directly rather than fiddling with the bounds
+	// TRange<FFrameNumber> NewSectionRange;
+	// if (SectionRange.GetLowerBound().IsClosed())
+	// {
+	// 	NewSectionRange = TRange<FFrameNumber>(
+	// 		TRangeBound<FFrameNumber>::Exclusive(SectionRange.GetLowerBoundValue() + TimeDelta),
+	// 		NewSectionRange.GetUpperBound()
+	// 		);
+	// }
+
+	// if (SectionRange.GetUpperBound().IsClosed())
+	// {
+	// 	NewSectionRange = TRange<FFrameNumber>(
+	// 		NewSectionRange.GetLowerBound(),
+	// 		TRangeBound<FFrameNumber>::Exclusive(SectionRange.GetUpperBoundValue() + TimeDelta)
+	// 		);
+	// }
+
+	TRange<FFrameNumber> ThisRange = SectionRange.Value;
 
 	for (const auto Section : Sections)
 	{
 		check(Section);
 		if ((this != Section) && (Section->GetRowIndex() == NewTrackIndex))
 		{
-			TRange<float> ExclusiveSectionRange = TRange<float>(TRange<float>::BoundsType::Exclusive(Section->GetRange().GetLowerBoundValue()), TRange<float>::BoundsType::Exclusive(Section->GetRange().GetUpperBoundValue()));
-			if (NewSectionRange.Overlaps(ExclusiveSectionRange))
+			//TRange<float> ExclusiveSectionRange = TRange<float>(TRange<float>::BoundsType::Exclusive(Section->GetRange().GetLowerBoundValue()), TRange<float>::BoundsType::Exclusive(Section->GetRange().GetUpperBoundValue()));
+			if (ThisRange.Overlaps(Section->GetRange()))
 			{
 				return Section;
 			}
@@ -115,12 +275,12 @@ const UMovieSceneSection* UMovieSceneSection::OverlapsWithSections(const TArray<
 }
 
 
-void UMovieSceneSection::InitialPlacement(const TArray<UMovieSceneSection*>& Sections, float InStartTime, float InEndTime, bool bAllowMultipleRows)
+void UMovieSceneSection::InitialPlacement(const TArray<UMovieSceneSection*>& Sections, FFrameNumber InStartTime, int32 Duration, bool bAllowMultipleRows)
 {
-	check(StartTime <= EndTime);
+	check(Duration >= 0);
 
-	StartTime = InStartTime;
-	EndTime = InEndTime;
+	// Inclusive lower, exclusive upper bounds
+	SectionRange = TRange<FFrameNumber>(InStartTime, InStartTime + Duration);
 	RowIndex = 0;
 
 	for (UMovieSceneSection* OtherSection : Sections)
@@ -140,14 +300,21 @@ void UMovieSceneSection::InitialPlacement(const TArray<UMovieSceneSection*>& Sec
 		for (;;)
 		{
 			const UMovieSceneSection* OverlappedSection = OverlapsWithSections(Sections);
-
 			if (OverlappedSection == nullptr)
 			{
 				break;
 			}
 
-			TSet<FKeyHandle> KeyHandles;
-			MoveSection(OverlappedSection->GetEndTime() - StartTime, KeyHandles);
+			TRange<FFrameNumber> OtherRange = OverlappedSection->GetRange();
+			if (OtherRange.GetUpperBound().IsClosed())
+			{
+				MoveSection(OtherRange.GetUpperBoundValue() - InStartTime);
+			}
+			else
+			{
+				++OverlapPriority;
+				break;
+			}
 		}
 	}
 
@@ -158,10 +325,51 @@ void UMovieSceneSection::InitialPlacement(const TArray<UMovieSceneSection*>& Sec
 	}
 }
 
+void UMovieSceneSection::InitialPlacementOnRow(const TArray<UMovieSceneSection*>& Sections, FFrameNumber InStartTime, int32 Duration, int32 InRowIndex)
+{
+	check(Duration >= 0);
+
+	// Inclusive lower, exclusive upper bounds
+	SectionRange = TRange<FFrameNumber>(InStartTime, InStartTime + Duration);
+	RowIndex = InRowIndex;
+
+	// If no given row index, put it on the next available row
+	if (RowIndex == INDEX_NONE)
+	{
+		RowIndex = 0;
+		while (OverlapsWithSections(Sections) != nullptr)
+		{
+			++RowIndex;
+		}
+	}
+
+	for (UMovieSceneSection* OtherSection : Sections)
+	{
+		OverlapPriority = FMath::Max(OtherSection->GetOverlapPriority()+1, OverlapPriority);
+	}
+
+	// If this overlaps with any sections, move out all the sections that are beyond this row
+	if (OverlapsWithSections(Sections))
+	{
+		for (UMovieSceneSection* OtherSection : Sections)
+		{
+			if (OtherSection != nullptr && OtherSection != this && OtherSection->GetRowIndex() >= RowIndex)
+			{
+				OtherSection->SetRowIndex(OtherSection->GetRowIndex()+1);
+			}
+		}
+	}
 
-UMovieSceneSection* UMovieSceneSection::SplitSection(float SplitTime)
+	UMovieSceneTrack* Track = GetTypedOuter<UMovieSceneTrack>();
+	if (Track)
+	{
+		Track->UpdateEasing();
+	}
+}
+
+UMovieSceneSection* UMovieSceneSection::SplitSection(FFrameNumber SplitTime)
 {
-	if (!IsTimeWithinSection(SplitTime))
+	if (!SectionRange.Value.Contains(SplitTime))
 	{
 		return nullptr;
 	}
@@ -170,10 +378,12 @@ UMovieSceneSection* UMovieSceneSection::SplitSection(float SplitTime)
 
 	if (TryModify())
 	{
-		float SectionEndTime = GetEndTime();
-				
+		TRange<FFrameNumber> StartingRange  = SectionRange.Value;
+		TRange<FFrameNumber> LeftHandRange  = TRange<FFrameNumber>(StartingRange.GetLowerBound(), TRangeBound<FFrameNumber>::Exclusive(SplitTime));
+		TRange<FFrameNumber> RightHandRange = TRange<FFrameNumber>(TRangeBound<FFrameNumber>::Inclusive(SplitTime), StartingRange.GetUpperBound());
+
 		// Trim off the right
-		SetEndTime(SplitTime);
+		SectionRange = LeftHandRange;
 
 		// Create a new section
 		UMovieSceneTrack* Track = CastChecked<UMovieSceneTrack>(GetOuter());
@@ -182,8 +392,7 @@ UMovieSceneSection* UMovieSceneSection::SplitSection(float SplitTime)
 		UMovieSceneSection* NewSection = DuplicateObject<UMovieSceneSection>(this, Track);
 		check(NewSection);
 
-		NewSection->SetStartTime(SplitTime);
-		NewSection->SetEndTime(SectionEndTime);
+		NewSection->SetRange(RightHandRange);
 		Track->AddSection(*NewSection);
 
 		return NewSection;
@@ -193,98 +402,73 @@ UMovieSceneSection* UMovieSceneSection::SplitSection(float SplitTime)
 }
 
 
-void UMovieSceneSection::TrimSection(float TrimTime, bool bTrimLeft)
+void UMovieSceneSection::TrimSection(FFrameNumber TrimTime, bool bTrimLeft)
 {
-	if (IsTimeWithinSection(TrimTime))
+	if (SectionRange.Value.Contains(TrimTime))
 	{
 		SetFlags(RF_Transactional);
 		if (TryModify())
 		{
 			if (bTrimLeft)
 			{
-				SetStartTime(TrimTime);
+				SectionRange.Value.SetLowerBound(TRangeBound<FFrameNumber>::Inclusive(TrimTime));
 			}
 			else
 			{
-				SetEndTime(TrimTime);
+				SectionRange.Value.SetUpperBound(TRangeBound<FFrameNumber>::Exclusive(TrimTime));
 			}
 		}
 	}
 }
 
 
-void UMovieSceneSection::AddKeyToCurve(FRichCurve& InCurve, float Time, float Value, EMovieSceneKeyInterpolation Interpolation, const bool bUnwindRotation)
-{
-	if (IsTimeWithinSection(Time))
-	{
-		if (TryModify())
-		{
-			FKeyHandle ExistingKeyHandle = InCurve.FindKey(Time);
-			FKeyHandle NewKeyHandle = InCurve.UpdateOrAddKey(Time, Value, bUnwindRotation);
-
-			if (!InCurve.IsKeyHandleValid(ExistingKeyHandle) && InCurve.IsKeyHandleValid(NewKeyHandle))
-			{
-				MovieSceneHelpers::SetKeyInterpolation(InCurve, NewKeyHandle, Interpolation);
-			}
-		}
-	}
-}
-
-
-void UMovieSceneSection::SetCurveDefault(FRichCurve& InCurve, float Value)
-{
-	if (InCurve.GetDefaultValue() != Value && TryModify())
-	{
-		InCurve.SetDefaultValue(Value);
-	}
-}
-
-
 FMovieSceneEvalTemplatePtr UMovieSceneSection::GenerateTemplate() const
 {
 	return FMovieSceneEvalTemplatePtr();
 }
 
 
-float UMovieSceneSection::EvaluateEasing(float InTime) const
+float UMovieSceneSection::EvaluateEasing(FFrameTime InTime) const
 {
-	TRange<float> CurrentRange = GetRange();
-
 	float EaseInValue = 1.f;
 	float EaseOutValue = 1.f;
 
-	if (!CurrentRange.GetLowerBound().IsOpen() && Easing.GetEaseInTime() > 0.f && Easing.EaseIn.GetObject())
+	if (HasStartFrame() && Easing.GetEaseInDuration() > 0 && Easing.EaseIn.GetObject())
 	{
-		float EaseInTime = (InTime - CurrentRange.GetLowerBoundValue()) / Easing.GetEaseInTime();
-		if (EaseInTime <= 0.f)
+		const int32  EaseFrame    = (InTime.FrameNumber - GetInclusiveStartFrame()).Value;
+		const double EaseInInterp = (double(EaseFrame) + InTime.GetSubFrame()) / Easing.GetEaseInDuration();
+
+		if (EaseInInterp <= 0.0)
 		{
-			EaseInValue = 0.f;
+			EaseInValue = 0.0;
 		}
-		else if (EaseInTime >= 1.f)
+		else if (EaseInInterp >= 1.0)
 		{
-			EaseInValue = 1.f;
+			EaseInValue = 1.0;
 		}
 		else
 		{
-			EaseInValue = IMovieSceneEasingFunction::EvaluateWith(Easing.EaseIn, EaseInTime);
+			EaseInValue = IMovieSceneEasingFunction::EvaluateWith(Easing.EaseIn, EaseInInterp);
 		}
 	}
 
 
-	if (!CurrentRange.GetUpperBound().IsOpen() && Easing.GetEaseOutTime() > 0.f && Easing.EaseOut.GetObject())
+	if (HasEndFrame() && Easing.GetEaseOutDuration() > 0 && Easing.EaseOut.GetObject())
 	{
-		float EaseOutTime = (InTime - (CurrentRange.GetUpperBoundValue() - Easing.GetEaseOutTime())) / Easing.GetEaseOutTime();
-		if (EaseOutTime <= 0.f)
+		const int32  EaseFrame     = (InTime.FrameNumber - GetExclusiveEndFrame() + Easing.GetEaseOutDuration()).Value;
+		const double EaseOutInterp = (double(EaseFrame) + InTime.GetSubFrame()) / Easing.GetEaseOutDuration();
+
+		if (EaseOutInterp <= 0.0)
 		{
-			EaseOutValue = 1.f;
+			EaseOutValue = 1.0;
 		}
-		else if (EaseOutTime >= 1.f)
+		else if (EaseOutInterp >= 1.0)
 		{
-			EaseOutValue = 0.f;
+			EaseOutValue = 0.0;
 		}
 		else
 		{
-			EaseOutValue = 1.f - IMovieSceneEasingFunction::EvaluateWith(Easing.EaseOut, EaseOutTime);
+			EaseOutValue = 1.f - IMovieSceneEasingFunction::EvaluateWith(Easing.EaseOut, EaseOutInterp);
 		}
 	}
 
@@ -292,51 +476,61 @@ float UMovieSceneSection::EvaluateEasing(float InTime) const
 }
 
 
-void UMovieSceneSection::EvaluateEasing(float InTime, TOptional<float>& OutEaseInValue, TOptional<float>& OutEaseOutValue, float* OutEaseInInterp, float* OutEaseOutInterp) const
+void UMovieSceneSection::EvaluateEasing(FFrameTime InTime, TOptional<float>& OutEaseInValue, TOptional<float>& OutEaseOutValue, float* OutEaseInInterp, float* OutEaseOutInterp) const
 {
-	TRange<float> CurrentRange = GetRange();
-
-	if (!CurrentRange.GetLowerBound().IsOpen() && Easing.EaseIn.GetObject() && GetEaseInRange().Contains(InTime))
+	if (HasStartFrame() && Easing.EaseIn.GetObject() && GetEaseInRange().Contains(InTime.FrameNumber))
 	{
-		float EaseInTime = (InTime - CurrentRange.GetLowerBoundValue()) / Easing.GetEaseInTime();
-		OutEaseInValue = IMovieSceneEasingFunction::EvaluateWith(Easing.EaseIn, EaseInTime);
+		const int32  EaseFrame    = (InTime.FrameNumber - GetInclusiveStartFrame()).Value;
+		const double EaseInInterp = (double(EaseFrame) + InTime.GetSubFrame()) / Easing.GetEaseInDuration();
+
+		OutEaseInValue = IMovieSceneEasingFunction::EvaluateWith(Easing.EaseIn, EaseInInterp);
 
 		if (OutEaseInInterp)
 		{
-			*OutEaseInInterp = EaseInTime;
+			*OutEaseInInterp = EaseInInterp;
 		}
 	}
 
-	if (!CurrentRange.GetUpperBound().IsOpen() && Easing.EaseOut.GetObject() && GetEaseOutRange().Contains(InTime))
+	if (HasEndFrame() && Easing.EaseOut.GetObject() && GetEaseOutRange().Contains(InTime.FrameNumber))
 	{
-		float EaseOutTime = (InTime - (CurrentRange.GetUpperBoundValue() - Easing.GetEaseOutTime())) / Easing.GetEaseOutTime();
-		OutEaseOutValue = 1.f - IMovieSceneEasingFunction::EvaluateWith(Easing.EaseOut, EaseOutTime);
+		const int32  EaseFrame     = (InTime.FrameNumber - GetExclusiveEndFrame() + Easing.GetEaseOutDuration()).Value;
+		const double EaseOutInterp = (double(EaseFrame) + InTime.GetSubFrame()) / Easing.GetEaseOutDuration();
+
+		OutEaseOutValue = 1.f - IMovieSceneEasingFunction::EvaluateWith(Easing.EaseOut, EaseOutInterp);
 
 		if (OutEaseOutInterp)
 		{
-			*OutEaseOutInterp = EaseOutTime;
+			*OutEaseOutInterp = EaseOutInterp;
 		}
 	}
 }
 
 
-TRange<float> UMovieSceneSection::GetEaseInRange() const
+TRange<FFrameNumber> UMovieSceneSection::GetEaseInRange() const
 {
-	const float MaxTime = FMath::Min(GetStartTime() + Easing.GetEaseInTime(), GetEndTime());
-	if (!bIsInfinite && Easing.GetEaseInTime() > 0.f)
+	if (HasStartFrame() && Easing.GetEaseInDuration() > 0)
 	{
-		return TRange<float>(GetStartTime(), TRangeBound<float>::Inclusive(MaxTime));
+		TRangeBound<FFrameNumber> LowerBound = TRangeBound<FFrameNumber>::Inclusive(GetInclusiveStartFrame());
+		TRangeBound<FFrameNumber> UpperBound = TRangeBound<FFrameNumber>::Exclusive(GetInclusiveStartFrame() + Easing.GetEaseInDuration());
+
+		UpperBound = TRangeBound<FFrameNumber>::MinUpper(UpperBound, SectionRange.Value.GetUpperBound());
+		return TRange<FFrameNumber>(LowerBound, UpperBound);
 	}
-	return TRange<float>::Empty();
+
+	return TRange<FFrameNumber>::Empty();
 }
 
 
-TRange<float> UMovieSceneSection::GetEaseOutRange() const
+TRange<FFrameNumber> UMovieSceneSection::GetEaseOutRange() const
 {
-	const float MinTime = FMath::Max(GetEndTime() - Easing.GetEaseOutTime(), GetStartTime());
-	if (!bIsInfinite && Easing.GetEaseOutTime() > 0.f)
+	if (HasEndFrame() && Easing.GetEaseOutDuration() > 0)
 	{
-		return TRange<float>(MinTime, TRangeBound<float>::Inclusive(GetEndTime()));
+		TRangeBound<FFrameNumber> UpperBound = TRangeBound<FFrameNumber>::Exclusive(GetExclusiveEndFrame());
+		TRangeBound<FFrameNumber> LowerBound = TRangeBound<FFrameNumber>::Inclusive(GetExclusiveEndFrame() - Easing.GetEaseOutDuration());
+
+		LowerBound = TRangeBound<FFrameNumber>::MaxLower(LowerBound, SectionRange.Value.GetLowerBound());
+		return TRange<FFrameNumber>(LowerBound, UpperBound);
 	}
-	return TRange<float>::Empty();
+
+	return TRange<FFrameNumber>::Empty();
 }
\ No newline at end of file
diff --git a/Engine/Source/Runtime/MovieScene/Private/MovieSceneSequenceInstanceData.cpp b/Engine/Source/Runtime/MovieScene/Private/MovieSceneSequenceInstanceData.cpp
--- a/Engine/Source/Runtime/MovieScene/Private/MovieSceneSequenceInstanceData.cpp
+++ b/Engine/Source/Runtime/MovieScene/Private/MovieSceneSequenceInstanceData.cpp
@@ -25,5 +25,6 @@ FMovieSceneSequenceInstanceDataPtr& FMovieSceneSequenceInstanceDataPtr::operator
 
 bool FMovieSceneSequenceInstanceDataPtr::Serialize(FArchive& Ar)
 {
-	return SerializeInlineValue(*this, Ar);
+	bool bShouldWarn = !WITH_EDITORONLY_DATA;
+	return SerializeInlineValue(*this, Ar, bShouldWarn);
 }
diff --git a/Engine/Source/Runtime/MovieScene/Private/MovieSceneSequencePlayer.cpp b/Engine/Source/Runtime/MovieScene/Private/MovieSceneSequencePlayer.cpp
--- a/Engine/Source/Runtime/MovieScene/Private/MovieSceneSequencePlayer.cpp
+++ b/Engine/Source/Runtime/MovieScene/Private/MovieSceneSequencePlayer.cpp
@@ -2,8 +2,10 @@
 
 #include "MovieSceneSequencePlayer.h"
 #include "MovieScene.h"
+#include "MovieSceneTimeHelpers.h"
 #include "MovieSceneSequence.h"
 #include "Engine/Engine.h"
+#include "GameFramework/WorldSettings.h"
 #include "Misc/RuntimeErrors.h"
 
 bool FMovieSceneSequencePlaybackSettings::SerializeFromMismatchedTag( const FPropertyTag& Tag, FArchive& Ar )
@@ -23,12 +25,11 @@ UMovieSceneSequencePlayer::UMovieSceneSequencePlayer(const FObjectInitializer& I
 	, bReversePlayback(false)
 	, bIsEvaluating(false)
 	, Sequence(nullptr)
-	, TimeCursorPosition(0.0f)
-	, StartTime(0.f)
-	, EndTime(0.f)
+	, StartTime(0)
+	, DurationFrames(0)
 	, CurrentNumLoops(0)
-	, LatentPlaybackPosition(0.f)
 {
+	PlayPosition.Reset(FFrameTime(0));
 }
 
 UMovieSceneSequencePlayer::~UMovieSceneSequencePlayer()
@@ -85,36 +86,46 @@ void UMovieSceneSequencePlayer::PlayLooping(int32 NumLoops)
 
 void UMovieSceneSequencePlayer::PlayInternal()
 {
-	if (!IsPlaying())
+	if (!IsPlaying() && Sequence && CanPlay())
 	{
 		// Start playing
-		StartPlayingNextTick();
+
+		// @todo Sequencer playback: Should we recreate the instance every time?
+		// We must not recreate the instance since it holds stateful information (such as which objects it has spawned). Recreating the instance would break any 
+		// @todo: Is this still the case now that eval state is stored (correctly) in the player?
+		if (!RootTemplateInstance.IsValid())
+		{
+			RootTemplateInstance.Initialize(*Sequence, *this);
+		}
 
 		// Update now
-		bPendingFirstUpdate = false;
-		
 		if (PlaybackSettings.bRestoreState)
 		{
 			PreAnimatedState.EnableGlobalCapture();
 		}
 
+		Status = EMovieScenePlayerStatus::Playing;
+		PlaybackSettings.TimeController->StartPlaying(GetCurrentTime());
+
+		OnStartedPlaying();
+
 		UMovieSceneSequence* MovieSceneSequence = RootTemplateInstance.GetSequence(MovieSceneSequenceID::Root);
-		const bool bHasValidMovieScene = (MovieSceneSequence && MovieSceneSequence->GetMovieScene());
-		TOptional<float> FixedFrameInterval = bHasValidMovieScene ? MovieSceneSequence->GetMovieScene()->GetOptionalFixedFrameInterval() : TOptional<float>();
+		UMovieScene*         MovieScene         = MovieSceneSequence ? MovieSceneSequence->GetMovieScene() : nullptr;
 
-		if (FixedFrameInterval.IsSet() && bHasValidMovieScene && MovieSceneSequence->GetMovieScene()->GetForceFixedFrameIntervalPlayback())
+		if (PlayPosition.GetEvaluationType() == EMovieSceneEvaluationType::FrameLocked)
 		{
 			OldMaxTickRate = GEngine->GetMaxFPS();
-			GEngine->SetMaxFPS(1.f / FixedFrameInterval.GetValue());
+			GEngine->SetMaxFPS(1.f / PlayPosition.GetInputRate().AsInterval());
 		}
 
-		if (!PlayPosition.GetPreviousPosition().IsSet() || PlayPosition.GetPreviousPosition().GetValue() != GetSequencePosition())
+		if (!PlayPosition.GetLastPlayEvalPostition().IsSet() || PlayPosition.GetLastPlayEvalPostition() != PlayPosition.GetCurrentPosition())
 		{
-			// Ensure we're at the current sequence position
-			PlayPosition.JumpTo(GetSequencePosition(), FixedFrameInterval);
+			UpdateMovieSceneInstance(PlayPosition.JumpTo(StartTime), Status);
+		}
 
-			// We pass the range of PlayTo here in order to correctly update the last evaluated time in the playposition
-			UpdateMovieSceneInstance(PlayPosition.PlayTo(GetSequencePosition(), FixedFrameInterval));
+		if (MovieSceneSequence)
+		{
+			UE_LOG(LogMovieScene, Verbose, TEXT("PlayInternal - MovieSceneSequence: %s"), *MovieSceneSequence->GetName());
 		}
 
 		if (bReversePlayback)
@@ -134,49 +145,25 @@ void UMovieSceneSequencePlayer::PlayInternal()
 	}
 }
 
-void UMovieSceneSequencePlayer::StartPlayingNextTick()
-{
-	if (IsPlaying() || !Sequence || !CanPlay())
-	{
-		return;
-	}
-
-	// @todo Sequencer playback: Should we recreate the instance every time?
-	// We must not recreate the instance since it holds stateful information (such as which objects it has spawned). Recreating the instance would break any 
-	// @todo: Is this still the case now that eval state is stored (correctly) in the player?
-	if (!RootTemplateInstance.IsValid())
-	{
-		RootTemplateInstance.Initialize(*Sequence, *this);
-	}
-
-	OnStartedPlaying();
-
-	bPendingFirstUpdate = true;
-	Status = EMovieScenePlayerStatus::Playing;
-}
-
 void UMovieSceneSequencePlayer::Pause()
 {
 	if (IsPlaying())
 	{
 		if (bIsEvaluating)
 		{
-			LatentActions.Add(ELatentAction::Pause);
+			LatentActions.Emplace(FLatentAction::Pause);
 			return;
 		}
 
 		Status = EMovieScenePlayerStatus::Paused;
+		PlaybackSettings.TimeController->StopPlaying(GetCurrentTime());
 
 		// Evaluate the sequence at its current time, with a status of 'stopped' to ensure that animated state pauses correctly. (ie. audio sounds should stop/pause)
 		{
 			bIsEvaluating = true;
 
-			UMovieSceneSequence* MovieSceneSequence = RootTemplateInstance.GetSequence(MovieSceneSequenceID::Root);
-			TOptional<float> FixedFrameInterval = MovieSceneSequence->GetMovieScene() ? MovieSceneSequence->GetMovieScene()->GetOptionalFixedFrameInterval() : TOptional<float>();
-
-			FMovieSceneEvaluationRange Range = PlayPosition.JumpTo(GetSequencePosition(), FixedFrameInterval);
-
-			const FMovieSceneContext Context(Range, EMovieScenePlayerStatus::Stopped);
+			FMovieSceneEvaluationRange CurrentTimeRange = PlayPosition.GetCurrentPositionAsRange();
+			const FMovieSceneContext Context(CurrentTimeRange, EMovieScenePlayerStatus::Stopped);
 			RootTemplateInstance.Evaluate(Context, *this);
 
 			bIsEvaluating = false;
@@ -184,6 +171,12 @@ void UMovieSceneSequencePlayer::Pause()
 
 		ApplyLatentActions();
 
+		UMovieSceneSequence* MovieSceneSequence = RootTemplateInstance.GetSequence(MovieSceneSequenceID::Root);
+		if (MovieSceneSequence)
+		{
+			UE_LOG(LogMovieScene, Verbose, TEXT("Pause - MovieSceneSequence: %s"), *MovieSceneSequence->GetName());
+		}
+
 		if (OnPause.IsBound())
 		{
 			OnPause.Broadcast();
@@ -205,6 +198,7 @@ void UMovieSceneSequencePlayer::Scrub()
 	}
 
 	Status = EMovieScenePlayerStatus::Scrubbing;
+	PlaybackSettings.TimeController->StopPlaying(GetCurrentTime());
 }
 
 void UMovieSceneSequencePlayer::Stop()
@@ -213,12 +207,17 @@ void UMovieSceneSequencePlayer::Stop()
 	{
 		if (bIsEvaluating)
 		{
-			LatentActions.Add(ELatentAction::Stop);
+			LatentActions.Emplace(FLatentAction::Stop);
 			return;
 		}
 
 		Status = EMovieScenePlayerStatus::Stopped;
-		TimeCursorPosition = bReversePlayback ? GetLength() : 0.f;
+
+		// Put the cursor back to the start
+		PlayPosition.Reset(bReversePlayback ? GetLastValidTime() : FFrameTime(StartTime));
+
+		PlaybackSettings.TimeController->StopPlaying(GetCurrentTime());
+
 		CurrentNumLoops = 0;
 
 		if (PlaybackSettings.bRestoreState)
@@ -235,6 +234,12 @@ void UMovieSceneSequencePlayer::Stop()
 
 		OnStopped();
 
+		UMovieSceneSequence* MovieSceneSequence = RootTemplateInstance.GetSequence(MovieSceneSequenceID::Root);
+		if (MovieSceneSequence)
+		{
+			UE_LOG(LogMovieScene, Verbose, TEXT("Stop - MovieSceneSequence: %s"), *MovieSceneSequence->GetName());
+		}
+
 		if (OnStop.IsBound())
 		{
 			OnStop.Broadcast();
@@ -244,31 +249,117 @@ void UMovieSceneSequencePlayer::Stop()
 
 void UMovieSceneSequencePlayer::GoToEndAndStop()
 {
-	SetPlaybackPosition(GetLength());
-
+	JumpToFrame(GetFrameDuration()-1);
 	Stop();
 }
 
-float UMovieSceneSequencePlayer::GetPlaybackPosition() const
+FQualifiedFrameTime UMovieSceneSequencePlayer::GetCurrentTime() const
 {
-	return TimeCursorPosition;
+	FFrameTime Time = PlayPosition.GetCurrentPosition();
+	return FQualifiedFrameTime(Time, PlayPosition.GetInputRate());
 }
 
-void UMovieSceneSequencePlayer::SetPlaybackPosition(float NewPlaybackPosition)
+FQualifiedFrameTime UMovieSceneSequencePlayer::GetDuration() const
 {
-	if (bIsEvaluating)
+	return FQualifiedFrameTime(DurationFrames, PlayPosition.GetInputRate());
+}
+
+int32 UMovieSceneSequencePlayer::GetFrameDuration() const
+{
+	return DurationFrames;
+}
+
+void UMovieSceneSequencePlayer::SetFrameRate(FFrameRate FrameRate)
+{
+	UMovieScene* MovieScene = Sequence ? Sequence->GetMovieScene() : nullptr;
+	if (MovieScene)
 	{
-		LatentActions.Add(ELatentAction::SetPlaybackPosition);
-		LatentPlaybackPosition = NewPlaybackPosition;
-		return;
+		if (MovieScene->GetEvaluationType() == EMovieSceneEvaluationType::FrameLocked && !FrameRate.IsMultipleOf(MovieScene->GetFrameResolution()))
+		{
+			UE_LOG(LogMovieScene, Warning, TEXT("Attempting to play back a sequence with resolution of %f ticks per second frame locked to %f fps, which is not a multiple of the resolution."), MovieScene->GetFrameResolution().AsDecimal(), FrameRate.AsDecimal());
+		}
+	}
+
+	FFrameRate CurrentInputRate = PlayPosition.GetInputRate();
+
+	StartTime      = ConvertFrameTime(StartTime,                    CurrentInputRate, FrameRate).FloorToFrame();
+	DurationFrames = ConvertFrameTime(FFrameNumber(DurationFrames), CurrentInputRate, FrameRate).RoundToFrame().Value;
+
+	PlayPosition.SetTimeBase(FrameRate, PlayPosition.GetOutputRate(), PlayPosition.GetEvaluationType());
+}
+
+void UMovieSceneSequencePlayer::SetFrameRange( int32 NewStartTime, int32 Duration )
+{
+	Duration = FMath::Max(Duration, 0);
+
+	StartTime      = NewStartTime;
+	DurationFrames = Duration;
+
+	TOptional<FFrameTime> CurrentTime = PlayPosition.GetCurrentPosition();
+	if (CurrentTime.IsSet())
+	{
+		FFrameTime LastValidTime = GetLastValidTime();
+
+		if (CurrentTime.GetValue() < StartTime)
+		{
+			PlayPosition.Reset(StartTime);
+		}
+		else if (CurrentTime.GetValue() > LastValidTime)
+		{
+			PlayPosition.Reset(LastValidTime);
+		}
+	}
+
+	if (PlaybackSettings.TimeController.IsValid())
+	{
+		PlaybackSettings.TimeController->Reset(GetCurrentTime());
 	}
+}
+
+void UMovieSceneSequencePlayer::SetTimeRange( float StartTimeSeconds, float DurationSeconds )
+{
+	const FFrameRate Rate = PlayPosition.GetInputRate();
 
-	UpdateTimeCursorPosition(NewPlaybackPosition);
+	const FFrameNumber StartFrame = ( StartTimeSeconds * Rate ).FloorToFrame();
+	const FFrameNumber Duration   = ( DurationSeconds  * Rate ).RoundToFrame();
+
+	SetFrameRange(StartFrame.Value, Duration.Value);
 }
 
-void UMovieSceneSequencePlayer::JumpToPosition(float NewPlaybackPosition)
+void UMovieSceneSequencePlayer::PlayToFrame(FFrameTime NewPosition)
 {
-	UpdateTimeCursorPosition(NewPlaybackPosition, EMovieScenePlayerStatus::Scrubbing);
+	UpdateTimeCursorPosition(NewPosition, EUpdatePositionMethod::Play);
+
+	PlaybackSettings.TimeController->Reset(GetCurrentTime());
+}
+
+void UMovieSceneSequencePlayer::ScrubToFrame(FFrameTime NewPosition)
+{
+	UpdateTimeCursorPosition(NewPosition, EUpdatePositionMethod::Scrub);
+
+	PlaybackSettings.TimeController->Reset(GetCurrentTime());
+}
+
+void UMovieSceneSequencePlayer::JumpToFrame(FFrameTime NewPosition)
+{
+	UpdateTimeCursorPosition(NewPosition, EUpdatePositionMethod::Jump);
+
+	PlaybackSettings.TimeController->Reset(GetCurrentTime());
+}
+
+void UMovieSceneSequencePlayer::PlayToSeconds(float TimeInSeconds)
+{
+	PlayToFrame(TimeInSeconds * PlayPosition.GetInputRate());
+}
+
+void UMovieSceneSequencePlayer::ScrubToSeconds(float TimeInSeconds)
+{
+	ScrubToFrame(TimeInSeconds * PlayPosition.GetInputRate());
+}
+
+void UMovieSceneSequencePlayer::JumpToSeconds(float TimeInSeconds)
+{
+	JumpToFrame(TimeInSeconds * PlayPosition.GetInputRate());
 }
 
 bool UMovieSceneSequencePlayer::IsPlaying() const
@@ -283,7 +374,7 @@ bool UMovieSceneSequencePlayer::IsPaused() const
 
 float UMovieSceneSequencePlayer::GetLength() const
 {
-	return EndTime - StartTime;
+	return GetDuration().AsSeconds();
 }
 
 float UMovieSceneSequencePlayer::GetPlayRate() const
@@ -296,26 +387,28 @@ void UMovieSceneSequencePlayer::SetPlayRate(float PlayRate)
 	PlaybackSettings.PlayRate = PlayRate;
 }
 
-void UMovieSceneSequencePlayer::SetPlaybackRange( const float NewStartTime, const float NewEndTime )
+void UMovieSceneSequencePlayer::SetPlaybackRange( float NewStartTime, float NewEndTime )
 {
-	StartTime = NewStartTime;
-	EndTime = FMath::Max(NewEndTime, StartTime);
+	SetTimeRange(NewStartTime, NewEndTime - NewStartTime);
+}
 
-	TimeCursorPosition = FMath::Clamp(TimeCursorPosition, 0.f, GetLength());
+FFrameTime UMovieSceneSequencePlayer::GetLastValidTime() const
+{
+	return DurationFrames > 0 ? FFrameTime(StartTime + DurationFrames-1, 0.99999994f) : FFrameTime(StartTime);
 }
 
-bool UMovieSceneSequencePlayer::ShouldStopOrLoop(float NewPosition) const
+bool UMovieSceneSequencePlayer::ShouldStopOrLoop(FFrameTime NewPosition) const
 {
 	bool bShouldStopOrLoop = false;
 	if (IsPlaying())
 	{
 		if (!bReversePlayback)
 		{
-			bShouldStopOrLoop = NewPosition >= GetLength();
+			bShouldStopOrLoop = NewPosition.FrameNumber >= StartTime + GetFrameDuration();
 		}
 		else
 		{
-			bShouldStopOrLoop = NewPosition < 0.f;
+			bShouldStopOrLoop = NewPosition.FrameNumber < StartTime ;
 		}
 	}
 
@@ -325,18 +418,62 @@ bool UMovieSceneSequencePlayer::ShouldStopOrLoop(float NewPosition) const
 void UMovieSceneSequencePlayer::Initialize(UMovieSceneSequence* InSequence, const FMovieSceneSequencePlaybackSettings& InSettings)
 {
 	check(InSequence);
-	
+
 	Sequence = InSequence;
 	PlaybackSettings = InSettings;
 
-	if (UMovieScene* MovieScene = Sequence->GetMovieScene())
+	EUpdateClockSource ClockToUse = EUpdateClockSource::Tick;
+
+	UMovieScene* MovieScene = Sequence->GetMovieScene();
+	if (MovieScene)
 	{
-		TRange<float> PlaybackRange = MovieScene->GetPlaybackRange();
-		SetPlaybackRange(PlaybackRange.GetLowerBoundValue(), PlaybackRange.GetUpperBoundValue());
+		EMovieSceneEvaluationType EvaluationType    = MovieScene->GetEvaluationType();
+		FFrameRate                FrameResolution   = MovieScene->GetFrameResolution();
+		FFrameRate                PlaybackFrameRate = MovieScene->GetPlaybackFrameRate();
+
+		UE_LOG(LogMovieScene, Verbose, TEXT("Initialize - MovieSceneSequence: %s, FrameResolution: %f, PlaybackFrameRate: %d, CurrentTime: %d"), *InSequence->GetName(), FrameResolution.Numerator, PlaybackFrameRate.Numerator);
+
+		// We set the play position in terms of the playback frame rate,
+		// but want evaluation ranges in the moviescene's frame resolution
+		PlayPosition.SetTimeBase(PlaybackFrameRate, FrameResolution, EvaluationType);
+
+		{
+			// Set up the default frame range from the sequence's play range
+			TRange<FFrameNumber> PlaybackRange   = MovieScene->GetPlaybackRange();
+
+			const FFrameNumber SrcStartFrame = MovieScene::DiscreteInclusiveLower(PlaybackRange);
+			const FFrameNumber SrcEndFrame   = MovieScene::DiscreteExclusiveUpper(PlaybackRange);
+
+			const FFrameNumber StartingFrame = ConvertFrameTime(SrcStartFrame, FrameResolution, PlaybackFrameRate).FloorToFrame();
+			const FFrameNumber EndingFrame   = ConvertFrameTime(SrcEndFrame,   FrameResolution, PlaybackFrameRate).FloorToFrame();
+
+			SetFrameRange(StartingFrame.Value, (EndingFrame - StartingFrame).Value);
+		}
+
+		// Reset the play position based on the user-specified start offset, or a random time
+		FFrameTime SpecifiedStartOffset = PlaybackSettings.StartTime * PlaybackFrameRate;
+
+		// Setup the starting time
+		FFrameTime StartingTimeOffset = PlaybackSettings.bRandomStartTime
+			? FFrameTime(FMath::Rand() % GetFrameDuration())
+			: FMath::Clamp<FFrameTime>(SpecifiedStartOffset, 0, GetFrameDuration()-1);
+
+		PlayPosition.Reset(StartTime + StartingTimeOffset);
+
+		ClockToUse = MovieScene->GetClockSource();
 	}
 
-	TimeCursorPosition = PlaybackSettings.bRandomStartTime ? FMath::FRand() * 0.99f * GetLength() : FMath::Clamp(PlaybackSettings.StartTime, 0.f, GetLength());
+	if (!PlaybackSettings.TimeController.IsValid())
+	{
+		switch (ClockToUse)
+		{
+		case EUpdateClockSource::Audio:    PlaybackSettings.TimeController = MakeShared<FMovieSceneTimeController_AudioClock>();    break;
+		case EUpdateClockSource::Platform: PlaybackSettings.TimeController = MakeShared<FMovieSceneTimeController_PlatformClock>(); break;
+		default:                           PlaybackSettings.TimeController = MakeShared<FMovieSceneTimeController_Tick>();          break;
+		}
+	}
 
+	PlaybackSettings.TimeController->Reset(GetCurrentTime());
 	RootTemplateInstance.Initialize(*Sequence, *this);
 
 	// Ensure everything is set up, ready for playback
@@ -347,44 +484,75 @@ void UMovieSceneSequencePlayer::Update(const float DeltaSeconds)
 {
 	if (IsPlaying())
 	{
+		// Delta seconds has already been multiplied by MatineeTimeDilation at this point, so don't pass that through to Tick
 		float PlayRate = bReversePlayback ? -PlaybackSettings.PlayRate : PlaybackSettings.PlayRate;
-		UpdateTimeCursorPosition(TimeCursorPosition + DeltaSeconds * PlayRate);
+		PlaybackSettings.TimeController->Tick(DeltaSeconds, PlayRate);
+
+		UWorld* World = GetPlaybackWorld();
+		if (World)
+		{
+			PlayRate *= World->GetWorldSettings()->MatineeTimeDilation;
+		}
+
+		FFrameTime NewTime = PlaybackSettings.TimeController->RequestCurrentTime(GetCurrentTime(), PlayRate);
+		UpdateTimeCursorPosition(NewTime, EUpdatePositionMethod::Play);
 	}
 }
 
-void UMovieSceneSequencePlayer::UpdateTimeCursorPosition(float NewPosition, TOptional<EMovieScenePlayerStatus::Type> OptionalStatus)
+void UMovieSceneSequencePlayer::UpdateTimeCursorPosition(FFrameTime NewPosition, EUpdatePositionMethod Method)
 {
-	float Length = GetLength();
+	if (bIsEvaluating)
+	{
+		LatentActions.Emplace(Method, NewPosition);
+	}
+	else
+	{
+		UpdateTimeCursorPosition_Internal(NewPosition, Method);
+	}
+}
 
-	UMovieSceneSequence* MovieSceneSequence = RootTemplateInstance.GetSequence(MovieSceneSequenceID::Root);
-	TOptional<float> FixedFrameInterval = (MovieSceneSequence && MovieSceneSequence->GetMovieScene()) ? MovieSceneSequence->GetMovieScene()->GetOptionalFixedFrameInterval() : TOptional<float>();
+void UMovieSceneSequencePlayer::UpdateTimeCursorPosition_Internal(FFrameTime NewPosition, EUpdatePositionMethod Method)
+{
+	EMovieScenePlayerStatus::Type StatusOverride = Method == EUpdatePositionMethod::Scrub ? EMovieScenePlayerStatus::Scrubbing : (EMovieScenePlayerStatus::Type)Status;
 
-	if (bPendingFirstUpdate)
+	const int32 Duration = DurationFrames;
+	if (Duration == 0)
 	{
-		NewPosition = TimeCursorPosition;
-		bPendingFirstUpdate = false;
+		UE_LOG(LogMovieScene, Warning, TEXT("Attempting to play back a sequence with zero duration"));
+		return;
 	}
 
+	FFrameTime PositionRelativeToStart = NewPosition.FrameNumber - StartTime;
+
 	if (ShouldStopOrLoop(NewPosition))
 	{
+		const int32 NumTimesLooped    = FMath::Abs(PositionRelativeToStart.FrameNumber.Value / Duration);
+		const bool  bLoopIndefinitely = PlaybackSettings.LoopCount < 0;
+
 		// loop playback
-		if (PlaybackSettings.LoopCount < 0 || CurrentNumLoops < PlaybackSettings.LoopCount)
+		if (bLoopIndefinitely || CurrentNumLoops + NumTimesLooped <= PlaybackSettings.LoopCount)
 		{
-			++CurrentNumLoops;
-			const float Overplay = FMath::Fmod(NewPosition, Length);
-			TimeCursorPosition = Overplay < 0 ? Length + Overplay : Overplay;
-			
-			PlayPosition.Reset(Overplay < 0 ? Length : 0.f);
-			FMovieSceneEvaluationRange Range = PlayPosition.PlayTo(GetSequencePosition(), FixedFrameInterval);
+			CurrentNumLoops += NumTimesLooped;
+
+			const FFrameTime Overplay       = FFrameTime(PositionRelativeToStart.FrameNumber.Value % Duration, PositionRelativeToStart.GetSubFrame());
+			const FFrameTime NewFrameOffset = (Overplay < 0) ? FFrameTime(Duration) + Overplay : Overplay;
 
 			if (SpawnRegister.IsValid())
 			{
 				SpawnRegister->ForgetExternallyOwnedSpawnedObjects(State, *this);
 			}
 
+			// Reset the play position, and generate a new range that gets us to the new frame time
+			PlayPosition.Reset(Overplay < 0 ? GetLastValidTime() : StartTime);
+
+			FMovieSceneEvaluationRange Range = PlayPosition.PlayTo(StartTime + NewFrameOffset);
+
 			const bool bHasJumped = true;
+			UpdateMovieSceneInstance(Range, StatusOverride, bHasJumped);
 
-			UpdateMovieSceneInstance(Range, OptionalStatus, bHasJumped);
+			// Use the exact time here rather than a frame locked time to ensure we don't skip the amount that was overplayed in the time controller
+			FQualifiedFrameTime ExactCurrentTime(StartTime + NewFrameOffset, PlayPosition.GetInputRate());
+			PlaybackSettings.TimeController->Reset(ExactCurrentTime);
 
 			OnLooped();
 		}
@@ -392,14 +560,25 @@ void UMovieSceneSequencePlayer::UpdateTimeCursorPosition(float NewPosition, TOpt
 		// stop playback
 		else
 		{
-			FMovieSceneEvaluationRange Range = PlayPosition.PlayTo(NewPosition + StartTime, FixedFrameInterval);
-			UpdateMovieSceneInstance(Range, OptionalStatus);
+			// Clamp the position to the duration
+			NewPosition = FMath::Clamp(NewPosition, FFrameTime(StartTime), GetLastValidTime());
 
-			Stop();
+			FMovieSceneEvaluationRange Range = PlayPosition.PlayTo(NewPosition);
+			UpdateMovieSceneInstance(Range, StatusOverride);
 
-			// When playback stops naturally, the time cursor is put at the boundary that was crossed to make ping-pong playback easy
-			TimeCursorPosition = bReversePlayback ? 0.f : GetLength();
-			PlayPosition.Reset(TimeCursorPosition);
+			if (PlaybackSettings.bPauseAtEnd)
+			{
+				Pause();
+			}
+			else
+			{
+				Stop();
+
+				// When playback stops naturally, the time cursor is put at the boundary that was crossed to make ping-pong playback easy
+				PlayPosition.Reset(NewPosition);
+			}
+
+			PlaybackSettings.TimeController->StopPlaying(GetCurrentTime());
 
 			if (OnFinished.IsBound())
 			{
@@ -410,24 +589,33 @@ void UMovieSceneSequencePlayer::UpdateTimeCursorPosition(float NewPosition, TOpt
 	else
 	{
 		// Just update the time and sequence
-		TimeCursorPosition = NewPosition;
-
-		FMovieSceneEvaluationRange Range = PlayPosition.PlayTo(NewPosition + StartTime, FixedFrameInterval);
-		UpdateMovieSceneInstance(Range, OptionalStatus);
+		FMovieSceneEvaluationRange Range = PlayPosition.PlayTo(NewPosition);
+		UpdateMovieSceneInstance(Range, StatusOverride);
 	}
 }
 
-void UMovieSceneSequencePlayer::UpdateMovieSceneInstance(FMovieSceneEvaluationRange InRange, TOptional<EMovieScenePlayerStatus::Type> OptionalStatus, bool bHasJumped)
+void UMovieSceneSequencePlayer::UpdateMovieSceneInstance(FMovieSceneEvaluationRange InRange, EMovieScenePlayerStatus::Type PlayerStatus, bool bHasJumped)
 {
+#if !NO_LOGGING
+	UMovieSceneSequence* MovieSceneSequence = RootTemplateInstance.GetSequence(MovieSceneSequenceID::Root);
+	if (MovieSceneSequence)
+	{
+		FQualifiedFrameTime CurrentTime = GetCurrentTime();
+		UE_LOG(LogMovieScene, VeryVerbose, TEXT("Evaluating sequence %s at frame %d, subframe %f (%f fps)."), *MovieSceneSequence->GetName(), CurrentTime.Time.FrameNumber.Value, CurrentTime.Time.GetSubFrame(), CurrentTime.Rate.AsDecimal());
+	}
+#endif
+
 	bIsEvaluating = true;
 
-	FMovieSceneContext Context(InRange, OptionalStatus.Get(GetPlaybackStatus()));
+	FMovieSceneContext Context(InRange, PlayerStatus);
 	Context.SetHasJumped(bHasJumped);
 
 	RootTemplateInstance.Evaluate(Context, *this);
 
 #if WITH_EDITOR
-	OnMovieSceneSequencePlayerUpdate.Broadcast(*this, Context.GetTime(), Context.GetPreviousTime());
+	FFrameTime CurrentTime  = ConvertFrameTime(Context.GetTime(),         Context.GetFrameRate(), PlayPosition.GetInputRate());
+	FFrameTime PreviousTime = ConvertFrameTime(Context.GetPreviousTime(), Context.GetFrameRate(), PlayPosition.GetInputRate());
+	OnMovieSceneSequencePlayerUpdate.Broadcast(*this, CurrentTime, PreviousTime);
 #endif
 	bIsEvaluating = false;
 
@@ -437,16 +625,23 @@ void UMovieSceneSequencePlayer::UpdateMovieSceneInstance(FMovieSceneEvaluationRa
 void UMovieSceneSequencePlayer::ApplyLatentActions()
 {
 	// Swap to a stack array to ensure no reentrancy if we evaluate during a pause, for instance
-	TArray<ELatentAction> TheseActions;
+	TArray<FLatentAction> TheseActions;
 	Swap(TheseActions, LatentActions);
 
-	for (ELatentAction LatentAction : TheseActions)
+	for (const FLatentAction& LatentAction : TheseActions)
 	{
-		switch(LatentAction)
+		switch (LatentAction.Type)
+		{
+		case FLatentAction::Stop:          Stop();                         continue;
+		case FLatentAction::Pause:         Pause();                        continue;
+		}
+
+		check(LatentAction.Type == FLatentAction::Update);
+		switch (LatentAction.UpdateMethod)
 		{
-		case ELatentAction::Stop:	Stop(); break;
-		case ELatentAction::Pause:	Pause(); break;
-		case ELatentAction::SetPlaybackPosition: SetPlaybackPosition(LatentPlaybackPosition); break;
+		case EUpdatePositionMethod::Play:  PlayToFrame( LatentAction.Position); continue;
+		case EUpdatePositionMethod::Jump:  JumpToFrame( LatentAction.Position); continue;
+		case EUpdatePositionMethod::Scrub: ScrubToFrame(LatentAction.Position); continue;
 		}
 	}
 }
@@ -475,3 +670,9 @@ void UMovieSceneSequencePlayer::BeginDestroy()
 
 	Super::BeginDestroy();
 }
+
+UWorld* UMovieSceneSequencePlayer::GetPlaybackWorld() const
+{
+	UObject* PlaybackContext = GetPlaybackContext();
+	return PlaybackContext ? PlaybackContext->GetWorld() : nullptr;
+}
\ No newline at end of file
