Copying //UE4/Dev-Sequencer to //UE4/Dev-Main (Source: //UE4/Dev-Sequencer @ 3946692)

diff --git a/Engine/Source/Runtime/MovieScene/Private/MovieSceneTrack.cpp b/Engine/Source/Runtime/MovieScene/Private/MovieSceneTrack.cpp
--- a/Engine/Source/Runtime/MovieScene/Private/MovieSceneTrack.cpp
+++ b/Engine/Source/Runtime/MovieScene/Private/MovieSceneTrack.cpp
@@ -2,6 +2,7 @@
 
 #include "MovieSceneTrack.h"
 #include "MovieScene.h"
+#include "MovieSceneTimeHelpers.h"
 #include "Evaluation/MovieSceneSegment.h"
 #include "Compilation/MovieSceneSegmentCompiler.h"
 #include "Compilation/MovieSceneCompilerRules.h"
@@ -17,6 +18,7 @@ UMovieSceneTrack::UMovieSceneTrack(const FObjectInitializer& InInitializer)
 {
 #if WITH_EDITORONLY_DATA
 	TrackTint = FColor(127, 127, 127, 0);
+	SortingOrder = -1;
 #endif
 }
 
@@ -64,11 +66,11 @@ void UMovieSceneTrack::UpdateEasing()
 
 			// Check overlaps with exclusive ranges so that sections can butt up against each other
 			UMovieSceneTrack* OuterTrack = CurrentSection->GetTypedOuter<UMovieSceneTrack>();
-			float MaxEaseIn = 0.f;
-			float MaxEaseOut = 0.f;
+			int32 MaxEaseIn = 0;
+			int32 MaxEaseOut = 0;
 			bool bIsEntirelyUnderlapped = false;
 
-			TRange<float> CurrentSectionRange = CurrentSection->GetRange();
+			TRange<FFrameNumber> CurrentSectionRange = CurrentSection->GetRange();
 			for (int32 OtherIndex = 0; OtherIndex < RowSections.Num(); ++OtherIndex)
 			{
 				if (OtherIndex == Index)
@@ -77,7 +79,7 @@ void UMovieSceneTrack::UpdateEasing()
 				}
 
 				UMovieSceneSection* Other = RowSections[OtherIndex];
-				TRange<float> OtherSectionRange = Other->GetRange();
+				TRange<FFrameNumber> OtherSectionRange = Other->GetRange();
 
 				bIsEntirelyUnderlapped = OtherSectionRange.Contains(CurrentSectionRange);
 
@@ -86,33 +88,35 @@ void UMovieSceneTrack::UpdateEasing()
 				const bool bSectionRangeContainsOtherLowerBound = !OtherSectionRange.GetLowerBound().IsOpen() && !CurrentSectionRange.GetUpperBound().IsOpen() && CurrentSectionRange.Contains(OtherSectionRange.GetLowerBoundValue());
 				if (bSectionRangeContainsOtherUpperBound && !bSectionRangeContainsOtherLowerBound)
 				{
-					MaxEaseIn = FMath::Max(MaxEaseIn, OtherSectionRange.GetUpperBoundValue() - CurrentSectionRange.GetLowerBoundValue());
+					const int32 Difference = MovieScene::DiscreteSize(TRange<FFrameNumber>(CurrentSectionRange.GetLowerBound(), OtherSectionRange.GetUpperBound()));
+					MaxEaseIn = FMath::Max(MaxEaseIn, Difference);
 				}
 
 				if (bSectionRangeContainsOtherLowerBound &&!bSectionRangeContainsOtherUpperBound)
 				{
-					MaxEaseOut = FMath::Max(MaxEaseOut, CurrentSectionRange.GetUpperBoundValue() - OtherSectionRange.GetLowerBoundValue());
+					const int32 Difference = MovieScene::DiscreteSize(TRange<FFrameNumber>(OtherSectionRange.GetLowerBoundValue(), CurrentSectionRange.GetUpperBoundValue()));
+					MaxEaseOut = FMath::Max(MaxEaseOut, Difference);
 				}
 			}
 
-			const bool bIsFinite = CurrentSectionRange.HasLowerBound() && CurrentSectionRange.HasUpperBound();
-			const float Max = bIsFinite ? CurrentSectionRange.Size<float>() : TNumericLimits<float>::Max();
+			const bool  bIsFinite = CurrentSectionRange.HasLowerBound() && CurrentSectionRange.HasUpperBound();
+			const int32 MaxSize   = bIsFinite ? MovieScene::DiscreteSize(CurrentSectionRange) : TNumericLimits<int32>::Max();
 
-			if (MaxEaseOut == 0.f && MaxEaseIn == 0.f && bIsEntirelyUnderlapped)
+			if (MaxEaseOut == 0 && MaxEaseIn == 0 && bIsEntirelyUnderlapped)
 			{
-				MaxEaseOut = MaxEaseIn = Max * 0.25f;
+				MaxEaseOut = MaxEaseIn = MaxSize / 4;
 			}
 
 			// Only modify the section if the ease in or out times have actually changed
-			MaxEaseIn = FMath::Clamp(MaxEaseIn, 0.f, Max);
-			MaxEaseOut = FMath::Clamp(MaxEaseOut, 0.f, Max);
+			MaxEaseIn  = FMath::Clamp(MaxEaseIn, 0, MaxSize);
+			MaxEaseOut = FMath::Clamp(MaxEaseOut, 0, MaxSize);
 
-			if (CurrentSection->Easing.AutoEaseInTime != MaxEaseIn || CurrentSection->Easing.AutoEaseOutTime != MaxEaseOut)
+			if (CurrentSection->Easing.AutoEaseInDuration != MaxEaseIn || CurrentSection->Easing.AutoEaseOutDuration != MaxEaseOut)
 			{
 				CurrentSection->Modify();
 
-				CurrentSection->Easing.AutoEaseInTime = MaxEaseIn;
-				CurrentSection->Easing.AutoEaseOutTime = MaxEaseOut;
+				CurrentSection->Easing.AutoEaseInDuration  = MaxEaseIn;
+				CurrentSection->Easing.AutoEaseOutDuration = MaxEaseOut;
 			}
 		}
 	}
@@ -199,7 +203,8 @@ EMovieSceneCompileResult UMovieSceneTrack::Compile(FMovieSceneEvaluationTrack& O
 	{
 		for (const UMovieSceneSection* Section : GetAllSections())
 		{
-			if (!Section->IsActive())
+			const TRange<FFrameNumber> SectionRange = Section->GetRange();
+			if (!Section->IsActive() || SectionRange.IsEmpty())
 			{
 				continue;
 			}
@@ -211,19 +216,17 @@ EMovieSceneCompileResult UMovieSceneTrack::Compile(FMovieSceneEvaluationTrack& O
 				NewTemplate->SetSourceSection(Section);
 
 				int32 TemplateIndex = OutTrack.AddChildTemplate(MoveTemp(NewTemplate));
-
-				const TRange<float> SectionRange = Section->IsInfinite() ? TRange<float>::All() : Section->GetRange();
 				OutTrack.AddTreeData(SectionRange, FSectionEvaluationData(TemplateIndex, ESectionEvaluationFlags::None));
 
-				if (!SectionRange.GetLowerBound().IsOpen() && Section->GetPreRollTime() > 0)
+				if (!SectionRange.GetLowerBound().IsOpen() && Section->GetPreRollFrames() > 0)
 				{
-					TRange<float> PreRollRange(SectionRange.GetLowerBoundValue() - Section->GetPreRollTime(), TRangeBound<float>::FlipInclusion(SectionRange.GetLowerBoundValue()));
+					TRange<FFrameNumber> PreRollRange = MovieScene::MakeDiscreteRangeFromUpper(TRangeBound<FFrameNumber>::FlipInclusion(SectionRange.GetLowerBoundValue()), Section->GetPreRollFrames());
 					OutTrack.AddTreeData(PreRollRange, FSectionEvaluationData(TemplateIndex, ESectionEvaluationFlags::PreRoll));
 				}
 
-				if (!SectionRange.GetUpperBound().IsOpen() && Section->GetPostRollTime() > 0)
+				if (!SectionRange.GetUpperBound().IsOpen() && Section->GetPostRollFrames() > 0)
 				{
-					TRange<float> PostRollRange(TRangeBound<float>::FlipInclusion(SectionRange.GetUpperBoundValue()), SectionRange.GetUpperBoundValue() + Section->GetPostRollTime());
+					TRange<FFrameNumber> PostRollRange = MovieScene::MakeDiscreteRangeFromLower(TRangeBound<FFrameNumber>::FlipInclusion(SectionRange.GetUpperBoundValue()), Section->GetPostRollFrames());
 					OutTrack.AddTreeData(PostRollRange, FSectionEvaluationData(TemplateIndex, ESectionEvaluationFlags::PostRoll));
 				}
 			}
diff --git a/Engine/Source/Runtime/MovieScene/Private/Sections/MovieSceneSubSection.cpp b/Engine/Source/Runtime/MovieScene/Private/Sections/MovieSceneSubSection.cpp
--- a/Engine/Source/Runtime/MovieScene/Private/Sections/MovieSceneSubSection.cpp
+++ b/Engine/Source/Runtime/MovieScene/Private/Sections/MovieSceneSubSection.cpp
@@ -4,6 +4,7 @@
 #include "MovieScene.h"
 #include "MovieSceneTrack.h"
 #include "MovieSceneSequence.h"
+#include "MovieSceneTimeHelpers.h"
 #include "Evaluation/MovieSceneEvaluationTemplate.h"
 
 TWeakObjectPtr<UMovieSceneSubSection> UMovieSceneSubSection::TheRecordingSection;
@@ -20,6 +21,34 @@ UMovieSceneSubSection::UMovieSceneSubSection()
 {
 }
 
+FMovieSceneSequenceTransform UMovieSceneSubSection::OuterToInnerTransform() const
+{
+	UMovieSceneSequence* SequencePtr   = GetSequence();
+	UMovieScene*         MovieScenePtr = SequencePtr->GetMovieScene();
+
+	TRange<FFrameNumber> SubRange = GetRange();
+	if (!MovieScenePtr || SubRange.GetLowerBound().IsOpen())
+	{
+		return FMovieSceneSequenceTransform();
+	}
+
+	const FFrameNumber InnerStartTime = MovieScene::DiscreteInclusiveLower(MovieScenePtr->GetPlaybackRange()) + Parameters.GetStartFrameOffset();
+	const FFrameNumber OuterStartTime = MovieScene::DiscreteInclusiveLower(SubRange);
+
+	const FFrameRate   InnerFrameRate = MovieScenePtr->GetFrameResolution();
+	const FFrameRate   OuterFrameRate = GetTypedOuter<UMovieScene>()->GetFrameResolution();
+
+	const float        FrameRateScale = (OuterFrameRate == InnerFrameRate) ? 1.f : (InnerFrameRate / OuterFrameRate).AsDecimal();
+
+	return
+		// Inner play offset
+		FMovieSceneSequenceTransform(InnerStartTime)
+		// Inner play rate
+		* FMovieSceneSequenceTransform(0, Parameters.TimeScale * FrameRateScale)
+		// Outer section start time
+		* FMovieSceneSequenceTransform(-OuterStartTime);
+}
+
 FString UMovieSceneSubSection::GetPathNameInMovieScene() const
 {
 	UMovieScene* OuterMovieScene = GetTypedOuter<UMovieScene>();
@@ -41,9 +70,22 @@ FMovieSceneSequenceID UMovieSceneSubSection::GetSequenceID() const
 
 void UMovieSceneSubSection::PostLoad()
 {
+	FFrameRate LegacyFrameRate = GetLegacyConversionFrameRate();
+
+	TOptional<double> StartOffsetToUpgrade;
 	if (StartOffset_DEPRECATED != DeprecatedMagicNumber)
 	{
-		Parameters.StartOffset = StartOffset_DEPRECATED;
+		StartOffsetToUpgrade = StartOffset_DEPRECATED;
+	}
+	else if (Parameters.StartOffset_DEPRECATED != 0.f)
+	{
+		StartOffsetToUpgrade = Parameters.StartOffset_DEPRECATED;
+	}
+
+	if (StartOffsetToUpgrade.IsSet())
+	{
+		FFrameNumber StartFrame = UpgradeLegacyMovieSceneTime(this, LegacyFrameRate, StartOffsetToUpgrade.GetValue());
+		Parameters.SetStartFrameOffset(StartFrame.Value);
 	}
 
 	if (TimeScale_DEPRECATED != DeprecatedMagicNumber)
@@ -59,12 +101,14 @@ void UMovieSceneSubSection::PostLoad()
 	// Pre and post roll is now supported generically
 	if (Parameters.PrerollTime_DEPRECATED > 0.f)
 	{
-		SetPreRollTime(Parameters.PrerollTime_DEPRECATED);
+		FFrameNumber ClampedPreRollFrames = UpgradeLegacyMovieSceneTime(this, LegacyFrameRate, Parameters.PrerollTime_DEPRECATED);
+		SetPreRollFrames(ClampedPreRollFrames.Value);
 	}
 
 	if (Parameters.PostrollTime_DEPRECATED > 0.f)
 	{
-		SetPostRollTime(Parameters.PostrollTime_DEPRECATED);
+		FFrameNumber ClampedPostRollFrames = UpgradeLegacyMovieSceneTime(this, LegacyFrameRate, Parameters.PostrollTime_DEPRECATED);
+		SetPreRollFrames(ClampedPostRollFrames.Value);
 	}
 
 	Super::PostLoad();
@@ -152,51 +196,59 @@ void UMovieSceneSubSection::PostEditChangeProperty(FPropertyChangedEvent& Proper
 }
 #endif
 
-UMovieSceneSection* UMovieSceneSubSection::SplitSection( float SplitTime )
+UMovieSceneSection* UMovieSceneSubSection::SplitSection( FFrameNumber SplitTime )
 {
-	if ( !IsTimeWithinSection( SplitTime ) )
+	TRange<FFrameNumber> InitialRange = GetRange();
+	if ( !InitialRange.Contains( SplitTime ) )
 	{
 		return nullptr;
 	}
 
-	float InitialStartTime = GetStartTime();
-	float InitialStartOffset = Parameters.StartOffset;
-	float NewStartOffset = ( SplitTime - InitialStartTime ) / Parameters.TimeScale;
-	NewStartOffset += InitialStartOffset;
-
-	// Ensure start offset is not less than 0
-	NewStartOffset = FMath::Max( NewStartOffset, 0.f );
+	FFrameNumber InitialStartOffset = Parameters.GetStartFrameOffset();
 
 	UMovieSceneSubSection* NewSection = Cast<UMovieSceneSubSection>( UMovieSceneSection::SplitSection( SplitTime ) );
 	if ( NewSection )
 	{
-		NewSection->Parameters.StartOffset = NewStartOffset;
+		if (InitialRange.GetLowerBound().IsClosed())
+		{
+			FFrameNumber NewStartOffset = ( SplitTime - InitialRange.GetLowerBoundValue() ) / Parameters.TimeScale;
+			NewStartOffset += InitialStartOffset;
+
+			if (NewStartOffset >= 0)
+			{
+				NewSection->Parameters.SetStartFrameOffset(NewStartOffset.Value);
+			}
+		}
+
 		return NewSection;
 	}
 
 	return nullptr;
 }
 
-void UMovieSceneSubSection::TrimSection( float TrimTime, bool bTrimLeft )
+void UMovieSceneSubSection::TrimSection( FFrameNumber TrimTime, bool bTrimLeft )
 {
-	if ( !IsTimeWithinSection( TrimTime ) )
+	TRange<FFrameNumber> InitialRange = GetRange();
+	if ( !InitialRange.Contains( TrimTime ) )
 	{
 		return;
 	}
 
-	float InitialStartTime = GetStartTime();
-	float InitialStartOffset = Parameters.StartOffset;
+	FFrameNumber InitialStartOffset = Parameters.GetStartFrameOffset();
 
 	UMovieSceneSection::TrimSection( TrimTime, bTrimLeft );
 
 	// If trimming off the left, set the offset of the shot
-	if ( bTrimLeft )
+	if ( bTrimLeft && InitialRange.GetLowerBound().IsClosed() )
 	{
-		float NewStartOffset = ( TrimTime - InitialStartTime ) / Parameters.TimeScale;
+		FFrameNumber NewStartOffset = ( TrimTime - InitialRange.GetLowerBoundValue() ) / Parameters.TimeScale;
 		NewStartOffset += InitialStartOffset;
 
 		// Ensure start offset is not less than 0
-		Parameters.StartOffset = FMath::Max( NewStartOffset, 0.f );
+		if (NewStartOffset >= 0)
+		{
+			Parameters.SetStartFrameOffset(NewStartOffset.Value);
+		}
 	}
 }
 
diff --git a/Engine/Source/Runtime/MovieScene/Private/Tests/MovieSceneBlendingTests.cpp b/Engine/Source/Runtime/MovieScene/Private/Tests/MovieSceneBlendingTests.cpp
--- a/Engine/Source/Runtime/MovieScene/Private/Tests/MovieSceneBlendingTests.cpp
+++ b/Engine/Source/Runtime/MovieScene/Private/Tests/MovieSceneBlendingTests.cpp
@@ -53,8 +53,8 @@ bool FMovieSceneBlendingTest::RunTest(const FString& Parameters)
 {
 	const FMovieSceneBlendingActuatorID ID = FInt32Actuator::GetID();
 	const FMovieSceneEvaluationScope Scope(FMovieSceneEvaluationKey(), EMovieSceneCompletionMode::KeepState);
-	const FMovieSceneContext Context(FMovieSceneEvaluationRange(0.f));
-	
+	const FMovieSceneContext Context(FMovieSceneEvaluationRange(0, FFrameRate()));
+
 	FNullPlayer Player;
 	FPersistentEvaluationData PersistentDataProxy(Player);
 
@@ -143,7 +143,7 @@ bool FMovieSceneBlendingStressTest::RunTest(const FString& Parameters)
 
 	const FMovieSceneBlendingActuatorID ID = FInt32Actuator::GetID();
 	const FMovieSceneEvaluationScope Scope(FMovieSceneEvaluationKey(), EMovieSceneCompletionMode::KeepState);
-	const FMovieSceneContext Context(FMovieSceneEvaluationRange(0.f));
+	const FMovieSceneContext Context(FMovieSceneEvaluationRange(0, FFrameRate()));
 	
 	FNullPlayer Player;
 	FPersistentEvaluationData PersistentDataProxy(Player);
diff --git a/Engine/Source/Runtime/MovieScene/Private/Tests/MovieSceneSegmentCompilerTests.cpp b/Engine/Source/Runtime/MovieScene/Private/Tests/MovieSceneSegmentCompilerTests.cpp
--- a/Engine/Source/Runtime/MovieScene/Private/Tests/MovieSceneSegmentCompilerTests.cpp
+++ b/Engine/Source/Runtime/MovieScene/Private/Tests/MovieSceneSegmentCompilerTests.cpp
@@ -15,9 +15,9 @@
 
 namespace Impl
 {
-	static const TRangeBound<float> Inf = TRangeBound<float>::Open();
+	static const TRangeBound<FFrameNumber> Inf = TRangeBound<FFrameNumber>::Open();
 
-	typedef TTuple<TRange<float>, TArray<int32>> FEvaluationTreeIteratorResult;
+	typedef TTuple<TRange<FFrameNumber>, TArray<int32>> FEvaluationTreeIteratorResult;
 
 	/** Compiler rules to sort by priority */
 	struct FSortByPrioritySegmentBlender : FMovieSceneTrackSegmentBlender
@@ -27,7 +27,7 @@ namespace Impl
 			bAllowEmptySegments = bInAllowEmptySegments;
 		}
 
-		virtual TOptional<FMovieSceneSegment> InsertEmptySpace(const TRange<float>& Range, const FMovieSceneSegment* PreviousSegment, const FMovieSceneSegment* NextSegment) const override
+		virtual TOptional<FMovieSceneSegment> InsertEmptySpace(const TRange<FFrameNumber>& Range, const FMovieSceneSegment* PreviousSegment, const FMovieSceneSegment* NextSegment) const override
 		{
 			return bAllowEmptySegments ? FMovieSceneSegment(Range) : TOptional<FMovieSceneSegment>();
 		}
@@ -63,7 +63,7 @@ namespace Impl
 				Result += TEXT(", ");
 			}
 
-			Result += FString::Printf(TEXT("(Impl: %d, ForcedTime: %.7f, Flags: %u)"), Thing.ImplIndex, Thing.ForcedTime, (uint8)Thing.Flags);
+			Result += FString::Printf(TEXT("(Impl: %d, ForcedTime: %i, Flags: %u)"), Thing.ImplIndex, Thing.ForcedTime.Value, (uint8)Thing.Flags);
 		}
 
 		return Result;
@@ -111,16 +111,16 @@ namespace Impl
 		}
 	}
 
-	FORCEINLINE_DEBUGGABLE void AssertSegmentAtTime(FAutomationTestBase* Test, FMovieSceneEvaluationTrack& InTrack, float InTime, TArrayView<const FSectionEvaluationData> ExpectedImpls)
+	FORCEINLINE_DEBUGGABLE void AssertSegmentAtTime(FAutomationTestBase* Test, FMovieSceneEvaluationTrack& InTrack, FFrameNumber InTime, TArrayView<const FSectionEvaluationData> ExpectedImpls)
 	{
 		FMovieSceneSegmentIdentifier ID = InTrack.GetSegmentFromTime(InTime);
 		if (!ID.IsValid())
 		{
-			Test->AddError(TEXT("No segment compiled for time %.fs"), InTime);
+			Test->AddError(TEXT("No segment compiled for frame %i"), InTime.Value);
 		}
 		else
 		{
-			AssertSegmentImpls(Test, ExpectedImpls, InTrack.GetSegment(ID).Impls, *FString::Printf(TEXT("time %.8fs"), InTime));
+			AssertSegmentImpls(Test, ExpectedImpls, InTrack.GetSegment(ID).Impls, *FString::Printf(TEXT("frame %i"), InTime.Value));
 		}
 	}
 
@@ -182,22 +182,22 @@ bool FMovieSceneCompilerBasicTest::RunTest(const FString& Parameters)
 
 	FMovieSceneEvaluationTrack Track;
 
-	Track.AddTreeData(TRange<float>(10.f,	20.f),										FSectionEvaluationData(0));
-	Track.AddTreeData(TRange<float>(20.f,	30.f),										FSectionEvaluationData(1));
-	Track.AddTreeData(TRange<float>(10.f,	30.f),										FSectionEvaluationData(2));
-	Track.AddTreeData(TRange<float>(Inf,	TRangeBound<float>::Exclusive(20.f)),		FSectionEvaluationData(3));
-	Track.AddTreeData(TRange<float>(25.f,	Inf),										FSectionEvaluationData(4));
+	Track.AddTreeData(TRange<FFrameNumber>(10,	20),										FSectionEvaluationData(0));
+	Track.AddTreeData(TRange<FFrameNumber>(20,	30),										FSectionEvaluationData(1));
+	Track.AddTreeData(TRange<FFrameNumber>(10,	30),										FSectionEvaluationData(2));
+	Track.AddTreeData(TRange<FFrameNumber>(Inf, TRangeBound<FFrameNumber>::Exclusive(20)),	FSectionEvaluationData(3));
+	Track.AddTreeData(TRange<FFrameNumber>(TRangeBound<FFrameNumber>::Inclusive(25), Inf),	FSectionEvaluationData(4));
 
 	FMovieSceneSegment Expected[] = {
-		FMovieSceneSegment(FFloatRange(Inf, 	TRangeBound<float>::Exclusive(10.f)),	{ FSectionEvaluationData(3) }),
-		FMovieSceneSegment(FFloatRange(10.f,	20.f),									{ FSectionEvaluationData(0), FSectionEvaluationData(2), FSectionEvaluationData(3) }),
-		FMovieSceneSegment(FFloatRange(20.f,	25.f),									{ FSectionEvaluationData(1), FSectionEvaluationData(2) }),
-		FMovieSceneSegment(FFloatRange(25.f,	30.f),									{ FSectionEvaluationData(1), FSectionEvaluationData(2), FSectionEvaluationData(4) }),
-		FMovieSceneSegment(FFloatRange(30.f,	Inf),									{ FSectionEvaluationData(4) })
+		FMovieSceneSegment(TRange<FFrameNumber>(Inf, TRangeBound<FFrameNumber>::Exclusive(10)),	{ FSectionEvaluationData(3) }),
+		FMovieSceneSegment(TRange<FFrameNumber>(10,	20),										{ FSectionEvaluationData(0), FSectionEvaluationData(2), FSectionEvaluationData(3) }),
+		FMovieSceneSegment(TRange<FFrameNumber>(20,	25),										{ FSectionEvaluationData(1), FSectionEvaluationData(2) }),
+		FMovieSceneSegment(TRange<FFrameNumber>(25,	30),										{ FSectionEvaluationData(1), FSectionEvaluationData(2), FSectionEvaluationData(4) }),
+		FMovieSceneSegment(TRange<FFrameNumber>(TRangeBound<FFrameNumber>::Inclusive(30),Inf),	{ FSectionEvaluationData(4) })
 	};
 
 	// Compile all the segments
-	Track.GetSegmentsInRange(TRange<float>::All());
+	Track.GetSegmentsInRange(TRange<FFrameNumber>::All());
 
 	AssertSegmentValues(this, Expected, Track.GetSortedSegments());
 	return true;
@@ -218,21 +218,21 @@ bool FMovieSceneCompilerEmptySpaceTest::RunTest(const FString& Parameters)
 
 	FMovieSceneEvaluationTrack Track;
 
-	Track.AddTreeData(TRange<float>(10.f,	20.f), 		FSectionEvaluationData(0));
-	Track.AddTreeData(TRange<float>(30.f,	40.f), 		FSectionEvaluationData(1));
+	Track.AddTreeData(TRange<FFrameNumber>(10,	20), FSectionEvaluationData(0));
+	Track.AddTreeData(TRange<FFrameNumber>(30,	40), FSectionEvaluationData(1));
 
 	// Override the blender (so it doesn't try and get it from the null track ptr)
 	FScopedOverrideTrackSegmentBlender ScopedBlenderOverride(FSortByPrioritySegmentBlender(true));
 
 	// Compile all the segments
-	Track.GetSegmentsInRange(TRange<float>::All());
+	Track.GetSegmentsInRange(TRange<FFrameNumber>::All());
 
 	FMovieSceneSegment Expected[] = {
-		FMovieSceneSegment(FFloatRange(Inf, 	TRangeBound<float>::Exclusive(10.f)), 	{ }),
-		FMovieSceneSegment(FFloatRange(10.f,	20.f), 									{ FSectionEvaluationData(0) }),
-		FMovieSceneSegment(FFloatRange(20.f,	30.f), 									{  }),
-		FMovieSceneSegment(FFloatRange(30.f,	40.f), 									{ FSectionEvaluationData(1) }),
-		FMovieSceneSegment(FFloatRange(40.f,	Inf), 									{  })
+		FMovieSceneSegment(TRange<FFrameNumber>(Inf, TRangeBound<FFrameNumber>::Exclusive(10)), 	{ }),
+		FMovieSceneSegment(TRange<FFrameNumber>(10,	20), 											{ FSectionEvaluationData(0) }),
+		FMovieSceneSegment(TRange<FFrameNumber>(20,	30), 											{  }),
+		FMovieSceneSegment(TRange<FFrameNumber>(30,	40), 											{ FSectionEvaluationData(1) }),
+		FMovieSceneSegment(TRange<FFrameNumber>(TRangeBound<FFrameNumber>::Inclusive(40),	Inf), 	{  })
 	};
 	AssertSegmentValues(this, Expected, Track.GetSortedSegments());
 	return true;
@@ -254,19 +254,19 @@ bool FMovieSceneCustomCompilerTest::RunTest(const FString& Parameters)
 
 	FMovieSceneEvaluationTrack Track;
 
-	Track.AddUniqueTreeData(TRange<float>(TRangeBound<float>::Inclusive(10.f),	TRangeBound<float>::Exclusive(20.f)), FSectionEvaluationData(0, ESectionEvaluationFlags::PreRoll));
-	Track.AddUniqueTreeData(TRange<float>(TRangeBound<float>::Inclusive(15.f),	TRangeBound<float>::Exclusive(25.f)), FSectionEvaluationData(0));
-	Track.AddUniqueTreeData(TRange<float>(TRangeBound<float>::Inclusive(20.f),	TRangeBound<float>::Exclusive(30.f)), FSectionEvaluationData(0));
-	Track.AddUniqueTreeData(TRange<float>(TRangeBound<float>::Inclusive(30.f),	TRangeBound<float>::Exclusive(40.f)), FSectionEvaluationData(0));
+	Track.AddUniqueTreeData(TRange<FFrameNumber>(TRangeBound<FFrameNumber>::Inclusive(10),	TRangeBound<FFrameNumber>::Exclusive(20)), FSectionEvaluationData(0, ESectionEvaluationFlags::PreRoll));
+	Track.AddUniqueTreeData(TRange<FFrameNumber>(TRangeBound<FFrameNumber>::Inclusive(15),	TRangeBound<FFrameNumber>::Exclusive(25)), FSectionEvaluationData(0));
+	Track.AddUniqueTreeData(TRange<FFrameNumber>(TRangeBound<FFrameNumber>::Inclusive(20),	TRangeBound<FFrameNumber>::Exclusive(30)), FSectionEvaluationData(0));
+	Track.AddUniqueTreeData(TRange<FFrameNumber>(TRangeBound<FFrameNumber>::Inclusive(30),	TRangeBound<FFrameNumber>::Exclusive(40)), FSectionEvaluationData(0));
 
 	// Override the blender (so it doesn't try and get it from the null track ptr)
 	FScopedOverrideTrackSegmentBlender ScopedBlenderOverride(FSortByPrioritySegmentBlender(false));
 
-	AssertSegmentAtTime(this, Track, 12.5f,		{ FSectionEvaluationData(0, ESectionEvaluationFlags::PreRoll) });
-	AssertSegmentAtTime(this, Track, 17.5f,		{ FSectionEvaluationData(0, ESectionEvaluationFlags::PreRoll), FSectionEvaluationData(0) });
-	AssertSegmentAtTime(this, Track, 22.5f,		{ FSectionEvaluationData(0) });
-	AssertSegmentAtTime(this, Track, 27.5f,		{ FSectionEvaluationData(0) });
-	AssertSegmentAtTime(this, Track, 30.f,		{ FSectionEvaluationData(0) });
+	AssertSegmentAtTime(this, Track, 12,		{ FSectionEvaluationData(0, ESectionEvaluationFlags::PreRoll) });
+	AssertSegmentAtTime(this, Track, 17,		{ FSectionEvaluationData(0, ESectionEvaluationFlags::PreRoll), FSectionEvaluationData(0) });
+	AssertSegmentAtTime(this, Track, 22,		{ FSectionEvaluationData(0) });
+	AssertSegmentAtTime(this, Track, 27,		{ FSectionEvaluationData(0) });
+	AssertSegmentAtTime(this, Track, 30,		{ FSectionEvaluationData(0) });
 
 	return true;
 }
@@ -283,7 +283,7 @@ bool FMovieSceneTrackCompilerTest::RunTest(const FString& Parameters)
 	//																						[=============== 1 =================]
 	//	------------------------------------------------------------------------------------------------------------------------------------------------------
 	//	Additive Camera Rules Expected		[ 					|			   0			|		(0,1)	|		 1			|							]
-	//	Nearest Section Expected			[ 		0 (10.f)	|			   0			|		(0,1)	|		 1			|			1 (30.f)		]
+	//	Nearest Section Expected			[ 		0 (10)		|			   0			|		(0,1)	|		 1			|			1 (30)			]
 	//	No Nearest Section Expected			[ 					|			   0			|		(0,1)	|		 1			|							]
 	//	High-pass Filter Expected			[ 					|						0					|		 1			|							]
 
@@ -292,13 +292,17 @@ bool FMovieSceneTrackCompilerTest::RunTest(const FString& Parameters)
 		Track->EvalOptions.bCanEvaluateNearestSection = true;
 
 		UMovieSceneSegmentCompilerTestSection* Section0 = NewObject<UMovieSceneSegmentCompilerTestSection>(Track);
-		Section0->SetStartTime(10.f);
-		Section0->SetEndTime(25.f);
+		Section0->SetRange(TRange<FFrameNumber>(
+			TRangeBound<FFrameNumber>::Inclusive(10),
+			TRangeBound<FFrameNumber>::Exclusive(25)
+		));
 		Section0->SetRowIndex(0);
 
 		UMovieSceneSegmentCompilerTestSection* Section1 = NewObject<UMovieSceneSegmentCompilerTestSection>(Track);
-		Section1->SetStartTime(20.f);
-		Section1->SetEndTime(30.f);
+		Section1->SetRange(TRange<FFrameNumber>(
+			TRangeBound<FFrameNumber>::Inclusive(20),
+			TRangeBound<FFrameNumber>::Inclusive(30)
+		));
 		Section1->SetRowIndex(1);
 		
 		Track->SectionArray.Add(Section0);
@@ -309,12 +313,12 @@ bool FMovieSceneTrackCompilerTest::RunTest(const FString& Parameters)
 			FScopedOverrideTrackSegmentBlender ScopedBlenderOverride{FMovieSceneAdditiveCameraTrackBlender()};
 
 			FMovieSceneEvaluationTrack EvalTrack = Track->GenerateTrackTemplate();
-			EvalTrack.GetSegmentsInRange(TRange<float>::All());
+			EvalTrack.GetSegmentsInRange(TRange<FFrameNumber>::All());
 
 			FMovieSceneSegment Expected[] = {
-				FMovieSceneSegment(FFloatRange(TRangeBound<float>::Inclusive(10.f),	TRangeBound<float>::Exclusive(20.f)), { FSectionEvaluationData(0) }),
-				FMovieSceneSegment(FFloatRange(TRangeBound<float>::Inclusive(20.f),	TRangeBound<float>::Inclusive(25.f)), { FSectionEvaluationData(0), FSectionEvaluationData(1) }),
-				FMovieSceneSegment(FFloatRange(TRangeBound<float>::Exclusive(25.f),	TRangeBound<float>::Inclusive(30.f)), { FSectionEvaluationData(1) }),
+				FMovieSceneSegment(TRange<FFrameNumber>(TRangeBound<FFrameNumber>::Inclusive(10),	TRangeBound<FFrameNumber>::Exclusive(20)), { FSectionEvaluationData(0) }),
+				FMovieSceneSegment(TRange<FFrameNumber>(TRangeBound<FFrameNumber>::Inclusive(20),	TRangeBound<FFrameNumber>::Exclusive(25)), { FSectionEvaluationData(0), FSectionEvaluationData(1) }),
+				FMovieSceneSegment(TRange<FFrameNumber>(TRangeBound<FFrameNumber>::Inclusive(25),	TRangeBound<FFrameNumber>::Inclusive(30)), { FSectionEvaluationData(1) }),
 			};
 			AssertSegmentValues(this, Expected, EvalTrack.GetSortedSegments());
 		}
@@ -324,14 +328,14 @@ bool FMovieSceneTrackCompilerTest::RunTest(const FString& Parameters)
 			Track->EvalOptions.bEvalNearestSection = true;
 
 			FMovieSceneEvaluationTrack EvalTrack = Track->GenerateTrackTemplate();
-			EvalTrack.GetSegmentsInRange(TRange<float>::All());
+			EvalTrack.GetSegmentsInRange(TRange<FFrameNumber>::All());
 
 			FMovieSceneSegment Expected[] = {
-				FMovieSceneSegment(FFloatRange(Inf,									TRangeBound<float>::Exclusive(10.f)),	{ FSectionEvaluationData(0, 10.f) }),
-				FMovieSceneSegment(FFloatRange(TRangeBound<float>::Inclusive(10.f),	TRangeBound<float>::Exclusive(20.f)),	{ FSectionEvaluationData(0) }),
-				FMovieSceneSegment(FFloatRange(TRangeBound<float>::Inclusive(20.f),	TRangeBound<float>::Inclusive(25.f)),	{ FSectionEvaluationData(0), FSectionEvaluationData(1) }),
-				FMovieSceneSegment(FFloatRange(TRangeBound<float>::Exclusive(25.f),	TRangeBound<float>::Inclusive(30.f)),	{ FSectionEvaluationData(1) }),
-				FMovieSceneSegment(FFloatRange(TRangeBound<float>::Exclusive(30.f),	Inf), 									{ FSectionEvaluationData(1, 30.f) }),
+				FMovieSceneSegment(TRange<FFrameNumber>(Inf,										TRangeBound<FFrameNumber>::Exclusive(10)),	{ FSectionEvaluationData(0, 10) }),
+				FMovieSceneSegment(TRange<FFrameNumber>(TRangeBound<FFrameNumber>::Inclusive(10),	TRangeBound<FFrameNumber>::Exclusive(20)),	{ FSectionEvaluationData(0) }),
+				FMovieSceneSegment(TRange<FFrameNumber>(TRangeBound<FFrameNumber>::Inclusive(20),	TRangeBound<FFrameNumber>::Exclusive(25)),	{ FSectionEvaluationData(0), FSectionEvaluationData(1) }),
+				FMovieSceneSegment(TRange<FFrameNumber>(TRangeBound<FFrameNumber>::Inclusive(25),	TRangeBound<FFrameNumber>::Inclusive(30)),	{ FSectionEvaluationData(1) }),
+				FMovieSceneSegment(TRange<FFrameNumber>(TRangeBound<FFrameNumber>::Exclusive(30),	Inf), 										{ FSectionEvaluationData(1, 30) }),
 			};
 			AssertSegmentValues(this, Expected, EvalTrack.GetSortedSegments());
 		}
@@ -341,12 +345,12 @@ bool FMovieSceneTrackCompilerTest::RunTest(const FString& Parameters)
 			Track->EvalOptions.bEvalNearestSection = false;
 
 			FMovieSceneEvaluationTrack EvalTrack = Track->GenerateTrackTemplate();
-			EvalTrack.GetSegmentsInRange(TRange<float>::All());
+			EvalTrack.GetSegmentsInRange(TRange<FFrameNumber>::All());
 
 			FMovieSceneSegment Expected[] = {
-				FMovieSceneSegment(FFloatRange(TRangeBound<float>::Inclusive(10.f),	TRangeBound<float>::Exclusive(20.f)), { FSectionEvaluationData(0) }),
-				FMovieSceneSegment(FFloatRange(TRangeBound<float>::Inclusive(20.f),	TRangeBound<float>::Inclusive(25.f)), { FSectionEvaluationData(0), FSectionEvaluationData(1) }),
-				FMovieSceneSegment(FFloatRange(TRangeBound<float>::Exclusive(25.f),	TRangeBound<float>::Inclusive(30.f)), { FSectionEvaluationData(1) }),
+				FMovieSceneSegment(TRange<FFrameNumber>(TRangeBound<FFrameNumber>::Inclusive(10),	TRangeBound<FFrameNumber>::Exclusive(20)), { FSectionEvaluationData(0) }),
+				FMovieSceneSegment(TRange<FFrameNumber>(TRangeBound<FFrameNumber>::Inclusive(20),	TRangeBound<FFrameNumber>::Exclusive(25)), { FSectionEvaluationData(0), FSectionEvaluationData(1) }),
+				FMovieSceneSegment(TRange<FFrameNumber>(TRangeBound<FFrameNumber>::Inclusive(25),	TRangeBound<FFrameNumber>::Inclusive(30)), { FSectionEvaluationData(1) }),
 			};
 			AssertSegmentValues(this, Expected, EvalTrack.GetSortedSegments());
 		}
@@ -357,11 +361,11 @@ bool FMovieSceneTrackCompilerTest::RunTest(const FString& Parameters)
 			Track->bHighPassFilter = true;
 
 			FMovieSceneEvaluationTrack EvalTrack = Track->GenerateTrackTemplate();
-			EvalTrack.GetSegmentsInRange(TRange<float>::All());
+			EvalTrack.GetSegmentsInRange(TRange<FFrameNumber>::All());
 
 			FMovieSceneSegment Expected[] = {
-				FMovieSceneSegment(FFloatRange(TRangeBound<float>::Inclusive(10.f),	TRangeBound<float>::Inclusive(25.f)), { FSectionEvaluationData(0) }),
-				FMovieSceneSegment(FFloatRange(TRangeBound<float>::Exclusive(25.f),	TRangeBound<float>::Inclusive(30.f)), { FSectionEvaluationData(1) }),
+				FMovieSceneSegment(TRange<FFrameNumber>(TRangeBound<FFrameNumber>::Inclusive(10),	TRangeBound<FFrameNumber>::Exclusive(25)), { FSectionEvaluationData(0) }),
+				FMovieSceneSegment(TRange<FFrameNumber>(TRangeBound<FFrameNumber>::Inclusive(25),	TRangeBound<FFrameNumber>::Inclusive(30)), { FSectionEvaluationData(1) }),
 			};
 			AssertSegmentValues(this, Expected, EvalTrack.GetSortedSegments());
 		}
@@ -382,22 +386,19 @@ bool FMovieSceneTrackCompilerTest::RunTest(const FString& Parameters)
 		UMovieSceneSegmentCompilerTestTrack* Track = NewObject<UMovieSceneSegmentCompilerTestTrack>(GetTransientPackage());
 
 		UMovieSceneSegmentCompilerTestSection* Section0 = NewObject<UMovieSceneSegmentCompilerTestSection>(Track);
-		Section0->SetStartTime(10.f);
-		Section0->SetEndTime(30.f);
+		Section0->SetRange(TRange<FFrameNumber>(10, 30));
 		Section0->SetRowIndex(1);
 
 		UMovieSceneSegmentCompilerTestSection* Section1 = NewObject<UMovieSceneSegmentCompilerTestSection>(Track);
-		Section1->SetIsInfinite(true);
+		Section1->SetRange(TRange<FFrameNumber>::All());
 		Section1->SetRowIndex(2);
 		
 		UMovieSceneSegmentCompilerTestSection* Section2 = NewObject<UMovieSceneSegmentCompilerTestSection>(Track);
-		Section2->SetStartTime(20.f);
-		Section2->SetEndTime(30.f);
+		Section2->SetRange(TRange<FFrameNumber>(20, 30));
 		Section2->SetRowIndex(0);
 
 		UMovieSceneSegmentCompilerTestSection* Section3 = NewObject<UMovieSceneSegmentCompilerTestSection>(Track);
-		Section3->SetStartTime(15.f);
-		Section3->SetEndTime(25.f);
+		Section3->SetRange(TRange<FFrameNumber>(15, 25));
 		Section3->SetRowIndex(0);
 		Section3->SetOverlapPriority(100.f);
 
@@ -408,11 +409,11 @@ bool FMovieSceneTrackCompilerTest::RunTest(const FString& Parameters)
 
 		// Additive camera rules prescribe that they are evaluated in order of start time
 		FMovieSceneSegment Expected[] = {
-			FMovieSceneSegment(FFloatRange(Inf,											TRangeBound<float>::Exclusive(10.f)),	{ }),
-			FMovieSceneSegment(FFloatRange(TRangeBound<float>::Inclusive(10.f),			TRangeBound<float>::Exclusive(15.f)),	{ }),
-			FMovieSceneSegment(FFloatRange(TRangeBound<float>::Inclusive(15.f),			TRangeBound<float>::Inclusive(25.f)),	{ }),
-			FMovieSceneSegment(FFloatRange(TRangeBound<float>::Exclusive(25.f),			TRangeBound<float>::Inclusive(30.f)),	{ }),
-			FMovieSceneSegment(FFloatRange(TRangeBound<float>::Exclusive(30.f),			Inf),									{ }),
+			FMovieSceneSegment(TRange<FFrameNumber>(Inf,											TRangeBound<FFrameNumber>::Exclusive(10)),	{ }),
+			FMovieSceneSegment(TRange<FFrameNumber>(TRangeBound<FFrameNumber>::Inclusive(10),		TRangeBound<FFrameNumber>::Exclusive(15)),	{ }),
+			FMovieSceneSegment(TRange<FFrameNumber>(TRangeBound<FFrameNumber>::Inclusive(15),		TRangeBound<FFrameNumber>::Exclusive(25)),	{ }),
+			FMovieSceneSegment(TRange<FFrameNumber>(TRangeBound<FFrameNumber>::Inclusive(25),		TRangeBound<FFrameNumber>::Exclusive(30)),	{ }),
+			FMovieSceneSegment(TRange<FFrameNumber>(TRangeBound<FFrameNumber>::Inclusive(30),		Inf),										{ }),
 		};
 
 		// Test compiling the track with the additive camera rules
@@ -420,7 +421,7 @@ bool FMovieSceneTrackCompilerTest::RunTest(const FString& Parameters)
 			FScopedOverrideTrackSegmentBlender ScopedBlenderOverride{FMovieSceneAdditiveCameraTrackBlender()};
 
 			FMovieSceneEvaluationTrack EvalTrack = Track->GenerateTrackTemplate();
-			EvalTrack.GetSegmentsInRange(TRange<float>::All());
+			EvalTrack.GetSegmentsInRange(TRange<FFrameNumber>::All());
 
 			// Additive camera rules prescribe that they are evaluated in order of start time
 			Expected[0].Impls = { FSectionEvaluationData(1) };
@@ -437,7 +438,7 @@ bool FMovieSceneTrackCompilerTest::RunTest(const FString& Parameters)
 			Track->EvalOptions.bEvalNearestSection = true;
 
 			FMovieSceneEvaluationTrack EvalTrack = Track->GenerateTrackTemplate();
-			EvalTrack.GetSegmentsInRange(TRange<float>::All());
+			EvalTrack.GetSegmentsInRange(TRange<FFrameNumber>::All());
 
 			Expected[0].Impls = { FSectionEvaluationData(1) };
 			Expected[1].Impls = { FSectionEvaluationData(0), FSectionEvaluationData(1) };
@@ -453,7 +454,7 @@ bool FMovieSceneTrackCompilerTest::RunTest(const FString& Parameters)
 			Track->EvalOptions.bEvalNearestSection = false;
 
 			FMovieSceneEvaluationTrack EvalTrack = Track->GenerateTrackTemplate();
-			EvalTrack.GetSegmentsInRange(TRange<float>::All());
+			EvalTrack.GetSegmentsInRange(TRange<FFrameNumber>::All());
 
 			Expected[0].Impls = { FSectionEvaluationData(1) };
 			Expected[1].Impls = { FSectionEvaluationData(0), FSectionEvaluationData(1) };
@@ -470,7 +471,7 @@ bool FMovieSceneTrackCompilerTest::RunTest(const FString& Parameters)
 			Track->bHighPassFilter = true;
 
 			FMovieSceneEvaluationTrack EvalTrack = Track->GenerateTrackTemplate();
-			EvalTrack.GetSegmentsInRange(TRange<float>::All());
+			EvalTrack.GetSegmentsInRange(TRange<FFrameNumber>::All());
 
 			Expected[0].Impls = { FSectionEvaluationData(1) };
 			Expected[1].Impls = { FSectionEvaluationData(0) };
@@ -503,11 +504,11 @@ bool FMovieSceneCompilerTreeBasicTest::RunTest(const FString& Parameters)
 
 	TMovieSceneEvaluationTree<FSectionEvaluationData> EvalTree;
 
-	EvalTree.Add(TRange<float>(10.f,	20.f), 								FSectionEvaluationData(0));
-	EvalTree.Add(TRange<float>(20.f,	30.f), 								FSectionEvaluationData(1));
-	EvalTree.Add(TRange<float>(10.f,	30.f), 								FSectionEvaluationData(2));
-	EvalTree.Add(TRange<float>(Inf,	TRangeBound<float>::Exclusive(20.f)), 	FSectionEvaluationData(3));
-	EvalTree.Add(TRange<float>(25.f,	Inf), 								FSectionEvaluationData(4));
+	EvalTree.Add(TRange<FFrameNumber>(10,	20), 										FSectionEvaluationData(0));
+	EvalTree.Add(TRange<FFrameNumber>(20,	30), 										FSectionEvaluationData(1));
+	EvalTree.Add(TRange<FFrameNumber>(10,	30), 										FSectionEvaluationData(2));
+	EvalTree.Add(TRange<FFrameNumber>(Inf,	TRangeBound<FFrameNumber>::Exclusive(20)), 	FSectionEvaluationData(3));
+	EvalTree.Add(TRange<FFrameNumber>(TRangeBound<FFrameNumber>::Inclusive(25),	Inf), 	FSectionEvaluationData(4));
 
 	EvalTree.Compact();
 
@@ -527,11 +528,11 @@ bool FMovieSceneCompilerTreeBasicTest::RunTest(const FString& Parameters)
 	}
 
 	FMovieSceneSegment Expected[] = {
-		FMovieSceneSegment(FFloatRange(Inf, 	TRangeBound<float>::Exclusive(10.f)), { FSectionEvaluationData(3) }),
-		FMovieSceneSegment(FFloatRange(10.f,	20.f), { FSectionEvaluationData(0), FSectionEvaluationData(2), FSectionEvaluationData(3) }),
-		FMovieSceneSegment(FFloatRange(20.f,	25.f), { FSectionEvaluationData(1), FSectionEvaluationData(2) }),
-		FMovieSceneSegment(FFloatRange(25.f,	30.f), { FSectionEvaluationData(1), FSectionEvaluationData(2), FSectionEvaluationData(4) }),
-		FMovieSceneSegment(FFloatRange(30.f,	Inf), { FSectionEvaluationData(4) })
+		FMovieSceneSegment(TRange<FFrameNumber>(Inf, 	TRangeBound<FFrameNumber>::Exclusive(10)), 	{ FSectionEvaluationData(3) }),
+		FMovieSceneSegment(TRange<FFrameNumber>(10,	20), 											{ FSectionEvaluationData(0), FSectionEvaluationData(2), FSectionEvaluationData(3) }),
+		FMovieSceneSegment(TRange<FFrameNumber>(20,	25), 											{ FSectionEvaluationData(1), FSectionEvaluationData(2) }),
+		FMovieSceneSegment(TRange<FFrameNumber>(25,	30), 											{ FSectionEvaluationData(1), FSectionEvaluationData(2), FSectionEvaluationData(4) }),
+		FMovieSceneSegment(TRange<FFrameNumber>(TRangeBound<FFrameNumber>::Inclusive(30),	Inf), 	{ FSectionEvaluationData(4) })
 	};
 	AssertSegmentValues(this, Expected, Segments);
 	return true;
@@ -554,18 +555,18 @@ bool FMovieSceneCompilerTreeIteratorTest::RunTest(const FString& Parameters)
 
 	TMovieSceneEvaluationTree<int32> Tree;
 
-	Tree.Add(TRange<float>(10.f,	20.f), 								0);
-	Tree.Add(TRange<float>(20.f,	30.f), 								1);
-	Tree.Add(TRange<float>(10.f,	30.f), 								2);
-	Tree.Add(TRange<float>(Inf,	TRangeBound<float>::Exclusive(20.f)), 	3);
-	Tree.Add(TRange<float>(25.f,	Inf), 								4);
+	Tree.Add(TRange<FFrameNumber>(10,	20), 										0);
+	Tree.Add(TRange<FFrameNumber>(20,	30), 										1);
+	Tree.Add(TRange<FFrameNumber>(10,	30), 										2);
+	Tree.Add(TRange<FFrameNumber>(Inf,	TRangeBound<FFrameNumber>::Exclusive(20)), 	3);
+	Tree.Add(TRange<FFrameNumber>(TRangeBound<FFrameNumber>::Inclusive(25),	Inf), 	4);
 
 	FEvaluationTreeIteratorResult Expected[] = {
-		MakeTuple( TRange<float>(Inf, 	TRangeBound<float>::Exclusive(10.f)),   TArray<int32>({ 3 })         ),
-		MakeTuple( TRange<float>(10.f,	20.f),                                  TArray<int32>({ 0, 2, 3 })   ),
-		MakeTuple( TRange<float>(20.f,	25.f),                                  TArray<int32>({ 1, 2 })      ),
-		MakeTuple( TRange<float>(25.f,	30.f),                                  TArray<int32>({ 1, 2, 4 })   ),
-		MakeTuple( TRange<float>(30.f,	Inf),                                   TArray<int32>({ 4 })         ),
+		MakeTuple( TRange<FFrameNumber>(Inf, 	TRangeBound<FFrameNumber>::Exclusive(10)),   TArray<int32>({ 3 })         ),
+		MakeTuple( TRange<FFrameNumber>(10,		20),                                         TArray<int32>({ 0, 2, 3 })   ),
+		MakeTuple( TRange<FFrameNumber>(20,		25),                                         TArray<int32>({ 1, 2 })      ),
+		MakeTuple( TRange<FFrameNumber>(25,		30),                                         TArray<int32>({ 1, 2, 4 })   ),
+		MakeTuple( TRange<FFrameNumber>(TRangeBound<FFrameNumber>::Inclusive(30), Inf),      TArray<int32>({ 4 })         ),
 	};
 
 	AssertEvaluationTree(this, Tree, Expected);
@@ -584,17 +585,17 @@ bool FMovieSceneCompilerTreeIteratorBoundsTest::RunTest(const FString& Parameter
 
 	TMovieSceneEvaluationTree<int32> Tree;
 
-	Tree.Add(TRange<float>(Inf,		TRangeBound<float>::Exclusive(10.f)), 		0);
-	Tree.Add(TRange<float>(10.f,	TRangeBound<float>::Exclusive(20.f)), 		1);
-	Tree.Add(TRange<float>(20.f,	TRangeBound<float>::Exclusive(30.f)), 		2);
-	Tree.Add(TRange<float>(30.f,	Inf), 										3);
+	Tree.Add(TRange<FFrameNumber>(Inf,												TRangeBound<FFrameNumber>::Exclusive(10)), 	0);
+	Tree.Add(TRange<FFrameNumber>(FFrameNumber(10),									TRangeBound<FFrameNumber>::Exclusive(20)), 	1);
+	Tree.Add(TRange<FFrameNumber>(FFrameNumber(20),									TRangeBound<FFrameNumber>::Exclusive(30)), 	2);
+	Tree.Add(TRange<FFrameNumber>(TRangeBound<FFrameNumber>::Inclusive(30),			Inf), 										3);
 
 
 	FEvaluationTreeIteratorResult Expected[] = {
-		MakeTuple( TRange<float>(Inf,	TRangeBound<float>::Exclusive(10.f)), 		TArray<int32>({ 0 }) ),
-		MakeTuple( TRange<float>(10.f,	TRangeBound<float>::Exclusive(20.f)), 		TArray<int32>({ 1 }) ),
-		MakeTuple( TRange<float>(20.f,	TRangeBound<float>::Exclusive(30.f)), 		TArray<int32>({ 2 }) ),
-		MakeTuple( TRange<float>(30.f,	Inf), 										TArray<int32>({ 3 }) ),
+		MakeTuple( TRange<FFrameNumber>(Inf,										TRangeBound<FFrameNumber>::Exclusive(10)), 	TArray<int32>({ 0 }) ),
+		MakeTuple( TRange<FFrameNumber>(FFrameNumber(10),							TRangeBound<FFrameNumber>::Exclusive(20)), 	TArray<int32>({ 1 }) ),
+		MakeTuple( TRange<FFrameNumber>(FFrameNumber(20),							TRangeBound<FFrameNumber>::Exclusive(30)), 	TArray<int32>({ 2 }) ),
+		MakeTuple( TRange<FFrameNumber>(TRangeBound<FFrameNumber>::Inclusive(30),	Inf), 										TArray<int32>({ 3 }) ),
 	};
 
 	AssertEvaluationTree(this, Tree, Expected);
@@ -636,7 +637,7 @@ FMovieSceneTrackSegmentBlenderPtr UMovieSceneSegmentCompilerTestTrack::GetTrackS
 			Algo::SortBy(BlendData, [](const FMovieSceneSectionData& In){ return In.Section->GetRowIndex(); });
 		}
 
-		virtual TOptional<FMovieSceneSegment> InsertEmptySpace(const TRange<float>& Range, const FMovieSceneSegment* PreviousSegment, const FMovieSceneSegment* NextSegment) const
+		virtual TOptional<FMovieSceneSegment> InsertEmptySpace(const TRange<FFrameNumber>& Range, const FMovieSceneSegment* PreviousSegment, const FMovieSceneSegment* NextSegment) const
 		{
 			return bEvaluateNearest ? MovieSceneSegmentCompiler::EvaluateNearestSegment(Range, PreviousSegment, NextSegment) : TOptional<FMovieSceneSegment>();
 		}
diff --git a/Engine/Source/Runtime/MovieScene/Private/Tests/MovieSceneTestsCommon.h b/Engine/Source/Runtime/MovieScene/Private/Tests/MovieSceneTestsCommon.h
--- a/Engine/Source/Runtime/MovieScene/Private/Tests/MovieSceneTestsCommon.h
+++ b/Engine/Source/Runtime/MovieScene/Private/Tests/MovieSceneTestsCommon.h
@@ -7,5 +7,5 @@
 
 namespace Lex
 {
-	MOVIESCENE_API FString ToString(const TRange<float>& InRange);
+	MOVIESCENE_API FString ToString(const TRange<FFrameNumber>& InRange);
 }
diff --git a/Engine/Source/Runtime/MovieScene/Private/Tests/MovieSceneTransformTests.cpp b/Engine/Source/Runtime/MovieScene/Private/Tests/MovieSceneTransformTests.cpp
--- a/Engine/Source/Runtime/MovieScene/Private/Tests/MovieSceneTransformTests.cpp
+++ b/Engine/Source/Runtime/MovieScene/Private/Tests/MovieSceneTransformTests.cpp
@@ -9,7 +9,7 @@
 
 #define LOCTEXT_NAMESPACE "MovieSceneTransformTests"
 
-bool IsNearly(TRangeBound<float> A, TRangeBound<float> B)
+bool IsEqual(TRangeBound<FFrameNumber> A, TRangeBound<FFrameNumber> B)
 {
 	if (A.IsOpen() || B.IsOpen())
 	{
@@ -20,38 +20,38 @@ bool IsNearly(TRangeBound<float> A, TRangeBound<float> B)
 		return false;
 	}
 	
-	return FMath::IsNearlyEqual(A.GetValue(), B.GetValue());
+	return A.GetValue() == B.GetValue();
 }
 
-bool IsNearly(TRange<float> A, TRange<float> B)
+bool IsEqual(TRange<FFrameNumber> A, TRange<FFrameNumber> B)
 {
-	return IsNearly(A.GetLowerBound(), B.GetLowerBound()) && IsNearly(A.GetUpperBound(), B.GetUpperBound());
+	return IsEqual(A.GetLowerBound(), B.GetLowerBound()) && IsEqual(A.GetUpperBound(), B.GetUpperBound());
 }
 
 namespace Lex
 {
-	FString ToString(const TRange<float>& InRange)
+	FString ToString(const TRange<FFrameNumber>& InRange)
 	{
-		TRangeBound<float> SourceLower = InRange.GetLowerBound();
-		TRangeBound<float> SourceUpper = InRange.GetUpperBound();
+		TRangeBound<FFrameNumber> SourceLower = InRange.GetLowerBound();
+		TRangeBound<FFrameNumber> SourceUpper = InRange.GetUpperBound();
 
 		return *FString::Printf(TEXT("%s-%s"),
 			SourceLower.IsOpen() ? 
 				TEXT("[...") : 
 				SourceLower.IsInclusive() ?
-					*FString::Printf(TEXT("[%.5f"), SourceLower.GetValue()) :
-					*FString::Printf(TEXT("(%.5f"), SourceLower.GetValue()),
+					*FString::Printf(TEXT("[%i"), SourceLower.GetValue().Value) :
+					*FString::Printf(TEXT("(%i"), SourceLower.GetValue().Value),
 
 			SourceUpper.IsOpen() ? 
 				TEXT("...]") : 
 				SourceUpper.IsInclusive() ?
-					*FString::Printf(TEXT("%.5f]"), SourceUpper.GetValue()) :
-					*FString::Printf(TEXT("%.5f)"), SourceUpper.GetValue())
+					*FString::Printf(TEXT("%i]"), SourceUpper.GetValue().Value) :
+					*FString::Printf(TEXT("%i)"), SourceUpper.GetValue().Value)
 			);
 	}
 }
 
-bool TestTransform(FAutomationTestBase& Test, FMovieSceneSequenceTransform Transform, TArrayView<TRange<float>> InSource, TArrayView<TRange<float>> InExpected, const TCHAR* TestName)
+bool TestTransform(FAutomationTestBase& Test, FMovieSceneSequenceTransform Transform, TArrayView<TRange<FFrameNumber>> InSource, TArrayView<TRange<FFrameNumber>> InExpected, const TCHAR* TestName)
 {
 	using namespace Lex;
 
@@ -60,14 +60,15 @@ bool TestTransform(FAutomationTestBase& Test, FMovieSceneSequenceTransform Trans
 	bool bSuccess = true;
 	for (int32 Index = 0; Index < InSource.Num(); ++Index)
 	{
-		TRange<float> Result = InSource[Index] * Transform;
-		if (!IsNearly(Result, InExpected[Index]))
+		TRange<FFrameNumber> Result = InSource[Index] * Transform;
+		if (!IsEqual(Result, InExpected[Index]))
 		{
-			Test.AddError(FString::Printf(TEXT("Test '%s' failed (Index %d). Transform (Scale %.3f, Offset %.3f) did not apply correctly (%s != %s)"),
+			Test.AddError(FString::Printf(TEXT("Test '%s' failed (Index %d). Transform (Scale %.3f, Offset %i+%.3f) did not apply correctly (%s != %s)"),
 				TestName,
 				Index,
 				Transform.TimeScale,
-				Transform.Offset,
+				Transform.Offset.FrameNumber.Value,
+				Transform.Offset.GetSubFrame(),
 				*ToString(Result),
 				*ToString(InExpected[Index])));
 
@@ -79,23 +80,24 @@ bool TestTransform(FAutomationTestBase& Test, FMovieSceneSequenceTransform Trans
 }
 
 // Calculate the transform that transforms from range A to range B
-FMovieSceneSequenceTransform TransformRange(float StartA, float EndA, float StartB, float EndB)
+FMovieSceneSequenceTransform TransformRange(FFrameNumber StartA, FFrameNumber EndA, FFrameNumber StartB, FFrameNumber EndB)
 {
-	return FMovieSceneSequenceTransform(StartB, (EndB - StartB) / (EndA - StartA)) * FMovieSceneSequenceTransform(-StartA);
+	float Scale = double( (EndB - StartB).Value ) / (EndA - StartA).Value;
+	return FMovieSceneSequenceTransform(StartB, Scale) * FMovieSceneSequenceTransform(-StartA);
 }
 
 IMPLEMENT_SIMPLE_AUTOMATION_TEST(FMovieSceneSubSectionCoreTransformsTest, "System.Engine.Sequencer.Core.Transforms", EAutomationTestFlags::EditorContext | EAutomationTestFlags::EngineFilter)
 bool FMovieSceneSubSectionCoreTransformsTest::RunTest(const FString& Parameters)
 {
-	// We test using ranges since that implicitly tests float transformation
-	static const TRangeBound<float> OpenBound;
+	// We test using ranges since that implicitly tests frame number transformation as well
+	static const TRangeBound<FFrameNumber> OpenBound;
 
-	TRange<float> InfiniteRange(OpenBound, OpenBound);
-	TRange<float> OpenLowerRange(OpenBound, 200.f);
-	TRange<float> OpenUpperRange(100.f, OpenBound);
-	TRange<float> ClosedRange(100.f, 200.f);
+	TRange<FFrameNumber> InfiniteRange(OpenBound, OpenBound);
+	TRange<FFrameNumber> OpenLowerRange(OpenBound, FFrameNumber(200));
+	TRange<FFrameNumber> OpenUpperRange(FFrameNumber(100), OpenBound);
+	TRange<FFrameNumber> ClosedRange(FFrameNumber(100), FFrameNumber(200));
 
-	TRange<float> SourceRanges[] = {
+	TRange<FFrameNumber> SourceRanges[] = {
 		InfiniteRange, OpenLowerRange, OpenUpperRange, ClosedRange
 	};
 
@@ -105,7 +107,7 @@ bool FMovieSceneSubSectionCoreTransformsTest::RunTest(const FString& Parameters)
 		// Test Multiplication with an identity transform
 		FMovieSceneSequenceTransform IdentityTransform;
 
-		TRange<float> Expected[] = {
+		TRange<FFrameNumber> Expected[] = {
 			InfiniteRange, OpenLowerRange, OpenUpperRange, ClosedRange
 		};
 		
@@ -114,10 +116,10 @@ bool FMovieSceneSubSectionCoreTransformsTest::RunTest(const FString& Parameters)
 
 	{
 		// Test a simple translation
-		FMovieSceneSequenceTransform Transform(100.f, 1);
+		FMovieSceneSequenceTransform Transform(100, 1);
 
-		TRange<float> Expected[] = {
-			InfiniteRange, TRange<float>(OpenBound, 300.f), TRange<float>(200.f, OpenBound), TRange<float>(200.f, 300.f)
+		TRange<FFrameNumber> Expected[] = {
+			InfiniteRange, TRange<FFrameNumber>(OpenBound, FFrameNumber(300)), TRange<FFrameNumber>(FFrameNumber(200), OpenBound), TRange<FFrameNumber>(200, 300)
 		};
 
 		bSuccess = TestTransform(*this, Transform, SourceRanges, Expected, TEXT("Simple Translation")) && bSuccess;
@@ -129,21 +131,21 @@ bool FMovieSceneSubSectionCoreTransformsTest::RunTest(const FString& Parameters)
 		// Transform 100 - 200 to -200 - 1000
 		FMovieSceneSequenceTransform Transform = TransformRange(100, 200, -200, 1000);
 
-		TRange<float> Expected[] = {
-			InfiniteRange, TRange<float>(OpenBound, 1000.f), TRange<float>(-200.f, OpenBound), TRange<float>(-200.f, 1000.f)
+		TRange<FFrameNumber> Expected[] = {
+			InfiniteRange, TRange<FFrameNumber>(OpenBound, FFrameNumber(1000)), TRange<FFrameNumber>(FFrameNumber(-200), OpenBound), TRange<FFrameNumber>(-200, 1000)
 		};
 
 		bSuccess = TestTransform(*this, Transform, SourceRanges, Expected, TEXT("Simple Translation + half speed")) && bSuccess;
 	}
 
 	{
-		// Test that transforming a float by the same transform multiple times, does the same as the equivalent accumulated transform
+		// Test that transforming a frame number by the same transform multiple times, does the same as the equivalent accumulated transform
 
 		// scales by 2, then offsets by 100
-		FMovieSceneSequenceTransform SeedTransform = FMovieSceneSequenceTransform(100.f, 0.5f);
+		FMovieSceneSequenceTransform SeedTransform = FMovieSceneSequenceTransform(100, 0.5f);
 		FMovieSceneSequenceTransform AccumulatedTransform;
 
-		float SeedValue = 10.f;
+		FFrameTime SeedValue = 10;
 		for (int32 i = 0; i < 5; ++i)
 		{
 			AccumulatedTransform = SeedTransform * AccumulatedTransform;
@@ -151,18 +153,18 @@ bool FMovieSceneSubSectionCoreTransformsTest::RunTest(const FString& Parameters)
 			SeedValue = SeedValue * SeedTransform;
 		}
 
-		float AccumValue = 10.f * AccumulatedTransform;
-		if (!FMath::IsNearlyEqual(AccumValue, SeedValue))
+		FFrameTime AccumValue = FFrameTime(10) * AccumulatedTransform;
+		if (AccumValue != SeedValue)
 		{
-			AddError(FString::Printf(TEXT("Accumulated transform does not have the same effect as separate transformations (%.7f != %.7f)"), AccumValue, SeedValue));
+			AddError(FString::Printf(TEXT("Accumulated transform does not have the same effect as separate transformations (%i+%.5f != %i+%.5f)"), AccumValue.FrameNumber.Value, AccumValue.GetSubFrame(), SeedValue.FrameNumber.Value, SeedValue.GetSubFrame()));
 		}
 
 		FMovieSceneSequenceTransform InverseTransform = AccumulatedTransform.Inverse();
 
-		float InverseValue = AccumValue * InverseTransform;
-		if (!FMath::IsNearlyEqual(InverseValue, 10.f))
+		FFrameTime InverseValue = AccumValue * InverseTransform;
+		if (InverseValue != 10)
 		{
-			AddError(FString::Printf(TEXT("Inverse accumulated transform does not return value back to its original value (%.7f != 10.f)"), InverseValue));
+			AddError(FString::Printf(TEXT("Inverse accumulated transform does not return value back to its original value (%i+%.5f != 10)"), InverseValue.FrameNumber.Value, InverseValue.GetSubFrame()));
 		}
 	}
 
diff --git a/Engine/Source/Runtime/MovieScene/Private/Tracks/MovieSceneSubTrack.cpp b/Engine/Source/Runtime/MovieScene/Private/Tracks/MovieSceneSubTrack.cpp
--- a/Engine/Source/Runtime/MovieScene/Private/Tracks/MovieSceneSubTrack.cpp
+++ b/Engine/Source/Runtime/MovieScene/Private/Tracks/MovieSceneSubTrack.cpp
@@ -23,66 +23,42 @@ UMovieSceneSubTrack::UMovieSceneSubTrack( const FObjectInitializer& ObjectInitia
 #endif
 }
 
-UMovieSceneSubSection* UMovieSceneSubTrack::AddSequence(UMovieSceneSequence* Sequence, float StartTime, float Duration, const bool& bInsertSequence)
+UMovieSceneSubSection* UMovieSceneSubTrack::AddSequenceOnRow(UMovieSceneSequence* Sequence, FFrameNumber StartTime, int32 Duration, int32 RowIndex)
 {
 	Modify();
 
-	// If inserting, push all existing sections out to make space for the new one
-	if (bInsertSequence)
+	UMovieSceneSubSection* NewSection = CastChecked<UMovieSceneSubSection>(CreateNewSection());
 	{
-		// If there's a shot that starts at the same time as this new shot, push the shots forward to make space for this one
-		bool bPushShotsForward = false;
-		for (auto Section : Sections)
-		{
-			float SectionStartTime = Section->GetStartTime();
-			if (FMath::IsNearlyEqual(SectionStartTime, StartTime))
-			{
-				bPushShotsForward = true;
-				break;
-			}
-		}
-		if (bPushShotsForward)
+		NewSection->SetSequence(Sequence);
+		NewSection->SetRange(TRange<FFrameNumber>(StartTime, StartTime + Duration));
+	}
+
+	// If no given row index, put it on the next available row
+	if (RowIndex == INDEX_NONE)
+	{
+		RowIndex = 0;
+		NewSection->SetRowIndex(RowIndex);
+		while (NewSection->OverlapsWithSections(Sections) != nullptr)
 		{
-			for (auto Section : Sections)
-			{
-				float SectionStartTime = Section->GetStartTime();
-				if (SectionStartTime >= StartTime)
-				{
-					Section->SetStartTime(SectionStartTime + Duration);
-					Section->SetEndTime(Section->GetEndTime() + Duration);
-				}
-			}
+			++RowIndex;
+			NewSection->SetRowIndex(RowIndex);
 		}
-		// Otherwise, see if there's a shot after the start time and clamp the duration to that next shot
-		else
-		{
-			bool bFoundAfterShot = false;
-			float MinDiff = FLT_MAX;
-			for (auto Section : Sections)
-			{
-				float SectionStartTime = Section->GetStartTime();
-				if (SectionStartTime > StartTime)
-				{
-					bFoundAfterShot = true;
-					float Diff = SectionStartTime - StartTime;
-					MinDiff = FMath::Min(MinDiff, Diff);
-				}
-			}
+	}
+
+	NewSection->SetRowIndex(RowIndex);
 
-			if (MinDiff != FLT_MAX)
+	// If this overlaps with any sections, move out all the sections that are beyond this row
+	if (NewSection->OverlapsWithSections(Sections))
+	{
+		for (UMovieSceneSection* OtherSection : Sections)
+		{
+			if (OtherSection != NewSection && OtherSection->GetRowIndex() >= RowIndex)
 			{
-				Duration = MinDiff;
+				OtherSection->SetRowIndex(OtherSection->GetRowIndex()+1);
 			}
 		}
 	}
 
-	UMovieSceneSubSection* NewSection = CastChecked<UMovieSceneSubSection>(CreateNewSection());
-	{
-		NewSection->SetSequence(Sequence);
-		NewSection->SetStartTime(StartTime);
-		NewSection->SetEndTime(StartTime + Duration);
-	}
-
 	Sections.Add(NewSection);
 
 	return NewSection;
@@ -91,7 +67,7 @@ UMovieSceneSubSection* UMovieSceneSubTrack::AddSequence(UMovieSceneSequence* Seq
 UMovieSceneSubSection* UMovieSceneSubTrack::AddSequenceToRecord()
 {
 	UMovieScene* MovieScene = CastChecked<UMovieScene>(GetOuter());
-	TRange<float> PlaybackRange = MovieScene->GetPlaybackRange();
+	TRange<FFrameNumber> PlaybackRange = MovieScene->GetPlaybackRange();
 
 	int32 MaxRowIndex = -1;
 	for (auto Section : Sections)
@@ -103,8 +79,7 @@ UMovieSceneSubSection* UMovieSceneSubTrack::AddSequenceToRecord()
 	{
 		NewSection->SetRowIndex(MaxRowIndex + 1);
 		NewSection->SetAsRecording(true);
-		NewSection->SetStartTime(PlaybackRange.GetLowerBoundValue());
-		NewSection->SetEndTime(PlaybackRange.GetUpperBoundValue());
+		NewSection->SetRange(PlaybackRange);
 	}
 
 	Sections.Add(NewSection);
@@ -180,19 +155,6 @@ const TArray<UMovieSceneSection*>& UMovieSceneSubTrack::GetAllSections() const
 }
 
 
-TRange<float> UMovieSceneSubTrack::GetSectionBoundaries() const
-{
-	TArray<TRange<float>> Bounds;
-
-	for (int32 SectionIndex = 0; SectionIndex < Sections.Num(); ++SectionIndex)
-	{
-		Bounds.Add(Sections[SectionIndex]->GetRange());
-	}
-
-	return TRange<float>::Hull(Bounds);
-}
-
-
 bool UMovieSceneSubTrack::HasSection(const UMovieSceneSection& Section) const
 {
 	return Sections.Contains(&Section);
diff --git a/Engine/Source/Runtime/MovieScene/Public/Channels/BatchChannelInterface.h b/Engine/Source/Runtime/MovieScene/Public/Channels/BatchChannelInterface.h
new file mode 100644
--- /dev/null
+++ b/Engine/Source/Runtime/MovieScene/Public/Channels/BatchChannelInterface.h
@@ -0,0 +1,133 @@
+// Copyright 1998-2018 Epic Games, Inc. All Rights Reserved.
+
+#pragma once
+
+#include "Misc/CoreDefines.h"
+#include "Containers/ArrayView.h"
+#include "Math/Range.h"
+#include "Channels/MovieSceneChannelTraits.h"
+
+struct IBatchChannelInterface
+{
+	virtual ~IBatchChannelInterface() {}
+
+	virtual void ChangeFrameResolution_Batch(TArrayView<void* const> Ptrs, FFrameRate SourceRate, FFrameRate DestinationRate) const = 0;
+	virtual TRange<FFrameNumber> ComputeEffectiveRange_Batch(TArrayView<void* const> Ptrs) const = 0;
+	virtual int32 GetNumKeys_Batch(TArrayView<void* const> Ptrs) const = 0;
+	virtual void Reset_Batch(TArrayView<void* const> Ptrs) const = 0;
+	virtual void Offset_Batch(TArrayView<void* const> Ptrs, FFrameNumber DeltaPosition) const = 0;
+	virtual void Dilate_Batch(TArrayView<void* const> Ptrs, FFrameNumber Origin, float DilationFactor) const = 0;
+	virtual void Optimize_Batch(TArrayView<void* const> Ptrs, const FKeyDataOptimizationParams& InParameters) const = 0;
+	virtual void ClearDefaults_Batch(TArrayView<void* const> Ptrs) const = 0;
+};
+
+template<typename ChannelType>
+struct TBatchChannelInterface : IBatchChannelInterface
+{
+	virtual TRange<FFrameNumber> ComputeEffectiveRange_Batch(TArrayView<void* const> Ptrs) const
+	{
+		using namespace MovieScene;
+
+		TRange<FFrameNumber> Range = TRange<FFrameNumber>::Empty();
+		for (void* Ptr : Ptrs)
+		{
+			if (Ptr)
+			{
+				Range = TRange<FFrameNumber>::Hull(Range, ComputeEffectiveRange(static_cast<ChannelType*>(Ptr)));
+			}
+		}
+		return Range;
+	}
+
+	virtual void ChangeFrameResolution_Batch(TArrayView<void* const> Ptrs, FFrameRate SourceRate, FFrameRate DestinationRate) const override
+	{
+		using namespace MovieScene;
+		for (void* Ptr : Ptrs)
+		{
+			if (Ptr)
+			{
+				ChangeFrameResolution(static_cast<ChannelType*>(Ptr), SourceRate, DestinationRate);
+			}
+		}
+	}
+
+	virtual int32 GetNumKeys_Batch(TArrayView<void* const> Ptrs) const
+	{
+		using namespace MovieScene;
+
+		int32 NumKeys = 0;
+		for (void* Ptr : Ptrs)
+		{
+			if (Ptr)
+			{
+				NumKeys += GetNumKeys(static_cast<ChannelType*>(Ptr));
+			}
+		}
+		return NumKeys;
+	}
+
+	virtual void Reset_Batch(TArrayView<void* const> Ptrs) const
+	{
+		using namespace MovieScene;
+
+		for (void* Ptr : Ptrs)
+		{
+			if (Ptr)
+			{
+				Reset(static_cast<ChannelType*>(Ptr));
+			}
+		}
+	}
+
+	virtual void Offset_Batch(TArrayView<void* const> Ptrs, FFrameNumber DeltaPosition) const
+	{
+		using namespace MovieScene;
+
+		for (void* Ptr : Ptrs)
+		{
+			if (Ptr)
+			{
+				Offset(static_cast<ChannelType*>(Ptr), DeltaPosition);
+			}
+		}
+	}
+
+	virtual void Dilate_Batch(TArrayView<void* const> Ptrs, FFrameNumber Origin, float DilationFactor) const
+	{
+		using namespace MovieScene;
+
+		for (void* Ptr : Ptrs)
+		{
+			if (Ptr)
+			{
+				Dilate(static_cast<ChannelType*>(Ptr), Origin, DilationFactor);
+			}
+		}
+	}
+
+	virtual void Optimize_Batch(TArrayView<void* const> Ptrs, const FKeyDataOptimizationParams& InParameters) const
+	{
+		using namespace MovieScene;
+
+		for (void* Ptr : Ptrs)
+		{
+			if (Ptr)
+			{
+				Optimize(static_cast<ChannelType*>(Ptr), InParameters);
+			}
+		}
+	}
+
+	virtual void ClearDefaults_Batch(TArrayView<void* const> Ptrs) const
+	{
+		using namespace MovieScene;
+
+		for (void* Ptr : Ptrs)
+		{
+			if (Ptr)
+			{
+				ClearChannelDefault(static_cast<ChannelType*>(Ptr));
+			}
+		}
+	}
+};
\ No newline at end of file
diff --git a/Engine/Source/Runtime/MovieScene/Public/Channels/MovieSceneBoolChannel.h b/Engine/Source/Runtime/MovieScene/Public/Channels/MovieSceneBoolChannel.h
new file mode 100644
--- /dev/null
+++ b/Engine/Source/Runtime/MovieScene/Public/Channels/MovieSceneBoolChannel.h
@@ -0,0 +1,141 @@
+// Copyright 1998-2018 Epic Games, Inc. All Rights Reserved.
+
+#pragma once
+
+#include "UObject/ObjectMacros.h"
+#include "FrameNumber.h"
+#include "MovieSceneChannel.h"
+#include "MovieSceneChannelTraits.h"
+#include "MovieSceneBoolChannel.generated.h"
+
+
+USTRUCT()
+struct FMovieSceneBoolChannel
+{
+	GENERATED_BODY()
+
+	FMovieSceneBoolChannel()
+		: DefaultValue(), bHasDefaultValue(false)
+	{}
+
+	/**
+	 * Serialize this type from another
+	 */
+	MOVIESCENE_API bool SerializeFromMismatchedTag(const FPropertyTag& Tag, FArchive& Ar);
+
+	/**
+	 * Access an integer that uniquely identifies this channel type.
+	 *
+	 * @return A static identifier that was allocated using FMovieSceneChannelEntry::RegisterNewID
+	 */
+	MOVIESCENE_API static uint32 GetChannelID();
+
+	/**
+	 * Access a mutable interface for this channel
+	 *
+	 * @return An object that is able to manipulate this channel
+	 */
+	FORCEINLINE TMovieSceneChannel<bool> GetInterface()
+	{
+		return TMovieSceneChannel<bool>(&Times, &Values, &KeyHandles);
+	}
+
+	/**
+	 * Access a constant interface for this channel
+	 *
+	 * @return An object that is able to interrogate this channel
+	 */
+	FORCEINLINE TMovieSceneChannel<const bool> GetInterface() const
+	{
+		return TMovieSceneChannel<const bool>(&Times, &Values);
+	}
+
+	/**
+	 * Const access to this channel's times
+	 */
+	FORCEINLINE TArrayView<const FFrameNumber> GetTimes() const
+	{
+		return Times;
+	}
+
+	/**
+	 * Const access to this channel's values
+	 */
+	FORCEINLINE TArrayView<const bool> GetValues() const
+	{
+		return Values;
+	}
+
+	/**
+	 * Evaluate this channel
+	 *
+	 * @param InTime     The time to evaluate at
+	 * @param OutValue   A value to receive the result
+	 * @return true if the channel was evaluated successfully, false otherwise
+	 */
+	MOVIESCENE_API bool Evaluate(FFrameTime InTime, bool& OutValue) const;
+
+public:
+
+	/**
+	 * Set this channel's default value that should be used when no keys are present
+	 *
+	 * @param InDefaultValue The desired default value
+	 */
+	FORCEINLINE void SetDefault(bool InDefaultValue)
+	{
+		bHasDefaultValue = true;
+		DefaultValue = InDefaultValue;
+	}
+
+	/**
+	 * Get this channel's default value that will be used when no keys are present
+	 *
+	 * @return (Optional) The channel's default value
+	 */
+	FORCEINLINE TOptional<bool> GetDefault() const
+	{
+		return bHasDefaultValue ? TOptional<bool>(DefaultValue) : TOptional<bool>();
+	}
+
+	/**
+	 * Remove this channel's default value causing the channel to have no effect where no keys are present
+	 */
+	FORCEINLINE void RemoveDefault()
+	{
+		bHasDefaultValue = false;
+	}
+
+private:
+
+	UPROPERTY()
+	TArray<FFrameNumber> Times;
+
+	UPROPERTY()
+	bool DefaultValue;
+
+	UPROPERTY()
+	bool bHasDefaultValue;
+
+	UPROPERTY()
+	TArray<bool> Values;
+
+	FMovieSceneKeyHandleMap KeyHandles;
+};
+
+template<>
+struct TStructOpsTypeTraits<FMovieSceneBoolChannel> : public TStructOpsTypeTraitsBase2<FMovieSceneBoolChannel>
+{
+	enum { WithSerializeFromMismatchedTag = true };
+};
+
+template<>
+struct TMovieSceneChannelTraits<FMovieSceneBoolChannel>
+{
+#if WITH_EDITORONLY_DATA
+
+	/** Bool channels can have external values (ie, they can get their values from external objects for UI purposes) */
+	typedef TMovieSceneExternalValue<bool> EditorDataType;
+
+#endif
+};
\ No newline at end of file
diff --git a/Engine/Source/Runtime/MovieScene/Public/Channels/MovieSceneByteChannel.h b/Engine/Source/Runtime/MovieScene/Public/Channels/MovieSceneByteChannel.h
new file mode 100644
--- /dev/null
+++ b/Engine/Source/Runtime/MovieScene/Public/Channels/MovieSceneByteChannel.h
@@ -0,0 +1,157 @@
+// Copyright 1998-2018 Epic Games, Inc. All Rights Reserved.
+
+#pragma once
+
+#include "UObject/ObjectMacros.h"
+#include "FrameNumber.h"
+#include "MovieSceneChannel.h"
+#include "MovieSceneChannelTraits.h"
+#include "MovieSceneByteChannel.generated.h"
+
+
+USTRUCT()
+struct FMovieSceneByteChannel
+{
+	GENERATED_BODY()
+
+	FMovieSceneByteChannel()
+		: DefaultValue(), bHasDefaultValue(false)
+	{}
+
+	/**
+	 * Serialize this type from another
+	 */
+	MOVIESCENE_API bool SerializeFromMismatchedTag(const FPropertyTag& Tag, FArchive& Ar);
+
+	/**
+	 * Access an integer that uniquely identifies this channel type.
+	 *
+	 * @return A static identifier that was allocated using FMovieSceneChannelEntry::RegisterNewID
+	 */
+	MOVIESCENE_API static uint32 GetChannelID();
+
+	/**
+	 * Access a mutable interface for this channel
+	 *
+	 * @return An object that is able to manipulate this channel
+	 */
+	FORCEINLINE TMovieSceneChannel<uint8> GetInterface()
+	{
+		return TMovieSceneChannel<uint8>(&Times, &Values, &KeyHandles);
+	}
+
+	/**
+	 * Access a constant interface for this channel
+	 *
+	 * @return An object that is able to interrogate this channel
+	 */
+	FORCEINLINE TMovieSceneChannel<const uint8> GetInterface() const
+	{
+		return TMovieSceneChannel<const uint8>(&Times, &Values);
+	}
+
+	/**
+	 * Const access to this channel's times
+	 */
+	FORCEINLINE TArrayView<const FFrameNumber> GetTimes() const
+	{
+		return Times;
+	}
+
+	/**
+	 * Const access to this channel's values
+	 */
+	FORCEINLINE TArrayView<const uint8> GetValues() const
+	{
+		return Values;
+	}
+
+	/**
+	 * Evaluate this channel
+	 *
+	 * @param InTime     The time to evaluate at
+	 * @param OutValue   A value to receive the result
+	 * @return true if the channel was evaluated successfully, false otherwise
+	 */
+	MOVIESCENE_API bool Evaluate(FFrameTime InTime, uint8& OutValue) const;
+
+public:
+
+	/**
+	 * Set this channel's default value that should be used when no keys are present
+	 *
+	 * @param InDefaultValue The desired default value
+	 */
+	FORCEINLINE void SetDefault(uint8 InDefaultValue)
+	{
+		bHasDefaultValue = true;
+		DefaultValue = InDefaultValue;
+	}
+
+	/**
+	 * Get this channel's default value that will be used when no keys are present
+	 *
+	 * @return (Optional) The channel's default value
+	 */
+	FORCEINLINE TOptional<uint8> GetDefault() const
+	{
+		return bHasDefaultValue ? TOptional<uint8>(DefaultValue) : TOptional<uint8>();
+	}
+
+	/**
+	 * Remove this channel's default value causing the channel to have no effect where no keys are present
+	 */
+	FORCEINLINE void RemoveDefault()
+	{
+		bHasDefaultValue = false;
+	}
+
+
+public:
+
+	UEnum* GetEnum() const
+	{
+		return Enum;
+	}
+
+	void SetEnum(UEnum* InEnum)
+	{
+		Enum = InEnum;
+	}
+
+private:
+
+	UPROPERTY()
+	TArray<FFrameNumber> Times;
+
+	UPROPERTY()
+	uint8 DefaultValue;
+
+	UPROPERTY()
+	bool bHasDefaultValue;
+
+	UPROPERTY()
+	TArray<uint8> Values;
+
+	UPROPERTY()
+	UEnum* Enum;
+
+	FMovieSceneKeyHandleMap KeyHandles;
+};
+
+template<>
+struct TStructOpsTypeTraits<FMovieSceneByteChannel> : public TStructOpsTypeTraitsBase2<FMovieSceneByteChannel>
+{
+	enum { WithSerializeFromMismatchedTag = true };
+};
+
+template<>
+struct TMovieSceneChannelTraits<FMovieSceneByteChannel>
+{
+#if WITH_EDITORONLY_DATA
+
+	/** Byte channels can have external values (ie, they can get their values from external objects for UI purposes) */
+	typedef TMovieSceneExternalValue<uint8> EditorDataType;
+
+#endif
+};
\ No newline at end of file
diff --git a/Engine/Source/Runtime/MovieScene/Public/Channels/MovieSceneChannel.h b/Engine/Source/Runtime/MovieScene/Public/Channels/MovieSceneChannel.h
new file mode 100644
--- /dev/null
+++ b/Engine/Source/Runtime/MovieScene/Public/Channels/MovieSceneChannel.h
@@ -0,0 +1,190 @@
+// Copyright 1998-2018 Epic Games, Inc. All Rights Reserved.
+
+#pragma once
+
+#include "Containers/ArrayView.h"
+#include "Curves/KeyHandle.h"
+#include "FrameNumber.h"
+#include "FrameTime.h"
+
+struct FKeyDataOptimizationParams;
+
+struct FMovieSceneKeyHandleMap : FKeyHandleLookupTable
+{
+	FMovieSceneKeyHandleMap() = default;
+	FMovieSceneKeyHandleMap(const FMovieSceneKeyHandleMap& RHS){}
+	FMovieSceneKeyHandleMap& operator=(const FMovieSceneKeyHandleMap& RHS)
+	{
+		Reset();
+		return *this;
+	}
+};
+
+namespace MovieScene
+{
+	MOVIESCENE_API void EvaluateTime(TArrayView<const FFrameNumber> InTimes, FFrameTime InTime, int32& OutIndex1, int32& OutIndex2);
+
+	MOVIESCENE_API void EvaluateTime(TArrayView<const FFrameNumber> InTimes, FFrameTime InTime, int32& OutIndex1, int32& OutIndex2, float& OutInterp);
+
+	MOVIESCENE_API void FindRange(TArrayView<const FFrameNumber> InTimes, FFrameNumber PredicateTime, FFrameNumber InTolerance, int32 MaxNum, int32& OutMin, int32& OutMax);
+}
+
+struct FMovieSceneChannel
+{
+	MOVIESCENE_API FMovieSceneChannel(TArray<FFrameNumber>* InTimes, FKeyHandleLookupTable* InKeyHandles);
+
+	MOVIESCENE_API FKeyHandle GetHandle(int32 Index);
+
+	MOVIESCENE_API int32 GetIndex(FKeyHandle Handle);
+
+	MOVIESCENE_API int32 FindKey(FFrameNumber InTime, int32 InTolerance = 0);
+
+	MOVIESCENE_API void FindKeys(FFrameNumber InTime, int32 MaxNum, int32& OutMinIndex, int32& OutMaxIndex, int32 InTolerance);
+
+	// mutable access to key times - any usage *must* keep times sorted. Any reordering of times will not be reflected in the values array.
+	FORCEINLINE TArrayView<FFrameNumber>       GetTimes()       { return *Times; }
+	FORCEINLINE TArrayView<const FFrameNumber> GetTimes() const { return *Times; }
+
+protected:
+
+	MOVIESCENE_API int32 MoveKeyInternal(int32 KeyIndex, FFrameNumber InNewTime);
+
+	MOVIESCENE_API int32 AddKeyInternal(FFrameNumber InTime);
+
+	TArray<FFrameNumber>*     Times;
+	FKeyHandleLookupTable* KeyHandles;
+};
+
+template<typename ValueType>
+struct TMovieSceneChannel : FMovieSceneChannel
+{
+	typedef typename TCallTraits<ValueType>::ParamType ParamType;
+
+	TMovieSceneChannel(TArray<FFrameNumber>* InTimes, TArray<ValueType>* InValues, FKeyHandleLookupTable* InKeyHandles)
+		: FMovieSceneChannel(InTimes, InKeyHandles), Values(InValues)
+	{
+		check(Times && Values);
+	}
+
+	operator TMovieSceneChannel<const ValueType>()
+	{
+		return TMovieSceneChannel<const ValueType>(Times, Values);
+	}
+
+	void RemoveKey(int32 KeyIndex)
+	{
+		check(Times->IsValidIndex(KeyIndex));
+		Times->RemoveAt(KeyIndex, 1, false);
+		Values->RemoveAt(KeyIndex, 1, false);
+
+		if (KeyHandles)
+		{
+			KeyHandles->DeallocateHandle(KeyIndex);
+		}
+	}
+
+	int32 MoveKey(int32 KeyIndex, FFrameNumber InNewTime)
+	{
+		int32 NewIndex = MoveKeyInternal(KeyIndex, InNewTime);
+		if (NewIndex != KeyIndex)
+		{
+			// We have to remove the key and re-add it in the right place
+			// This could probably be done better by just shuffling up/down the items that need to move, without ever changing the size of the array
+			ValueType OldValue = (*Values)[KeyIndex];
+			Values->RemoveAt(KeyIndex, 1, false);
+			Values->Insert(OldValue, NewIndex);
+		}
+		return NewIndex;
+	}
+
+	int32 AddKey(FFrameNumber InTime, ParamType InValue)
+	{
+		int32 KeyIndex = AddKeyInternal(InTime);
+		Values->Insert(InValue, KeyIndex);
+		return KeyIndex;
+	}
+
+	int32 SetKeyTime(int32 KeyIndex, FFrameNumber InNewTime)
+	{
+		return MoveKey(KeyIndex, InNewTime);
+	}
+
+	FKeyHandle UpdateOrAddKey(FFrameNumber InTime, ParamType InValue)
+	{
+		int32 ExistingKey = FindKey(InTime);
+		if (ExistingKey != INDEX_NONE)
+		{
+			(*Values)[ExistingKey] = InValue;
+		}
+		else
+		{
+			ExistingKey = AddKey(InTime, InValue);
+		}
+
+		return GetHandle(ExistingKey);
+	}
+
+	void Reset()
+	{
+		Times->Reset();
+		Values->Reset();
+		if (KeyHandles)
+		{
+			KeyHandles->Reset();
+		}
+	}
+
+	TArrayView<const ValueType> GetValues() const
+	{
+		return *Values;
+	}
+
+	TArrayView<ValueType> GetValues()
+	{
+		return *Values;
+	}
+
+private:
+
+	TArray<ValueType>* Values;
+};
+
+
+template<typename ValueType>
+struct TMovieSceneChannel<const ValueType>
+{
+	typedef typename TCallTraits<ValueType>::ParamType ParamType;
+
+	TMovieSceneChannel(const TArray<FFrameNumber>* InTimes, const TArray<ValueType>* InValues)
+		: Times(InTimes), Values(InValues)
+	{
+		check(Times && Values);
+	}
+
+	int32 FindKey(FFrameNumber InTime, FFrameNumber InTolerance) const
+	{
+		int32 MinIndex = 0, MaxIndex = 0;
+		MovieScene::FindRange(*Times, InTime, InTolerance, 1, MinIndex, MaxIndex);
+		return MinIndex;
+	}
+
+	void FindKeys(FFrameNumber InTime, int32 MaxNum, int32& OutMinIndex, int32& OutMaxIndex, FFrameNumber InTolerance) const
+	{
+		MovieScene::FindRange(*Times, InTime, InTolerance, MaxNum, OutMinIndex, OutMaxIndex);
+	}
+
+	TArrayView<const FFrameNumber> GetTimes() const
+	{
+		return *Times;
+	}
+
+	TArrayView<const ValueType> GetValues() const
+	{
+		return *Values;
+	}
+
+private:
+
+	const TArray<FFrameNumber>* Times;
+	const TArray<ValueType>*    Values;
+};
\ No newline at end of file
diff --git a/Engine/Source/Runtime/MovieScene/Public/Channels/MovieSceneChannelEditorDataEntry.h b/Engine/Source/Runtime/MovieScene/Public/Channels/MovieSceneChannelEditorDataEntry.h
new file mode 100644
--- /dev/null
+++ b/Engine/Source/Runtime/MovieScene/Public/Channels/MovieSceneChannelEditorDataEntry.h
@@ -0,0 +1,141 @@
+// Copyright 1998-2018 Epic Games, Inc. All Rights Reserved.
+
+#pragma once
+
+#include "UObject/ObjectMacros.h"
+#include "Misc/InlineValue.h"
+#include "Algo/Sort.h"
+#include "Misc/Attribute.h"
+#include "Misc/Optional.h"
+#include "Channels/MovieSceneChannelEditorData.h"
+
+#if WITH_EDITOR
+
+/**
+ * Base entry type for use in FMovieSceneChannelProxy that stores a piece of editor data for each channel or a given type (one entry per type).
+ */
+struct FMovieSceneChannelEditorDataEntry
+{
+	/**
+	 * Templated constructor that uses the channel parameter to construct
+	 * the editor data array from the correct editor data type
+	 */
+	template<typename ChannelType>
+	explicit FMovieSceneChannelEditorDataEntry(const ChannelType& Channel)
+	{
+		ConstructSpecializedEditorDataArray<ChannelType>((typename TMovieSceneChannelTraits<ChannelType>::EditorDataType*)nullptr);
+	}
+
+	/**
+	 * Get the common editor data for all channels
+	 */
+	TArrayView<const FMovieSceneChannelEditorData> GetCommonEditorData() const
+	{
+		return CommonEditorDataArray;
+	}
+
+	/**
+	 * Access the specialized editor data for a specific channel
+	 */
+	const void* GetSpecializedEditorData(int32 ChannelIndex) const
+	{
+		if (SpecializedEditorDataArray.IsValid())
+		{
+			return SpecializedEditorDataArray->GetChannel(ChannelIndex);
+		}
+		return nullptr;
+	}
+
+protected:
+
+	/**
+	 * Add new editor data for the specified channel type at the last index in the array
+	 */
+	template<typename ChannelType>
+	void AddEditorData(const FMovieSceneChannelEditorData& CommonData)
+	{
+		static_assert(TIsSame<typename TMovieSceneChannelTraits<ChannelType>::EditorDataType, void>::Value, "Must supply typed editor data according to the channel's traits.");
+
+		// Add the common editor data
+		CommonEditorDataArray.Add(CommonData);
+	}
+
+	/**
+	 * Add new editor data for the specified channel type at the last index in the arrays
+	 */
+	template<typename ChannelType, typename SpecializedEditorData>
+	void AddEditorData(const FMovieSceneChannelEditorData& CommonData, SpecializedEditorData&& InTypeSpecificEditorData)
+	{
+		// Add the common editor data
+		CommonEditorDataArray.Add(CommonData);
+
+		auto& TypedImpl = static_cast<TMovieSceneEditorDataArray<ChannelType>&>(SpecializedEditorDataArray.GetValue());
+		TypedImpl.Data.Add(Forward<SpecializedEditorData>(InTypeSpecificEditorData));
+	}
+
+	/**
+	 * Access the specialized editor data for channels stored in this entry
+	 */
+	template<typename ChannelType>
+	TArrayView<const typename TMovieSceneChannelTraits<ChannelType>::EditorDataType> GetAllSpecializedEditorData() const
+	{
+		static_assert(!TIsSame<typename TMovieSceneChannelTraits<ChannelType>::EditorDataType, void>::Value, "This channel type does not define any additional editor data.");
+
+		const auto& TypedImpl = static_cast<const TMovieSceneEditorDataArray<ChannelType>&>(SpecializedEditorDataArray.GetValue());
+		return TypedImpl.Data;
+	}
+
+private:
+
+	/** Construct the specialized editor data container for channel types that require it */
+	template<typename ChannelType, typename EditorDataType>
+	void ConstructSpecializedEditorDataArray(EditorDataType*)
+	{
+		SpecializedEditorDataArray = TMovieSceneEditorDataArray<ChannelType>();
+	}
+
+	template<typename ChannelType>
+	void ConstructSpecializedEditorDataArray(void*)
+	{
+	}
+
+private:
+
+	/** Base editor data, one per channel */
+	TArray<FMovieSceneChannelEditorData, TInlineAllocator<1>> CommonEditorDataArray;
+
+	/**
+	 * We store the array behind an interface whose access is via void*
+	 * Typed access is only permitted using the original ChannelType templated methods
+	 * to ensure safe casting of the array
+	 */
+	struct IMovieSceneEditorDataArray
+	{
+		virtual ~IMovieSceneEditorDataArray() {}
+		virtual const void* GetChannel(int32 Index) const = 0;
+	};
+
+	template<typename ChannelType>
+	struct TMovieSceneEditorDataArray : IMovieSceneEditorDataArray
+	{
+		typedef typename TMovieSceneChannelTraits<ChannelType>::EditorDataType EditorDataType;
+
+		virtual const void* GetChannel(int32 Index) const { return &Data[Index]; }
+
+		/** The actual editor data */
+		TArray<EditorDataType> Data;
+	};
+
+	/** Typed editor data, one per channel, defined by TMovieSceneChannelTraits::EditorDataType */
+	TInlineValue<IMovieSceneEditorDataArray, 8> SpecializedEditorDataArray;
+};
+
+#else // WITH_EDITOR
+
+/** Empty stub in non-editor bulds */
+struct FMovieSceneChannelEditorDataEntry
+{
+	template<typename ChannelType> FMovieSceneChannelEditorDataEntry(const ChannelType&){}
+};
+
+#endif // WITH_EDITOR
diff --git a/Engine/Source/Runtime/MovieScene/Public/Channels/MovieSceneChannelHandle.h b/Engine/Source/Runtime/MovieScene/Public/Channels/MovieSceneChannelHandle.h
new file mode 100644
--- /dev/null
+++ b/Engine/Source/Runtime/MovieScene/Public/Channels/MovieSceneChannelHandle.h
@@ -0,0 +1,58 @@
+// Copyright 1998-2018 Epic Games, Inc. All Rights Reserved.
+
+#pragma once
+
+#include "Templates/SharedPointer.h"
+
+/**
+ * Handle to a specific channel in a UMovieSceneSection. Will become nullptr when the FMovieSceneChannelProxy it was created with is reallocated.
+ * Implemented internally as a weak pointer, but not advertised as such to ensure that no shared references are held to FMovieSceneChannelProxy
+ * outside of the owning MovieSceneSection's reference.
+ */
+template<typename ChannelType>
+struct TMovieSceneChannelHandle
+{
+	TMovieSceneChannelHandle()
+	{}
+
+	/** Construction from a weak pointer produced from FMovieSceneChannelProxy */
+	TMovieSceneChannelHandle(TWeakPtr<ChannelType> InWeakChannel)
+		: WeakChannel(InWeakChannel)
+	{}
+
+	/**
+	 * Get the channel pointer this handle represents.
+	 *
+	 * @return the channel's pointer, or nullptr if the proxy it was created with is no longer alive.
+	 */
+	ChannelType* Get() const
+	{
+		return WeakChannel.Pin().Get();
+	}
+
+	/**
+	 * Cast this handle to a handle of a related type
+	 */
+	template<typename OtherChannelType>
+	TMovieSceneChannelHandle<OtherChannelType> Cast() const;
+
+private:
+
+	/** Weak pointer to the channel, proxy alisased to the channel proxy's shared reference controller to ensure it becomes null when the proxy is re-allocated */
+	TWeakPtr<ChannelType> WeakChannel;
+};
+
+
+
+template<typename ChannelType>
+template<typename OtherChannelType>
+TMovieSceneChannelHandle<OtherChannelType> TMovieSceneChannelHandle<ChannelType>::Cast() const
+{
+	TSharedPtr<ChannelType> StrongRef = WeakChannel.Pin();
+	if (StrongRef.IsValid())
+	{
+		return TWeakPtr<OtherChannelType>(StaticCastSharedPtr<OtherChannelType>(StrongRef));
+	}
+
+	return TMovieSceneChannelHandle<OtherChannelType>();
+}
diff --git a/Engine/Source/Runtime/MovieScene/Public/Channels/MovieSceneChannelProxy.h b/Engine/Source/Runtime/MovieScene/Public/Channels/MovieSceneChannelProxy.h
new file mode 100644
--- /dev/null
+++ b/Engine/Source/Runtime/MovieScene/Public/Channels/MovieSceneChannelProxy.h
@@ -0,0 +1,463 @@
+// Copyright 1998-2018 Epic Games, Inc. All Rights Reserved.
+
+#pragma once
+
+#include "CoreTypes.h"
+#include "Templates/UnrealTemplate.h"
+#include "Misc/InlineValue.h"
+#include "Channels/BatchChannelInterface.h"
+#include "Channels/MovieSceneChannelEditorDataEntry.h"
+#include "Channels/MovieSceneChannelHandle.h"
+
+struct FMovieSceneChannelData;
+struct FMovieSceneChannelProxy;
+
+
+
+/**
+ * An entry within FMovieSceneChannelProxy that contains all channels (and editor data) for any given channel type
+ */
+struct FMovieSceneChannelEntry : FMovieSceneChannelEditorDataEntry
+{
+	/**
+	 * Register a new, unique identifier that should be associated with a channel type through ChannelType::GetChannelID
+	 */
+	MOVIESCENE_API static uint32 RegisterNewID();
+
+	/** 
+	 * Get the ID of the channels stored in this entry
+	 */
+	uint32 GetChannelID() const
+	{
+		return ChannelID;
+	}
+
+	/** 
+	 * Access the multi-channel interface that can interact with this entry's channels
+	 */
+	const IBatchChannelInterface& GetBatchChannelInterface() const
+	{
+		return SectionDataInterface.GetValue();
+	}
+
+	/** 
+	 * Access all the channels contained within this entry
+	 */
+	TArrayView<void* const> GetChannels() const
+	{
+		return Channels;
+	}
+
+#if WITH_EDITOR
+
+	/**
+	 * Access specialized typed editor data for channels stored in this entry
+	 */
+	template<typename ChannelType>
+	TArrayView<const typename TMovieSceneChannelTraits<ChannelType>::EditorDataType> GetAllSpecializedEditorData() const
+	{
+		check(ChannelType::GetChannelID() == ChannelID);
+		return FMovieSceneChannelEditorDataEntry::GetAllSpecializedEditorData<ChannelType>();
+	}
+
+#endif
+
+private:
+
+	// only FMovieSceneChannelData and FMovieSceneChannelProxy can create entries
+	friend FMovieSceneChannelData;
+	friend FMovieSceneChannelProxy;
+
+	/** Templated constructor from the channel and its ID */
+	template<typename ChannelType>
+	explicit FMovieSceneChannelEntry(uint32 InChannelID, const ChannelType& Channel)
+		: FMovieSceneChannelEditorDataEntry(Channel)
+		, ChannelID(InChannelID)
+		, SectionDataInterface(TBatchChannelInterface<ChannelType>())
+	{
+		check(InChannelID == ChannelType::GetChannelID());
+	}
+
+	/** The ID of the channels contained in this entry, generated from FMovieSceneChannelEntry::RegisterNewID */
+	uint32 ChannelID;
+
+	/** Pointers to the channels that this entry contains. Pointers are assumed to stay alive as long as this entry is. If channels are reallocated, a new channel proxy should be created */
+	TArray<void*> Channels;
+
+	/** Implementation of the channel interface that can interact with the typed channels */
+	TInlineValue<IBatchChannelInterface, 8> SectionDataInterface;
+};
+
+
+
+
+/**
+ * Construction helper that is required to create a new FMovieSceneChannelProxy from multiple channels
+ */
+struct FMovieSceneChannelData
+{
+#if WITH_EDITOR
+
+	/**
+	 * Add a new channel to the proxy. Channel's address is stored internally as a voic* and should exist as long as the channel proxy does
+	 *
+	 * @param InChannel          The channel to add to this proxy. Should live for as long as the proxy does. Any re-allocation should be accompanied with a re-creation of the proxy
+	 * @param InCommonEditorData The editor data to be associated with this channel
+	 */
+	template<typename ChannelType>
+	void Add(ChannelType& InChannel, const FMovieSceneChannelEditorData& InCommonEditorData)
+	{
+		static_assert(TIsSame<typename TMovieSceneChannelTraits<ChannelType>::EditorDataType, void>::Value, "Must supply typed editor data according to the channel's traits.");
+
+		// Add the channel
+		const int32 ChannelTypeIndex = AddInternal(InChannel);
+		// Add the editor data at the same index
+		Entries[ChannelTypeIndex].AddEditorData<ChannelType>(InCommonEditorData);
+	}
+
+	/**
+	 * Add a new channel to the proxy. Channel's address is stored internally as a voic* and should exist as long as the channel proxy does
+	 *
+	 * @param InChannel          The channel to add to this proxy. Should live for as long as the proxy does. Any re-allocation should be accompanied with a re-creation of the proxy
+	 * @param InCommonEditorData The editor data to be associated with this channel
+	 * @param InTypedEditorData  Additional editor data to be associated with this channel as per its traits
+	 */
+	template<typename ChannelType, typename TypedEditorData>
+	void Add(ChannelType& InChannel, const FMovieSceneChannelEditorData& InCommonEditorData, TypedEditorData&& InTypedEditorData)
+	{
+		// Add the channel
+		const int32 ChannelTypeIndex = AddInternal(InChannel);
+		// Add the editor data at the same index
+		Entries[ChannelTypeIndex].AddEditorData<ChannelType>(InCommonEditorData, Forward<TypedEditorData>(InTypedEditorData));
+	}
+
+#else
+
+	/**
+	 * Add a new channel to the proxy. Channel's address is stored internally as a voic* and should exist as long as the channel proxy does
+	 *
+	 * @param InChannel    The channel to add to this proxy. Should live for as long as the proxy does. Any re-allocation should be accompanied with a re-creation of the proxy
+	 */
+	template<typename ChannelType>
+	void Add(ChannelType& InChannel)
+	{
+		AddInternal(InChannel);
+	}
+
+#endif
+
+private:
+
+	/**
+	 * Implementation that adds a channel to an entry, creating a new entry for this channel type if necessary
+	 *
+	 * @param InChannel    The channel to add to this proxy. Should live for as long as the proxy does. Any re-allocation should be accompanied with a re-creation of the proxy
+	 */
+	template<typename ChannelType>
+	int32 AddInternal(ChannelType& InChannel);
+
+	friend struct FMovieSceneChannelProxy;
+	/** Array of entryies, one per channel type. Inline allocation space for one entry since most sections only have one channel type. */
+	TArray<FMovieSceneChannelEntry, TInlineAllocator<1>> Entries;
+};
+
+
+/**
+ * Proxy type stored inside UMovieSceneSection for access to all its channels. Construction via either a single channel, or a FMovieSceneChannelData structure
+ * This proxy exists as a generic accessor to any channel data existing in derived types
+ */
+struct MOVIESCENE_API FMovieSceneChannelProxy : TSharedFromThis<FMovieSceneChannelProxy>
+{
+public:
+
+	/** Default construction - emtpy proxy */
+	FMovieSceneChannelProxy(){}
+
+	/**
+	 * Construction via multiple channels
+	 */
+	FMovieSceneChannelProxy(FMovieSceneChannelData&& InChannels)
+		: Entries(MoveTemp(InChannels.Entries))
+	{}
+
+	/** Not copyable or moveable to ensure that previously retrieved pointers remain valid for the lifetime of this object. */
+	FMovieSceneChannelProxy(const FMovieSceneChannelProxy&) = delete;
+	FMovieSceneChannelProxy& operator=(const FMovieSceneChannelProxy&) = delete;
+
+	FMovieSceneChannelProxy(FMovieSceneChannelProxy&&) = delete;
+	FMovieSceneChannelProxy& operator=(FMovieSceneChannelProxy&&) = delete;
+
+public:
+
+	/**
+	 * Const access to all the entries in this proxy
+	 *
+	 * @return Array view of all this proxy's entries (channels grouped by type)
+	 */
+	TArrayView<const FMovieSceneChannelEntry> GetAllEntries() const
+	{
+		return Entries;
+	}
+
+	/**
+	 * Find an entry by its channel type ID
+	 *
+	 * @return A pointer to the channel, or nullptr
+	 */
+	const FMovieSceneChannelEntry* FindEntry(uint32 ChannelTypeID) const;
+
+	/**
+	 * Find the index of the specified channel ptr in this proxy
+	 *
+	 * @param ChannelTypeID    The type ID of the channel
+	 * @param ChannelPtr       The channel pointer to find
+	 * @return The index of the channel if found, else INDEX_NONE
+	 */
+	int32 FindIndex(uint32 ChannelTypeID, void* ChannelPtr) const;
+
+	/**
+	 * Get all channels of the specified type
+	 *
+	 * @return A possibly empty array view of all the channels in this proxy that match the template type
+	 */
+	template<typename ChannelType>
+	TArrayView<ChannelType*> GetChannels() const;
+
+	/**
+	 * Get the channel for the specified index of a particular type.
+	 *
+	 * @return A pointer to the channel, or nullptr if the index was invalid, or the type was not present
+	 */
+	template<typename ChannelType>
+	ChannelType* GetChannel(int32 ChannelIndex) const;
+
+	/**
+	 * Get the channel for the specified index of a particular type.
+	 *
+	 * @return A pointer to the channel, or nullptr if the index was invalid, or the type was not present
+	 */
+	void* GetChannel(uint32 ChannelTypeID, int32 ChannelIndex) const;
+
+	/**
+	 * Make a channel handle out of the specified raw pointer that will become nullptr if this proxy is re-allocated
+	 *
+	 * @return A handle to the supplied channel that will become nullptr when the proxy is reallocated
+	 */
+	template<typename ChannelType>
+	TMovieSceneChannelHandle<ChannelType> MakeHandle(ChannelType* RawChannelPtr) const;
+
+#if !WITH_EDITOR
+
+
+	/**
+	 * Construction via a single channel, and its editor data
+	 * Channel's address is stored internally as a voic* and should exist as long as this channel proxy does.
+	 */
+	template<typename ChannelType>
+	FMovieSceneChannelProxy(ChannelType& InChannel);
+
+#else
+
+	/**
+	 * Construction via a single channel, and its editor data
+	 * Channel's address is stored internally as a voic* and should exist as long as this channel proxy does.
+	 */
+	template<typename ChannelType>
+	FMovieSceneChannelProxy(ChannelType& InChannel, const FMovieSceneChannelEditorData& InCommonEditorData);
+
+	/**
+	 * Construction via a single channel, and its editor data
+	 * Channel's address is stored internally as a voic* and should exist as long as this channel proxy does.
+	 */
+	template<typename ChannelType, typename TypedEditorData>
+	FMovieSceneChannelProxy(ChannelType& InChannel, const FMovieSceneChannelEditorData& InCommonEditorData, TypedEditorData&& InTypedEditorData);
+
+
+	/**
+	 * Access all the common editor data for the templated channel type
+	 *
+	 * @return A potentially empty array view for all the common editor data for the template channel type
+	 */
+	template<typename ChannelType>
+	TArrayView<const FMovieSceneChannelEditorData> GetEditorData() const;
+
+	/**
+	 * Access all the specialized data for the templated channel type
+	 *
+	 * @return A potentially empty array view for all the specialized editor data for the template channel type
+	 */
+	template<typename ChannelType>
+	TArrayView<const typename TMovieSceneChannelTraits<ChannelType>::EditorDataType> GetAllSpecializedEditorData() const;
+
+#endif
+
+private:
+
+	/** Do not expose shared-ownership semantics of this object */
+	using TSharedFromThis::AsShared;
+
+	/** Array of channel entries, one per channel type. Should never be changed or reallocated after construction to keep pointers alive. */
+	TArray<FMovieSceneChannelEntry, TInlineAllocator<1>> Entries;
+};
+
+
+/**
+ * Implementation that adds a channel to an entry, creating a new entry for this channel type if necessary
+ *
+ * @param InChannel    The channel to add to this proxy. Should live for as long as the proxy does. Any re-allocation should be accompanied with a re-creation of the proxy
+ */
+template<typename ChannelType>
+int32 FMovieSceneChannelData::AddInternal(ChannelType& InChannel)
+{
+	// Find the entry for this channel's type
+	const uint32 ChannelID = ChannelType::GetChannelID();
+
+	// Find the first entry that has a >= channel ID
+	int32 ChannelTypeIndex = Algo::LowerBoundBy(Entries, ChannelID, &FMovieSceneChannelEntry::ChannelID);
+
+	// If the index we found isn't valid, or it's not the channel we want, we need to add a new entry there
+	if (ChannelTypeIndex >= Entries.Num() || Entries[ChannelTypeIndex].GetChannelID() != ChannelID)
+	{
+		Entries.Insert(FMovieSceneChannelEntry(ChannelID, InChannel), ChannelTypeIndex);
+	}
+
+	check(Entries.IsValidIndex(ChannelTypeIndex));
+
+	// Add the channel to the channels array
+	Entries[ChannelTypeIndex].Channels.Add(&InChannel);
+	return ChannelTypeIndex;
+}
+
+
+/**
+ * Get all channels of the specified type
+ *
+ * @return A possibly empty array view of all the channels in this proxy that match the template type
+ */
+template<typename ChannelType>
+TArrayView<ChannelType*> FMovieSceneChannelProxy::GetChannels() const
+{
+	const uint32 ChannelTypeID = ChannelType::GetChannelID();
+	const FMovieSceneChannelEntry* FoundEntry = FindEntry(ChannelTypeID);
+
+	if (FoundEntry)
+	{
+		return TArrayView<ChannelType*>((ChannelType**)(FoundEntry->Channels.GetData()), FoundEntry->Channels.Num());
+	}
+	return TArrayView<ChannelType*>();
+}
+
+
+/**
+ * Get the channel for the specified index of a particular type.
+ *
+ * @return A pointer to the channel, or nullptr if the index was invalid, or the type was not present
+ */
+template<typename ChannelType>
+ChannelType* FMovieSceneChannelProxy::GetChannel(int32 ChannelIndex) const
+{
+	TArrayView<ChannelType*> Channels = GetChannels<ChannelType>();
+	return Channels.IsValidIndex(ChannelIndex) ? Channels[ChannelIndex] : nullptr;
+}
+
+
+/**
+ * Make a channel handle out of the specified raw pointer that will become nullptr if this proxy is re-allocated
+ *
+ * @return A handle to the supplied channel that will become nullptr when the proxy is reallocated
+ */
+template<typename ChannelType>
+TMovieSceneChannelHandle<ChannelType> FMovieSceneChannelProxy::MakeHandle(ChannelType* RawChannelPtr) const
+{
+	TWeakPtr<ChannelType> WeakChannel = TSharedPtr<ChannelType>(AsShared(), RawChannelPtr);
+	return TMovieSceneChannelHandle<ChannelType>(WeakChannel);
+}
+
+
+#if !WITH_EDITOR
+
+
+/**
+ * Construction via a single channel, and its editor data
+ * Channel's address is stored internally as a voic* and should exist as long as this channel proxy does.
+ */
+template<typename ChannelType>
+FMovieSceneChannelProxy::FMovieSceneChannelProxy(ChannelType& InChannel)
+{
+	const uint32 ChannelID = ChannelType::GetChannelID();
+	Entries.Add(FMovieSceneChannelEntry(ChannelID, InChannel));
+	Entries[0].Channels.Add(&InChannel);
+}
+
+
+#else	// !WITH_EDITOR
+
+
+/**
+ * Construction via a single channel, and its editor data
+ * Channel's address is stored internally as a voic* and should exist as long as this channel proxy does.
+ */
+template<typename ChannelType>
+FMovieSceneChannelProxy::FMovieSceneChannelProxy(ChannelType& InChannel, const FMovieSceneChannelEditorData& InCommonEditorData)
+{
+	static_assert(TIsSame<typename TMovieSceneChannelTraits<ChannelType>::EditorDataType, void>::Value, "Must supply typed editor data according to the channel's traits.");
+
+	const uint32 ChannelID = ChannelType::GetChannelID();
+	Entries.Add(FMovieSceneChannelEntry(ChannelID, InChannel));
+	Entries[0].Channels.Add(&InChannel);
+	Entries[0].AddEditorData<ChannelType>(InCommonEditorData);
+}
+
+
+/**
+ * Construction via a single channel, its editor data, and its specialized editor data. Compulsary for channel types with a TMovieSceneChannelTraits<ChannelType>::EditorDataType.
+ * Channel's address is stored internally as a voic* and should exist as long as this channel proxy does.
+ */
+template<typename ChannelType, typename TypedEditorData>
+FMovieSceneChannelProxy::FMovieSceneChannelProxy(ChannelType& InChannel, const FMovieSceneChannelEditorData& InCommonEditorData, TypedEditorData&& InTypedEditorData)
+{
+	const uint32 ChannelID = ChannelType::GetChannelID();
+	Entries.Add(FMovieSceneChannelEntry(ChannelID, InChannel));
+	Entries[0].Channels.Add(&InChannel);
+	Entries[0].AddEditorData<ChannelType>(InCommonEditorData, Forward<TypedEditorData>(InTypedEditorData));
+}
+
+
+/**
+ * Access all the specialized data for the templated channel type
+ *
+ * @return A potentially empty array view for all the specialized editor data for the template channel type
+ */
+template<typename ChannelType>
+TArrayView<const typename TMovieSceneChannelTraits<ChannelType>::EditorDataType> FMovieSceneChannelProxy::GetAllSpecializedEditorData() const
+{
+	const uint32 ChannelTypeID = ChannelType::GetChannelID();
+	const FMovieSceneChannelEntry* FoundEntry = FindEntry(ChannelTypeID);
+
+	if (FoundEntry)
+	{
+		return FoundEntry->GetAllSpecializedEditorData<ChannelType>();
+	}
+	return TArrayView<const typename TMovieSceneChannelTraits<ChannelType>::EditorDataType>();
+}
+
+
+/**
+ * Access all the common editor data for the templated channel type
+ *
+ * @return A potentially empty array view for all the common editor data for the template channel type
+ */
+template<typename ChannelType>
+TArrayView<const FMovieSceneChannelEditorData> FMovieSceneChannelProxy::GetEditorData() const
+{
+	const uint32 ChannelTypeID = ChannelType::GetChannelID();
+	const FMovieSceneChannelEntry* FoundEntry = FindEntry(ChannelTypeID);
+
+	if (FoundEntry)
+	{
+		return FoundEntry->GetCommonEditorData();
+	}
+	return TArrayView<const FMovieSceneChannelEditorData>();
+}
+
+#endif	// !WITH_EDITOR
\ No newline at end of file
