Copying //UE4/Dev-Sequencer to //UE4/Dev-Main (Source: //UE4/Dev-Sequencer @ 3946692)

diff --git a/Engine/Source/Runtime/MovieScene/Public/Channels/MovieSceneChannelTraits.h b/Engine/Source/Runtime/MovieScene/Public/Channels/MovieSceneChannelTraits.h
new file mode 100644
--- /dev/null
+++ b/Engine/Source/Runtime/MovieScene/Public/Channels/MovieSceneChannelTraits.h
@@ -0,0 +1,270 @@
+// Copyright 1998-2018 Epic Games, Inc. All Rights Reserved.
+
+#pragma once
+
+#include "Containers/ArrayView.h"
+#include "FrameNumber.h"
+#include "FrameTime.h"
+#include "FrameRate.h"
+#include "MovieSceneChannelEditorData.h"
+
+enum class EMovieSceneKeyInterpolation : uint8;
+
+/**
+ * Parameter structure passed to MovieScene::Optimize that defines optimization parameters
+ */
+struct FKeyDataOptimizationParams
+{
+	/** An arbitrary tolerance under which keys should be optimizied */
+	float Tolerance = KINDA_SMALL_NUMBER;
+
+	/** Whether to automatically set interpolation tangents or not */
+	bool bAutoSetInterpolation = false;
+
+	/** A range inside which to optimize keys */
+	TRange<FFrameNumber> Range = TRange<FFrameNumber>::All();
+};
+
+/**
+ * Traits structure to be specialized for any channel type passed to FMovieSceneChannelProxy
+ */
+template<typename ChannelType>
+struct TMovieSceneChannelTraitsBase
+{
+};
+
+/**
+ * Traits structure to be specialized for any channel type passed to FMovieSceneChannelProxy
+ */
+template<typename ChannelType>
+struct TMovieSceneChannelTraits : TMovieSceneChannelTraitsBase<ChannelType>
+{
+#if WITH_EDITOR
+
+	/** Type that specifies what editor data should be associated with ChannelType */
+	typedef void EditorDataType;
+
+#endif
+};
+
+namespace MovieScene
+{
+	/**
+	 * Called to evaluate a channel. Overload with specific channel types for custom behaviour.
+	 *
+	 * @param InChannel   The channel to evaluate
+	 * @param InTime      The time to evaluate at
+	 * @param OutValue    Value to receive the result
+	 * @return true if the channel was evaluated successfully, false otherwise
+	 */
+	template<typename ChannelType, typename ValueType>
+	bool EvaluateChannel(const ChannelType* InChannel, FFrameTime InTime, ValueType& OutValue)
+	{
+		return InChannel->Evaluate(InTime, OutValue);
+	}
+
+	/**
+	 * Called to assign a specific value in a channel.
+	 *
+	 * @param InChannel     The channel the value is contained within
+	 * @param InValueIndex  The index of the value to set
+	 * @param InValue       The new value
+	 */
+	template<typename ChannelType, typename ValueType>
+	void AssignValue(ChannelType* InChannel, int32 InValueIndex, ValueType&& InValue)
+	{
+		InChannel->GetInterface().GetValues()[InValueIndex] = Forward<ValueType>(InValue);
+	}
+
+	/**
+	 * Add a key to a channel, or update an existing key if one already exists at this time
+	 *
+	 * @param InChannel     The channel the value is contained within
+	 * @param InTime        The time to add or update the key at
+	 * @param InValue       The new value
+	 * @return A handle to the key that was added
+	 */
+	template<typename ChannelType, typename ValueType>
+	FKeyHandle AddKeyToChannel(ChannelType* InChannel, FFrameNumber InTime, ValueType&& Value, EMovieSceneKeyInterpolation Interpolation)
+	{
+		auto ChannelInterface = InChannel->GetInterface();
+		int32 ExistingIndex = ChannelInterface.FindKey(InTime);
+		if (ExistingIndex != INDEX_NONE)
+		{
+			AssignValue(InChannel, ExistingIndex, Forward<ValueType>(Value));
+		}
+		else
+		{
+			ExistingIndex = ChannelInterface.AddKey(InTime, Forward<ValueType>(Value));
+		}
+		return ChannelInterface.GetHandle(ExistingIndex);
+	}
+
+	/**
+	 * Check whether the specified value already exists at the specified time
+	 *
+	 * @param InChannel     The channel to check
+	 * @param InTime        The time to check for
+	 * @param InValue       The value to check - compared against the curve's existing value at this time
+	 * @return true if this value already exists at the time, false otherwise
+	 */
+	template<typename ChannelType, typename ValueType>
+	bool ValueExistsAtTime(const ChannelType* InChannel, FFrameNumber InTime, const ValueType& InValue)
+	{
+		ValueType ExistingValue{};
+		return EvaluateChannel(InChannel, InTime, ExistingValue) && ExistingValue == InValue;
+	}
+
+	/**
+	 * Compute the effective range of the specified channel. Generally just means the range of its keys.
+	 *
+	 * @param InChannel     The channel to compute the range for
+	 * @return The range of this channel
+	 */
+	template<typename ChannelType>
+	TRange<FFrameNumber> ComputeEffectiveRange(const ChannelType* InChannel)
+	{
+		TArrayView<const FFrameNumber> Times = InChannel->GetInterface().GetTimes();
+		return Times.Num() ? TRange<FFrameNumber>(Times[0], TRangeBound<FFrameNumber>::Inclusive(Times[Times.Num()-1])) : TRange<FFrameNumber>::Empty();
+	}
+
+	/**
+	 * Convert the frame resolution of a movie scene channel by moving the key times to the equivalent frame time
+	 *
+	 * @param InChannel       The channel to compute the range for
+	 * @param SourceRate      The frame rate the channel is currently in
+	 * @param DestinationRate The new frame rate to convert the channel to
+	 */
+	template<typename ChannelType>
+	void ChangeFrameResolution(ChannelType* InChannel, FFrameRate SourceRate, FFrameRate DestinationRate)
+	{
+		TArrayView<FFrameNumber> Times = InChannel->GetInterface().GetTimes();
+		for (int32 Index = 0; Index < Times.Num(); ++Index)
+		{
+			Times[Index] = ConvertFrameTime(Times[Index], SourceRate, DestinationRate).RoundToFrame();
+		}
+	}
+
+	/**
+	 * Get the number of keys contained within the specified channel
+	 *
+	 * @param InChannel     The channel to get the number of keys for
+	 * @return The number of keys in this channel
+	 */
+	template<typename ChannelType>
+	int32 GetNumKeys(const ChannelType* InChannel)
+	{
+		return InChannel->GetInterface().GetTimes().Num();
+	}
+
+	/**
+	 * Reset the specified channel back to its default state
+	 *
+	 * @param InChannel     The channel to reset
+	 */
+	template<typename ChannelType>
+	void Reset(ChannelType* InChannel)
+	{
+		*InChannel = ChannelType();
+	}
+
+	/**
+	 * Offset the specified channel by a given delta time
+	 *
+	 * @param InChannel     The channel to offset
+	 * @param DeltaTime     The time to offset by
+	 */
+	template<typename ChannelType>
+	void Offset(ChannelType* InChannel, FFrameNumber DeltaTime)
+	{
+		TArrayView<FFrameNumber> Times = InChannel->GetInterface().GetTimes();
+		for (FFrameNumber& Time : Times)
+		{
+			Time += DeltaTime;
+		}
+	}
+
+	/**
+	 * Dilate the specified channel with a given factor and origin
+	 *
+	 * @param InChannel      The channel to offset
+	 * @param Origin         The origin around which key times should scale
+	 * @param DilationFactor The amount to dilate by
+	 */
+	template<typename ChannelType>
+	void Dilate(ChannelType* InChannel, FFrameNumber Origin, float DilationFactor)
+	{
+		// @todo: sequencer-timecode: need to rethink this
+		// TArrayView<FFrameNumber> Times = InChannel->GetInterface().GetTimes();
+		// for (FFrameNumber& Time : Times)
+		// {
+		// 	Time = Origin + (Time - Origin) * DilationFactor;
+		// }
+	}
+
+	/**
+	 * Set a channel's default value
+	 *
+	 * @param InChannel      The channel to set the default on
+	 * @param DefaultValue   The new default value
+	 */
+	template<typename ChannelType, typename ValueType>
+	void SetChannelDefault(ChannelType* Channel, ValueType&& DefaultValue)
+	{
+		Channel->SetDefault(DefaultValue);
+	}
+
+	/**
+	 * Clear a channel's default value
+	 *
+	 * @param InChannel      The channel to clear the default on
+	 */
+	template<typename ChannelType>
+	void ClearChannelDefault(ChannelType* InChannel)
+	{
+		InChannel->RemoveDefault();
+	}
+
+	/**
+	 * Optimize the specified channel by removing any redundant keys
+	 *
+	 * @param InChannel     The channel to optimize
+	 * @param Params        Optimization parameters
+	 */
+	template<typename ChannelType>
+	void Optimize(ChannelType* InChannel, const FKeyDataOptimizationParams& Params)
+	{
+		using namespace MovieScene;
+
+		auto ChannelInterface = InChannel->GetInterface();
+		if (ChannelInterface.GetTimes().Num() > 1)
+		{
+			int32 StartIndex = 0;
+			int32 EndIndex = 0;
+
+			{
+				TArrayView<const FFrameNumber> Times = ChannelInterface.GetTimes();
+				StartIndex = Params.Range.GetLowerBound().IsClosed() ? Algo::LowerBound(Times, Params.Range.GetLowerBoundValue()) : 0;
+				EndIndex   = Params.Range.GetUpperBound().IsClosed() ? Algo::UpperBound(Times, Params.Range.GetUpperBoundValue()) : Times.Num();
+			}
+
+			for (int32 Index = StartIndex; Index < EndIndex && Index < ChannelInterface.GetTimes().Num(); ++Index)
+			{
+				// Reget times and values as they may be reallocated
+				FFrameNumber Time   = ChannelInterface.GetTimes()[Index];
+				auto  OriginalValue = ChannelInterface.GetValues()[Index];
+
+				// If the channel evaluates the same with this key removed, we can leave it out
+				ChannelInterface.RemoveKey(Index);
+				if (ValueExistsAtTime(InChannel, Time, OriginalValue))
+				{
+					Index--;
+				}
+				else
+				{
+					ChannelInterface.AddKey(Time, MoveTemp(OriginalValue));
+				}
+			}
+		}
+	}
+}
\ No newline at end of file
diff --git a/Engine/Source/Runtime/MovieScene/Public/Channels/MovieSceneFloatChannel.h b/Engine/Source/Runtime/MovieScene/Public/Channels/MovieSceneFloatChannel.h
new file mode 100644
--- /dev/null
+++ b/Engine/Source/Runtime/MovieScene/Public/Channels/MovieSceneFloatChannel.h
@@ -0,0 +1,305 @@
+// Copyright 1998-2018 Epic Games, Inc. All Rights Reserved.
+
+#pragma once
+
+#include "UObject/ObjectMacros.h"
+#include "FrameNumber.h"
+#include "MovieSceneChannel.h"
+#include "MovieSceneChannelTraits.h"
+#include "KeyParams.h"
+
+#include "Curves/RichCurve.h"
+
+#include "MovieSceneFloatChannel.generated.h"
+
+
+USTRUCT()
+struct FMovieSceneTangentData
+{
+	GENERATED_BODY()
+
+	FMovieSceneTangentData()
+		: ArriveTangent(0.f)
+		, LeaveTangent(0.f)
+	{}
+
+	/** If RCIM_Cubic, the arriving tangent at this key */
+	UPROPERTY(EditAnywhere, Category="Key")
+	float ArriveTangent;
+
+	/** If RCIM_Cubic, the leaving tangent at this key */
+	UPROPERTY(EditAnywhere, Category="Key")
+	float LeaveTangent;
+};
+
+USTRUCT()
+struct FMovieSceneFloatValue
+{
+	GENERATED_BODY()
+
+	FMovieSceneFloatValue()
+		: Value(0.f), InterpMode(RCIM_Cubic), TangentMode(RCTM_Auto)
+	{}
+
+	explicit FMovieSceneFloatValue(float InValue)
+		: Value(InValue), InterpMode(RCIM_Cubic), TangentMode(RCTM_Auto)
+	{}
+
+	UPROPERTY(EditAnywhere, Category="Key")
+	float Value;
+
+	UPROPERTY(EditAnywhere, Category="Key")
+	TEnumAsByte<ERichCurveInterpMode> InterpMode;
+
+	UPROPERTY(EditAnywhere, Category="Key")
+	TEnumAsByte<ERichCurveTangentMode> TangentMode;
+
+	UPROPERTY(EditAnywhere, Category="Key")
+	FMovieSceneTangentData Tangent;
+};
+
+USTRUCT()
+struct FMovieSceneFloatChannel
+{
+	GENERATED_BODY()
+
+	FMovieSceneFloatChannel()
+		: PreInfinityExtrap(RCCE_Constant), PostInfinityExtrap(RCCE_Constant), DefaultValue(), bHasDefaultValue(false)
+	{}
+
+	/**
+	 * Access an integer that uniquely identifies this channel type.
+	 *
+	 * @return A static identifier that was allocated using FMovieSceneChannelEntry::RegisterNewID
+	 */
+	MOVIESCENE_API static uint32 GetChannelID();
+
+	/**
+	 * Access a mutable interface for this channel
+	 *
+	 * @return An object that is able to manipulate this channel
+	 */
+	FORCEINLINE TMovieSceneChannel<FMovieSceneFloatValue> GetInterface()
+	{
+		return TMovieSceneChannel<FMovieSceneFloatValue>(&Times, &Values, &KeyHandles);
+	}
+
+	/**
+	 * Access a constant interface for this channel
+	 *
+	 * @return An object that is able to interrogate this channel
+	 */
+	FORCEINLINE TMovieSceneChannel<const FMovieSceneFloatValue> GetInterface() const
+	{
+		return TMovieSceneChannel<const FMovieSceneFloatValue>(&Times, &Values);
+	}
+
+	/**
+	 * Const access to this channel's times
+	 */
+	FORCEINLINE TArrayView<const FFrameNumber> GetTimes() const
+	{
+		return Times;
+	}
+
+	/**
+	 * Const access to this channel's values
+	 */
+	FORCEINLINE TArrayView<const FMovieSceneFloatValue> GetValues() const
+	{
+		return Values;
+	}
+
+	/**
+	 * Evaluate this channel
+	 *
+	 * @param InTime     The time to evaluate at
+	 * @param OutValue   A value to receive the result
+	 * @return true if the channel was evaluated successfully, false otherwise
+	 */
+	MOVIESCENE_API bool Evaluate(FFrameTime InTime, float& OutValue) const;
+
+	/**
+	 * Set the channel's times and values to the requested values
+	 */
+	FORCEINLINE void Set(TArray<FFrameNumber> InTimes, TArray<FMovieSceneFloatValue> InValues)
+	{
+		check(InTimes.Num() == InValues.Num());
+
+		Times = MoveTemp(InTimes);
+		Values = MoveTemp(InValues);
+
+		KeyHandles.Reset();
+		for (int32 Index = 0; Index < Times.Num(); ++Index)
+		{
+			KeyHandles.AllocateHandle(Index);
+		}
+	}
+
+public:
+
+	/**
+	 * Set this channel's default value that should be used when no keys are present
+	 *
+	 * @param InDefaultValue The desired default value
+	 */
+	FORCEINLINE void SetDefault(float InDefaultValue)
+	{
+		bHasDefaultValue = true;
+		DefaultValue = InDefaultValue;
+	}
+
+	/**
+	 * Get this channel's default value that will be used when no keys are present
+	 *
+	 * @return (Optional) The channel's default value
+	 */
+	FORCEINLINE TOptional<float> GetDefault() const
+	{
+		return bHasDefaultValue ? TOptional<float>(DefaultValue) : TOptional<float>();
+	}
+
+	/**
+	 * Remove this channel's default value causing the channel to have no effect where no keys are present
+	 */
+	FORCEINLINE void RemoveDefault()
+	{
+		bHasDefaultValue = false;
+	}
+
+public:
+
+	/** Serialize this float function from a mismatching property tag (FRichCurve) */
+	bool SerializeFromMismatchedTag(const FPropertyTag& Tag, FArchive& Ar);
+
+	MOVIESCENE_API int32 AddConstantKey(FFrameNumber InTime, float InValue);
+
+	MOVIESCENE_API int32 AddLinearKey(FFrameNumber InTime, float InValue);
+
+	MOVIESCENE_API int32 AddCubicKey(FFrameNumber InTime, float InValue, ERichCurveTangentMode TangentMode = RCTM_Auto, const FMovieSceneTangentData& Tangent = FMovieSceneTangentData());
+
+	MOVIESCENE_API void AutoSetTangents(float Tension = 0.f);
+
+	/**
+	 * Populate the specified array with times and values that represent the smooth interpolation of this channel across the specified range
+	 *
+	 * @param StartTimeSeconds      The first time in seconds to include in the resulting array
+	 * @param EndTimeSeconds        The last time in seconds to include in the resulting array
+	 * @param TimeThreshold         A small time threshold in seconds below which we should stop adding new points
+	 * @param ValueThreshold        A small value threshold below which we should stop adding new points where the linear interpolation would suffice
+	 * @param FrameResolution       The frame resolution with which to interpret this channel's times
+	 * @param InOutPoints           An array to populate with the evaluated points
+	 */
+	MOVIESCENE_API void PopulateCurvePoints(double StartTimeSeconds, double EndTimeSeconds, double TimeThreshold, float ValueThreshold, FFrameRate FrameResolution, TArray<TTuple<double, double>>& InOutPoints) const;	
+
+public:
+
+	/** Pre-infinity extrapolation state */
+	UPROPERTY()
+	TEnumAsByte<ERichCurveExtrapolation> PreInfinityExtrap;
+
+	/** Post-infinity extrapolation state */
+	UPROPERTY()
+	TEnumAsByte<ERichCurveExtrapolation> PostInfinityExtrap;
+
+private:
+
+	int32 InsertKeyInternal(FFrameNumber InTime);
+
+	/**
+	 * Evaluate this channel's extrapolation. Assumes more than 1 key is present.
+	 *
+	 * @param InTime     The time to evaluate at
+	 * @param OutValue   A value to receive the result
+	 * @return true if the time was evaluated with extrapolation, false otherwise
+	 */
+	bool EvaluateExtrapolation(FFrameTime InTime, float& OutValue) const;
+
+	/**
+	 * Adds median points between each of the supplied points if their evaluated value is significantly different than the linear interpolation of those points
+	 *
+	 * @param FrameResolution       The frame resolution with which to interpret this channel's times
+	 * @param TimeThreshold         A small time threshold in seconds below which we should stop adding new points
+	 * @param ValueThreshold        A small value threshold below which we should stop adding new points where the linear interpolation would suffice
+	 * @param InOutPoints           An array to populate with the evaluated points
+	 */
+	void RefineCurvePoints(FFrameRate FrameResolution, double TimeThreshold, float ValueThreshold, TArray<TTuple<double, double>>& InOutPoints) const;
+
+	UPROPERTY()
+	TArray<FFrameNumber> Times;
+
+	UPROPERTY()
+	TArray<FMovieSceneFloatValue> Values;
+
+	UPROPERTY()
+	float DefaultValue;
+
+	UPROPERTY()
+	bool bHasDefaultValue;
+
+	FMovieSceneKeyHandleMap KeyHandles;
+};
+
+template<>
+struct TStructOpsTypeTraits<FMovieSceneFloatChannel> : public TStructOpsTypeTraitsBase2<FMovieSceneFloatChannel>
+{
+	enum { WithSerializeFromMismatchedTag = true };
+};
+
+template<>
+struct TMovieSceneChannelTraits<FMovieSceneFloatChannel>
+{
+#if WITH_EDITORONLY_DATA
+
+	/** Float channels can have external values (ie, they can get their values from external objects for UI purposes) */
+	typedef TMovieSceneExternalValue<float> EditorDataType;
+
+#endif
+};
+
+inline bool EvaluateChannel(const FMovieSceneFloatChannel* InChannel, float InTime, FMovieSceneFloatValue& OutValue)
+{
+	checkf(false, TEXT("Unimplemented: Please use EvaluateChannel(const FMovieSceneFloatChannel&, float, float&) instead."));
+	return false;
+}
+
+inline bool ValueExistsAtTime(const FMovieSceneFloatChannel* Channel, FFrameNumber InFrameNumber, float Value, float Tolerance = KINDA_SMALL_NUMBER)
+{
+	const FFrameTime FrameTime(InFrameNumber);
+
+	float ExistingValue = 0.f;
+	return Channel->Evaluate(FrameTime, ExistingValue) && FMath::IsNearlyEqual(ExistingValue, Value, Tolerance);
+}
+
+inline bool ValueExistsAtTime(const FMovieSceneFloatChannel* Channel, FFrameNumber InFrameNumber, const FMovieSceneFloatValue& InValue, float Tolerance = KINDA_SMALL_NUMBER)
+{
+	return ValueExistsAtTime(Channel, InFrameNumber, InValue.Value, Tolerance);
+}
+
+inline void AssignValue(FMovieSceneFloatChannel* InChannel, int32 ValueIndex, float InValue)
+{
+	InChannel->GetInterface().GetValues()[ValueIndex].Value = InValue;
+}
+
+/**
+ * Overload for adding a new key to a float channel at a given time. See MovieScene::AddKeyToChannel for default implementation.
+ */
+MOVIESCENE_API FKeyHandle AddKeyToChannel(FMovieSceneFloatChannel* Channel, FFrameNumber InFrameNumber, float InValue, EMovieSceneKeyInterpolation Interpolation);
+
+
+/**
+ * Overload for optimizing key data for a float channel (removing superfluous keys). See MovieScene::Optimize for default implementation.
+ */
+MOVIESCENE_API void Optimize(FMovieSceneFloatChannel* InChannel, const FKeyDataOptimizationParams& Params);
+
+
+/**
+ * Overload for dilating float channel data. See MovieScene::Dilate for default implementation.
+ */
+MOVIESCENE_API void Dilate(FMovieSceneFloatChannel* InChannel, FFrameNumber Origin, float DilationFactor);
+
+
+/**
+ * Overload for Converting a float channel from one frame rate to another. Ensures that both key times and tangents are migrated correctly. See MovieScene::ChangeFrameResolution for default implementation
+ */
+MOVIESCENE_API void ChangeFrameResolution(FMovieSceneFloatChannel* InChannel, FFrameRate SourceRate, FFrameRate DestinationRate);
\ No newline at end of file
diff --git a/Engine/Source/Runtime/MovieScene/Public/Channels/MovieSceneIntegerChannel.h b/Engine/Source/Runtime/MovieScene/Public/Channels/MovieSceneIntegerChannel.h
new file mode 100644
--- /dev/null
+++ b/Engine/Source/Runtime/MovieScene/Public/Channels/MovieSceneIntegerChannel.h
@@ -0,0 +1,141 @@
+// Copyright 1998-2018 Epic Games, Inc. All Rights Reserved.
+
+#pragma once
+
+#include "UObject/ObjectMacros.h"
+#include "FrameNumber.h"
+#include "MovieSceneChannel.h"
+#include "MovieSceneChannelTraits.h"
+#include "MovieSceneIntegerChannel.generated.h"
+
+
+USTRUCT()
+struct FMovieSceneIntegerChannel
+{
+	GENERATED_BODY()
+
+	FMovieSceneIntegerChannel()
+		: DefaultValue(), bHasDefaultValue(false)
+	{}
+
+	/**
+	 * Serialize this type from another
+	 */
+	MOVIESCENE_API bool SerializeFromMismatchedTag(const FPropertyTag& Tag, FArchive& Ar);
+
+	/**
+	 * Access an integer that uniquely identifies this channel type.
+	 *
+	 * @return A static identifier that was allocated using FMovieSceneChannelEntry::RegisterNewID
+	 */
+	MOVIESCENE_API static uint32 GetChannelID();
+
+	/**
+	 * Access a mutable interface for this channel
+	 *
+	 * @return An object that is able to manipulate this channel
+	 */
+	FORCEINLINE TMovieSceneChannel<int32> GetInterface()
+	{
+		return TMovieSceneChannel<int32>(&Times, &Values, &KeyHandles);
+	}
+
+	/**
+	 * Access a constant interface for this channel
+	 *
+	 * @return An object that is able to interrogate this channel
+	 */
+	FORCEINLINE TMovieSceneChannel<const int32> GetInterface() const
+	{
+		return TMovieSceneChannel<const int32>(&Times, &Values);
+	}
+
+	/**
+	 * Const access to this channel's times
+	 */
+	FORCEINLINE TArrayView<const FFrameNumber> GetTimes() const
+	{
+		return Times;
+	}
+
+	/**
+	 * Const access to this channel's values
+	 */
+	FORCEINLINE TArrayView<const int32> GetValues() const
+	{
+		return Values;
+	}
+
+	/**
+	 * Evaluate this channel
+	 *
+	 * @param InTime     The time to evaluate at
+	 * @param OutValue   A value to receive the result
+	 * @return true if the channel was evaluated successfully, false otherwise
+	 */
+	MOVIESCENE_API bool Evaluate(FFrameTime InTime, int32& OutValue) const;
+
+public:
+
+	/**
+	 * Set this channel's default value that should be used when no keys are present
+	 *
+	 * @param InDefaultValue The desired default value
+	 */
+	FORCEINLINE void SetDefault(int32 InDefaultValue)
+	{
+		bHasDefaultValue = true;
+		DefaultValue = InDefaultValue;
+	}
+
+	/**
+	 * Get this channel's default value that will be used when no keys are present
+	 *
+	 * @return (Optional) The channel's default value
+	 */
+	FORCEINLINE TOptional<int32> GetDefault() const
+	{
+		return bHasDefaultValue ? TOptional<int32>(DefaultValue) : TOptional<int32>();
+	}
+
+	/**
+	 * Remove this channel's default value causing the channel to have no effect where no keys are present
+	 */
+	FORCEINLINE void RemoveDefault()
+	{
+		bHasDefaultValue = false;
+	}
+
+private:
+
+	UPROPERTY()
+	TArray<FFrameNumber> Times;
+
+	UPROPERTY()
+	int32 DefaultValue;
+
+	UPROPERTY()
+	bool bHasDefaultValue;
+
+	UPROPERTY()
+	TArray<int32> Values;
+
+	FMovieSceneKeyHandleMap KeyHandles;
+};
+
+template<>
+struct TStructOpsTypeTraits<FMovieSceneIntegerChannel> : public TStructOpsTypeTraitsBase2<FMovieSceneIntegerChannel>
+{
+	enum { WithSerializeFromMismatchedTag = true };
+};
+
+template<>
+struct TMovieSceneChannelTraits<FMovieSceneIntegerChannel>
+{
+#if WITH_EDITORONLY_DATA
+
+	/** Integer channels can have external values (ie, they can get their values from external objects for UI purposes) */
+	typedef TMovieSceneExternalValue<int32> EditorDataType;
+
+#endif
+};
\ No newline at end of file
diff --git a/Engine/Source/Runtime/MovieScene/Public/Compilation/MovieSceneCompiler.h b/Engine/Source/Runtime/MovieScene/Public/Compilation/MovieSceneCompiler.h
--- a/Engine/Source/Runtime/MovieScene/Public/Compilation/MovieSceneCompiler.h
+++ b/Engine/Source/Runtime/MovieScene/Public/Compilation/MovieSceneCompiler.h
@@ -27,10 +27,10 @@ class UMovieSceneSubSection;
 struct FCompiledGroupResult
 {
 	/** Construction to an empty result for the specified range */
-	FCompiledGroupResult(TRange<float> InRange) : Range(InRange) {}
+	FCompiledGroupResult(TRange<FFrameNumber> InRange) : Range(InRange) {}
 
 	/** The range within the sequence that the goup and meta-data apply to */
-	TRange<float> Range;
+	TRange<FFrameNumber> Range;
 
 	/** Evaluation group specifying all the entities that occur in the sequence, sorted and redy for evaluation */
 	FMovieSceneEvaluationGroup Group;
@@ -58,7 +58,7 @@ public:
 	 * @param InTemplateStore 			A template store that contains (or knows how to fetch) evaluation templates for a given sequence
 	 * @return The compiled data, ready for insertion into an FMovieSceneEvaluationField
 	 */
-	MOVIESCENE_API static TOptional<FCompiledGroupResult> CompileTime(float InGlobalTime, UMovieSceneSequence& InCompileSequence, IMovieSceneSequenceTemplateStore& InTemplateStore);
+	MOVIESCENE_API static TOptional<FCompiledGroupResult> CompileTime(FFrameNumber InGlobalTime, UMovieSceneSequence& InCompileSequence, IMovieSceneSequenceTemplateStore& InTemplateStore);
 
 	/**
 	 * Compile the hierarchy data for the specified sequence into a specific hierarchy
@@ -81,7 +81,7 @@ private:
 
 	static void PopulateEvaluationGroup(FCompiledGroupResult& OutResult, const TArray<FCompileOnTheFlyData>& SortedCompileData);
 
-	static void PopulateMetaData(FCompiledGroupResult& OutResult, const FMovieSceneSequenceHierarchy& RootHierarchy, const TArray<FCompileOnTheFlyData>& SortedCompileData, TMovieSceneEvaluationTreeDataIterator<FMovieSceneSequenceID> SubSequences);
+	static void PopulateMetaData(FCompiledGroupResult& OutResult, const FMovieSceneSequenceHierarchy& RootHierarchy, IMovieSceneSequenceTemplateStore& TemplateStore, const TArray<FCompileOnTheFlyData>& SortedCompileData, TMovieSceneEvaluationTreeDataIterator<FMovieSceneSequenceID> SubSequences);
 
 private:
 
diff --git a/Engine/Source/Runtime/MovieScene/Public/Compilation/MovieSceneCompilerRules.h b/Engine/Source/Runtime/MovieScene/Public/Compilation/MovieSceneCompilerRules.h
--- a/Engine/Source/Runtime/MovieScene/Public/Compilation/MovieSceneCompilerRules.h
+++ b/Engine/Source/Runtime/MovieScene/Public/Compilation/MovieSceneCompilerRules.h
@@ -9,12 +9,12 @@
 
 namespace MovieSceneSegmentCompiler
 {
-	static TOptional<FMovieSceneSegment> EvaluateNearestSegment(const TRange<float>& Range, const FMovieSceneSegment* PreviousSegment, const FMovieSceneSegment* NextSegment)
+	static TOptional<FMovieSceneSegment> EvaluateNearestSegment(const TRange<FFrameNumber>& Range, const FMovieSceneSegment* PreviousSegment, const FMovieSceneSegment* NextSegment)
 	{
 		if (PreviousSegment)
 		{
 			// There is a preceeding segment
-			float PreviousSegmentRangeBound = PreviousSegment->Range.GetUpperBoundValue();
+			FFrameNumber PreviousSegmentRangeBound = PreviousSegment->Range.GetUpperBoundValue();
 
 			FMovieSceneSegment EmptySpace(Range);
 			for (FSectionEvaluationData Data : PreviousSegment->Impls)
@@ -26,7 +26,7 @@ namespace MovieSceneSegmentCompiler
 		else if (NextSegment)
 		{
 			// Before any sections
-			float NextSegmentRangeBound = NextSegment->Range.GetLowerBoundValue();
+			FFrameNumber NextSegmentRangeBound = NextSegment->Range.GetLowerBoundValue();
 
 			FMovieSceneSegment EmptySpace(Range);
 			for (FSectionEvaluationData Data : NextSegment->Impls)
@@ -127,7 +127,7 @@ struct FEvaluateNearestSegmentBlender : FMovieSceneTrackSegmentBlender
 		bCanFillEmptySpace = true;
 	}
 
-	virtual TOptional<FMovieSceneSegment> InsertEmptySpace(const TRange<float>& Range, const FMovieSceneSegment* PreviousSegment, const FMovieSceneSegment* NextSegment) const
+	virtual TOptional<FMovieSceneSegment> InsertEmptySpace(const TRange<FFrameNumber>& Range, const FMovieSceneSegment* PreviousSegment, const FMovieSceneSegment* NextSegment) const
 	{
 		return MovieSceneSegmentCompiler::EvaluateNearestSegment(Range, PreviousSegment, NextSegment);
 	}
diff --git a/Engine/Source/Runtime/MovieScene/Public/Compilation/MovieSceneSegmentCompiler.h b/Engine/Source/Runtime/MovieScene/Public/Compilation/MovieSceneSegmentCompiler.h
--- a/Engine/Source/Runtime/MovieScene/Public/Compilation/MovieSceneSegmentCompiler.h
+++ b/Engine/Source/Runtime/MovieScene/Public/Compilation/MovieSceneSegmentCompiler.h
@@ -114,7 +114,7 @@ struct FMovieSceneTrackSegmentBlender
 	 *
 	 * @return An optional new segment to define
 	 */
-	virtual TOptional<FMovieSceneSegment> InsertEmptySpace(const TRange<float>& Range, const FMovieSceneSegment* PreviousSegment, const FMovieSceneSegment* NextSegment) const
+	virtual TOptional<FMovieSceneSegment> InsertEmptySpace(const TRange<FFrameNumber>& Range, const FMovieSceneSegment* PreviousSegment, const FMovieSceneSegment* NextSegment) const
 	{
 		return TOptional<FMovieSceneSegment>();
 	}
diff --git a/Engine/Source/Runtime/MovieScene/Public/Evaluation/MovieSceneEvalTemplate.h b/Engine/Source/Runtime/MovieScene/Public/Evaluation/MovieSceneEvalTemplate.h
--- a/Engine/Source/Runtime/MovieScene/Public/Evaluation/MovieSceneEvalTemplate.h
+++ b/Engine/Source/Runtime/MovieScene/Public/Evaluation/MovieSceneEvalTemplate.h
@@ -135,7 +135,7 @@ public:
 	 * @param Container				Container to populate with the desired output from this track
 	 * @param BindingOverride		Optional binding to specify the object that is being animated by this track
 	 */
-	virtual void Interrogate(const FMovieSceneContext& Context, TRange<float> SweptRange, FMovieSceneInterrogationData& Container, UObject* BindingOverride) const
+	virtual void Interrogate(const FMovieSceneContext& Context, TRange<FFrameNumber> SweptRange, FMovieSceneInterrogationData& Container, UObject* BindingOverride) const
 	{
 	}
 
@@ -166,7 +166,7 @@ protected:
 	/**
 	 * Evaluate this template's easing functions based on the specified time
 	 */
-	MOVIESCENE_API float EvaluateEasing(float CurrentTime) const;
+	MOVIESCENE_API float EvaluateEasing(FFrameTime CurrentTime) const;
 
 	/**
 	 * Enum evaluation flag structure defining which functions are to be called in implementations of this struct
@@ -229,6 +229,10 @@ struct FMovieSceneEvalTemplatePtr
 			StructOps.Construct(Allocation);
 			StructOps.Copy(Allocation, &RHS.GetValue(), 1);
 		}
+		else
+		{
+			Reset();
+		}
 
 		return *this;
 	}
diff --git a/Engine/Source/Runtime/MovieScene/Public/Evaluation/MovieSceneEvaluationField.h b/Engine/Source/Runtime/MovieScene/Public/Evaluation/MovieSceneEvaluationField.h
--- a/Engine/Source/Runtime/MovieScene/Public/Evaluation/MovieSceneEvaluationField.h
+++ b/Engine/Source/Runtime/MovieScene/Public/Evaluation/MovieSceneEvaluationField.h
@@ -7,11 +7,13 @@
 #include "MovieSceneSequenceID.h"
 #include "Evaluation/MovieSceneEvaluationKey.h"
 #include "Evaluation/MovieSceneSegment.h"
+#include "MovieSceneFrameMigration.h"
 #include "Evaluation/MovieSceneTrackIdentifier.h"
 #include "MovieSceneEvaluationField.generated.h"
 
+struct FFrameNumber;
 struct FMovieSceneSequenceHierarchy;
-
+struct IMovieSceneSequenceTemplateStore;
 class UMovieSceneSequence;
 
 /** A pointer to a track held within an evaluation template */
@@ -189,11 +191,12 @@ struct FMovieSceneEvaluationMetaData
 	 *
 	 * @param RootSequence				The sequence that corresponds to this meta-data's root sequence
 	 * @param RootHierarchy				The hierarchy that corresponds to this meta-data's root sequence
+	 * @param TemplateStore				The template store used to retrieve templates for sub sequences
 	 * @param OutSubRangeToInvalidate	(Optional) A range to fill with a range to invalidate in the sequence's evaluation field (in root space)
 	 *
 	 * @return true if the meta-data needs re-generating, false otherwise
 	 */
-	bool IsDirty(UMovieSceneSequence& RootSequence, const FMovieSceneSequenceHierarchy& RootHierarchy, TRange<float>* OutSubRangeToInvalidate = nullptr) const;
+	bool IsDirty(UMovieSceneSequence& RootSequence, const FMovieSceneSequenceHierarchy& RootHierarchy, IMovieSceneSequenceTemplateStore& TemplateStore, TRange<FFrameNumber>* OutSubRangeToInvalidate = nullptr) const;
 
 	/** Array of sequences that are active in this time range. */
 	UPROPERTY()
@@ -203,9 +206,9 @@ struct FMovieSceneEvaluationMetaData
 	UPROPERTY()
 	TArray<FMovieSceneOrderedEvaluationKey> ActiveEntities;
 
-	/** Map of sub sequence IDs to signatures that this meta data was generated with (not including root. */
+	/** Map of sub sequence IDs to FMovieSceneEvaluationTemplate::TemplateSignature that this meta data was generated with (not including root). */
 	UPROPERTY()
-	TMap<FMovieSceneSequenceID, FGuid> SubSequenceSignatures;
+	TMap<FMovieSceneSequenceID, FGuid> SubTemplateSignatures;
 };
 
 /**
@@ -222,7 +225,7 @@ struct FMovieSceneEvaluationField
 	 * @param Time			The time at which to seach
 	 * @return The index within Ranges, Groups and MetaData that the current time resides, or INDEX_NONE if there is nothing to do at the requested time
 	 */
-	MOVIESCENE_API int32 GetSegmentFromTime(float Time) const;
+	MOVIESCENE_API int32 GetSegmentFromTime(FFrameNumber Time) const;
 
 	/**
 	 * Deduce the indices into Ranges and Groups that overlap with the specified time range
@@ -230,7 +233,7 @@ struct FMovieSceneEvaluationField
 	 * @param Range			The range to overlap with our field
 	 * @return A range of indices into Ranges and Groups that overlap with the requested range
 	 */
-	MOVIESCENE_API TRange<int32> OverlapRange(TRange<float> Range) const;
+	MOVIESCENE_API TRange<int32> OverlapRange(TRange<FFrameNumber> Range) const;
 
 	/**
 	 * Invalidate a range in this field
@@ -238,17 +241,18 @@ struct FMovieSceneEvaluationField
 	 * @param Range			The range to overlap with our field
 	 * @return A range of indices into Ranges and Groups that overlap with the requested range
 	 */
-	MOVIESCENE_API void Invalidate(TRange<float> Range);
+	MOVIESCENE_API void Invalidate(TRange<FFrameNumber> Range);
 
 	/**
 	 * Insert a new range into this field
 	 *
-	 * @param InRange		The range to add
+	 * @param InsertTime	The time at which to insert the range
+	 * @param InRange		The total range to insert to the field. Will potentially be intersected with preexisting adjacent ranges
 	 * @param InGroup		The group defining what should happen at this time
 	 * @param InMetaData	The meta-data defining efficient access to what happens in this frame
 	 * @return The index the entries were inserted at
 	 */
-	MOVIESCENE_API int32 Insert(TRange<float> InRange, FMovieSceneEvaluationGroup&& InGroup, FMovieSceneEvaluationMetaData&& InMetaData);
+	MOVIESCENE_API int32 Insert(FFrameNumber InsertTime, TRange<FFrameNumber> InRange, FMovieSceneEvaluationGroup&& InGroup, FMovieSceneEvaluationMetaData&& InMetaData);
 
 	/**
 	 * Add the specified data to this field, assuming the specified range lies after any other entries
@@ -257,7 +261,7 @@ struct FMovieSceneEvaluationField
 	 * @param InGroup		The group defining what should happen at this time
 	 * @param InMetaData	The meta-data defining efficient access to what happens in this frame
 	 */
-	MOVIESCENE_API void Add(TRange<float> InRange, FMovieSceneEvaluationGroup&& InGroup, FMovieSceneEvaluationMetaData&& InMetaData);
+	MOVIESCENE_API void Add(TRange<FFrameNumber> InRange, FMovieSceneEvaluationGroup&& InGroup, FMovieSceneEvaluationMetaData&& InMetaData);
 
 	/**
 	 * Access this field's signature
@@ -282,9 +286,9 @@ struct FMovieSceneEvaluationField
 	 * @param Index 	The valid index within the ranges to lookup
 	 * @return The range
 	 */
-	const TRange<float>& GetRange(int32 Index) const
+	const TRange<FFrameNumber>& GetRange(int32 Index) const
 	{
-		return Ranges[Index];
+		return Ranges[Index].Value;
 	}
 
 	/**
@@ -316,7 +320,7 @@ private:
 
 	/** Ranges stored separately for fast (cache efficient) lookup. Each index has a corresponding entry in FMovieSceneEvaluationField::Groups. */
 	UPROPERTY()
-	TArray<FFloatRange> Ranges;
+	TArray<FMovieSceneFrameRange> Ranges;
 
 	/** Groups that store segment pointers for each of the above ranges. Each index has a corresponding entry in FMovieSceneEvaluationField::Ranges. */
 	UPROPERTY()
diff --git a/Engine/Source/Runtime/MovieScene/Public/Evaluation/MovieSceneEvaluationTemplate.h b/Engine/Source/Runtime/MovieScene/Public/Evaluation/MovieSceneEvaluationTemplate.h
--- a/Engine/Source/Runtime/MovieScene/Public/Evaluation/MovieSceneEvaluationTemplate.h
+++ b/Engine/Source/Runtime/MovieScene/Public/Evaluation/MovieSceneEvaluationTemplate.h
@@ -9,6 +9,7 @@
 #include "Misc/Guid.h"
 #include "UObject/Class.h"
 #include "MovieSceneTrack.h"
+#include "MovieSceneFrameMigration.h"
 #include "Evaluation/MovieSceneTrackIdentifier.h"
 #include "Evaluation/MovieSceneEvaluationField.h"
 #include "Containers/ArrayView.h"
@@ -58,7 +59,7 @@ public:
 
 	/** Map of sub section ranges that exist in a template */
 	UPROPERTY()
-	TMap<FGuid, FFloatRange> SubSectionRanges;
+	TMap<FGuid, FMovieSceneFrameRange> SubSectionRanges;
 };
 template<> struct TStructOpsTypeTraits<FMovieSceneTemplateGenerationLedger> : public TStructOpsTypeTraitsBase2<FMovieSceneTemplateGenerationLedger> { enum { WithCopy = true }; };
 
@@ -195,7 +196,7 @@ public:
 	/**
 	 * Add a new sub section
 	 */
-	MOVIESCENE_API void AddSubSectionRange(UMovieSceneSubSection& InSubSection, const FGuid& InObjectBindingId, const TRange<float>& InRange, ESectionEvaluationFlags InFlags);
+	MOVIESCENE_API void AddSubSectionRange(UMovieSceneSubSection& InSubSection, const FGuid& InObjectBindingId, const TRange<FFrameNumber>& InRange, ESectionEvaluationFlags InFlags);
 
 	/**
 	 * Add a new track for the specified identifier
@@ -292,6 +293,10 @@ public:
 	UPROPERTY()
 	FGuid SequenceSignature;
 
+	/** Signature that is re-generated every time this template is re-generated through FMovieSceneEvaluationTemplateGenerator */
+	UPROPERTY()
+	FGuid TemplateSignature;
+
 private:
 
 	UPROPERTY()
diff --git a/Engine/Source/Runtime/MovieScene/Public/Evaluation/MovieSceneEvaluationTrack.h b/Engine/Source/Runtime/MovieScene/Public/Evaluation/MovieSceneEvaluationTrack.h
--- a/Engine/Source/Runtime/MovieScene/Public/Evaluation/MovieSceneEvaluationTrack.h
+++ b/Engine/Source/Runtime/MovieScene/Public/Evaluation/MovieSceneEvaluationTrack.h
@@ -368,7 +368,7 @@ public:
 	 * @param InLocalRange 			The range in this track's space to overlap
 	 * @return A (potentially invalid) segment identifier for the first segment whose range overlaps the specified range
 	 */
-	FMovieSceneSegmentIdentifier FindFirstSegment(TRange<float> InLocalRange);
+	FMovieSceneSegmentIdentifier FindFirstSegment(TRange<FFrameNumber> InLocalRange);
 
 	/**
 	 * Find or compile a segment for the specified time
@@ -376,7 +376,7 @@ public:
 	 * @param InTime 				The time to lookup or compile for
 	 * @return A (potentially invalid) segment identifier for the segment whose range overlaps the specified time
 	 */
-	FMovieSceneSegmentIdentifier GetSegmentFromTime(float InTime);
+	FMovieSceneSegmentIdentifier GetSegmentFromTime(FFrameNumber InTime);
 
 	/**
 	 * Find or compile a segment for the specified iterator
@@ -392,7 +392,7 @@ public:
 	 * @param InLocalRange 			The range in this track's space to overlap
 	 * @return An (potentially empty) array of segment identifiers that overlap the specified range
 	 */
-	TArray<FMovieSceneSegmentIdentifier> GetSegmentsInRange(TRange<float> InLocalRange);
+	TArray<FMovieSceneSegmentIdentifier> GetSegmentsInRange(TRange<FFrameNumber> InLocalRange);
 
 	/**
 	 * Get the smallest range of unique FSectionEvaluationData combinations that overlaps the specified lower bound
@@ -400,7 +400,7 @@ public:
 	 * @param InLowerBound 			The lower bound from which to start looking for a time range
 	 * @return The smallest time range of unique evaluation data entries that encompasses the specified lower bound
 	 */
-	TRange<float> GetUniqueRangeFromLowerBound(TRangeBound<float> InLowerBound) const;
+	TRange<FFrameNumber> GetUniqueRangeFromLowerBound(TRangeBound<FFrameNumber> InLowerBound) const;
 
 	/**
 	 * Set the source track from which this track originates
@@ -424,7 +424,7 @@ public:
 	 * @param Range 		The range that the specified section data should apply to
 	 * @param EvalData 		The actual evaluation data
 	 */
-	void AddTreeData(TRange<float> Range, FSectionEvaluationData EvalData)
+	void AddTreeData(TRange<FFrameNumber> Range, FSectionEvaluationData EvalData)
 	{
 		EvaluationTree.Tree.Add(Range, EvalData);
 	}
@@ -435,7 +435,7 @@ public:
 	 * @param Range 		The range that the specified section data should apply to
 	 * @param EvalData 		The actual evaluation data
 	 */
-	void AddUniqueTreeData(TRange<float> Range, FSectionEvaluationData EvalData)
+	void AddUniqueTreeData(TRange<FFrameNumber> Range, FSectionEvaluationData EvalData)
 	{
 		EvaluationTree.Tree.AddUnique(Range, EvalData);
 	}
diff --git a/Engine/Source/Runtime/MovieScene/Public/Evaluation/MovieSceneEvaluationTree.h b/Engine/Source/Runtime/MovieScene/Public/Evaluation/MovieSceneEvaluationTree.h
--- a/Engine/Source/Runtime/MovieScene/Public/Evaluation/MovieSceneEvaluationTree.h
+++ b/Engine/Source/Runtime/MovieScene/Public/Evaluation/MovieSceneEvaluationTree.h
@@ -6,6 +6,7 @@
 #include "Containers/ArrayView.h"
 #include "Containers/Array.h"
 #include "Templates/Function.h"
+#include "FrameNumber.h"
 
 #include "UObject/SequencerObjectVersion.h"
 
@@ -249,11 +250,11 @@ struct FMovieSceneEvaluationTreeNode
 {
 	/** Constructors */
 	FMovieSceneEvaluationTreeNode()
-		: Range(FFloatRange::Empty())
+		: Range(TRange<FFrameNumber>::Empty())
 		, Parent(FMovieSceneEvaluationTreeNodeHandle::Invalid())
 	{}
 
-	FMovieSceneEvaluationTreeNode(TRange<float> InRange, FMovieSceneEvaluationTreeNodeHandle InParent)
+	FMovieSceneEvaluationTreeNode(TRange<FFrameNumber> InRange, FMovieSceneEvaluationTreeNodeHandle InParent)
 		: Range(InRange)
 		, Parent(InParent)
 	{}
@@ -265,7 +266,7 @@ struct FMovieSceneEvaluationTreeNode
 	}
 
 	/** The time-range that this node represents */
-	FFloatRange Range;
+	TRange<FFrameNumber> Range;
 	/** Handle to the parent node */
 	FMovieSceneEvaluationTreeNodeHandle Parent;
 	/** Identifier for the child node entries associated with this node (FMovieSceneEvaluationTree::ChildNodes) */
@@ -296,7 +297,7 @@ struct FMovieSceneEvaluationTreeNode
 struct MOVIESCENE_API FMovieSceneEvaluationTree
 {
 	FMovieSceneEvaluationTree()
-		: RootNode(FMovieSceneEvaluationTreeNode(TRange<float>::All(), FMovieSceneEvaluationTreeNodeHandle::Invalid()))
+		: RootNode(FMovieSceneEvaluationTreeNode(TRange<FFrameNumber>::All(), FMovieSceneEvaluationTreeNodeHandle::Invalid()))
 	{}
 
 	/**
@@ -305,7 +306,7 @@ struct MOVIESCENE_API FMovieSceneEvaluationTree
 	 * @param Time The time at which we should start iterating
 	 * @return A bi-directional iterator that is set to the unique range that overlaps the current time
 	 */
-	FMovieSceneEvaluationTreeRangeIterator IterateFromTime(float Time) const;
+	FMovieSceneEvaluationTreeRangeIterator IterateFromTime(FFrameNumber Time) const;
 
 	/**
 	 * Start iterating this tree from the specified lower boundary
@@ -313,7 +314,7 @@ struct MOVIESCENE_API FMovieSceneEvaluationTree
 	 * @param InStartingLowerBound The lowerbound at which we should start iterating
 	 * @return A bi-directional iterator that is set to the unique range that overlaps the specified lowerbound
 	 */
-	FMovieSceneEvaluationTreeRangeIterator IterateFromLowerBound(TRangeBound<float> InStartingLowerBound) const;
+	FMovieSceneEvaluationTreeRangeIterator IterateFromLowerBound(TRangeBound<FFrameNumber> InStartingLowerBound) const;
 
 	/**
 	 * Access this tree's root node (infinite range)
@@ -383,7 +384,7 @@ struct MOVIESCENE_API FMovieSceneEvaluationTree
 	 *
 	 * @param InTimeRange 		The time range to add to the tree
 	 */
-	void AddTimeRange(TRange<float> InTimeRange);
+	void AddTimeRange(TRange<FFrameNumber> InTimeRange);
 
 	/**
 	 * Serialize this evaluation tree
@@ -420,7 +421,7 @@ protected:
 	 * @param InOperator 		Operator implementation to call for each node that is relevant to the time range
 	 * @param InParent 			The current parent node
 	 */
-	void AddTimeRange(TRange<float> InTimeRange, const IMovieSceneEvaluationTreeNodeOperator& InOperator, FMovieSceneEvaluationTreeNodeHandle InParent);
+	void AddTimeRange(TRange<FFrameNumber> InTimeRange, const IMovieSceneEvaluationTreeNodeOperator& InOperator, FMovieSceneEvaluationTreeNodeHandle InParent);
 
 	/**
 	 * Helper function that creates a new child for the specified parent node
@@ -430,7 +431,7 @@ protected:
 	 * @param InsertIndex		The index at which to insert the new child (children must be sorted)
 	 * @param InParent 			The node to add the child to
 	 */
-	void InsertNewChild(TRange<float> InEffectiveRange, const IMovieSceneEvaluationTreeNodeOperator& InOperator, int32 InsertIndex, FMovieSceneEvaluationTreeNodeHandle InParent);
+	void InsertNewChild(TRange<FFrameNumber> InEffectiveRange, const IMovieSceneEvaluationTreeNodeOperator& InOperator, int32 InsertIndex, FMovieSceneEvaluationTreeNodeHandle InParent);
 
 	/** This tree's root node */
 	FMovieSceneEvaluationTreeNode RootNode;
@@ -448,7 +449,7 @@ struct MOVIESCENE_API FMovieSceneEvaluationTreeRangeIterator
 	FMovieSceneEvaluationTreeRangeIterator(const FMovieSceneEvaluationTree& InTree);
 
 	/** Iterate the tree ranges starting at the range that overlaps the specified lower bound */
-	FMovieSceneEvaluationTreeRangeIterator(const FMovieSceneEvaluationTree& InTree, TRangeBound<float> StartingBound);
+	FMovieSceneEvaluationTreeRangeIterator(const FMovieSceneEvaluationTree& InTree, TRangeBound<FFrameNumber> StartingBound);
 
 	/** Move onto the next time range */
 	FMovieSceneEvaluationTreeRangeIterator& operator++()
@@ -479,7 +480,7 @@ struct MOVIESCENE_API FMovieSceneEvaluationTreeRangeIterator
 	}
 
 	/** Get the current range */
-	TRange<float> Range() const
+	TRange<FFrameNumber> Range() const
 	{
 		return CurrentRange;
 	}
@@ -509,19 +510,19 @@ private:
 	}
 
 	/** Compare a bound with a range based on whether we're iterating forwards or backwards */
-	static bool CompareBound(bool bForwards, TRange<float> Range, TRangeBound<float> Bound)
+	static bool CompareBound(bool bForwards, TRange<FFrameNumber> Range, TRangeBound<FFrameNumber> Bound)
 	{
 		return bForwards ? Range.GetLowerBound() == Bound : Range.GetUpperBound() == Bound;
 	}
 
 	/** Access the 'leading' bound of a range (lowerbound if forwards, upperbound if backwards) */
-	static TRangeBound<float> GetLeadingBound(bool bForwards, TRange<float> Range)
+	static TRangeBound<FFrameNumber> GetLeadingBound(bool bForwards, TRange<FFrameNumber> Range)
 	{
 		return bForwards ? Range.GetLowerBound() : Range.GetUpperBound();
 	}
 
 	/** Access the 'trailing' bound of a range (upperbound if forwards, lowerbound if backwards) */
-	static TRangeBound<float> GetTrailingBound(bool bForwards, TRange<float> Range)
+	static TRangeBound<FFrameNumber> GetTrailingBound(bool bForwards, TRange<FFrameNumber> Range)
 	{
 		return bForwards ? Range.GetUpperBound() : Range.GetLowerBound();
 	}
@@ -530,11 +531,11 @@ private:
 	void Iter(bool bForwards);
 
 	/** Find the next child within the given parent that is >= the specified bound based on whether we're going forwards or backwards */
-	FMovieSceneEvaluationTreeNodeHandle FindNextChild(FMovieSceneEvaluationTreeNodeHandle ParentNodeHandle, TRangeBound<float> PredicateBound, bool bForwards);
+	FMovieSceneEvaluationTreeNodeHandle FindNextChild(FMovieSceneEvaluationTreeNodeHandle ParentNodeHandle, TRangeBound<FFrameNumber> PredicateBound, bool bForwards);
 
 private:
 	/** The unique time range that we're currently at */
-	TRange<float> CurrentRange;
+	TRange<FFrameNumber> CurrentRange;
 	/** Handle of the current node */
 	FMovieSceneEvaluationTreeNodeHandle CurrentNodeHandle;
 	/** The tree we're iterating */
@@ -552,7 +553,7 @@ struct TMovieSceneEvaluationTree : FMovieSceneEvaluationTree
 	 *
 	 * @param InTimeRange 		The range with which this data should be associated
 	 */
-	void Add(TRange<float> InTimeRange)
+	void Add(TRange<FFrameNumber> InTimeRange)
 	{
 		struct FNullOperator : IMovieSceneEvaluationTreeNodeOperator
 		{
@@ -569,7 +570,7 @@ struct TMovieSceneEvaluationTree : FMovieSceneEvaluationTree
 	 * @param InTimeRange 		The range with which this data should be associated
 	 * @param InData 			The data to assoicate with this time range
 	 */
-	void AddUnique(TRange<float> InTimeRange, DataType InData)
+	void AddUnique(TRange<FFrameNumber> InTimeRange, DataType InData)
 	{
 		AddTimeRange(InTimeRange, FAddUniqueOperator(*this, MoveTemp(InData)), FMovieSceneEvaluationTreeNodeHandle::Root());
 	}
@@ -580,7 +581,7 @@ struct TMovieSceneEvaluationTree : FMovieSceneEvaluationTree
 	 * @param InTimeRange 		The range with which this data should be associated
 	 * @param InData 			The data to assoicate with this time range
 	 */
-	void Add(TRange<float> InTimeRange, DataType InData)
+	void Add(TRange<FFrameNumber> InTimeRange, DataType InData)
 	{
 		AddTimeRange(InTimeRange, FAddOperator(*this, MoveTemp(InData)), FMovieSceneEvaluationTreeNodeHandle::Root());
 	}
diff --git a/Engine/Source/Runtime/MovieScene/Public/Evaluation/MovieSceneExecutionTokens.h b/Engine/Source/Runtime/MovieScene/Public/Evaluation/MovieSceneExecutionTokens.h
--- a/Engine/Source/Runtime/MovieScene/Public/Evaluation/MovieSceneExecutionTokens.h
+++ b/Engine/Source/Runtime/MovieScene/Public/Evaluation/MovieSceneExecutionTokens.h
@@ -18,7 +18,7 @@
 struct FMovieSceneExecutionTokens
 {
 	FMovieSceneExecutionTokens()
-		: Context(FMovieSceneEvaluationRange(0.f), EMovieScenePlayerStatus::Stopped)
+		: Context(FMovieSceneEvaluationRange(FFrameTime(0), FFrameRate()), EMovieScenePlayerStatus::Stopped)
 	{
 	}
 
diff --git a/Engine/Source/Runtime/MovieScene/Public/Evaluation/MovieScenePlayback.h b/Engine/Source/Runtime/MovieScene/Public/Evaluation/MovieScenePlayback.h
--- a/Engine/Source/Runtime/MovieScene/Public/Evaluation/MovieScenePlayback.h
+++ b/Engine/Source/Runtime/MovieScene/Public/Evaluation/MovieScenePlayback.h
@@ -4,6 +4,7 @@
 
 #include "CoreMinimal.h"
 #include "MovieSceneFwd.h"
+#include "FrameTime.h"
 #include "Evaluation/MovieSceneSequenceTransform.h"
 
 
@@ -20,26 +21,44 @@ struct MOVIESCENE_API FMovieSceneEvaluationRange
 	/**
 	 * Construct this range from a single fixed time
 	 */
-	FMovieSceneEvaluationRange(float InTime);
+	FMovieSceneEvaluationRange(FFrameTime InTime, FFrameRate InFrameRate);
 
 	/**
 	 * Construct this range from a raw range and a direction
 	 */
-	FMovieSceneEvaluationRange(TRange<float> InRange, EPlayDirection InDirection);
+	FMovieSceneEvaluationRange(TRange<FFrameTime> InRange, FFrameRate InFrameRate, EPlayDirection InDirection);
 
 	/**
 	 * Construct this range from 2 times, and whether the range should include the previous time or not
 	 */
-	FMovieSceneEvaluationRange(float InCurrentTime, float InPreviousTime, bool bInclusivePreviousTime = false);
+	FMovieSceneEvaluationRange(FFrameTime InCurrentTime, FFrameTime InPreviousTime, FFrameRate InFrameRate, bool bInclusivePreviousTime = false);
+
+	/**
+	 * Convert a frame time range to a frame number range comprising all the frame numbers traversed in the range
+	 */
+	static TRange<FFrameNumber> TimeRangeToNumberRange(const TRange<FFrameTime>& InFrameTimeRange);
+
+	/**
+	 * Convert a frame number range to a frame time range
+	 */
+	static TRange<FFrameTime> NumberRangeToTimeRange(const TRange<FFrameNumber>& InFrameTimeRange);
 
 	/**
 	 * Get the range that we should be evaluating
 	 */
-	FORCEINLINE TRange<float> GetRange() const
+	FORCEINLINE TRange<FFrameTime> GetRange() const
 	{
 		return EvaluationRange;
 	}
 
+	/**
+	 * Get the range of frame numbers traversed over this evaluation range, not including partial frames
+	 */
+	FORCEINLINE TRange<FFrameNumber> GetFrameNumberRange() const
+	{
+		return TimeRangeToNumberRange(EvaluationRange);
+	}
+
 	/**
 	 * Get the direction to evaluate our range
 	 */
@@ -51,9 +70,9 @@ struct MOVIESCENE_API FMovieSceneEvaluationRange
 	/**
 	 * Get the current time of evaluation.
 	 */
-	FORCEINLINE float GetTime() const
+	FORCEINLINE FFrameTime GetTime() const
 	{
-		if (TimeOverride != TNumericLimits<float>::Lowest())
+		if (TimeOverride != TNumericLimits<int32>::Lowest())
 		{
 			return TimeOverride;
 		}
@@ -64,15 +83,15 @@ struct MOVIESCENE_API FMovieSceneEvaluationRange
 	/**
 	 * Get the absolute amount of time that has passed since the last update (will always be >= 0)
 	 */
-	FORCEINLINE float GetDelta() const
+	FORCEINLINE FFrameTime GetDelta() const
 	{
-		return EvaluationRange.Size<float>();
+		return EvaluationRange.Size<FFrameTime>();
 	}
 
 	/**
 	 * Get the previous time of evaluation. Should not generally be used. Prefer GetRange instead.
 	 */
-	FORCEINLINE float GetPreviousTime() const
+	FORCEINLINE FFrameTime GetPreviousTime() const
 	{
 		return Direction == EPlayDirection::Forwards ? EvaluationRange.GetLowerBoundValue() : EvaluationRange.GetUpperBoundValue();
 	}
@@ -80,21 +99,33 @@ struct MOVIESCENE_API FMovieSceneEvaluationRange
 	/**
 	 * Override the time that we're actually evaluating at
 	 */
-	FORCEINLINE void OverrideTime(float InTimeOverride)
+	FORCEINLINE void OverrideTime(FFrameNumber InTimeOverride)
 	{
 		TimeOverride = InTimeOverride;
 	}
 
+	/**
+	 * Get the framerate that this context's times are in
+	 * @return The framerate that all times are relative to
+	 */
+	FORCEINLINE FFrameRate GetFrameRate() const
+	{
+		return CurrentFrameRate;
+	}
+
 protected:
 
 	/** The range to evaluate */
-	TRange<float> EvaluationRange;
+	TRange<FFrameTime> EvaluationRange;
+
+	/** The framerate of the current sequence. */
+	FFrameRate CurrentFrameRate;
 
 	/** Whether to evaluate the range forwards, or backwards */
 	EPlayDirection Direction;
 
 	/** Overridden current time (doesn't manipulate the actual evaluated range) */
-	float TimeOverride;
+	FFrameNumber TimeOverride;
 };
 
 /** MovieScene evaluation context. Should remain bitwise copyable, and contain no external state since this has the potential to be used on a thread */
@@ -106,7 +137,7 @@ struct FMovieSceneContext : FMovieSceneEvaluationRange
 	FMovieSceneContext(FMovieSceneEvaluationRange InRange)
 		: FMovieSceneEvaluationRange(InRange)
 		, Status(EMovieScenePlayerStatus::Stopped)
-		, PrePostRollStartEndTime(TNumericLimits<float>::Lowest())
+		, PrePostRollStartEndTime(TNumericLimits<int32>::Lowest())
 		, HierarchicalBias(0)
 		, bHasJumped(false)
 		, bSilent(false)
@@ -122,7 +153,7 @@ struct FMovieSceneContext : FMovieSceneEvaluationRange
 	FMovieSceneContext(FMovieSceneEvaluationRange InRange, EMovieScenePlayerStatus::Type InStatus)
 		: FMovieSceneEvaluationRange(InRange)
 		, Status(InStatus)
-		, PrePostRollStartEndTime(TNumericLimits<float>::Lowest())
+		, PrePostRollStartEndTime(TNumericLimits<int32>::Lowest())
 		, HierarchicalBias(0)
 		, bHasJumped(false)
 		, bSilent(false)
@@ -207,21 +238,22 @@ public:
 	/**
 	 * Clamp the current evaluation range to the specified range (in the current transform space)
 	 */
-	FMovieSceneContext Clamp(TRange<float> NewRange) const
+	FMovieSceneContext Clamp(TRange<FFrameTime> NewRange) const
 	{
 		FMovieSceneContext NewContext = *this;
-		NewContext.EvaluationRange = NewRange;
+		NewContext.EvaluationRange = TRange<FFrameTime>::Intersection(NewRange, NewContext.EvaluationRange);
 		return NewContext;
 	}
 
 	/**
 	 * Transform this context to a different sub sequence space
 	 */
-	FMovieSceneContext Transform(const FMovieSceneSequenceTransform& InTransform) const
+	FMovieSceneContext Transform(const FMovieSceneSequenceTransform& InTransform, FFrameRate NewFrameRate) const
 	{
 		FMovieSceneContext NewContext = *this;
 		NewContext.EvaluationRange = EvaluationRange * InTransform;
 		NewContext.RootToSequenceTransform = NewContext.RootToSequenceTransform * InTransform;
+		NewContext.CurrentFrameRate = NewFrameRate;
 		return NewContext;
 	}
 
@@ -285,9 +317,9 @@ public:
 	 * Access the time at which preroll will stop, and evaluation will commence
 	 * @note: Only valid to call when HasPreRollEndTime() is true
 	 */
-	float GetPreRollEndTime() const
+	FFrameNumber GetPreRollEndFrame() const
 	{
-		checkf(bHasPreRollEndTime, TEXT("It's invalid to call GetPreRollEndTime() without first checking HasPreRollEndTime()"));
+		checkf(bHasPreRollEndTime, TEXT("It's invalid to call GetPreRollEndFrame() without first checking HasPreRollEndTime()"));
 		return PrePostRollStartEndTime;
 	}
 
@@ -295,9 +327,9 @@ public:
 	 * Access the time at which post roll started (or in other terms: when evaluation stopped)
 	 * @note: Only valid to call when HasPostRollStartTime() is true
 	 */
-	float GetPostRollStartTime() const
+	FFrameNumber GetPostRollStartFrame() const
 	{
-		checkf(bHasPostRollStartTime, TEXT("It's invalid to call GetPostRollStartTime() without first checking HasPostRollStartTime()"));
+		checkf(bHasPostRollStartTime, TEXT("It's invalid to call GetPostRollStartFrame() without first checking HasPostRollStartTime()"));
 		return PrePostRollStartEndTime;
 	}
 
@@ -307,9 +339,9 @@ public:
 	 * @param InLeadingRange			The leading (preroll) range in front of the outer section, in the current transformation's time space
 	 * @param InTrailingRange			The trailing (postroll) range at the end of the outer section, in the current transformation's time space
 	 */
-	void ReportOuterSectionRanges(TRange<float> InLeadingRange, TRange<float> InTrailingRange)
+	void ReportOuterSectionRanges(TRange<FFrameNumber> InLeadingRange, TRange<FFrameNumber> InTrailingRange)
 	{
-		const float Now = GetTime();
+		const FFrameNumber Now = GetTime().FrameNumber;
 		if (InLeadingRange.Contains(Now) && InLeadingRange.HasUpperBound())
 		{
 			PrePostRollStartEndTime = InLeadingRange.GetUpperBoundValue();
@@ -327,7 +359,7 @@ public:
 		else
 		{
 			bHasPreRollEndTime = bHasPostRollStartTime = false;
-			PrePostRollStartEndTime = TNumericLimits<float>::Lowest();
+			PrePostRollStartEndTime = FFrameNumber(TNumericLimits<int32>::Lowest());
 		}
 	}
 
@@ -339,8 +371,8 @@ protected:
 	/** The current playback status */
 	EMovieScenePlayerStatus::Type Status;
 
-	/** When bHasPreRollEndTime or bHasPostRollStartTime is true, this defines either the time at which 'real' evaluation commences, or finished */
-	float PrePostRollStartEndTime;
+	/** When bHasPreRollEndTime or bHasPostRollStartTime is true, this defines either the frame at which 'real' evaluation commences, or finished */
+	FFrameNumber PrePostRollStartEndTime;
 
 	/** Hierachical bias. Higher bias should take precedence. */
 	int32 HierarchicalBias;
@@ -369,49 +401,116 @@ protected:
 /** Helper class designed to abstract the complexity of calculating evaluation ranges for previous times and fixed time intervals */
 struct MOVIESCENE_API FMovieScenePlaybackPosition
 {
+	FMovieScenePlaybackPosition()
+		: InputRate(0,0), OutputRate(0,0), EvaluationType(EMovieSceneEvaluationType::WithSubFrames)
+	{}
+
+	/**
+	 * @return Whether we are evaluating with sub frames, or frame-locked
+	 */
+	FORCEINLINE EMovieSceneEvaluationType GetEvaluationType() const
+	{
+		return EvaluationType;
+	}
+
+	/**
+	 * @return The input frame rate that all frame times provided to this class will be interpreted as
+	 */
+	FORCEINLINE FFrameRate GetInputRate() const
+	{
+		return InputRate;
+	}
+
+	/**
+	 * @return The output frame rate that all frame times returned from this class will be interpreted as
+	 */
+	FORCEINLINE FFrameRate GetOutputRate() const
+	{
+		return OutputRate;
+	}
+
+public:
+
+	/**
+	 * Assign the input and output rates that frame times should be interpreted as.
+	 *
+	 * @param InInputRate           The framerate to interpret any frame time provided to this class
+	 * @param InOutputRate          The framerate to use when returning any frame range from this class
+	 * @param InputEvaluationType   Whether we're using frame-locked or sub-frame evaluation
+	 */
+	void SetTimeBase(FFrameRate InInputRate, FFrameRate InOutputRate, EMovieSceneEvaluationType InputEvaluationType);
+
 	/**
 	 * Reset this position to the specified time.
 	 * @note Future calls to 'PlayTo' will include this time in its resulting evaluation range
 	 */
-	void Reset(float StartPos);
+	void Reset(FFrameTime StartPos);
+
+	/**
+	 * Get the last position that was set, in InputRate space
+	 */
+	FFrameTime GetCurrentPosition() const { return CurrentPosition; }
 
 	/**
-	 * Jump to the specified time.
+	 * Get the last actual time that was evaluated during playback, in InputRate space.
+	 */
+	TOptional<FFrameTime> GetLastPlayEvalPostition() const { return PreviousPlayEvalPosition; }
+
+public:
+
+	/**
+	 * Jump to the specified input time.
 	 * @note Will reset previous play position. Any subsequent call to 'PlayTo' will include NewPosition.
-	 * @return A range encompassing only the specified time.
+	 *
+	 * @param NewPosition         The new frame time to set, in InputRate space
+	 * @return A range encompassing only the specified time, in OutputRate space.
 	 */
-	FMovieSceneEvaluationRange JumpTo(float NewPosition, TOptional<float> FixedInterval);
+	FMovieSceneEvaluationRange JumpTo(FFrameTime NewPosition);
 
 	/**
 	 * Play from the previously evaluated play time, to the specified time
-	 * @return An evaluation range from the previously evaluated time, to the specified time
+	 *
+	 * @param NewPosition         The new frame time to set, in InputRate space
+	 * @return An evaluation range from the previously evaluated time to the specified time, in OutputRate space.
 	 */
-	FMovieSceneEvaluationRange PlayTo(float NewPosition, TOptional<float> FixedInterval);
+	FMovieSceneEvaluationRange PlayTo(FFrameTime NewPosition);
 
 	/**
-	 * Get a range that encompasses the last evaluated range
-	 * @return An optional evaluation range
+	 * Get a range that encompasses the last evaluated range in OutputRate space.
+	 * @return An optional evaluation range in OutputRate space.
 	 */
 	TOptional<FMovieSceneEvaluationRange> GetLastRange() const;
 
 	/**
-	 * Get the last position that was set
+	 * Get a range encompassing only the current time, if available (in OutputRate space)
+	 * @return An optional evaluation range in OutputRate space.
 	 */
-	TOptional<float> GetPreviousPosition() const { return PreviousPosition; }
+	FMovieSceneEvaluationRange GetCurrentPositionAsRange() const;
+
+private:
 
 	/**
-	 * Get the last actual time that was evaluated during playback
+	 * Check this class's invariants
 	 */
-	TOptional<float> GetLastPlayEvalPostition() const { return PreviousPlayEvalPosition; }
-
+	void CheckInvariants() const;
 
 private:
-	/** The previous *actual* time position set. Never rounded to a fixed interval. */
-	TOptional<float> PreviousPosition;
 
-	/** The previous evaluated position when playing, potentially rounded to a frame interval. */
-	TOptional<float> PreviousPlayEvalPosition;
+	/** The framerate to be used when interpreting frame time values provided to this class (i.e. moviescene playback frame rate) */
+	FFrameRate InputRate;
+
+	/** The framerate to be used when returning frame time values from this class (i.e. moviescene frame resolution) */
+	FFrameRate OutputRate;
+
+	/** The type of evaluation to use */
+	EMovieSceneEvaluationType EvaluationType;
+
+	/** The current time position set, in 'InputRate' time-space. */
+	FFrameTime CurrentPosition;
+
+	/** The previously evaluated position when playing, in 'InputRate' time-space. */
+	TOptional<FFrameTime> PreviousPlayEvalPosition;
 
-	/** The previously evaluated range, if available */
+	/** The previously evaluated range if available, in 'OutputRate' time-space */
 	TOptional<FMovieSceneEvaluationRange> LastRange;
 };
diff --git a/Engine/Source/Runtime/MovieScene/Public/Evaluation/MovieSceneSectionParameters.h b/Engine/Source/Runtime/MovieScene/Public/Evaluation/MovieSceneSectionParameters.h
--- a/Engine/Source/Runtime/MovieScene/Public/Evaluation/MovieSceneSectionParameters.h
+++ b/Engine/Source/Runtime/MovieScene/Public/Evaluation/MovieSceneSectionParameters.h
@@ -14,16 +14,25 @@ struct FMovieSceneSectionParameters
 
 	/** Default constructor */
 	FMovieSceneSectionParameters()
-		: StartOffset(0.0f)
+		: StartFrameOffset(0)
 		, TimeScale(1.0f)
 		, HierarchicalBias(100)
+		, StartOffset_DEPRECATED(0.f)
 		, PrerollTime_DEPRECATED(0.0f)
 		, PostrollTime_DEPRECATED(0.0f)
 	{}
 
-	/** Number of seconds to skip at the beginning of the sub-sequence. */
-	UPROPERTY(EditAnywhere, Category="Clipping")
-	float StartOffset;
+	/** StartFrameOffset is treated as a FFrameNumber for details customizations but we continue to expose the API publicly as it was before via int32. */
+	int32 GetStartFrameOffset() const { return StartFrameOffset.Value; }
+	void SetStartFrameOffset(int32 InOffset) { StartFrameOffset = FFrameNumber(InOffset); }
+
+private:
+
+	/** Number of frames to skip at the beginning of the sub-sequence. */
+	UPROPERTY(EditAnywhere, Category = "Clipping")
+	FFrameNumber StartFrameOffset;
+
+public:
 
 	/** Playback time scaling factor. */
 	UPROPERTY(EditAnywhere, Category="Timing")
@@ -33,6 +42,8 @@ struct FMovieSceneSectionParameters
 	UPROPERTY(EditAnywhere, Category="Sequence")
 	int32 HierarchicalBias;
 	UPROPERTY()
+	float StartOffset_DEPRECATED;
+	UPROPERTY()
 	float PrerollTime_DEPRECATED;
 	UPROPERTY()
 	float PostrollTime_DEPRECATED;
diff --git a/Engine/Source/Runtime/MovieScene/Public/Evaluation/MovieSceneSegment.h b/Engine/Source/Runtime/MovieScene/Public/Evaluation/MovieSceneSegment.h
--- a/Engine/Source/Runtime/MovieScene/Public/Evaluation/MovieSceneSegment.h
+++ b/Engine/Source/Runtime/MovieScene/Public/Evaluation/MovieSceneSegment.h
@@ -3,11 +3,12 @@
 #pragma once
 
 #include "CoreMinimal.h"
+#include "FrameTime.h"
 #include "UObject/ObjectMacros.h"
 #include "UObject/Class.h"
 #include "Containers/ArrayView.h"
 #include "UObject/SequencerObjectVersion.h"
-
+#include "MovieSceneFrameMigration.h"
 #include "MovieSceneSegment.generated.h"
 
 struct FMovieSceneEvaluationTrackSegments;
@@ -83,21 +84,21 @@ struct FSectionEvaluationData
 	GENERATED_BODY()
 
 	/** Default constructor */
-	FSectionEvaluationData() : ImplIndex(-1), ForcedTime(TNumericLimits<float>::Lowest()), Flags(ESectionEvaluationFlags::None) {}
+	FSectionEvaluationData() : ImplIndex(-1), ForcedTime(TNumericLimits<int32>::Lowest()), Flags(ESectionEvaluationFlags::None) {}
 
 	/** Construction from an implementaiton index (probably a section) */
 	explicit FSectionEvaluationData(int32 InImplIndex)
-		: ImplIndex(InImplIndex), ForcedTime(TNumericLimits<float>::Lowest()), Flags(ESectionEvaluationFlags::None)
+		: ImplIndex(InImplIndex), ForcedTime(TNumericLimits<int32>::Lowest()), Flags(ESectionEvaluationFlags::None)
 	{}
 
 	/** Construction from an implementaiton index and a time to force evaluation at */
-	FSectionEvaluationData(int32 InImplIndex, float InForcedTime)
+	FSectionEvaluationData(int32 InImplIndex, FFrameNumber InForcedTime)
 		: ImplIndex(InImplIndex), ForcedTime(InForcedTime), Flags(ESectionEvaluationFlags::None)
 	{}
 
 	/** Construction from an implementaiton index and custom eval flags */
 	FSectionEvaluationData(int32 InImplIndex, ESectionEvaluationFlags InFlags)
-		: ImplIndex(InImplIndex), ForcedTime(TNumericLimits<float>::Lowest()), Flags(InFlags)
+		: ImplIndex(InImplIndex), ForcedTime(TNumericLimits<int32>::Lowest()), Flags(InFlags)
 	{}
 
 	friend bool operator==(FSectionEvaluationData A, FSectionEvaluationData B)
@@ -105,9 +106,9 @@ struct FSectionEvaluationData
 		return A.ImplIndex == B.ImplIndex && A.ForcedTime == B.ForcedTime && A.Flags == B.Flags;
 	}
 
-	float GetTime(float InActualTime)
+	FFrameTime GetTime(FFrameTime InActualTime)
 	{
-		return ForcedTime == TNumericLimits<float>::Lowest() ? InActualTime : ForcedTime;
+		return ForcedTime == TNumericLimits<int32>::Lowest() ? InActualTime : FFrameTime(ForcedTime);
 	}
 
 	/** Check if this is a preroll eval */
@@ -128,7 +129,7 @@ struct FSectionEvaluationData
 
 	/** A forced time to evaluate this section at */
 	UPROPERTY()
-	float ForcedTime;
+	FFrameNumber ForcedTime;
 
 	/** Additional flags for evaluating this section */
 	UPROPERTY()
@@ -146,12 +147,12 @@ struct FMovieSceneSegment
 	FMovieSceneSegment()
 	{}
 
-	FMovieSceneSegment(const TRange<float>& InRange)
+	FMovieSceneSegment(const TRange<FFrameNumber>& InRange)
 		: Range(InRange)
 		, bAllowEmpty(false)
 	{}
 
-	FMovieSceneSegment(const TRange<float>& InRange, TArrayView<const FSectionEvaluationData> InApplicationImpls)
+	FMovieSceneSegment(const TRange<FFrameNumber>& InRange, TArrayView<const FSectionEvaluationData> InApplicationImpls)
 		: Range(InRange)
 		, bAllowEmpty(true)
 	{
@@ -176,7 +177,7 @@ struct FMovieSceneSegment
 	{
 		if (Range.Adjoins(OtherSegment.Range) && Impls == OtherSegment.Impls && bAllowEmpty == OtherSegment.bAllowEmpty)
 		{
-			Range = TRange<float>::Hull(OtherSegment.Range, Range);
+			Range = TRange<FFrameNumber>::Hull(OtherSegment.Range, Range);
 			return true;
 		}
 		return false;
@@ -185,9 +186,19 @@ struct FMovieSceneSegment
 	/** Custom serializer to accomodate the inline allocator on our array */
 	bool Serialize(FArchive& Ar)
 	{
-		Ar << Range;
-
 		Ar.UsingCustomVersion(FSequencerObjectVersion::GUID);
+
+		if (Ar.IsLoading() && Ar.CustomVer(FSequencerObjectVersion::GUID) < FSequencerObjectVersion::FloatToIntConversion)
+		{
+			TRange<float> OldFloatRange;
+			Ar << OldFloatRange;
+			Range = FMovieSceneFrameRange::FromFloatRange(OldFloatRange);
+		}
+		else
+		{
+			Ar << Range;
+		}
+
 		if (Ar.CustomVer(FSequencerObjectVersion::GUID) >= FSequencerObjectVersion::EvaluationTree)
 		{
 			Ar << ID;
@@ -217,7 +228,7 @@ struct FMovieSceneSegment
 	}
 
 	/** The segment's range */
-	FFloatRange Range;
+	TRange<FFrameNumber> Range;
 
 	FMovieSceneSegmentIdentifier ID;
 
diff --git a/Engine/Source/Runtime/MovieScene/Public/Evaluation/MovieSceneSequenceHierarchy.h b/Engine/Source/Runtime/MovieScene/Public/Evaluation/MovieSceneSequenceHierarchy.h
--- a/Engine/Source/Runtime/MovieScene/Public/Evaluation/MovieSceneSequenceHierarchy.h
+++ b/Engine/Source/Runtime/MovieScene/Public/Evaluation/MovieSceneSequenceHierarchy.h
@@ -3,6 +3,7 @@
 #pragma once
 
 #include "CoreMinimal.h"
+#include "FrameRate.h"
 #include "UObject/ObjectMacros.h"
 #include "Evaluation/MovieSceneSequenceTransform.h"
 #include "Evaluation/MovieSceneSectionParameters.h"
@@ -10,6 +11,7 @@
 #include "MovieSceneSequenceID.h"
 #include "Evaluation/MovieSceneSequenceInstanceData.h"
 #include "Containers/ArrayView.h"
+#include "MovieSceneFrameMigration.h"
 #include "MovieSceneSequenceHierarchy.generated.h"
 
 class UMovieSceneSequence;
@@ -52,21 +54,25 @@ struct FMovieSceneSubSequenceData
 	UPROPERTY()
 	FMovieSceneSequenceTransform RootToSequenceTransform;
 
+	/** The frame resolution of the inner sequence. */
+	UPROPERTY()
+	FFrameRate FrameResolution;
+
 	/** This sequence's deterministic sequence ID. Used in editor to reduce the risk of collisions on recompilation. */ 
 	UPROPERTY()
 	FMovieSceneSequenceID DeterministicSequenceID;
 
 	/** This sub sequence's playback range according to its parent sub section. Clamped recursively during template generation */
 	UPROPERTY()
-	FFloatRange PlayRange;
+	FMovieSceneFrameRange PlayRange;
 
 	/** The sequence preroll range considering the start offset */
 	UPROPERTY()
-	FFloatRange PreRollRange;
+	FMovieSceneFrameRange PreRollRange;
 
 	/** The sequence postroll range considering the start offset */
 	UPROPERTY()
-	FFloatRange PostRollRange;
+	FMovieSceneFrameRange PostRollRange;
 
 	/** The accumulated hierarchical bias of this sequence. Higher bias will take precedence */
 	UPROPERTY()
diff --git a/Engine/Source/Runtime/MovieScene/Public/Evaluation/MovieSceneSequenceTransform.h b/Engine/Source/Runtime/MovieScene/Public/Evaluation/MovieSceneSequenceTransform.h
--- a/Engine/Source/Runtime/MovieScene/Public/Evaluation/MovieSceneSequenceTransform.h
+++ b/Engine/Source/Runtime/MovieScene/Public/Evaluation/MovieSceneSequenceTransform.h
@@ -4,6 +4,7 @@
 
 #include "CoreMinimal.h"
 #include "UObject/ObjectMacros.h"
+#include "FrameTime.h"
 #include "MovieSceneSequenceTransform.generated.h"
 
 /**
@@ -26,7 +27,7 @@ struct FMovieSceneSequenceTransform
 	 */
 	FMovieSceneSequenceTransform()
 		: TimeScale(1.f)
-		, Offset(0.f)
+		, Offset(0)
 	{}
 
 	/**
@@ -35,7 +36,7 @@ struct FMovieSceneSequenceTransform
 	 * @param InOffset 			The offset to translate by
 	 * @param InTimeScale 		The timescale. For instance, if a sequence is playing twice as fast, pass 2.f
 	 */
-	FMovieSceneSequenceTransform(float InOffset, float InTimeScale = 1.f)
+	FMovieSceneSequenceTransform(FFrameTime InOffset, float InTimeScale = 1.f)
 		: TimeScale(InTimeScale)
 		, Offset(InOffset)
 	{}
@@ -45,16 +46,17 @@ struct FMovieSceneSequenceTransform
 	 */
 	FMovieSceneSequenceTransform Inverse() const
 	{
-		return FMovieSceneSequenceTransform(-Offset/TimeScale, 1.f/TimeScale);
+		const FFrameTime NewOffset = -Offset/TimeScale;
+		return FMovieSceneSequenceTransform(NewOffset, 1.f/TimeScale);
 	}
 
 	/** The sequence's time scale (or play rate) */
 	UPROPERTY()
 	float TimeScale;
 
-	/** Scalar time offset applied before the scale */
+	/** Scalar frame offset applied before the scale */
 	UPROPERTY()
-	float Offset;
+	FFrameTime Offset;
 };
 
 /**
@@ -63,8 +65,14 @@ struct FMovieSceneSequenceTransform
  * @param InTime 			The time to transform
  * @param RHS 				The transform
  */
-inline float operator*(float InTime, const FMovieSceneSequenceTransform& RHS)
+inline FFrameTime operator*(FFrameTime InTime, const FMovieSceneSequenceTransform& RHS)
 {
+	// Avoid floating point conversion when in the same time-space
+	if (RHS.TimeScale == 1.f)
+	{
+		return InTime + RHS.Offset;
+	}
+
 	return RHS.Offset + InTime * RHS.TimeScale;
 }
 
@@ -74,7 +82,7 @@ inline float operator*(float InTime, const FMovieSceneSequenceTransform& RHS)
  * @param InTime 			The time to transform
  * @param RHS 				The transform
  */
-inline float& operator*=(float& InTime, const FMovieSceneSequenceTransform& RHS)
+inline FFrameTime& operator*=(FFrameTime& InTime, const FMovieSceneSequenceTransform& RHS)
 {
 	InTime = InTime * RHS;
 	return InTime;
@@ -89,23 +97,44 @@ inline float& operator*=(float& InTime, const FMovieSceneSequenceTransform& RHS)
 template<typename T>
 TRange<T> operator*(const TRange<T>& LHS, const FMovieSceneSequenceTransform& RHS)
 {
-	TRangeBound<float> SourceLower = LHS.GetLowerBound();
-	TRangeBound<float> TransformedLower =
+	TRangeBound<T> SourceLower = LHS.GetLowerBound();
+	TRangeBound<T> TransformedLower =
+		SourceLower.IsOpen() ? 
+			TRangeBound<T>() : 
+			SourceLower.IsInclusive() ?
+				TRangeBound<T>::Inclusive(SourceLower.GetValue() * RHS) :
+				TRangeBound<T>::Exclusive(SourceLower.GetValue() * RHS);
+
+	TRangeBound<T> SourceUpper = LHS.GetUpperBound();
+	TRangeBound<T> TransformedUpper =
+		SourceUpper.IsOpen() ? 
+			TRangeBound<T>() : 
+			SourceUpper.IsInclusive() ?
+				TRangeBound<T>::Inclusive(SourceUpper.GetValue() * RHS) :
+				TRangeBound<T>::Exclusive(SourceUpper.GetValue() * RHS);
+
+	return TRange<T>(TransformedLower, TransformedUpper);
+}
+
+inline TRange<FFrameNumber> operator*(const TRange<FFrameNumber>& LHS, const FMovieSceneSequenceTransform& RHS)
+{
+	TRangeBound<FFrameNumber> SourceLower = LHS.GetLowerBound();
+	TRangeBound<FFrameNumber> TransformedLower =
 		SourceLower.IsOpen() ? 
-			TRangeBound<float>() : 
+			TRangeBound<FFrameNumber>() : 
 			SourceLower.IsInclusive() ?
-				TRangeBound<float>::Inclusive(SourceLower.GetValue() * RHS) :
-				TRangeBound<float>::Exclusive(SourceLower.GetValue() * RHS);
+				TRangeBound<FFrameNumber>::Inclusive((SourceLower.GetValue() * RHS).FloorToFrame()) :
+				TRangeBound<FFrameNumber>::Exclusive((SourceLower.GetValue() * RHS).FloorToFrame());
 
-	TRangeBound<float> SourceUpper = LHS.GetUpperBound();
-	TRangeBound<float> TransformedUpper =
+	TRangeBound<FFrameNumber> SourceUpper = LHS.GetUpperBound();
+	TRangeBound<FFrameNumber> TransformedUpper =
 		SourceUpper.IsOpen() ? 
-			TRangeBound<float>() : 
+			TRangeBound<FFrameNumber>() : 
 			SourceUpper.IsInclusive() ?
-				TRangeBound<float>::Inclusive(SourceUpper.GetValue() * RHS) :
-				TRangeBound<float>::Exclusive(SourceUpper.GetValue() * RHS);
+				TRangeBound<FFrameNumber>::Inclusive((SourceUpper.GetValue() * RHS).FloorToFrame()) :
+				TRangeBound<FFrameNumber>::Exclusive((SourceUpper.GetValue() * RHS).FloorToFrame());
 
-	return TRange<float>(TransformedLower, TransformedUpper);
+	return TRange<FFrameNumber>(TransformedLower, TransformedUpper);
 }
 
 /**
@@ -132,8 +161,9 @@ inline FMovieSceneSequenceTransform operator*(const FMovieSceneSequenceTransform
 	// | TimeScaleA	, OffsetA	|	.	| TimeScaleB, OffsetB	|
 	// | 0			, 1			|		| 0			, 1			|
 
+	const FFrameTime ScaledOffsetRHS = LHS.TimeScale == 1.f ? RHS.Offset : RHS.Offset*LHS.TimeScale;
 	return FMovieSceneSequenceTransform(
-		RHS.Offset*LHS.TimeScale + LHS.Offset,		// New Offset
-		LHS.TimeScale * RHS.TimeScale				// New TimeScale
+		LHS.Offset + ScaledOffsetRHS,		// New Offset
+		LHS.TimeScale * RHS.TimeScale		// New TimeScale
 		);
 }
diff --git a/Engine/Source/Runtime/MovieScene/Public/Evaluation/MovieSceneTrackImplementation.h b/Engine/Source/Runtime/MovieScene/Public/Evaluation/MovieSceneTrackImplementation.h
--- a/Engine/Source/Runtime/MovieScene/Public/Evaluation/MovieSceneTrackImplementation.h
+++ b/Engine/Source/Runtime/MovieScene/Public/Evaluation/MovieSceneTrackImplementation.h
@@ -149,6 +149,10 @@ struct FMovieSceneTrackImplementationPtr
 			StructOps.Construct(Allocation);
 			StructOps.Copy(Allocation, &RHS.GetValue(), 1);
 		}
+		else
+		{
+			Reset();
+		}
 
 		return *this;
 	}
diff --git a/Engine/Source/Runtime/MovieScene/Public/IMovieSceneObjectSpawner.h b/Engine/Source/Runtime/MovieScene/Public/IMovieSceneObjectSpawner.h
--- a/Engine/Source/Runtime/MovieScene/Public/IMovieSceneObjectSpawner.h
+++ b/Engine/Source/Runtime/MovieScene/Public/IMovieSceneObjectSpawner.h
@@ -87,7 +87,7 @@ public:
 	 * @param	Sequencer		The sequencer this spawnable was just created by
 	 * @param	Settings		The settings for this sequencer
 	 */
-	virtual void SetupDefaultsForSpawnable(UObject* SpawnedObject, const FGuid& Guid, const FTransformData& TransformData, TSharedRef<ISequencer> Sequencer, USequencerSettings* Settings) {}
+	virtual void SetupDefaultsForSpawnable(UObject* SpawnedObject, const FGuid& Guid, const TOptional<FTransformData>& TransformData, TSharedRef<ISequencer> Sequencer, USequencerSettings* Settings) {}
 
 	/*
  	 * Whether this spawner can set up defaults
diff --git a/Engine/Source/Runtime/MovieScene/Public/MovieScene.h b/Engine/Source/Runtime/MovieScene/Public/MovieScene.h
--- a/Engine/Source/Runtime/MovieScene/Public/MovieScene.h
+++ b/Engine/Source/Runtime/MovieScene/Public/MovieScene.h
@@ -14,6 +14,8 @@
 #include "MovieSceneSignedObject.h"
 #include "MovieSceneSequenceID.h"
 #include "MovieSceneObjectBindingID.h"
+#include "MovieSceneFrameMigration.h"
+#include "MovieSceneTimeController.h"
 #include "MovieScene.generated.h"
 
 class UMovieSceneFolder;
@@ -41,17 +43,43 @@ struct FMovieSceneEditorData
 {
 	GENERATED_USTRUCT_BODY()
 
+	FMovieSceneEditorData()
+		: ViewStart(0.0), ViewEnd(0.0)
+		, WorkStart(0.0), WorkEnd(0.0)
+	{}
+
+	TRange<double> GetViewRange() const
+	{
+		return TRange<double>(ViewStart, ViewEnd);
+	}
+
+	TRange<double> GetWorkingRange() const
+	{
+		return TRange<double>(WorkStart, WorkEnd);
+	}
+
 	/** Map of node path -> expansion state. */
 	UPROPERTY()
 	TMap<FString, FMovieSceneExpansionState> ExpansionStates;
 
-	/** User-defined working range in which the entire sequence should reside. */
+	/** The last view-range start that the user was observing */
 	UPROPERTY()
-	FFloatRange WorkingRange;
+	double ViewStart;
+	/** The last view-range end that the user was observing */
+	UPROPERTY()
+	double ViewEnd;
+
+	/** User-defined working range start in which the entire sequence should reside. */
+	UPROPERTY()
+	double WorkStart;
+	/** User-defined working range end in which the entire sequence should reside. */
+	UPROPERTY()
+	double WorkEnd;
 
-	/** The last view-range that the user was observing */
 	UPROPERTY()
-	FFloatRange ViewRange;
+	FFloatRange WorkingRange_DEPRECATED;
+	UPROPERTY()
+	FFloatRange ViewRange_DEPRECATED;
 };
 
 
@@ -428,9 +456,9 @@ public:
 	}
 
 	/** Get the current selection range. */
-	TRange<float> GetSelectionRange() const
+	TRange<FFrameNumber> GetSelectionRange() const
 	{
-		return SelectionRange;
+		return SelectionRange.Value;
 	}
 
 	/**
@@ -442,9 +470,78 @@ public:
 	FText GetObjectDisplayName(const FGuid& ObjectId);
 
 	/** Get the playback time range of this movie scene, relative to its 0-time offset. */
-	TRange<float> GetPlaybackRange() const
+	TRange<FFrameNumber> GetPlaybackRange() const
+	{
+		return PlaybackRange.Value;
+	}
+
+	/**
+	 * Retrieve the frame resolution at which all frame numbers within this movie scene are defined
+	 */
+	FFrameRate GetFrameResolution() const
+	{
+		return FrameResolution;
+	}
+
+	/**
+	 * Directly set the frame resolution for this movie scene without applying any conversion whatsoever, or modifying the data
+	 */
+	void SetFrameResolutionDirectly(FFrameRate InFrameResolution)
+	{
+		FrameResolution = InFrameResolution;
+	}
+
+	/**
+	 * Retrieve the frame resolution at which all frame numbers within this movie scene are defined
+	 */
+	FFrameRate GetPlaybackFrameRate() const
+	{
+		return PlayRate;
+	}
+
+	/**
+	 * Set the play rate for this movie scene
+	 */
+	void SetPlaybackFrameRate(FFrameRate InPlayRate)
+	{
+		PlayRate = InPlayRate;
+	}
+
+	/**
+	 * Retrieve a value signifying how to evaluate this movie scene data
+	 */
+	EMovieSceneEvaluationType GetEvaluationType() const
+	{
+		return EvaluationType;
+	}
+
+	/**
+	 * Assign a value signifying how to evaluate this movie scene data
+	 */
+	void SetEvaluationType(EMovieSceneEvaluationType InNewEvaluationType)
+	{
+		EvaluationType = InNewEvaluationType;
+
+		if (EvaluationType == EMovieSceneEvaluationType::FrameLocked && ClockSource == EUpdateClockSource::Tick)
+		{
+			ClockSource = EUpdateClockSource::Platform;
+		}
+	}
+
+	/**
+	 * retrieve the clock source to be used for this moviescene
+	 */
+	EUpdateClockSource GetClockSource() const
+	{
+		return ClockSource;
+	}
+
+	/**
+	 * Assign the clock source to be used for this moviescene
+	 */
+	void SetClockSource(EUpdateClockSource InNewClockSource)
 	{
-		return PlaybackRange;
+		ClockSource = InNewClockSource;
 	}
 
 	/*
@@ -461,9 +558,9 @@ public:
 	}
 
 	/** Set the selection range. */
-	void SetSelectionRange(TRange<float> Range)
+	void SetSelectionRange(TRange<FFrameNumber> Range)
 	{
-		SelectionRange = Range;
+		SelectionRange.Value = Range;
 	}
 
 	/**
@@ -484,11 +581,20 @@ public:
 	 * Set the start and end playback positions (playback range) for this movie scene
 	 *
 	 * @param Start The offset from 0-time to start playback of this movie scene
-	 * @param End The offset from 0-time to end playback of this movie scene
+	 * @param Duration The number of frames the movie scene should play for
 	 * @param bAlwaysMarkDirty Whether to always mark the playback range dirty when changing it. 
 	 *        In the case where the playback range is dynamic and based on section bounds, the playback range doesn't need to be dirtied when set
 	 */
-	void SetPlaybackRange(float Start, float End, bool bAlwaysMarkDirty = true);
+	void SetPlaybackRange(FFrameNumber Start, int32 Duration, bool bAlwaysMarkDirty = true);
+
+	/**
+	 * Set the playback range for this movie scene
+	 *
+	 * @param Range The new playback range. Must not have any open bounds (ie must be a finite range)
+	 * @param bAlwaysMarkDirty Whether to always mark the playback range dirty when changing it. 
+	 *        In the case where the playback range is dynamic and based on section bounds, the playback range doesn't need to be dirtied when set
+	 */
+	void SetPlaybackRange(const TRange<FFrameNumber>& NewRange, bool bAlwaysMarkDirty = true);
 
 	/**
 	 * Set the start and end working range (outer) for this movie scene
@@ -507,6 +613,9 @@ public:
 	void SetViewRange(float Start, float End);
 
 #if WITH_EDITORONLY_DATA
+
+public:
+
 	/** 
 	 * Return whether the playback range is locked.
 	 */
@@ -516,46 +625,7 @@ public:
 	 * Set whether the playback range is locked.
 	 */
 	void SetPlaybackRangeLocked(bool bLocked);
-#endif
 
-	/**
-	 * Gets whether or not playback should be forced to match the fixed frame interval.  When true all time values will be rounded to a fixed
-	 * frame value which will force editor and runtime playback to match exactly, but will result in duplicate frames if the runtime and editor
-	 * frame rates aren't exactly the same.
-	 */
-	bool GetForceFixedFrameIntervalPlayback() const;
-
-	/**
-	* Sets whether or not playback should be forced to match the fixed frame interval.  When true all time values will be rounded to a fixed
-	* frame value which will force editor and runtime playback to match exactly, but will result in duplicate frames if the runtime and editor
-	* frame rates aren't exactly the same.
-	*/
-	void SetForceFixedFrameIntervalPlayback( bool bInForceFixedFrameIntervalPlayback );
-
-	/**
-	* Gets the fixed frame interval to be used when "force fixed frame interval playback" is set.
-	*/
-	float GetFixedFrameInterval() const { return FixedFrameInterval; }
-
-	/**
-	* Gets the fixed frame interval to be used when "force fixed frame interval playback" is set.
-	*/
-	void SetFixedFrameInterval( float InFixedFrameInterval );
-
-	/**
-	 * Gets the fixed frame interval to be used when "force fixed frame interval playback" is set. Only returns a valid result when GetForceFixedFrameIntervalPlayback() is true, and the interval is > 0.
-	 */
-	TOptional<float> GetOptionalFixedFrameInterval() const { return (!GetForceFixedFrameIntervalPlayback() || GetFixedFrameInterval() <= 0) ? TOptional<float>() : GetFixedFrameInterval(); }
-
-	/**
-	 * Calculates a fixed frame time based on a current time, a fixed frame interval, and an internal epsilon to account
-	 * for floating point consistency.
-	 */ 
-	static float CalculateFixedFrameTime( float Time, float FixedFrameInterval );
-
-public:
-	
-#if WITH_EDITORONLY_DATA
 	/**
 	 * @return The editor only data for use with this movie scene
 	 */
@@ -568,7 +638,8 @@ public:
 	{
 		EditorData = InEditorData;
 	}
-#endif
+
+#endif	// WITH_EDITORONLY_DATA
 
 protected:
 
@@ -598,10 +669,6 @@ protected:
 
 private:
 
-	// Small value added for fixed frame interval calculations to make up for consistency in
-	// floating point calculations.
-	static const float FixedFrameIntervalEpsilon;
-
 	/**
 	 * Data-only blueprints for all of the objects that we we're able to spawn.
 	 * These describe objects and actors that we may instantiate at runtime,
@@ -628,25 +695,33 @@ private:
 
 	/** User-defined selection range. */
 	UPROPERTY()
-	FFloatRange SelectionRange;
+	FMovieSceneFrameRange SelectionRange;
 
 	/** User-defined playback range for this movie scene. Must be a finite range. Relative to this movie-scene's 0-time origin. */
 	UPROPERTY()
-	FFloatRange PlaybackRange;
+	FMovieSceneFrameRange PlaybackRange;
 
-	/** User-defined playback range is locked. */
-#if WITH_EDITORONLY_DATA
+	/** The resolution at which all frame numbers within this movie-scene data are stored */
 	UPROPERTY()
-	bool bPlaybackRangeLocked;
-#endif
+	FFrameRate FrameResolution;
 
+	/** The rate at which to playback this moviescene data */
 	UPROPERTY()
-	bool bForceFixedFrameIntervalPlayback;
+	FFrameRate PlayRate;
 
+	/** The type of evaluation to use when playing back this sequence */
 	UPROPERTY()
-	float FixedFrameInterval;
+	EMovieSceneEvaluationType EvaluationType;
+
+	UPROPERTY()
+	EUpdateClockSource ClockSource;
 
 #if WITH_EDITORONLY_DATA
+
+	/** User-defined playback range is locked. */
+	UPROPERTY()
+	bool bPlaybackRangeLocked;
+
 	/** Maps object GUIDs to user defined display names. */
 	UPROPERTY()
 	TMap<FString, FText> ObjectsToDisplayNames;
@@ -662,10 +737,9 @@ private:
 	/** The root folders for this movie scene. */
 	UPROPERTY()
 	TArray<UMovieSceneFolder*> RootFolders;
-#endif
 
 private:
-	
+
 	UPROPERTY()
 	float InTime_DEPRECATED;
 
@@ -677,4 +751,12 @@ private:
 
 	UPROPERTY()
 	float EndTime_DEPRECATED;
+
+	UPROPERTY()
+	bool bForceFixedFrameIntervalPlayback_DEPRECATED;
+
+	UPROPERTY()
+	float FixedFrameInterval_DEPRECATED;
+
+#endif
 };
diff --git a/Engine/Source/Runtime/MovieScene/Public/MovieSceneBinding.h b/Engine/Source/Runtime/MovieScene/Public/MovieSceneBinding.h
--- a/Engine/Source/Runtime/MovieScene/Public/MovieSceneBinding.h
+++ b/Engine/Source/Runtime/MovieScene/Public/MovieSceneBinding.h
@@ -18,7 +18,11 @@ struct FMovieSceneBinding
 	GENERATED_USTRUCT_BODY()
 
 	/** Default constructor. */
-	FMovieSceneBinding() { }
+	FMovieSceneBinding()
+#if WITH_EDITORONLY_DATA
+		: SortingOrder(-1)
+#endif
+	{ }
 
 	/**
 	 * Creates and initializes a new instance.
@@ -31,6 +35,9 @@ struct FMovieSceneBinding
 		: ObjectGuid(InObjectGuid)
 		, BindingName(InBindingName)
 		, Tracks(InTracks)
+#if WITH_EDITORONLY_DATA
+		, SortingOrder(-1)
+#endif
 	{ }
 
 	/**
@@ -42,13 +49,11 @@ struct FMovieSceneBinding
 	FMovieSceneBinding(const FGuid& InObjectGuid, const FString& InBindingName)
 		: ObjectGuid(InObjectGuid)
 		, BindingName(InBindingName)
+#if WITH_EDITORONLY_DATA
+		, SortingOrder(-1)
+#endif
 	{ }
 
-	/**
-	 * @return The time range of all tracks in this binding
-	 */
-	TRange<float> GetTimeRange() const;
-
 	/**
 	* Set the object guid
 	*
@@ -115,6 +120,26 @@ struct FMovieSceneBinding
 
 #endif
 
+#if WITH_EDITORONLY_DATA
+	/**
+	* Get this folder's desired sorting order
+	*/
+	int32 GetSortingOrder() const
+	{
+		return SortingOrder;
+	}
+
+	/**
+	* Set this folder's desired sorting order.
+	*
+	* @param InSortingOrder The higher the value the further down the list the folder will be.
+	*/
+	void SetSortingOrder(const int32 InSortingOrder)
+	{
+		SortingOrder = InSortingOrder;
+	}
+#endif
+
 private:
 
 	/** Object binding guid for runtime objects */
@@ -128,4 +153,10 @@ private:
 	/** All tracks in this binding */
 	UPROPERTY(Instanced)
 	TArray<UMovieSceneTrack*> Tracks;
+
+#if WITH_EDITORONLY_DATA
+	/** The desired sorting order for this binding in Sequencer */
+	UPROPERTY()
+	int32 SortingOrder;
+#endif
 };
diff --git a/Engine/Source/Runtime/MovieScene/Public/MovieSceneClipboard.h b/Engine/Source/Runtime/MovieScene/Public/MovieSceneClipboard.h
--- a/Engine/Source/Runtime/MovieScene/Public/MovieSceneClipboard.h
+++ b/Engine/Source/Runtime/MovieScene/Public/MovieSceneClipboard.h
@@ -32,7 +32,7 @@ class MOVIESCENE_API FMovieSceneClipboardKey
 public:
 	/** Templated constructor accepting a specific value type */
 	template<typename T>
-	FMovieSceneClipboardKey(float InTime, T InValue)
+	FMovieSceneClipboardKey(FFrameNumber InTime, T InValue)
 		: Time(InTime)
 		, Data(TKey<T>(MoveTemp(InValue)))
 	{
@@ -43,10 +43,10 @@ public:
 	FMovieSceneClipboardKey& operator=(const FMovieSceneClipboardKey& In);
 
 	/** Get the time at which this key is defined */
-	float GetTime() const;
+	FFrameNumber GetTime() const;
 
 	/** Set the time at which this key is defined */
-	void SetTime(float InTime);
+	void SetTime(FFrameNumber InTime);
 
 	/**
 	 * Get the value of this key as the specified type.
@@ -141,7 +141,7 @@ private:
 	};
 
 	/** The time that this key is defined at */
-	float Time;
+	FFrameNumber Time;
 
 	/** Type-erased bytes containing the key's value */
 	TInlineValue<IKey, 64> Data;
@@ -202,7 +202,7 @@ public:
 
 	/** Add a key of the specified type to this track. KeyType must match the type this track was constructed with */
 	template<typename KeyType>
-	void AddKey(float Time, KeyType Value)
+	void AddKey(FFrameNumber Time, KeyType Value)
 	{
 		checkf(IsKeyOfType<KeyType>(), TEXT("Unable to add a key of a different value type to the track"));
 		Keys.Add(FMovieSceneClipboardKey(Time, MoveTemp(Value)));
@@ -258,15 +258,18 @@ private:
 struct FMovieSceneClipboardEnvironment
 {
 	FMovieSceneClipboardEnvironment()
-		: CardinalTime(0.f)
+		: CardinalTime(0)
 		, DateTime(FDateTime::UtcNow())
 	{}
 
-	/** The cardinal time for a copy-paste operation. Keys are copied with absolute time values */
-	float CardinalTime;
+	/** The cardinal time for a copy-paste operation. */
+	FFrameTime CardinalTime;
 
 	/** The date/time at which the copy operation was performed */
 	FDateTime DateTime;
+
+	/** The frame resolution within which all clipboard data is stored */
+	FFrameRate FrameResolution;
 };
 
 /** A clipboard representing serializable copied data for a movie scene */
@@ -312,7 +315,7 @@ class MOVIESCENE_API FMovieSceneClipboardBuilder
 public:
 
 	/** Generate a clipboard for the current state of this builder, resetting the builder back to its default state */
-	FMovieSceneClipboard Commit(TOptional<float> CopyRelativeTo);
+	FMovieSceneClipboard Commit(TOptional<FFrameNumber> CopyRelativeTo);
 
 	/**
 	 * Find or add a key track. Key tracks are group primarily by MovieSceneTrack instance, then by name
diff --git a/Engine/Source/Runtime/MovieScene/Public/MovieSceneCommonHelpers.h b/Engine/Source/Runtime/MovieScene/Public/MovieSceneCommonHelpers.h
--- a/Engine/Source/Runtime/MovieScene/Public/MovieSceneCommonHelpers.h
+++ b/Engine/Source/Runtime/MovieScene/Public/MovieSceneCommonHelpers.h
@@ -7,6 +7,7 @@
 #include "UObject/UnrealType.h"
 #include "UObject/ObjectKey.h"
 #include "Curves/KeyHandle.h"
+#include "FrameNumber.h"
 
 class AActor;
 class UCameraComponent;
@@ -19,23 +20,13 @@ class MOVIESCENE_API MovieSceneHelpers
 {
 public:
 
-	/**
-	 * Gets the sections that were traversed over between the current time and the previous time, including overlapping sections
-	 */
-	static TArray<UMovieSceneSection*> GetAllTraversedSections( const TArray<UMovieSceneSection*>& Sections, float CurrentTime, float PreviousTime );
-
-	/**
-	 * Gets the sections that were traversed over between the current time and the previous time, excluding overlapping sections (highest wins)
-	 */
-	static TArray<UMovieSceneSection*> GetTraversedSections( const TArray<UMovieSceneSection*>& Sections, float CurrentTime, float PreviousTime );
-
 	/**
 	 * Finds a section that exists at a given time
 	 *
 	 * @param Time	The time to find a section at
 	 * @return The found section or null
 	 */
-	static UMovieSceneSection* FindSectionAtTime( const TArray<UMovieSceneSection*>& Sections, float Time );
+	static UMovieSceneSection* FindSectionAtTime( const TArray<UMovieSceneSection*>& Sections, FFrameNumber Time );
 
 	/**
 	 * Finds the nearest section to the given time
@@ -43,7 +34,7 @@ public:
 	 * @param Time	The time to find a section at
 	 * @return The found section or null
 	 */
-	static UMovieSceneSection* FindNearestSectionAtTime( const TArray<UMovieSceneSection*>& Sections, float Time );
+	static UMovieSceneSection* FindNearestSectionAtTime( const TArray<UMovieSceneSection*>& Sections, FFrameNumber Time );
 
 	/*
 	 * Fix up consecutive sections so that there are no gaps
@@ -91,30 +82,26 @@ public:
 	 */
 	static void SetRuntimeObjectMobility(UObject* Object, EComponentMobility::Type ComponentMobility = EComponentMobility::Movable);
 
-	/*
-	 * Set the key interpolation
-	 *
-	 * @param InCurve The curve that contains the key handle to set
-	 * @param InKeyHandle The key handle to set
-	 * @param InInterpolation The interpolation to set
-	 */
-	static void SetKeyInterpolation(FRichCurve& InCurve, FKeyHandle InKeyHandle, EMovieSceneKeyInterpolation InKeyInterpolation);
-
 	/**
 	 * Sort predicate that sorts lower bounds of a range
 	 */
-	static bool SortLowerBounds(TRangeBound<float> A, TRangeBound<float> B)
+	static bool SortLowerBounds(TRangeBound<FFrameNumber> A, TRangeBound<FFrameNumber> B)
 	{
-		return TRangeBound<float>::MinLower(A, B) == A && A != B;
+		return TRangeBound<FFrameNumber>::MinLower(A, B) == A && A != B;
 	}
 
 	/**
 	 * Sort predicate that sorts upper bounds of a range
 	 */
-	static bool SortUpperBounds(TRangeBound<float> A, TRangeBound<float> B)
+	static bool SortUpperBounds(TRangeBound<FFrameNumber> A, TRangeBound<FFrameNumber> B)
 	{
-		return TRangeBound<float>::MinUpper(A, B) == A && A != B;
+		return TRangeBound<FFrameNumber>::MinUpper(A, B) == A && A != B;
 	}
+
+	/**
+	 * Sort predicate that sorts overlapping sections by row primarily, then by overlap priority
+	 */
+	static bool SortOverlappingSections(const UMovieSceneSection* A, const UMovieSceneSection* B);
 };
 
 /**
@@ -190,6 +177,21 @@ public:
 		return Val ? *Val : ValueType();
 	}
 
+	/**
+	 * Optionally gets the current value of a property on an object
+	 *
+	 * @param Object	The object to get the property from
+	 * @return (Optional) The current value of the property on the object
+	 */
+	template <typename ValueType>
+	TOptional<ValueType> GetOptionalValue(const UObject& Object)
+	{
+		FPropertyAndFunction PropAndFunction = FindOrAdd(Object);
+
+		const ValueType* Val = PropAndFunction.GetPropertyAddress<ValueType>();
+		return Val ? *Val : TOptional<ValueType>();
+	}
+
 	/**
 	 * Gets the current value of a property on an object
 	 *
diff --git a/Engine/Source/Runtime/MovieScene/Public/MovieSceneFolder.h b/Engine/Source/Runtime/MovieScene/Public/MovieSceneFolder.h
--- a/Engine/Source/Runtime/MovieScene/Public/MovieSceneFolder.h
+++ b/Engine/Source/Runtime/MovieScene/Public/MovieSceneFolder.h
@@ -73,6 +73,24 @@ class MOVIESCENE_API UMovieSceneFolder : public UObject
 	{
 		FolderColor = InFolderColor;
 	}
+
+	/**
+	 * Get this folder's desired sorting order 
+	 */
+	int32 GetSortingOrder() const
+	{
+		return SortingOrder;
+	}
+
+	/**
+	 * Set this folder's desired sorting order.
+	 *
+	 * @param InSortingOrder The higher the value the further down the list the folder will be.
+	 */
+	void SetSortingOrder(const int32 InSortingOrder)
+	{
+		SortingOrder = InSortingOrder;
+	}
 #endif
 
 private:
@@ -96,6 +114,10 @@ private:
 	/** This folder's color */
 	UPROPERTY(EditAnywhere, Category=General, DisplayName=Color)
 	FColor FolderColor;
+
+	/** This folder's desired sorting order */
+	UPROPERTY()
+	int32 SortingOrder;
 #endif
 
 	/** The guids for the object bindings contained by this folder. */
diff --git a/Engine/Source/Runtime/MovieScene/Public/MovieSceneFwd.h b/Engine/Source/Runtime/MovieScene/Public/MovieSceneFwd.h
--- a/Engine/Source/Runtime/MovieScene/Public/MovieSceneFwd.h
+++ b/Engine/Source/Runtime/MovieScene/Public/MovieSceneFwd.h
@@ -4,6 +4,7 @@
 
 #include "CoreMinimal.h"
 #include "Stats/Stats.h"
+#include "FrameRate.h"
 #include "MovieSceneSequenceID.h"
 
 UENUM()
@@ -22,9 +23,39 @@ namespace EMovieScenePlayerStatus
 	};
 }
 
+UENUM()
+enum class EMovieSceneEvaluationType : uint8
+{
+	/** Play the sequence frame-locked to its playback rate (snapped to the frame resolution - no sub-frames) */
+	FrameLocked,
+
+	/** Play the sequence in real-time, with sub-frame interpolation if necessary */
+	WithSubFrames,
+};
+
+/**
+ * Enum used to define how to update to a particular time
+ */
+UENUM()
+enum class EUpdateClockSource : uint8
+{
+	/** Use the default world tick delta for timing. Honors world and actor pause state, but is susceptible to accumulation errors */
+	Tick,
+
+	/** Use the platform clock for timing. Does not honor world or actor pause state. */
+	Platform,
+
+	/** Use the audio clock for timing. Does not honor world or actor pause state. */
+	Audio,
+};
+
 MOVIESCENE_API DECLARE_LOG_CATEGORY_EXTERN(LogMovieScene, Log, All);
 DECLARE_STATS_GROUP(TEXT("Movie Scene Evaluation"), STATGROUP_MovieSceneEval, STATCAT_Advanced);
 
+MOVIESCENE_API FFrameRate GetLegacyConversionFrameRate();
+MOVIESCENE_API void EmitLegacyOutOfBoundsError(UObject* ErrorContext, FFrameRate InFrameRate, double InTime);
+MOVIESCENE_API FFrameNumber UpgradeLegacyMovieSceneTime(UObject* ErrorContext, FFrameRate InFrameRate, double InTime);
+
 #ifndef MOVIESCENE_DETAILED_STATS
 	#define MOVIESCENE_DETAILED_STATS 0
 #endif
