Copying //UE4/Dev-Sequencer to //UE4/Dev-Main (Source: //UE4/Dev-Sequencer @ 3946692)

diff --git a/Engine/Source/Runtime/MovieScene/Public/MovieSceneKeyStruct.h b/Engine/Source/Runtime/MovieScene/Public/MovieSceneKeyStruct.h
--- a/Engine/Source/Runtime/MovieScene/Public/MovieSceneKeyStruct.h
+++ b/Engine/Source/Runtime/MovieScene/Public/MovieSceneKeyStruct.h
@@ -4,6 +4,8 @@
 
 #include "CoreMinimal.h"
 #include "UObject/ObjectMacros.h"
+#include "MovieSceneKeyStructHelper.h"
+#include "UObject/StructOnScope.h"
 #include "MovieSceneKeyStruct.generated.h"
 
 struct FPropertyChangedEvent;
@@ -22,7 +24,60 @@ struct FMovieSceneKeyStruct
 	 *
 	 * @param ChangeEvent The property change event.
 	 */
-	virtual void PropagateChanges(const FPropertyChangedEvent& ChangeEvent) { };
+	virtual void PropagateChanges(const FPropertyChangedEvent& ChangeEvent) { }
 
-	virtual ~FMovieSceneKeyStruct() {};
+	virtual ~FMovieSceneKeyStruct() {}
 };
+
+
+USTRUCT()
+struct FMovieSceneKeyTimeStruct : public FMovieSceneKeyStruct
+{
+	GENERATED_BODY()
+
+	FMovieSceneKeyTimeStruct(){}
+
+	UPROPERTY(EditAnywhere, Category="Key", meta=(Units=s))
+	FFrameNumber Time;
+
+	FMovieSceneKeyStructHelper KeyStructInterop;
+
+	/**
+	 * Propagate changes from this key structure to the corresponding key values.
+	 *
+	 * @param ChangeEvent The property change event.
+	 */
+	virtual void PropagateChanges(const FPropertyChangedEvent& ChangeEvent)
+	{
+		KeyStructInterop.Apply(Time);
+	}
+};
+template<> struct TStructOpsTypeTraits<FMovieSceneKeyTimeStruct> : public TStructOpsTypeTraitsBase2<FMovieSceneKeyTimeStruct> { enum { WithCopy = false }; };
+
+/**
+ * Templated helper to aid in the creation of key structs
+ */
+template<typename KeyStructType, typename ChannelType>
+TSharedPtr<FStructOnScope> CreateKeyStruct(TMovieSceneChannelHandle<ChannelType> ChannelHandle, FKeyHandle InHandle)
+{
+	TSharedPtr<FStructOnScope> KeyStruct;
+
+	ChannelType* Channel = ChannelHandle.Get();
+	if (Channel)
+	{
+		auto ChannelInterface = Channel->GetInterface();
+		const int32 KeyIndex = ChannelInterface.GetIndex(InHandle);
+
+		if (KeyIndex != INDEX_NONE)
+		{
+			KeyStruct = MakeShared<FStructOnScope>(KeyStructType::StaticStruct());
+			KeyStructType* Struct = reinterpret_cast<KeyStructType*>(KeyStruct->GetStructMemory());
+
+			Struct->Time  = ChannelInterface.GetTimes()[KeyIndex];
+			Struct->Value = ChannelInterface.GetValues()[KeyIndex];
+
+			Struct->KeyStructInterop.Add(FMovieSceneChannelValueHelper(ChannelHandle, &Struct->Value, MakeTuple(InHandle, Struct->Time)));
+		}
+	}
+	return KeyStruct;
+}
\ No newline at end of file
diff --git a/Engine/Source/Runtime/MovieScene/Public/MovieSceneSection.h b/Engine/Source/Runtime/MovieScene/Public/MovieSceneSection.h
--- a/Engine/Source/Runtime/MovieScene/Public/MovieSceneSection.h
+++ b/Engine/Source/Runtime/MovieScene/Public/MovieSceneSection.h
@@ -3,20 +3,24 @@
 #pragma once
 
 #include "CoreMinimal.h"
+#include "FrameTime.h"
 #include "UObject/ObjectMacros.h"
 #include "MovieSceneFwd.h"
 #include "KeyParams.h"
-#include "Curves/KeyHandle.h"
-#include "Curves/RichCurve.h"
 #include "MovieSceneSignedObject.h"
 #include "Evaluation/Blending/MovieSceneBlendType.h"
 #include "Generators/MovieSceneEasingFunction.h"
+#include "MovieSceneFrameMigration.h"
 #include "MovieSceneSection.generated.h"
 
 class FStructOnScope;
-class UAISenseEvent;
+
+struct FKeyHandle;
+struct FMovieSceneChannelProxy;
 struct FMovieSceneEvalTemplatePtr;
 
+template<typename> class TArrayView;
+
 /** Enumeration specifying how to handle state when this section is no longer evaluated */
 UENUM()
 enum class EMovieSceneCompletionMode : uint8
@@ -59,32 +63,35 @@ struct FMovieSceneEasingSettings
 	GENERATED_BODY()
 
 	FMovieSceneEasingSettings()
-		: AutoEaseInTime(0.f), AutoEaseOutTime(0.f)
-		, EaseIn(nullptr), bManualEaseIn(false), ManualEaseInTime(0.f)
-		, EaseOut(nullptr), bManualEaseOut(false), ManualEaseOutTime(0.f)
+		: AutoEaseInDuration(0), AutoEaseOutDuration(0)
+		, EaseIn(nullptr), bManualEaseIn(false), ManualEaseInDuration(0)
+		, EaseOut(nullptr), bManualEaseOut(false), ManualEaseOutDuration(0)
+#if WITH_EDITORONLY_DATA
+		, AutoEaseInTime_DEPRECATED(0.f), AutoEaseOutTime_DEPRECATED(0.f), ManualEaseInTime_DEPRECATED(0.f), ManualEaseOutTime_DEPRECATED(0.f)
+#endif
 	{}
 
 public:
 
-	float GetEaseInTime() const
+	int32 GetEaseInDuration() const
 	{
-		return bManualEaseIn ? ManualEaseInTime : AutoEaseInTime;
+		return bManualEaseIn ? ManualEaseInDuration : AutoEaseInDuration;
 	}
 
-	float GetEaseOutTime() const
+	int32 GetEaseOutDuration() const
 	{
-		return bManualEaseOut ? ManualEaseOutTime : AutoEaseOutTime;
+		return bManualEaseOut ? ManualEaseOutDuration : AutoEaseOutDuration;
 	}
 
 public:
 
-	/** Automatically applied ease in time */
+	/** Automatically applied ease in duration in frames */
 	UPROPERTY()
-	float AutoEaseInTime;
+	int32 AutoEaseInDuration;
 
 	/** Automatically applied ease out time */
 	UPROPERTY()
-	float AutoEaseOutTime;
+	int32 AutoEaseOutDuration;
 
 	UPROPERTY()
 	TScriptInterface<IMovieSceneEasingFunction> EaseIn;
@@ -93,9 +100,9 @@ public:
 	UPROPERTY()
 	bool bManualEaseIn;
 
-	/** Manually override this section's ease in time */
+	/** Manually override this section's ease in duration in frames */
 	UPROPERTY()
-	float ManualEaseInTime;
+	int32 ManualEaseInDuration;
 
 	UPROPERTY()
 	TScriptInterface<IMovieSceneEasingFunction> EaseOut;
@@ -104,9 +111,20 @@ public:
 	UPROPERTY()
 	bool bManualEaseOut;
 
-	/** Manually override this section's ease-out time */
+	/** Manually override this section's ease-out duration in frames */
+	UPROPERTY()
+	int32 ManualEaseOutDuration;
+
+#if WITH_EDITORONLY_DATA
+	UPROPERTY()
+	float AutoEaseInTime_DEPRECATED;
+	UPROPERTY()
+	float AutoEaseOutTime_DEPRECATED;
+	UPROPERTY()
+	float ManualEaseInTime_DEPRECATED;
 	UPROPERTY()
-	float ManualEaseOutTime;
+	float ManualEaseOutTime_DEPRECATED;
+#endif
 };
 
 /**
@@ -130,96 +148,101 @@ public:
 	 *
 	 * @return Returns whether this section is locked or not
 	 */
-	virtual MOVIESCENE_API bool TryModify(bool bAlwaysMarkDirty=true);
+	MOVIESCENE_API bool TryModify(bool bAlwaysMarkDirty=true);
 
 	/**
-	 * @return The start time of the section
+	 * @return The range of times of the section
 	 */
-	float GetStartTime() const
+	TRange<FFrameNumber> GetRange() const
 	{
-		return StartTime;
+		return SectionRange.Value;
 	}
 
 	/**
-	 * @return The end time of the section
+	 * Expands this section's range to include the specified time
 	 */
-	float GetEndTime() const
+	void ExpandToFrame(FFrameNumber InFrame)
 	{
-		return EndTime;
-	}
-	
-	/**
-	 * @return The size of the time range of the section
-	 */
-	float GetTimeSize() const
-	{
-		return EndTime - StartTime;
+		SetRange(TRange<FFrameNumber>::Hull(GetRange(), TRange<FFrameNumber>::Inclusive(InFrame, InFrame)));
 	}
 
 	/**
-	 * Sets a new end time for this section
+	 * Sets a new range of times for this section
 	 * 
-	 * @param InEndTime	The new end time
+	 * @param NewRange	The new range of times
 	 */
-	void SetStartTime(float NewStartTime)
-	{ 
+	void SetRange(const TRange<FFrameNumber>& NewRange)
+	{
 		if (TryModify())
 		{
-			StartTime = NewStartTime;
+			check(NewRange.GetLowerBound().IsOpen() || NewRange.GetUpperBound().IsOpen() || NewRange.GetLowerBoundValue() <= NewRange.GetUpperBoundValue());
+			SectionRange.Value = NewRange;
 		}
 	}
 
 	/**
-	 * Sets a new end time for this section
-	 * 
-	 * @param InEndTime	The new end time
+	 * Check whether this section has a start frame (else infinite)
+	 * @return true if this section has an inclusive or exclusive start frame, false if it's open (infinite)
 	 */
-	void SetEndTime(float NewEndTime)
-	{ 
-		if (TryModify())
-		{
-			EndTime = NewEndTime;
-		}
+	bool HasStartFrame() const
+	{
+		return !SectionRange.Value.GetLowerBound().IsOpen();
 	}
-	
+
 	/**
-	 * @return The range of times of the section
+	 * Check whether this section has an end frame (else infinite)
+	 * @return true if this section has an inclusive or exclusive end frame, false if it's open (infinite)
 	 */
-	TRange<float> GetRange() const 
+	bool HasEndFrame() const
 	{
-		// Use the single value constructor for zero sized ranges because it creates a range that is inclusive on both upper and lower
-		// bounds which isn't considered "empty".  Use the standard constructor for non-zero sized ranges so that they work well when
-		// calculating overlap with other non-zero sized ranges.
-		return (StartTime == EndTime)
-			? TRange<float>(StartTime)
-			: TRange<float>(StartTime, TRangeBound<float>::Inclusive(EndTime));
+		return !SectionRange.Value.GetUpperBound().IsOpen();
 	}
-	
+
 	/**
-	 * Sets a new range of times for this section
-	 * 
-	 * @param NewRange	The new range of times
+	 * Gets the frame number at which this section starts
+	 *
+	 * @note Assumes a non-infinite start time. Check HasStartFrame first.
+	 * @return The frame number at which this section starts.
 	 */
-	void SetRange(TRange<float> NewRange)
+	FFrameNumber GetInclusiveStartFrame() const
 	{
-		check(NewRange.HasLowerBound() && NewRange.HasUpperBound());
+		TRangeBound<FFrameNumber> LowerBound = SectionRange.GetLowerBound();
+		return LowerBound.IsInclusive() ? LowerBound.GetValue() : LowerBound.GetValue() + 1;
+	}
 
-		if (TryModify())
-		{
-			StartTime = NewRange.GetLowerBoundValue();
-			EndTime = NewRange.GetUpperBoundValue();
-		}
+	/**
+	 * Gets the first frame number after the end of this section
+	 *
+	 * @note Assumes a non-infinite end time. Check HasEndFrame first.
+	 * @return The first frame after this section ends
+	 */
+	FFrameNumber GetExclusiveEndFrame() const
+	{
+		TRangeBound<FFrameNumber> UpperBound = SectionRange.GetUpperBound();
+		return UpperBound.IsInclusive() ? UpperBound.GetValue() + 1 : UpperBound.GetValue();
 	}
 
+	/**
+	 * Set this section's start frame in sequence resolution space.
+	 * @note: Will be clamped to the current end frame if necessary
+	 */
+	MOVIESCENE_API void SetStartFrame(TRangeBound<FFrameNumber> NewStartFrame);
+
+	/**
+	 * Set this section's end frame in sequence resolution space
+	 * @note: Will be clamped to the current start frame if necessary
+	 */
+	MOVIESCENE_API void SetEndFrame(TRangeBound<FFrameNumber> NewEndFrame);
+
 	/**
 	 * Returns whether or not a provided position in time is within the timespan of the section 
 	 *
 	 * @param Position	The position to check
 	 * @return true if the position is within the timespan, false otherwise
 	 */
-	bool IsTimeWithinSection(float Position) const 
+	bool IsTimeWithinSection(FFrameNumber Position) const 
 	{
-		return Position >= StartTime && Position <= EndTime;
+		return SectionRange.Value.Contains(Position);
 	}
 
 	/**
@@ -250,40 +273,23 @@ public:
 	 * Moves the section by a specific amount of time
 	 *
 	 * @param DeltaTime	The distance in time to move the curve
-	 * @param KeyHandles The key handles to operate on
 	 */
-	virtual void MoveSection(float DeltaTime, TSet<FKeyHandle>& KeyHandles)
-	{
-		if (TryModify())
-		{
-			StartTime += DeltaTime;
-			EndTime += DeltaTime;
-		}
-	}
-	
+	MOVIESCENE_API void MoveSection(FFrameNumber DeltaTime);
+
 	/**
-	 * Dilates the section by a specific factor
+	 * Return the range within which this section is effective. Used for automatic calculation of sequence bounds.
 	 *
-	 * @param DilationFactor The multiplier which scales this section
-	 * @param bFromStart Whether to dilate from the beginning or end (whichever stays put)
-	 * @param KeyHandles The key handles to operate on
+	 * @return the range within which this section is effective
 	 */
-	virtual void DilateSection(float DilationFactor, float Origin, TSet<FKeyHandle>& KeyHandles)
-	{
-		if (TryModify())
-		{
-			StartTime = (StartTime - Origin) * DilationFactor + Origin;
-			EndTime = (EndTime - Origin) * DilationFactor + Origin;
-		}
-	}
-	
+	MOVIESCENE_API TRange<FFrameNumber> ComputeEffectiveRange() const;
+
 	/**
 	 * Split a section in two at the split time
 	 *
 	 * @param SplitTime The time at which to split
 	 * @return The newly created split section
 	 */
-	virtual MOVIESCENE_API UMovieSceneSection* SplitSection(float SplitTime);
+	MOVIESCENE_API virtual UMovieSceneSection* SplitSection(FFrameNumber SplitTime);
 
 	/**
 	 * Trim a section at the trim time
@@ -291,26 +297,15 @@ public:
 	 * @param TrimTime The time at which to trim
 	 * @param bTrimLeft Whether to trim left or right
 	 */
-	virtual MOVIESCENE_API void TrimSection(float TrimTime, bool bTrimLeft);
-
-	/**
-	 * Get the key handles for the keys on the curves within this section
-	 *
-	 * @param OutKeyHandles Will contain the key handles of the keys on the curves within this section
-	 * @param TimeRange Optional time range that the keys must be in (default = all)
-	 */
-	virtual void GetKeyHandles(TSet<FKeyHandle>& OutKeyHandles, TRange<float> TimeRange) const { }
+	MOVIESCENE_API virtual void TrimSection(FFrameNumber TrimTime, bool bTrimLeft);
 
 	/**
-	 * Get the data structure representing the specified key.
+	 * Get the data structure representing the specified keys.
 	 *
-	 * @param KeyHandle The handle of the key.
-	 * @return The key's data structure representation, or nullptr if key not found or no structure available.
+	 * @param KeyHandles The handles of the keys.
+	 * @return The keys' data structure representation, or nullptr if key not found or no structure available.
 	 */
-	virtual TSharedPtr<FStructOnScope> GetKeyStruct(const TArray<FKeyHandle>& KeyHandles)
-	{
-		return nullptr;
-	}
+	MOVIESCENE_API virtual TSharedPtr<FStructOnScope> GetKeyStruct(TArrayView<const FKeyHandle> KeyHandles);
 
 	/**
 	 * Generate an evaluation template for this section
@@ -324,17 +319,22 @@ public:
 	 * @param OutSnapTimes The array of times we will to output
 	 * @param bGetSectionBorders Gets the section borders in addition to any custom snap times
 	 */
-	virtual void GetSnapTimes(TArray<float>& OutSnapTimes, bool bGetSectionBorders) const
+	virtual void GetSnapTimes(TArray<FFrameNumber>& OutSnapTimes, bool bGetSectionBorders) const
 	{
 		if (bGetSectionBorders)
 		{
-			OutSnapTimes.Add(StartTime);
-			OutSnapTimes.Add(EndTime);
+			if (SectionRange.Value.GetLowerBound().IsClosed())
+			{
+				OutSnapTimes.Add(SectionRange.Value.GetLowerBoundValue());
+			}
+
+			if (SectionRange.Value.GetUpperBound().IsClosed())
+			{
+				OutSnapTimes.Add(SectionRange.Value.GetUpperBoundValue());
+			}
 		}
 	}
 
-	virtual void GetAllKeyTimes(TArray<float>& OutKeyTimes) const {}
-
 	/** Sets this section's new row index */
 	void SetRowIndex(int32 NewRowIndex) {RowIndex = NewRowIndex;}
 
@@ -353,25 +353,6 @@ public:
 		return OverlapPriority;
 	}
 
-	/**
-	 * Adds a key to a rich curve, finding an existing key to modify or adding a new one.
-	 *
-	 * @param InCurve The curve to add keys to.
-	 * @param Time The time where the key should be added.
-	 * @param Value The value at the given time.
-	 * @param Interpolation The key interpolation to use.
-	 * @param bUnwindRotation Unwind rotation.
-	 */
-	void MOVIESCENE_API AddKeyToCurve(FRichCurve& InCurve, float Time, float Value, EMovieSceneKeyInterpolation Interpolation, const bool bUnwindRotation = false);
-
-	/**
-	 * Sets the default value for a curve.
-	 *
-	 * @param InCurve The curve to set a default value on.
-	 * @param Value The value to use as the default.
-	 */
-	void MOVIESCENE_API SetCurveDefault(FRichCurve& InCurve, float Value);
-
 	/**
 	 * Checks to see if this section overlaps with an array of other sections
 	 * given an optional time and track delta.
@@ -381,17 +362,27 @@ public:
 	 * @param TimeDelta Optional offset to this section's time delta.
 	 * @return The first section that overlaps, or null if there is no overlap.
 	 */
-	virtual MOVIESCENE_API const UMovieSceneSection* OverlapsWithSections(const TArray<UMovieSceneSection*>& Sections, int32 TrackDelta = 0, float TimeDelta = 0.f) const;
+	virtual MOVIESCENE_API const UMovieSceneSection* OverlapsWithSections(const TArray<UMovieSceneSection*>& Sections, int32 TrackDelta = 0, int32 TimeDelta = 0) const;
 	
 	/**
 	 * Places this section at the first valid row at the specified time. Good for placement upon creation.
 	 *
 	 * @param Sections Sections that we can not overlap with.
 	 * @param InStartTime The new start time.
-	 * @param InEndTime The new end time.
+	 * @param InDuration The duration.
 	 * @param bAllowMultipleRows If false, it will move the section in the time direction to make it fit, rather than the row direction.
 	 */
-	virtual MOVIESCENE_API void InitialPlacement(const TArray<UMovieSceneSection*>& Sections, float InStartTime, float InEndTime, bool bAllowMultipleRows);
+	virtual MOVIESCENE_API void InitialPlacement(const TArray<UMovieSceneSection*>& Sections, FFrameNumber InStartTime, int32 InDuration, bool bAllowMultipleRows);
+
+	/**
+	 * Places this section at the specified row at the specified time. Overlapping sections will be moved down a row. Good for placement upon creation.
+	 *
+	 * @param Sections Sections that we can not overlap with.
+	 * @param InStartTime The new start time.
+	 * @param InDuration The duration.
+	 * @param InRowIndex The row index to place this section on.
+	 */
+	virtual MOVIESCENE_API void InitialPlacementOnRow(const TArray<UMovieSceneSection*>& Sections, FFrameNumber InStartTime, int32 InDuration, int32 InRowIndex);
 
 	/** Whether or not this section is active. */
 	void SetIsActive(bool bInIsActive) { bIsActive = bInIsActive; }
@@ -401,26 +392,17 @@ public:
 	void SetIsLocked(bool bInIsLocked) { bIsLocked = bInIsLocked; }
 	bool IsLocked() const { return bIsLocked; }
 
-	/** Whether or not this section is infinite. An infinite section will draw the entire width of the track. StartTime and EndTime will be ignored but not discarded. */
-	void SetIsInfinite(bool bInIsInfinite) { bIsInfinite = bInIsInfinite; }
-	bool IsInfinite() const { return bIsInfinite; }
-
-	/** Gets the amount of time to prepare this section for evaluation before it actually starts. */
-	void SetPreRollTime(float InPreRollTime) { PreRollTime = InPreRollTime; }
-	float GetPreRollTime() const { return PreRollTime; }
+	/** Gets the number of frames to prepare this section for evaluation before it actually starts. */
+	void SetPreRollFrames(int32 InPreRollFrames) { if (TryModify()){ PreRollFrames = InPreRollFrames; } }
+	int32 GetPreRollFrames() const { return PreRollFrames.Value; }
 
-	/** Gets/sets the amount of time to continue 'postrolling' this section for after evaluation has ended. */
-	void SetPostRollTime(float InPostRollTime) { PostRollTime = InPostRollTime; }
-	float GetPostRollTime() const { return PostRollTime; }
+	/** Gets/sets the number of frames to continue 'postrolling' this section for after evaluation has ended. */
+	void SetPostRollFrames(int32 InPostRollFrames) { if (TryModify()){ PostRollFrames = InPostRollFrames; } }
+	int32 GetPostRollFrames() const { return PostRollFrames.Value; }
 
 	/** The optional offset time of this section */
-	virtual TOptional<float> GetOffsetTime() const { return TOptional<float>(); }
+	virtual TOptional<FFrameTime> GetOffsetTime() const { return TOptional<FFrameTime>(); }
 
-	/** Gets the time for the key referenced by the supplied key handle. */
-	virtual TOptional<float> GetKeyTime( FKeyHandle KeyHandle ) const PURE_VIRTUAL( UAISenseEvent::GetKeyTime, return TOptional<float>(); );
-
-	/** Sets the time for the key referenced by the supplied key handle. */
-	virtual void SetKeyTime( FKeyHandle KeyHandle, float Time ) PURE_VIRTUAL( UAISenseEvent::SetKeyTime, );
 	/**
 	 * When guid bindings are updated to allow this section to fix-up any internal bindings
 	 *
@@ -435,36 +417,54 @@ public:
 	/**
 	 * Evaluate this sections's easing functions based on the specified time
 	 */
-	MOVIESCENE_API float EvaluateEasing(float InTime) const;
+	MOVIESCENE_API float EvaluateEasing(FFrameTime InTime) const;
 
 	/**
 	 * Evaluate this sections's easing functions based on the specified time
 	 */
-	MOVIESCENE_API void EvaluateEasing(float InTime, TOptional<float>& OutEaseInValue, TOptional<float>& OutEaseOutValue, float* OutEaseInInterp, float* OutEaseOutInterp) const;
+	MOVIESCENE_API void EvaluateEasing(FFrameTime InTime, TOptional<float>& OutEaseInValue, TOptional<float>& OutEaseOutValue, float* OutEaseInInterp, float* OutEaseOutInterp) const;
 
-	MOVIESCENE_API TRange<float> GetEaseInRange() const;
+	MOVIESCENE_API TRange<FFrameNumber> GetEaseInRange() const;
 
-	MOVIESCENE_API TRange<float> GetEaseOutRange() const;
+	MOVIESCENE_API TRange<FFrameNumber> GetEaseOutRange() const;
 
+	/**
+	 * Access this section's channel proxy, containing pointers to all existing data channels in this section
+	 * @note: Proxy can be reallocated at any time; this accessor is only for immediate use.
+	 *
+	 * @return A reference to this section's channel proxy.
+	 */
+	MOVIESCENE_API FMovieSceneChannelProxy& GetChannelProxy() const;
+
+	/** Does this movie section support infinite ranges for evaluation */
+	MOVIESCENE_API bool GetSupportsInfiniteRange() const { return bSupportsInfiniteRange; }
 protected:
 
 	//~ UObject interface
 	MOVIESCENE_API virtual void PostInitProperties() override;
+	MOVIESCENE_API virtual void Serialize(FArchive& Ar) override;
+
+	virtual void OnMoved(int32 DeltaTime) {}
+	virtual void OnDilated(float DilationFactor, FFrameNumber Origin) {}
 
 public:
 
 	UPROPERTY(EditAnywhere, Category="Easing", meta=(ShowOnlyInnerProperties))
 	FMovieSceneEasingSettings Easing;
 
+	/** The range in which this section is active */
+	UPROPERTY(EditAnywhere, Category="Section")
+	FMovieSceneFrameRange SectionRange;
 private:
 
-	/** The start time of the section */
-	UPROPERTY(EditAnywhere, Category="Section")
-	float StartTime;
 
-	/** The end time of the section */
-	UPROPERTY(EditAnywhere, Category="Section")
-	float EndTime;
+	/** The amount of time to prepare this section for evaluation before it actually starts. */
+	UPROPERTY(EditAnywhere, AdvancedDisplay, Category="Section", meta=(UIMin=0))
+	FFrameNumber PreRollFrames;
+
+	/** The amount of time to continue 'postrolling' this section for after evaluation has ended. */
+	UPROPERTY(EditAnywhere, AdvancedDisplay, Category="Section", meta=(UIMin=0))
+	FFrameNumber PostRollFrames;
 
 	/** The row index that this section sits on */
 	UPROPERTY()
@@ -482,20 +482,36 @@ private:
 	UPROPERTY(EditAnywhere, Category="Section")
 	uint32 bIsLocked : 1;
 
-	/** Toggle to set this section to be infinite */
-	UPROPERTY(EditAnywhere, Category="Section")
-	uint32 bIsInfinite : 1;
+protected:
+
+	/** The start time of the section */
+	UPROPERTY()
+	float StartTime_DEPRECATED;
+
+	/** The end time of the section */
+	UPROPERTY()
+	float EndTime_DEPRECATED;
 
 	/** The amount of time to prepare this section for evaluation before it actually starts. */
-	UPROPERTY(EditAnywhere, AdvancedDisplay, Category="Section", meta=(Units=s, UIMin=0))
-	float PreRollTime;
+	UPROPERTY()
+	float PreRollTime_DEPRECATED;
 
 	/** The amount of time to continue 'postrolling' this section for after evaluation has ended. */
-	UPROPERTY(EditAnywhere, AdvancedDisplay, Category="Section", meta=(Units=s, UIMin=0))
-	float PostRollTime;
+	UPROPERTY()
+	float PostRollTime_DEPRECATED;
+
+	/** Toggle to set this section to be infinite */
+	UPROPERTY()
+	uint32 bIsInfinite_DEPRECATED : 1;
 
 protected:
+	/** Does this section support infinite ranges in the track editor? */
+	UPROPERTY()
+	bool bSupportsInfiniteRange;
 
 	UPROPERTY()
 	FOptionalMovieSceneBlendType BlendType;
+
+	/** Shared channel proxy - to be populated and invalidated by derived types */
+	TSharedPtr<FMovieSceneChannelProxy> ChannelProxy;
 };
diff --git a/Engine/Source/Runtime/MovieScene/Public/MovieSceneSequencePlayer.h b/Engine/Source/Runtime/MovieScene/Public/MovieSceneSequencePlayer.h
--- a/Engine/Source/Runtime/MovieScene/Public/MovieSceneSequencePlayer.h
+++ b/Engine/Source/Runtime/MovieScene/Public/MovieSceneSequencePlayer.h
@@ -9,10 +9,25 @@
 #include "MovieScene.h"
 #include "Evaluation/MovieSceneEvaluationTemplateInstance.h"
 #include "MovieSceneBindingOverridesInterface.h"
+#include "QualifiedFrameTime.h"
+#include "MovieSceneTimeController.h"
 #include "MovieSceneSequencePlayer.generated.h"
 
 DECLARE_DYNAMIC_MULTICAST_DELEGATE(FOnMovieSceneSequencePlayerEvent);
 
+/**
+ * Enum used to define how to update to a particular time
+ */
+enum class EUpdatePositionMethod
+{
+	/** Update from the current position to a specified position (including triggering events), using the current player status */
+	Play,
+	/** Jump to a specified position (without triggering events in between), using the current player status */
+	Jump,
+	/** Jump to a specified position, temporarily using EMovieScenePlayerStatus::Scrubbing */
+	Scrub,
+};
+
 /**
  * Settings for the level sequence player actor.
  */
@@ -30,6 +45,7 @@ struct FMovieSceneSequencePlaybackSettings
 		, bHidePlayer(false)
 		, bHideHud(false)
 		, bDisableCameraCuts(false)
+		, bPauseAtEnd(false)
 		, InstanceData(nullptr)
 	{ }
 
@@ -47,8 +63,8 @@ struct FMovieSceneSequencePlaybackSettings
 	UPROPERTY(EditAnywhere, BlueprintReadWrite, Category="Playback")
 	bool bRandomStartTime;
 
-	/** Start playback at the specified time */
-	UPROPERTY(EditAnywhere, BlueprintReadWrite, Category="Playback", meta=(Units=s, EditCondition="!bRandomStartTime"))
+	/** Start playback at the specified offset from the start of the sequence's playback range */
+	UPROPERTY(EditAnywhere, BlueprintReadWrite, Category="Playback", DisplayName="Start Offset", meta=(EditCondition="!bRandomStartTime", Units=s))
 	float StartTime;
 
 	/** Flag used to specify whether actor states should be restored on stop */
@@ -75,6 +91,10 @@ struct FMovieSceneSequencePlaybackSettings
 	UPROPERTY(EditAnywhere, BlueprintReadWrite, Category="Cinematic")
 	bool bDisableCameraCuts;
 
+	/** Pause the sequence when playback reaches the end rather than stopping it */
+	UPROPERTY(EditAnywhere, BlueprintReadWrite, Category="Cinematic")
+	bool bPauseAtEnd;
+
 	/** An object that can implement specific instance overrides for the sequence */
 	UPROPERTY(BlueprintReadWrite, Category="Cinematic")
 	UObject* InstanceData;
@@ -83,6 +103,9 @@ struct FMovieSceneSequencePlaybackSettings
 	UPROPERTY()
 	TScriptInterface<IMovieSceneBindingOverridesInterface> BindingOverrides;
 
+	/** (Optional) Externally supplied time controller */
+	TSharedPtr<FMovieSceneTimeController> TimeController;
+
 	MOVIESCENE_API bool SerializeFromMismatchedTag(const FPropertyTag& Tag, FArchive& Ar);
 };
 
@@ -123,10 +146,6 @@ public:
 	 */
 	UFUNCTION(BlueprintCallable, Category="Game|Cinematic")
 	void PlayLooping(int32 NumLoops = -1);
-
-	/** Start playback from the current time cursor position, using the current play rate. Does not update the animation until next tick. */
-	UFUNCTION(BlueprintCallable, Category="Game|Cinematic")
-	void StartPlayingNextTick();
 	
 	/** Pause playback. */
 	UFUNCTION(BlueprintCallable, Category="Game|Cinematic")
@@ -144,66 +163,202 @@ public:
 	UFUNCTION(BlueprintCallable, Category="Game|Cinematic", meta = (ToolTip = "Go to end of the sequence and stop. Adheres to 'When Finished' section rules."))
 	void GoToEndAndStop();
 
-	/** Get the current playback position */
+public:
+
+	/**
+	 * Get the current playback position
+	 * @return The current playback position
+	 */
+	DEPRECATED(4.20, "Please use GetCurrentTime instead")
 	UFUNCTION(BlueprintCallable, Category="Game|Cinematic")
-	float GetPlaybackPosition() const;
+	float GetPlaybackPosition() const { return GetCurrentTime().AsSeconds() - StartTime / PlayPosition.GetInputRate(); }
+
+	/**
+	 * Get the playback length of the sequence
+	 */
+	DEPRECATED(4.20, "Please use GetDuration instead")
+	UFUNCTION(BlueprintCallable, Category="Game|Cinematic")
+	float GetLength() const;
+
+	/**
+	 * Get the offset within the level sequence to start playing
+	 */
+	DEPRECATED(4.20, "Please use GetStartTime instead")
+	UFUNCTION(BlueprintCallable, Category="Game|Cinematic")
+	float GetPlaybackStart() const { return StartTime / PlayPosition.GetInputRate(); }
+
+	/**
+	 * Get the offset within the level sequence to finish playing
+	 */
+	DEPRECATED(4.20, "Please use GetEndTime instead")
+	UFUNCTION(BlueprintCallable, Category="Game|Cinematic")
+	float GetPlaybackEnd() const { return (StartTime + DurationFrames) / PlayPosition.GetInputRate(); }
 
 	/**
 	 * Set the current playback position
 	 * @param NewPlaybackPosition - The new playback position to set.
 	 * If the animation is currently playing, it will continue to do so from the new position
 	 */
+	DEPRECATED(4.20, "Please use PlayToFrame instead")
 	UFUNCTION(BlueprintCallable, Category="Game|Cinematic")
-	void SetPlaybackPosition(float NewPlaybackPosition);
+	void SetPlaybackPosition(float NewPlaybackPosition) { PlayToSeconds(NewPlaybackPosition + StartTime / PlayPosition.GetInputRate()); }
+
+	/**
+	 * Sets the range in time to be played back by this player, overriding the default range stored in the asset
+	 *
+	 * @param	NewStartTime	The new starting time for playback
+	 * @param	NewEndTime		The new ending time for playback.  Must be larger than the start time.
+	 */
+	DEPRECATED(4.20, "Please use SetFrameRange or SetTimeRange instead")
+	UFUNCTION(BlueprintCallable, Category="Game|Cinematic")
+	void SetPlaybackRange( const float NewStartTime, const float NewEndTime );
 
 	/**
 	 * Jump to new playback position
 	 * @param NewPlaybackPosition - The new playback position to set.
 	 * This can be used to update sequencer repeatedly, as if in a scrubbing state
 	 */
+	DEPRECATED(4.20, "Please use ScrubToTime instead")
 	UFUNCTION(BlueprintCallable, Category="Game|Cinematic")
-	void JumpToPosition(float NewPlaybackPosition);
+	void JumpToPosition(float NewPlaybackPosition) { ScrubToSeconds(NewPlaybackPosition); }
 
-	/** Check whether the sequence is actively playing. */
+public:
+
+	/**
+	 * Get the current playback position
+	 * @return The current playback position
+	 */
 	UFUNCTION(BlueprintCallable, Category="Game|Cinematic")
-	bool IsPlaying() const;
+	FQualifiedFrameTime GetCurrentTime() const;
 
-	/** Check whether the sequence is paused. */
+	/**
+	 * Get the total duration of the sequence
+	 */
 	UFUNCTION(BlueprintCallable, Category="Game|Cinematic")
-	bool IsPaused() const;
+	FQualifiedFrameTime GetDuration() const;
 
-	/** Get the playback length of the sequence */
+	/**
+	 * Get this sequence's duration in frames
+	 */
 	UFUNCTION(BlueprintCallable, Category="Game|Cinematic")
-	float GetLength() const;
+	int32 GetFrameDuration() const;
 
-	/** Get the playback rate of this player. */
+	/**
+	 * Get this sequence's playback frame rate.
+	 */
 	UFUNCTION(BlueprintCallable, Category="Game|Cinematic")
-	float GetPlayRate() const;
+	FFrameRate GetFrameRate() const { return PlayPosition.GetInputRate(); }
 
 	/**
-	 * Set the playback rate of this player. Negative values will play the animation in reverse.
-	 * @param PlayRate - The new rate of playback for the animation.
+	 * Set the frame-rate that this player should play with, making all frame numbers in the specified time-space
 	 */
 	UFUNCTION(BlueprintCallable, Category="Game|Cinematic")
-	void SetPlayRate(float PlayRate);
+	void SetFrameRate(FFrameRate FrameRate);
 
 	/**
-	 * Sets the range in time to be played back by this player, overriding the default range stored in the asset
+	 * Get the offset within the level sequence to start playing
+	 */
+	UFUNCTION(BlueprintCallable, Category="Game|Cinematic")
+	FQualifiedFrameTime GetStartTime() const { return FQualifiedFrameTime(StartTime, PlayPosition.GetInputRate()); }
+
+	/**
+	 * Get the offset within the level sequence to finish playing
+	 */
+	UFUNCTION(BlueprintCallable, Category="Game|Cinematic")
+	FQualifiedFrameTime GetEndTime() const { return FQualifiedFrameTime(StartTime + DurationFrames, PlayPosition.GetInputRate()); }
+
+public:
+
+	/**
+	 * Set the valid play range for this sequence, determined by a starting frame number (in this sequence player's plaback frame), and a number of frames duration
 	 *
-	 * @param	NewStartTime	The new starting time for playback
-	 * @param	NewEndTime		The new ending time for playback.  Must be larger than the start time.
+	 * @param StartFrame      The frame number to start playing back the sequence
+	 * @param Duration        The number of frames to play
+	 */
+	UFUNCTION(BlueprintCallable, Category="Game|Cinematic", DisplayName="Set Play Range (Frames)")
+	void SetFrameRange( int32 StartFrame, int32 Duration );
+
+	/**
+	 * Set the valid play range for this sequence, determined by a starting time  and a duration (in seconds)
+	 *
+	 * @param StartTime       The time to start playing back the sequence in seconds
+	 * @param Duration        The length to play for
+	 */
+	UFUNCTION(BlueprintCallable, Category="Game|Cinematic", DisplayName="Set Play Range (Seconds)")
+	void SetTimeRange( float StartTime, float Duration );
+
+public:
+
+	/**
+	 * Play the sequence from the current time, to the specified frame position
+	 *
+	 * @param NewPosition     The new frame time to play to
+	 */
+	UFUNCTION(BlueprintCallable, Category="Game|Cinematic", DisplayName="Play To (Frames)")
+	void PlayToFrame(FFrameTime NewPosition);
+
+	/**
+	 * Scrub the sequence from the current time, to the specified frame position
+	 *
+	 * @param NewPosition     The new frame time to scrub to
+	 */
+	UFUNCTION(BlueprintCallable, Category="Game|Cinematic", DisplayName="Scrub To (Frames)")
+	void ScrubToFrame(FFrameTime NewPosition);
+
+	/**
+	 * Jump to the specified frame position, without evaluating the sequence in between the current and desired time (as if in a paused state)
+	 *
+	 * @param NewPosition     The new frame time to jump to
+	 */
+	UFUNCTION(BlueprintCallable, Category="Game|Cinematic", DisplayName="Jump To (Frames)")
+	void JumpToFrame(FFrameTime NewPosition);
+
+
+	/**
+	 * Play the sequence from the current time, to the specified time in seconds
+	 *
+	 * @param TimeInSeconds   The desired time in seconds
+	 */
+	UFUNCTION(BlueprintCallable, Category="Game|Cinematic", DisplayName="Play To (Seconds)")
+	void PlayToSeconds(float TimeInSeconds);
+
+	/**
+	 * Scrub the sequence from the current time, to the specified time in seconds
+	 *
+	 * @param TimeInSeconds   The desired time in seconds
+	 */
+	UFUNCTION(BlueprintCallable, Category="Game|Cinematic", DisplayName="Scrub To (Seconds)")
+	void ScrubToSeconds(float TimeInSeconds);
+
+	/**
+	 * Jump to the specified time in seconds, without evaluating the sequence in between the current and desired time (as if in a paused state)
+	 *
+	 * @param TimeInSeconds   The desired time in seconds
 	 */
+	UFUNCTION(BlueprintCallable, Category="Game|Cinematic", DisplayName="Jump To (Seconds)")
+	void JumpToSeconds(float TimeInSeconds);
+
+public:
+
+	/** Check whether the sequence is actively playing. */
 	UFUNCTION(BlueprintCallable, Category="Game|Cinematic")
-	void SetPlaybackRange( const float NewStartTime, const float NewEndTime );
+	bool IsPlaying() const;
 
-	/** Get the offset within the level sequence to start playing */
+	/** Check whether the sequence is paused. */
 	UFUNCTION(BlueprintCallable, Category="Game|Cinematic")
-	float GetPlaybackStart() const { return StartTime; }
+	bool IsPaused() const;
 
-	/** Get the offset within the level sequence to finish playing */
+	/** Get the playback rate of this player. */
 	UFUNCTION(BlueprintCallable, Category="Game|Cinematic")
-	float GetPlaybackEnd() const { return EndTime; }
-	
+	float GetPlayRate() const;
+
+	/**
+	 * Set the playback rate of this player. Negative values will play the animation in reverse.
+	 * @param PlayRate - The new rate of playback for the animation.
+	 */
+	UFUNCTION(BlueprintCallable, Category="Game|Cinematic")
+	void SetPlayRate(float PlayRate);
+
 	/** Set whether to disable camera cuts */
 	UFUNCTION(BlueprintCallable, Category="Game|Cinematic")
 	void SetDisableCameraCuts(bool bInDisableCameraCuts) { PlaybackSettings.bDisableCameraCuts = bInDisableCameraCuts; }
@@ -213,7 +368,7 @@ public:
 	bool GetDisableCameraCuts() { return PlaybackSettings.bDisableCameraCuts; }
 
 	/** An event that is broadcast each time this level sequence player is updated */
-	DECLARE_EVENT_ThreeParams( UMovieSceneSequencePlayer, FOnMovieSceneSequencePlayerUpdated, const UMovieSceneSequencePlayer&, float /*current time*/, float /*previous time*/ );
+	DECLARE_EVENT_ThreeParams( UMovieSceneSequencePlayer, FOnMovieSceneSequencePlayerUpdated, const UMovieSceneSequencePlayer&, FFrameTime /*current time*/, FFrameTime /*previous time*/ );
 	FOnMovieSceneSequencePlayerUpdated& OnSequenceUpdated() const { return OnMovieSceneSequencePlayerUpdate; }
 
 	/** Event triggered when the level sequence player is played */
@@ -245,14 +400,14 @@ public:
 
 public:
 
-	/** Update based on the specified delta seconds */
+	/** Update the sequence for the current time, if playing */
 	void Update(const float DeltaSeconds);
 
 	/** Initialize this player with a sequence and some settings */
 	void Initialize(UMovieSceneSequence* InSequence, const FMovieSceneSequencePlaybackSettings& InSettings);
 
 	/** Begin play called */
-	virtual void BeginPlay() {};
+	virtual void BeginPlay() {}
 
 public:
 
@@ -266,13 +421,14 @@ protected:
 
 	void PlayInternal();
 
-	void UpdateMovieSceneInstance(FMovieSceneEvaluationRange InRange, TOptional<EMovieScenePlayerStatus::Type> OptionalStatus = TOptional<EMovieScenePlayerStatus::Type>(), bool bHasJumped = false);
+	void UpdateMovieSceneInstance(FMovieSceneEvaluationRange InRange, EMovieScenePlayerStatus::Type PlayerStatus, bool bHasJumped = false);
 
-	void UpdateTimeCursorPosition(float NewPosition, TOptional<EMovieScenePlayerStatus::Type> OptionalStatus = TOptional<EMovieScenePlayerStatus::Type>());
+	void UpdateTimeCursorPosition(FFrameTime NewPosition, EUpdatePositionMethod Method);
+	bool ShouldStopOrLoop(FFrameTime NewPosition) const;
 
-	bool ShouldStopOrLoop(float NewPosition) const;
+	UWorld* GetPlaybackWorld() const;
 
-	FORCEINLINE float GetSequencePosition() const { return TimeCursorPosition + StartTime; }
+	FFrameTime GetLastValidTime() const;
 
 protected:
 
@@ -306,6 +462,8 @@ private:
 	/** Apply any latent actions which may have accumulated while the sequence was being evaluated */
 	void ApplyLatentActions();
 
+	void UpdateTimeCursorPosition_Internal(FFrameTime NewPosition, EUpdatePositionMethod Method);
+
 protected:
 
 	/** Movie player status. */
@@ -316,10 +474,6 @@ protected:
 	UPROPERTY()
 	uint32 bReversePlayback : 1;
 
-	/** True where we're waiting for the first update of the sequence after calling StartPlayingNextTick. */
-	UPROPERTY()
-	uint32 bPendingFirstUpdate : 1;
-
 	/** Set to true while evaluating to prevent reentrancy */
 	bool bIsEvaluating : 1;
 
@@ -327,29 +481,37 @@ protected:
 	UPROPERTY(transient)
 	UMovieSceneSequence* Sequence;
 
-	/** The current time cursor position within the sequence (in seconds) */
-	UPROPERTY()
-	float TimeCursorPosition;
-
-	/** Time time at which to start playing the sequence (defaults to the lower bound of the sequence's play range) */
+	/** Time (in playback frames) at which to start playing the sequence (defaults to the lower bound of the sequence's play range) */
 	UPROPERTY()
-	float StartTime;
+	FFrameNumber StartTime;
 
-	/** Time time at which to end playing the sequence (defaults to the upper bound of the sequence's play range) */
+	/** Time (in playback frames) at which to stop playing the sequence (defaults to the upper bound of the sequence's play range) */
 	UPROPERTY()
-	float EndTime;
+	int32 DurationFrames;
 
 	/** The number of times we have looped in the current playback */
 	UPROPERTY(transient)
 	int32 CurrentNumLoops;
 
-	enum class ELatentAction
+	struct FLatentAction
 	{
-		Stop, Pause, SetPlaybackPosition
+		enum EType { Stop, Pause, Update };
+
+		FLatentAction(EType InType)
+			: Type(InType)
+		{}
+
+		FLatentAction(EUpdatePositionMethod InUpdateMethod, FFrameTime DesiredTime)
+			: Type(Update), UpdateMethod(InUpdateMethod), Position(DesiredTime)
+		{}
+
+		EType                 Type;
+		EUpdatePositionMethod UpdateMethod;
+		FFrameTime            Position;
 	};
 
 	/** Set of latent actions that are to be performed when the sequence has finished evaluating this frame */
-	TArray<ELatentAction> LatentActions;
+	TArray<FLatentAction> LatentActions;
 
 	/** Specific playback settings for the animation. */
 	UPROPERTY()
@@ -370,7 +532,4 @@ private:
 
 	/** The maximum tick rate prior to playing (used for overriding delta time during playback). */
 	TOptional<double> OldMaxTickRate;
-
-	/** Playback position to use when invoking latent SetPlaybackPosition */
-	float LatentPlaybackPosition;
 };
diff --git a/Engine/Source/Runtime/MovieScene/Public/MovieSceneSignedObject.h b/Engine/Source/Runtime/MovieScene/Public/MovieSceneSignedObject.h
--- a/Engine/Source/Runtime/MovieScene/Public/MovieSceneSignedObject.h
+++ b/Engine/Source/Runtime/MovieScene/Public/MovieSceneSignedObject.h
@@ -11,17 +11,17 @@
 class ITransactionObjectAnnotation;
 
 UCLASS()
-class UMovieSceneSignedObject : public UObject
+class MOVIESCENE_API UMovieSceneSignedObject : public UObject
 {
 public:
 	GENERATED_BODY()
 
-	MOVIESCENE_API UMovieSceneSignedObject(const FObjectInitializer& Init);
+	UMovieSceneSignedObject(const FObjectInitializer& Init);
 
 	/**
 	 * 
 	 */
-	MOVIESCENE_API void MarkAsChanged();
+	void MarkAsChanged();
 
 	/**
 	 * 
@@ -37,14 +37,14 @@ public:
 
 public:
 
-	MOVIESCENE_API virtual bool Modify(bool bAlwaysMarkDirty = true) override;
-	MOVIESCENE_API virtual void PostInitProperties() override;
-	MOVIESCENE_API virtual void PostLoad() override;
+	virtual bool Modify(bool bAlwaysMarkDirty = true) override;
+	virtual void PostInitProperties() override;
+	virtual void PostLoad() override;
 
 #if WITH_EDITOR
-	MOVIESCENE_API virtual void PostEditChangeProperty( struct FPropertyChangedEvent& PropertyChangedEvent) override;
-	MOVIESCENE_API virtual void PostEditUndo() override;
-	MOVIESCENE_API virtual void PostEditUndo(TSharedPtr<ITransactionObjectAnnotation> TransactionAnnotation) override;
+	virtual void PostEditChangeProperty( struct FPropertyChangedEvent& PropertyChangedEvent) override;
+	virtual void PostEditUndo() override;
+	virtual void PostEditUndo(TSharedPtr<ITransactionObjectAnnotation> TransactionAnnotation) override;
 #endif
 
 private:
diff --git a/Engine/Source/Runtime/MovieScene/Public/MovieSceneSpawnRegister.h b/Engine/Source/Runtime/MovieScene/Public/MovieSceneSpawnRegister.h
--- a/Engine/Source/Runtime/MovieScene/Public/MovieSceneSpawnRegister.h
+++ b/Engine/Source/Runtime/MovieScene/Public/MovieSceneSpawnRegister.h
@@ -105,6 +105,13 @@ public:
 	
 public:
 
+	/**
+	 * Whether this class is supported as a spawnable
+	 * @param InClass The class
+	 * @return Whether this class is supported as a spawnable
+	 */
+	virtual bool CanSpawnObject(UClass* InClass) const { return false; }
+
 	/**
 	 * Create a new spawnable type from the given source object
 	 *
@@ -129,7 +136,7 @@ public:
 	 * @param	Sequencer		The sequencer this spawnable was just created by
 	 * @param	Settings		The settings for this sequencer
 	 */
-	virtual void SetupDefaultsForSpawnable(UObject* SpawnedObject, const FGuid& Guid, const FTransformData& TransformData, TSharedRef<ISequencer> Sequencer, USequencerSettings* Settings) {}
+	virtual void SetupDefaultsForSpawnable(UObject* SpawnedObject, const FGuid& Guid, const TOptional<FTransformData>& TransformData, TSharedRef<ISequencer> Sequencer, USequencerSettings* Settings) {}
 
 	/**
 	 * Called to handle cleanup of objects when we convert a possessable to a spawnable object
@@ -138,7 +145,7 @@ public:
 	 * @param	Player		The current player
 	 * @param	OutTransformData	Transform data that can be used to recreate objects in the same location
 	 */
-	virtual void HandleConvertPossessableToSpawnable(UObject* OldObject, IMovieScenePlayer& Player, FTransformData& OutTransformData) {}
+	virtual void HandleConvertPossessableToSpawnable(UObject* OldObject, IMovieScenePlayer& Player, TOptional<FTransformData>& OutTransformData) {}
 
 
 	/**
diff --git a/Engine/Source/Runtime/MovieScene/Public/MovieSceneTimeController.h b/Engine/Source/Runtime/MovieScene/Public/MovieSceneTimeController.h
new file mode 100644
--- /dev/null
+++ b/Engine/Source/Runtime/MovieScene/Public/MovieSceneTimeController.h
@@ -0,0 +1,135 @@
+// Copyright 1998-2018 Epic Games, Inc. All Rights Reserved.
+
+#pragma once
+
+#include "CoreTypes.h"
+#include "QualifiedFrameTime.h"
+#include "MovieSceneFwd.h"
+#include "Misc/Attribute.h"
+#include "UObject/ObjectMacros.h"
+
+class IMovieScenePlayer;
+
+struct MOVIESCENE_API FMovieSceneTimeController
+{
+public:
+
+	virtual ~FMovieSceneTimeController() {}
+
+	/**
+	 * Called whenever a sequence starts or resumes playback from a non-playing state
+	 */
+	void StartPlaying(const FQualifiedFrameTime& InStartTime);
+
+	/**
+	 * Called whenever a sequence stops playback
+	 */
+	void StopPlaying(const FQualifiedFrameTime& InStopTime);
+
+	/**
+	 * Ticks this controller
+	 *
+	 * @param DeltaSeconds     The tick delta in seconds, dilated by the current world settings global dilation
+	 * @param InPlayRate       The current play rate of the sequence
+	 */
+	void Tick(float DeltaSeconds, float InPlayRate);
+
+	/**
+	 * Request the current time based on the specified existing time and play rate.
+	 * Times should be returned in the same play rate as that specified by InCurrentTime
+	 *
+	 * @param InCurrentTime    The current time of the sequence
+	 * @param InPlayRate       The current play rate of the sequence, multiplied by any world actor settings global dilation
+	 */
+	FFrameTime RequestCurrentTime(const FQualifiedFrameTime& InCurrentTime, float InPlayRate);
+
+	/**
+	 * Called when the status of the owning IMovieScenePlayer has changed
+	 */
+	void PlayerStatusChanged(EMovieScenePlayerStatus::Type InStatus, const FQualifiedFrameTime& InCurrentTime);
+
+	/**
+	 * Called to stop and resume playback from the specified time
+	 */
+	void Reset(const FQualifiedFrameTime& InNewStartTime);
+
+protected:
+
+	virtual void OnTick(float DeltaSeconds, float InPlayRate){}
+	virtual void OnStartPlaying(const FQualifiedFrameTime& InStartTime){}
+	virtual void OnStopPlaying(const FQualifiedFrameTime& InStopTime){}
+	virtual FFrameTime OnRequestCurrentTime(const FQualifiedFrameTime& InCurrentTime, float InPlayRate) = 0;
+
+protected:
+
+	TOptional<FQualifiedFrameTime> GetPlaybackStartTime() const
+	{
+		return PlaybackStartTime;
+	}
+
+private:
+
+	TOptional<FQualifiedFrameTime> PlaybackStartTime;
+};
+
+/**
+ * A timing manager that retrieves its time from an external clock source
+ */
+struct MOVIESCENE_API FMovieSceneTimeController_ExternalClock : FMovieSceneTimeController
+{
+protected:
+
+	FMovieSceneTimeController_ExternalClock()
+		: AccumulatedDilation(0.0)
+	{}
+
+	virtual double GetCurrentTime() const = 0;
+
+protected:
+
+	virtual void OnStartPlaying(const FQualifiedFrameTime& InStartTime) override;
+	virtual void OnStopPlaying(const FQualifiedFrameTime& InStopTime) override;
+	virtual FFrameTime OnRequestCurrentTime(const FQualifiedFrameTime& InCurrentTime, float InPlayRate) override;
+
+private:
+
+	double AccumulatedDilation;
+
+	TOptional<double> ClockStartTime;
+	TOptional<double> ClockLastUpdateTime;
+};
+
+/**
+ * A timing manager that retrieves its time from the platform clock
+ */
+struct MOVIESCENE_API FMovieSceneTimeController_PlatformClock : FMovieSceneTimeController_ExternalClock
+{
+	virtual double GetCurrentTime() const override;
+};
+
+/**
+ * A timing manager that retrieves its time from the audio clock
+ */
+struct MOVIESCENE_API FMovieSceneTimeController_AudioClock : FMovieSceneTimeController_ExternalClock
+{
+	virtual double GetCurrentTime() const override;
+};
+
+/**
+ * A timing manager that accumulates delta times from a world
+ */
+struct MOVIESCENE_API FMovieSceneTimeController_Tick : FMovieSceneTimeController
+{
+	FMovieSceneTimeController_Tick()
+		: CurrentOffsetSeconds(0.0)
+	{}
+
+protected:
+
+	virtual void OnTick(float DeltaSeconds, float InPlayRate) override;
+	virtual void OnStartPlaying(const FQualifiedFrameTime& InStartTime) override;
+	virtual FFrameTime OnRequestCurrentTime(const FQualifiedFrameTime& InCurrentTime, float InPlayRate) override;
+
+private:
+	double CurrentOffsetSeconds;
+};
\ No newline at end of file
diff --git a/Engine/Source/Runtime/MovieScene/Public/MovieSceneTrack.h b/Engine/Source/Runtime/MovieScene/Public/MovieSceneTrack.h
--- a/Engine/Source/Runtime/MovieScene/Public/MovieSceneTrack.h
+++ b/Engine/Source/Runtime/MovieScene/Public/MovieSceneTrack.h
@@ -253,14 +253,6 @@ public:
 	 */
 	MOVIESCENE_API bool FixRowIndices();
 
-	/**
-	 * @return Whether or not this track's section bounds should be added to the play range
-	 */
-	virtual bool AddsSectionBoundsToPlayRange() const
-	{
-		return false;
-	}
-
 public:
 
 	/**
@@ -289,7 +281,8 @@ public:
 	 * 
 	 * @return The range of time boundaries.
 	 */
-	virtual TRange<float> GetSectionBoundaries() const PURE_VIRTUAL(UMovieSceneTrack::GetSectionBoundaries, return TRange<float>::Empty(););
+	DEPRECATED(4.20, "This function is no longer used.")
+	virtual TRange<FFrameNumber> GetSectionBoundaries() const PURE_VIRTUAL(UMovieSceneTrack::GetSectionBoundaries, return TRange<FFrameNumber>::Empty(););
 
 	/**
 	 * Checks to see if the section is in this track.
@@ -345,12 +338,33 @@ public:
 		TrackTint = InTrackTint;
 	}
 
+	/**
+	* Get this folder's desired sorting order
+	*/
+	int32 GetSortingOrder() const
+	{
+		return SortingOrder;
+	}
+
+	/**
+	* Set this folder's desired sorting order.
+	*
+	* @param InSortingOrder The higher the value the further down the list the folder will be.
+	*/
+	void SetSortingOrder(const int32 InSortingOrder)
+	{
+		SortingOrder = InSortingOrder;
+	}
+
 protected:
 
 	/** This track's tint color */
 	UPROPERTY(EditAnywhere, Category=General, DisplayName=Color)
 	FColor TrackTint;
 
+	/** This folder's desired sorting order */
+	UPROPERTY()
+	int32 SortingOrder;
 public:
 #endif
 
diff --git a/Engine/Source/Runtime/MovieScene/Public/Sections/MovieSceneSubSection.h b/Engine/Source/Runtime/MovieScene/Public/Sections/MovieSceneSubSection.h
--- a/Engine/Source/Runtime/MovieScene/Public/Sections/MovieSceneSubSection.h
+++ b/Engine/Source/Runtime/MovieScene/Public/Sections/MovieSceneSubSection.h
@@ -66,6 +66,11 @@ public:
 
 public:
 
+	/**
+	 * Get the transform that converts time from this section's time-base to its inner sequence's
+	 */
+	FMovieSceneSequenceTransform OuterToInnerTransform() const;
+
 	/**
 	 * Sets the sequence played by this section.
 	 *
@@ -129,11 +134,9 @@ public:
 
 	//~ UMovieSceneSection interface
 
-	virtual UMovieSceneSection* SplitSection( float SplitTime ) override;
-	virtual void TrimSection( float TrimTime, bool bTrimLeft ) override;
-	virtual TOptional<float> GetOffsetTime() const override { return TOptional<float>(Parameters.StartOffset); }
-	virtual TOptional<float> GetKeyTime( FKeyHandle KeyHandle ) const override { return TOptional<float>(); }
-	virtual void SetKeyTime( FKeyHandle KeyHandle, float Time ) override { }
+	virtual UMovieSceneSection* SplitSection( FFrameNumber SplitTime ) override;
+	virtual void TrimSection( FFrameNumber TrimTime, bool bTrimLeft ) override;
+	virtual TOptional<FFrameTime> GetOffsetTime() const override { return TOptional<FFrameTime>(FFrameTime(Parameters.GetStartFrameOffset())); }
 
 public:
 
diff --git a/Engine/Source/Runtime/MovieScene/Public/Tracks/MovieSceneSubTrack.h b/Engine/Source/Runtime/MovieScene/Public/Tracks/MovieSceneSubTrack.h
--- a/Engine/Source/Runtime/MovieScene/Public/Tracks/MovieSceneSubTrack.h
+++ b/Engine/Source/Runtime/MovieScene/Public/Tracks/MovieSceneSubTrack.h
@@ -30,11 +30,21 @@ public:
 	 *
 	 * @param Sequence The sequence to add
 	 * @param StartTime The time to add the section at
-	 * @param Duration The duration of the section
+	 * @param Duration The duration of the section in frames
+	 * @return The newly created sub section
+	 */
+	virtual UMovieSceneSubSection* AddSequence(UMovieSceneSequence* Sequence, FFrameNumber StartTime, int32 Duration) { return AddSequenceOnRow(Sequence, StartTime, Duration, INDEX_NONE); }
+
+	/**
+	 * Adds a movie scene section at the requested time.
+	 *
+	 * @param Sequence The sequence to add
+	 * @param StartTime The time to add the section at
+	 * @param Duration The duration of the section in frames
 	 * @param bInsertSequence Whether or not to insert the sequence and push existing sequences out
 	 * @return The newly created sub section
 	 */
-	virtual UMovieSceneSubSection* AddSequence(UMovieSceneSequence* Sequence, float StartTime, float Duration, const bool& bInsertSequence = false);
+	virtual UMovieSceneSubSection* AddSequenceOnRow(UMovieSceneSequence* Sequence, FFrameNumber StartTime, int32 Duration, int32 RowIndex);
 
 	/**
 	 * Check whether this track contains the given sequence.
@@ -57,7 +67,6 @@ public:
 	virtual void AddSection(UMovieSceneSection& Section) override;
 	virtual UMovieSceneSection* CreateNewSection() override;
 	virtual const TArray<UMovieSceneSection*>& GetAllSections() const override;
-	virtual TRange<float> GetSectionBoundaries() const override;
 	virtual bool HasSection(const UMovieSceneSection& Section) const override;
 	virtual bool IsEmpty() const override;
 	virtual void RemoveAllAnimationData() override;
diff --git a/Engine/Source/Runtime/MovieSceneCapture/Private/AutomatedLevelSequenceCapture.cpp b/Engine/Source/Runtime/MovieSceneCapture/Private/AutomatedLevelSequenceCapture.cpp
--- a/Engine/Source/Runtime/MovieSceneCapture/Private/AutomatedLevelSequenceCapture.cpp
+++ b/Engine/Source/Runtime/MovieSceneCapture/Private/AutomatedLevelSequenceCapture.cpp
@@ -14,6 +14,58 @@
 #include "Sections/MovieSceneCinematicShotSection.h"
 #include "TimerManager.h"
 #include "GameFramework/PlayerController.h"
+#include "Camera/CameraComponent.h"
+#include "MovieSceneCaptureModule.h"
+#include "MovieSceneTimeHelpers.h"
+
+struct FMovieSceneTimeController_FrameStep : FMovieSceneTimeController
+{
+	FMovieSceneTimeController_FrameStep()
+		: DeltaTime(0)
+		, CurrentTime(-1)
+	{}
+
+	virtual void OnTick(float DeltaSeconds, float InPlayRate) override
+	{
+		// Move onto the next frame in the sequence. Play rate dilation occurs in OnRequestCurrentTime, since this InPlayRate does not consider the global world settings dilation
+		DeltaTime = FFrameTime(1);
+	}
+
+	virtual void OnStartPlaying(const FQualifiedFrameTime& InStartTime)
+	{
+		DeltaTime   = FFrameTime(0);
+		CurrentTime = FFrameTime(-1);
+	}
+
+	virtual FFrameTime OnRequestCurrentTime(const FQualifiedFrameTime& InCurrentTime, float InPlayRate) override
+	{
+		TOptional<FQualifiedFrameTime> StartTimeIfPlaying = GetPlaybackStartTime();
+		if (!StartTimeIfPlaying.IsSet())
+		{
+			return InCurrentTime.Time;
+		}
+		else
+		{
+			// Scale the delta time (should be one frame) by this frame's play rate, and add it to the current time offset
+			if (InPlayRate == 1.f)
+			{
+				CurrentTime += DeltaTime;
+			}
+			else
+			{
+				CurrentTime += DeltaTime * InPlayRate;
+			}
+
+			DeltaTime = FFrameTime(0);
+
+			ensure(CurrentTime >= 0);
+			return StartTimeIfPlaying->ConvertTo(InCurrentTime.Rate) + CurrentTime;
+		}
+	}
+
+	FFrameTime DeltaTime;
+	FFrameTime CurrentTime;
+};
 
 UAutomatedLevelSequenceCapture::UAutomatedLevelSequenceCapture(const FObjectInitializer& Init)
 	: Super(Init)
@@ -30,6 +82,7 @@ UAutomatedLevelSequenceCapture::UAutomatedLevelSequenceCapture(const FObjectInit
 	EndFrame = 1;
 	WarmUpFrameCount = 0;
 	DelayBeforeWarmUp = 0.0f;
+	DelayBeforeShotWarmUp = 0.0f;
 	bWriteEditDecisionList = true;
 
 	RemainingWarmUpFrames = 0;
@@ -51,9 +104,12 @@ void UAutomatedLevelSequenceCapture::SetLevelSequenceAsset(FString AssetPath)
 void UAutomatedLevelSequenceCapture::AddFormatMappings(TMap<FString, FStringFormatArg>& OutFormatMappings, const FFrameMetrics& FrameMetrics) const
 {
 	OutFormatMappings.Add(TEXT("shot"), CachedState.CurrentShotName);
+	OutFormatMappings.Add(TEXT("shot_frame"), FString::Printf(TEXT("%0*d"), Settings.ZeroPadFrameNumbers, CachedState.CurrentShotLocalTime.Time.FrameNumber.Value));
 
-	const int32 FrameNumber = FMath::RoundToInt(CachedState.CurrentShotLocalTime * CachedState.Settings.FrameRate);
-	OutFormatMappings.Add(TEXT("shot_frame"), FString::Printf(TEXT("%0*d"), Settings.ZeroPadFrameNumbers, FrameNumber));
+	if (CachedState.CameraComponent && CachedState.CameraComponent->GetOwner())
+	{
+		OutFormatMappings.Add(TEXT("camera"), CachedState.CameraComponent->GetOwner()->GetName());
+	}
 }
 
 void UAutomatedLevelSequenceCapture::Initialize(TSharedPtr<FSceneViewport> InViewport, int32 PIEInstance)
@@ -96,6 +152,12 @@ void UAutomatedLevelSequenceCapture::Initialize(TSharedPtr<FSceneViewport> InVie
 		{
 			DelayBeforeWarmUp = DelayBeforeWarmUpOverride;
 		}
+
+		float DelayBeforeShotWarmUpOverride;
+		if( FParse::Value( FCommandLine::Get(), TEXT( "-MovieDelayBeforeShotWarmUp=" ), DelayBeforeShotWarmUpOverride ) )
+		{
+			DelayBeforeShotWarmUp = DelayBeforeShotWarmUpOverride;
+		}
 	}
 
 	ALevelSequenceActor* Actor = LevelSequenceActor.Get();
@@ -131,8 +193,6 @@ void UAutomatedLevelSequenceCapture::Initialize(TSharedPtr<FSceneViewport> InVie
 			// Spawn a new actor
 			Actor = InViewport->GetClient()->GetWorld()->SpawnActor<ALevelSequenceActor>();
 			Actor->SetSequence(Asset);
-			// Ensure it doesn't loop (-1 is indefinite)
-			Actor->PlaybackSettings.LoopCount = 0;
 		
 			LevelSequenceActor = Actor;
 		}
@@ -144,6 +204,11 @@ void UAutomatedLevelSequenceCapture::Initialize(TSharedPtr<FSceneViewport> InVie
 
 	if (Actor)
 	{
+		// Ensure it doesn't loop (-1 is indefinite)
+		Actor->PlaybackSettings.LoopCount = 0;
+		Actor->PlaybackSettings.TimeController = MakeShared<FMovieSceneTimeController_FrameStep>();
+		Actor->PlaybackSettings.bPauseAtEnd = true;
+
 		if (BurnInOptions)
 		{
 			Actor->BurnInOptions = BurnInOptions;
@@ -157,19 +222,27 @@ void UAutomatedLevelSequenceCapture::Initialize(TSharedPtr<FSceneViewport> InVie
 
 		Actor->RefreshBurnIn();
 
-		// Make sure we're not playing yet (in case AutoPlay was called from BeginPlay)
-		if( Actor->SequencePlayer != nullptr && Actor->SequencePlayer->IsPlaying() )
+		// Make sure we're not playing yet, and have a fully up to date player based on the above settings (in case AutoPlay was called from BeginPlay)
+		if( Actor->SequencePlayer != nullptr )
 		{
-			Actor->SequencePlayer->Stop();
+			if (Actor->SequencePlayer->IsPlaying())
+			{
+				Actor->SequencePlayer->Stop();
+			}
+			Actor->InitializePlayer();
 		}
 		Actor->bAutoPlay = false;
 
 		if (InitializeShots())
 		{
-			float StartTime, EndTime;
+			FFrameNumber StartTime, EndTime;
 			SetupShot(StartTime, EndTime);
 		}
 	}
+	else
+	{
+		UE_LOG(LogMovieSceneCapture, Error, TEXT("Could not find or create a Level Sequence Actor for this capture. Capturing will fail."));
+	}
 
 	ExportEDL();
 
@@ -228,26 +301,37 @@ bool UAutomatedLevelSequenceCapture::InitializeShots()
 		return false;
 	}
 
-	NumShots = CinematicShotTrack->GetAllSections().Num();					
+	NumShots = CinematicShotTrack->GetAllSections().Num();
 	ShotIndex = 0;
 	CachedPlaybackRange = MovieScene->GetPlaybackRange();
 
-	float HandleTime = (float)Settings.HandleFrames / (float)Settings.FrameRate;
-					
+	// Compute handle frames in frame resolution space since that is what the section ranges are defined in
+	FFrameNumber HandleFramesResolutionSpace = ConvertFrameTime(Settings.HandleFrames, Settings.FrameRate, MovieScene->GetFrameResolution()).RoundToFrame();
+
 	for (int32 SectionIndex = 0; SectionIndex < CinematicShotTrack->GetAllSections().Num(); ++SectionIndex)
 	{
 		UMovieSceneCinematicShotSection* ShotSection = Cast<UMovieSceneCinematicShotSection>(CinematicShotTrack->GetAllSections()[SectionIndex]);
 		UMovieScene* ShotMovieScene = ShotSection->GetSequence() ? ShotSection->GetSequence()->GetMovieScene() : nullptr;
-		CachedShotStates.Add(FCinematicShotCache(ShotSection->IsActive(), ShotSection->IsLocked(), TRange<float>(ShotSection->GetStartTime(), ShotSection->GetEndTime()), ShotMovieScene ? ShotMovieScene->GetPlaybackRange() : TRange<float>::Empty()));
 
-		if (ShotMovieScene)
+		if (ShotMovieScene != nullptr)
 		{
-			ShotMovieScene->SetPlaybackRange(ShotMovieScene->GetPlaybackRange().GetLowerBoundValue() - HandleTime, ShotMovieScene->GetPlaybackRange().GetUpperBoundValue() + HandleTime, false);
+			// Expand the inner shot section range by the handle size, multiplied by the difference between the outer and inner frame resolutions (and factoring in the time scale)
+			const float OuterToInnerRateDilation = (MovieScene->GetFrameResolution() == ShotMovieScene->GetFrameResolution()) ? 1.f : (ShotMovieScene->GetFrameResolution() / MovieScene->GetFrameResolution()).AsDecimal();
+			const float OuterToInnerScale = OuterToInnerRateDilation * ShotSection->Parameters.TimeScale;
+
+			CachedShotStates.Add(FCinematicShotCache(ShotSection->IsActive(), ShotSection->IsLocked(), ShotSection->GetRange(), ShotMovieScene ? ShotMovieScene->GetPlaybackRange() : TRange<FFrameNumber>::Empty()));
+
+			if (ShotMovieScene)
+			{
+				TRange<FFrameNumber> NewPlaybackRange = MovieScene::ExpandRange(ShotMovieScene->GetPlaybackRange(), FFrameNumber(FMath::FloorToInt(HandleFramesResolutionSpace.Value * OuterToInnerScale)));
+				ShotMovieScene->SetPlaybackRange(NewPlaybackRange, false);
+			}
+
+			ShotSection->SetIsLocked(false);
+			ShotSection->SetIsActive(false);
+
+			ShotSection->SetRange(MovieScene::ExpandRange(ShotSection->GetRange(), HandleFramesResolutionSpace));
 		}
-		ShotSection->SetIsLocked(false);
-		ShotSection->SetIsActive(false);
-		ShotSection->SetStartTime(ShotSection->GetStartTime() - HandleTime);
-		ShotSection->SetEndTime(ShotSection->GetEndTime() + HandleTime);
 	}
 	return NumShots > 0;
 }
@@ -271,26 +355,23 @@ void UAutomatedLevelSequenceCapture::RestoreShots()
 		return;
 	}
 
-	MovieScene->SetPlaybackRange(CachedPlaybackRange.GetLowerBoundValue(), CachedPlaybackRange.GetUpperBoundValue(), false);
+	MovieScene->SetPlaybackRange(CachedPlaybackRange, false);
 
-	float HandleTime = (float)Settings.HandleFrames / (float)Settings.FrameRate;
-					
 	for (int32 SectionIndex = 0; SectionIndex < CinematicShotTrack->GetAllSections().Num(); ++SectionIndex)
 	{
 		UMovieSceneCinematicShotSection* ShotSection = Cast<UMovieSceneCinematicShotSection>(CinematicShotTrack->GetAllSections()[SectionIndex]);
 		UMovieScene* ShotMovieScene = ShotSection->GetSequence() ? ShotSection->GetSequence()->GetMovieScene() : nullptr;
 		if (ShotMovieScene)
 		{
-			ShotMovieScene->SetPlaybackRange(CachedShotStates[SectionIndex].MovieSceneRange.GetLowerBoundValue(), CachedShotStates[SectionIndex].MovieSceneRange.GetUpperBoundValue(), false);
+			ShotMovieScene->SetPlaybackRange(CachedShotStates[SectionIndex].MovieSceneRange, false);
 		}
 		ShotSection->SetIsActive(CachedShotStates[SectionIndex].bActive);
-		ShotSection->SetStartTime(CachedShotStates[SectionIndex].ShotRange.GetLowerBoundValue());
-		ShotSection->SetEndTime(CachedShotStates[SectionIndex].ShotRange.GetUpperBoundValue());
+		ShotSection->SetRange(CachedShotStates[SectionIndex].ShotRange);
 		ShotSection->SetIsLocked(CachedShotStates[SectionIndex].bLocked);
 	}
 }
 
-bool UAutomatedLevelSequenceCapture::SetupShot(float& StartTime, float& EndTime)
+bool UAutomatedLevelSequenceCapture::SetupShot(FFrameNumber& StartTime, FFrameNumber& EndTime)
 {
 	if (Settings.HandleFrames <= 0)
 	{
@@ -314,8 +395,6 @@ bool UAutomatedLevelSequenceCapture::SetupShot(float& StartTime, float& EndTime)
 		return false;
 	}
 
-	float HandleTime = (float)Settings.HandleFrames / (float)Settings.FrameRate;
-
 	// Disable all shots unless it's the current one being rendered
 	for (int32 SectionIndex = 0; SectionIndex < CinematicShotTrack->GetAllSections().Num(); ++SectionIndex)
 	{
@@ -325,12 +404,12 @@ bool UAutomatedLevelSequenceCapture::SetupShot(float& StartTime, float& EndTime)
 
 		if (SectionIndex == ShotIndex)
 		{
-			StartTime = ShotSection->GetStartTime();
-			EndTime = ShotSection->GetEndTime();
+			TRange<FFrameNumber> TotalRange = TRange<FFrameNumber>::Intersection(ShotSection->GetRange(), CachedPlaybackRange);
 
-			StartTime = FMath::Clamp(StartTime, CachedPlaybackRange.GetLowerBoundValue(), CachedPlaybackRange.GetUpperBoundValue());
-			EndTime = FMath::Clamp(EndTime, CachedPlaybackRange.GetLowerBoundValue(), CachedPlaybackRange.GetUpperBoundValue());
-			MovieScene->SetPlaybackRange(StartTime, EndTime, false);
+			StartTime = TotalRange.IsEmpty() ? FFrameNumber(0) : MovieScene::DiscreteInclusiveLower(TotalRange);
+			EndTime   = TotalRange.IsEmpty() ? FFrameNumber(0) : MovieScene::DiscreteExclusiveUpper(TotalRange);
+
+			MovieScene->SetPlaybackRange(StartTime, (EndTime - StartTime).Value, false);
 		}
 	}
 
@@ -348,53 +427,43 @@ void UAutomatedLevelSequenceCapture::SetupFrameRange()
 			UMovieScene* MovieScene = LevelSequence->GetMovieScene();
 			if( MovieScene != nullptr )
 			{
-				const int32 SequenceStartFrame = FMath::RoundToInt( MovieScene->GetPlaybackRange().GetLowerBoundValue() * Settings.FrameRate );
-				const int32 SequenceEndFrame = FMath::Max( SequenceStartFrame, FMath::RoundToInt( MovieScene->GetPlaybackRange().GetUpperBoundValue() * Settings.FrameRate ) );
+				FFrameRate           SourceFrameRate = MovieScene->GetFrameResolution();
+				TRange<FFrameNumber> SequenceRange   = MovieScene->GetPlaybackRange();
 
-				// Default to playing back the sequence's stored playback range
-				int32 PlaybackStartFrame = SequenceStartFrame;
-				int32 PlaybackEndFrame = SequenceEndFrame;
+				FFrameNumber PlaybackStartFrame = ConvertFrameTime(MovieScene::DiscreteInclusiveLower(SequenceRange), SourceFrameRate, Settings.FrameRate).FloorToFrame();
+				FFrameNumber PlaybackEndFrame   = ConvertFrameTime(MovieScene::DiscreteExclusiveUpper(SequenceRange), SourceFrameRate, Settings.FrameRate).FloorToFrame();
 
 				if( bUseCustomStartFrame )
 				{
-					PlaybackStartFrame = Settings.bUseRelativeFrameNumbers ? ( SequenceStartFrame + StartFrame ) : StartFrame;
+					PlaybackStartFrame = Settings.bUseRelativeFrameNumbers ? ( PlaybackStartFrame + StartFrame ) : StartFrame;
 				}
 
 				if( !Settings.bUseRelativeFrameNumbers )
 				{
-					// NOTE: The frame number will be an offset from the first frame that we start capturing on, not the frame
-					// that we start playback at (in the case of WarmUpFrameCount being non-zero).  So we'll cache out frame
-					// number offset before adjusting for the warm up frames.
-					this->FrameNumberOffset = PlaybackStartFrame;
+				 	// NOTE: The frame number will be an offset from the first frame that we start capturing on, not the frame
+				 	// that we start playback at (in the case of WarmUpFrameCount being non-zero).  So we'll cache out frame
+				 	// number offset before adjusting for the warm up frames.
+				 	this->FrameNumberOffset = PlaybackStartFrame.Value;
 				}
 
 				if( bUseCustomEndFrame )
 				{
-					PlaybackEndFrame = FMath::Max( PlaybackStartFrame, Settings.bUseRelativeFrameNumbers ? ( SequenceEndFrame + EndFrame ) : EndFrame );
-
-					// We always add 1 to the number of frames we want to capture, because we want to capture both the start and end frames (which if the play range is 0, would still yield a single frame)
-					this->FrameCount = ( PlaybackEndFrame - PlaybackStartFrame ) + 1;
-				}
-				else
-				{
-					FrameCount = 0;
+				 	PlaybackEndFrame = FMath::Max( PlaybackStartFrame, Settings.bUseRelativeFrameNumbers ? ( PlaybackEndFrame + EndFrame ) : EndFrame );
 				}
 
 				RemainingWarmUpFrames = FMath::Max( WarmUpFrameCount, 0 );
 				if( RemainingWarmUpFrames > 0 )
 				{
-					// We were asked to playback additional frames before we start capturing
-					PlaybackStartFrame -= RemainingWarmUpFrames;
+				 	// We were asked to playback additional frames before we start capturing
+				 	PlaybackStartFrame -= RemainingWarmUpFrames;
 				}
 
 				// Override the movie scene's playback range
-				Actor->SequencePlayer->SetPlaybackRange(
-					(float)PlaybackStartFrame / (float)Settings.FrameRate,
-					(float)PlaybackEndFrame / (float)Settings.FrameRate );
-				Actor->SequencePlayer->SetPlaybackPosition(0.f);
+				Actor->SequencePlayer->SetFrameRate(Settings.FrameRate);
+				Actor->SequencePlayer->SetFrameRange(PlaybackStartFrame.Value, (PlaybackEndFrame - PlaybackStartFrame).Value);
+				Actor->SequencePlayer->JumpToFrame(PlaybackStartFrame.Value);
 
-				const float WarmupTime = WarmUpFrameCount / CachedState.Settings.FrameRate;
-				Actor->SequencePlayer->SetSnapshotOffsetTime(WarmupTime);
+				Actor->SequencePlayer->SetSnapshotOffsetFrames(WarmUpFrameCount);
 			}
 		}
 	}
@@ -463,12 +532,12 @@ void UAutomatedLevelSequenceCapture::Tick(float DeltaSeconds)
 	if( CaptureState == ELevelSequenceCaptureState::DelayBeforeWarmUp )
 	{
 		// Ensure evaluation at the start of the sequence/shot
-		Actor->SequencePlayer->SetPlaybackPosition(0.f);
+		Actor->SequencePlayer->JumpToFrame(0);
 	}
 	else if( CaptureState == ELevelSequenceCaptureState::ReadyToWarmUp )
 	{
 		Actor->SequencePlayer->SetSnapshotSettings(FLevelSequenceSnapshotSettings(Settings.ZeroPadFrameNumbers, Settings.FrameRate));
-		Actor->SequencePlayer->StartPlayingNextTick();
+		Actor->SequencePlayer->Play();
 		// Start warming up
 		CaptureState = ELevelSequenceCaptureState::WarmingUp;
 	}
@@ -488,12 +557,17 @@ void UAutomatedLevelSequenceCapture::Tick(float DeltaSeconds)
 	{
 		++ShotIndex;
 
-		float StartTime, EndTime;
+		FFrameNumber StartTime, EndTime;
 		if (SetupShot(StartTime, EndTime))
 		{
-			Actor->SequencePlayer->SetPlaybackRange(StartTime, EndTime);
-			Actor->SequencePlayer->SetPlaybackPosition(0.f);
-			Actor->SequencePlayer->StartPlayingNextTick();
+			UMovieScene* MovieScene = GetMovieScene(LevelSequenceActor);
+
+			FFrameNumber StartTimePlayRateSpace = ConvertFrameTime(StartTime, MovieScene->GetFrameResolution(), Settings.FrameRate).FloorToFrame();
+			FFrameNumber EndTimePlayRateSpace   = ConvertFrameTime(EndTime,   MovieScene->GetFrameResolution(), Settings.FrameRate).FloorToFrame();
+
+			Actor->SequencePlayer->SetFrameRange(StartTimePlayRateSpace.Value, (EndTimePlayRateSpace - StartTimePlayRateSpace).Value);
+			Actor->SequencePlayer->JumpToFrame(0);
+			Actor->SequencePlayer->Play();
 			CaptureState = ELevelSequenceCaptureState::FinishedWarmUp;
 			UpdateFrameState();
 		}
@@ -521,17 +595,17 @@ void UAutomatedLevelSequenceCapture::PauseFinished()
 	if (CachedPlayRate.IsSet())
 	{
 		ALevelSequenceActor* Actor = LevelSequenceActor.Get();
-	
+
 		// Force an evaluation to capture this frame
-		Actor->SequencePlayer->SetPlaybackPosition(Actor->SequencePlayer->GetPlaybackPosition());
-		
+		Actor->SequencePlayer->JumpToFrame(Actor->SequencePlayer->GetCurrentTime().Time);
+
 		// Continue playing forwards
 		Actor->SequencePlayer->SetPlayRate(CachedPlayRate.GetValue());
 		CachedPlayRate.Reset();
 	}
 }
 
-void UAutomatedLevelSequenceCapture::SequenceUpdated(const UMovieSceneSequencePlayer& Player, float CurrentTime, float PreviousTime)
+void UAutomatedLevelSequenceCapture::SequenceUpdated(const UMovieSceneSequencePlayer& Player, FFrameTime CurrentTime, FFrameTime PreviousTime)
 {
 	if (bCapturing)
 	{
@@ -546,16 +620,21 @@ void UAutomatedLevelSequenceCapture::SequenceUpdated(const UMovieSceneSequencePl
 			bool bHasMultipleShots = PreviousState.CurrentShotName != PreviousState.MasterName;
 			bool bNewShot = bHasMultipleShots && PreviousState.ShotID != CachedState.ShotID;
 			
-			if (bNewShot && Actor->SequencePlayer->IsPlaying() && DelayBeforeWarmUp > 0)
+			if (bNewShot && Actor->SequencePlayer->IsPlaying() && DelayBeforeShotWarmUp > 0)
 			{
 				CaptureState = ELevelSequenceCaptureState::Paused;
-				Actor->GetWorld()->GetTimerManager().SetTimer(DelayTimer, FTimerDelegate::CreateUObject(this, &UAutomatedLevelSequenceCapture::PauseFinished), DelayBeforeWarmUp, false);
+				Actor->GetWorld()->GetTimerManager().SetTimer(DelayTimer, FTimerDelegate::CreateUObject(this, &UAutomatedLevelSequenceCapture::PauseFinished), DelayBeforeShotWarmUp, false);
 				CachedPlayRate = Actor->SequencePlayer->GetPlayRate();
 				Actor->SequencePlayer->SetPlayRate(0.f);
 			}
 			else if (CaptureState == ELevelSequenceCaptureState::FinishedWarmUp)
 			{
-				CaptureThisFrame(CurrentTime - PreviousTime);
+				CaptureThisFrame( (CurrentTime - PreviousTime) / Settings.FrameRate);
+
+				if (bFinalizeWhenReady)
+				{
+					Actor->SequencePlayer->OnSequenceUpdated().Remove( OnPlayerUpdatedBinding );
+				}
 			}
 		}
 	}
@@ -585,8 +664,8 @@ void UAutomatedLevelSequenceCapture::LoadFromConfig()
 
 void UAutomatedLevelSequenceCapture::SaveToConfig()
 {
-	int32 CurrentStartFrame = StartFrame;
-	int32 CurrentEndFrame = EndFrame;
+	FFrameNumber CurrentStartFrame = StartFrame;
+	FFrameNumber CurrentEndFrame = EndFrame;
 	bool bRestoreFrameOverrides = RestoreFrameOverrides();
 
 	BurnInOptions->SaveConfig();
@@ -640,7 +719,7 @@ bool UAutomatedLevelSequenceCapture::RestoreFrameOverrides()
 	return bAnySet;
 }
 
-void UAutomatedLevelSequenceCapture::SetFrameOverrides(int32 InStartFrame, int32 InEndFrame)
+void UAutomatedLevelSequenceCapture::SetFrameOverrides(FFrameNumber InStartFrame, FFrameNumber InEndFrame)
 {
 	CachedStartFrame = StartFrame;
 	CachedEndFrame = EndFrame;
diff --git a/Engine/Source/Runtime/MovieSceneCapture/Private/CompositionGraphCaptureProtocol.cpp b/Engine/Source/Runtime/MovieSceneCapture/Private/CompositionGraphCaptureProtocol.cpp
--- a/Engine/Source/Runtime/MovieSceneCapture/Private/CompositionGraphCaptureProtocol.cpp
+++ b/Engine/Source/Runtime/MovieSceneCapture/Private/CompositionGraphCaptureProtocol.cpp
@@ -15,7 +15,7 @@
 
 struct FFrameCaptureViewExtension : public FSceneViewExtensionBase
 {
-	FFrameCaptureViewExtension( const FAutoRegister& AutoRegister, const TArray<FString>& InRenderPasses, bool bInCaptureFramesInHDR, int32 InHDRCompressionQuality, int32 InCaptureGamut, UMaterialInterface* InPostProcessingMaterial)
+	FFrameCaptureViewExtension( const FAutoRegister& AutoRegister, const TArray<FString>& InRenderPasses, bool bInCaptureFramesInHDR, int32 InHDRCompressionQuality, int32 InCaptureGamut, UMaterialInterface* InPostProcessingMaterial, bool bInDisableScreenPercentage)
 		: FSceneViewExtensionBase(AutoRegister)
 		, RenderPasses(InRenderPasses)
 		, bNeedsCapture(true)
@@ -23,6 +23,7 @@ struct FFrameCaptureViewExtension : public FSceneViewExtensionBase
 		, HDRCompressionQuality(InHDRCompressionQuality)
 		, CaptureGamut(InCaptureGamut)
 		, PostProcessingMaterial(InPostProcessingMaterial)
+		, bDisableScreenPercentage(bInDisableScreenPercentage)
 	{
 		CVarDumpFrames = IConsoleManager::Get().FindConsoleVariable(TEXT("r.BufferVisualizationDumpFrames"));
 		CVarDumpFramesAsHDR = IConsoleManager::Get().FindConsoleVariable(TEXT("r.BufferVisualizationDumpFramesAsHDR"));
@@ -125,8 +126,11 @@ struct FFrameCaptureViewExtension : public FSceneViewExtensionBase
 
 	virtual void SetupViewFamily(FSceneViewFamily& InViewFamily) override
 	{
-		// Ensure we're rendering at full size.
-		InViewFamily.EngineShowFlags.ScreenPercentage = false;
+		if (bDisableScreenPercentage)
+		{
+			// Ensure we're rendering at full size.
+			InViewFamily.EngineShowFlags.ScreenPercentage = false;
+		}
 	}
 
 	virtual void BeginRenderViewFamily(FSceneViewFamily& InViewFamily) {}
@@ -147,6 +151,8 @@ private:
 
 	UMaterialInterface* PostProcessingMaterial;
 
+	bool bDisableScreenPercentage;
+
 	IConsoleVariable* CVarDumpFrames;
 	IConsoleVariable* CVarDumpFramesAsHDR;
 	IConsoleVariable* CVarHDRCompressionQuality;
@@ -208,6 +214,7 @@ bool FCompositionGraphCaptureProtocol::Initialize(const FCaptureProtocolInitSett
 	bool bCaptureFramesInHDR = false;
 	int32 HDRCompressionQuality = 0;
 	int32 CaptureGamut = HCGM_Rec709;
+	bool bDisableScreenPercentage = true;
 
 	UMaterialInterface* PostProcessingMaterial = nullptr;
 	UCompositionGraphCaptureSettings* ProtocolSettings = CastChecked<UCompositionGraphCaptureSettings>(InSettings.ProtocolSettings);
@@ -218,6 +225,7 @@ bool FCompositionGraphCaptureProtocol::Initialize(const FCaptureProtocolInitSett
 		HDRCompressionQuality = ProtocolSettings->HDRCompressionQuality;
 		CaptureGamut = ProtocolSettings->CaptureGamut;
 		PostProcessingMaterial = Cast<UMaterialInterface>(ProtocolSettings->PostProcessingMaterial.TryLoad());
+		bDisableScreenPercentage = ProtocolSettings->bDisableScreenPercentage;
 
 		FString OverrideRenderPasses;
 		if (FParse::Value(FCommandLine::Get(), TEXT("-CustomRenderPasses="), OverrideRenderPasses))
@@ -242,9 +250,15 @@ bool FCompositionGraphCaptureProtocol::Initialize(const FCaptureProtocolInitSett
 		{
 			CaptureGamut = OverrideCaptureGamut;
 		}
+
+		bool bOverrideDisableScreenPercentage;
+		if (FParse::Bool(FCommandLine::Get(), TEXT("-DisableScreenPercentage="), bOverrideDisableScreenPercentage))
+		{
+			bDisableScreenPercentage = bOverrideDisableScreenPercentage;
+		}
 	}
 
-	ViewExtension = FSceneViewExtensions::NewExtension<FFrameCaptureViewExtension>(RenderPasses, bCaptureFramesInHDR, HDRCompressionQuality, CaptureGamut, PostProcessingMaterial);
+	ViewExtension = FSceneViewExtensions::NewExtension<FFrameCaptureViewExtension>(RenderPasses, bCaptureFramesInHDR, HDRCompressionQuality, CaptureGamut, PostProcessingMaterial, bDisableScreenPercentage);
 
 	return true;
 }
diff --git a/Engine/Source/Runtime/MovieSceneCapture/Private/MovieSceneCapture.cpp b/Engine/Source/Runtime/MovieSceneCapture/Private/MovieSceneCapture.cpp
--- a/Engine/Source/Runtime/MovieSceneCapture/Private/MovieSceneCapture.cpp
+++ b/Engine/Source/Runtime/MovieSceneCapture/Private/MovieSceneCapture.cpp
@@ -40,7 +40,7 @@ FMovieSceneCaptureSettings::FMovieSceneCaptureSettings()
 	HandleFrames = 0;
 	GameModeOverride = nullptr;
 	OutputFormat = TEXT("{world}");
-	FrameRate = 24;
+	FrameRate = FFrameRate(24, 1);
 	ZeroPadFrameNumbers = 4;
 	bEnableTextureStreaming = false;
 	bCinematicEngineScalability = true;
@@ -67,7 +67,6 @@ UMovieSceneCapture::UMovieSceneCapture(const FObjectInitializer& Initializer)
 
 	Handle = FUniqueMovieSceneCaptureHandle();
 
-	FrameCount = 0;
 	bCapturing = false;
 	FrameNumberOffset = 0;
 	CaptureType = TEXT("Video");
@@ -129,7 +128,7 @@ void UMovieSceneCapture::Initialize(TSharedPtr<FSceneViewport> InSceneViewport,
 		{
 			Settings.bUseRelativeFrameNumbers = bOverrideRelativeFrameNumbers;
 		}
-				
+
 		int32 HandleFramesOverride;
 		if( FParse::Value( FCommandLine::Get(), TEXT( "-HandleFrames=" ), HandleFramesOverride ) )
 		{
@@ -155,10 +154,13 @@ void UMovieSceneCapture::Initialize(TSharedPtr<FSceneViewport> InSceneViewport,
 			InitializeSettings();
 		}
 
-		int32 FrameRateOverride;
-		if( FParse::Value( FCommandLine::Get(), TEXT( "-MovieFrameRate=" ), FrameRateOverride ) )
+		FString FrameRateOverrideString;
+		if ( FParse::Value( FCommandLine::Get(), TEXT( "-MovieFrameRate=" ), FrameRateOverrideString ) )
 		{
-			Settings.FrameRate = FrameRateOverride;
+			if (!TryParseString(Settings.FrameRate, *FrameRateOverrideString))
+			{
+				UE_LOG(LogMovieSceneCapture, Error, TEXT("Unrecognized capture frame rate: %s."), *FrameRateOverrideString);
+			}
 		}
 	}
 
@@ -170,8 +172,17 @@ void UMovieSceneCapture::Initialize(TSharedPtr<FSceneViewport> InSceneViewport,
 	CachedMetrics.Width = InitSettings->DesiredSize.X;
 	CachedMetrics.Height = InitSettings->DesiredSize.Y;
 
+	double FrameRate = Settings.FrameRate.AsDecimal();
+
 	FormatMappings.Reserve(10);
-	FormatMappings.Add(TEXT("fps"), FString::Printf(TEXT("%d"), Settings.FrameRate));
+	if (FrameRate == FMath::RoundToDouble(FrameRate))
+	{
+		FormatMappings.Add(TEXT("fps"), FString::Printf(TEXT("%d"), FMath::RoundToInt(FrameRate)));
+	}
+	else
+	{
+		FormatMappings.Add(TEXT("fps"), FString::Printf(TEXT("%.2f"), FrameRate));
+	}
 	FormatMappings.Add(TEXT("width"), FString::Printf(TEXT("%d"), CachedMetrics.Width));
 	FormatMappings.Add(TEXT("height"), FString::Printf(TEXT("%d"), CachedMetrics.Height));
 	FormatMappings.Add(TEXT("world"), InSceneViewport->GetClient()->GetWorld()->GetName());
@@ -242,12 +253,9 @@ void UMovieSceneCapture::CaptureThisFrame(float DeltaSeconds)
 			);
 		CaptureProtocol->CaptureFrame(ThisFrameMetrics, *this);
 
-		++CachedMetrics.Frame;
+		UE_LOG(LogMovieSceneCapture, Verbose, TEXT("Captured frame: %d"), CachedMetrics.Frame);
 
-		if (!bFinalizeWhenReady && (FrameCount != 0 && CachedMetrics.Frame >= FrameCount))
-		{
-			FinalizeWhenReady();
-		}
+		++CachedMetrics.Frame;
 	}
 }
 
@@ -443,20 +451,20 @@ void UMovieSceneCapture::PostEditChangeProperty( struct FPropertyChangedEvent& P
 }
 #endif
 
-FFixedTimeStepCaptureStrategy::FFixedTimeStepCaptureStrategy(uint32 InTargetFPS)
-	: TargetFPS(InTargetFPS)
+FFixedTimeStepCaptureStrategy::FFixedTimeStepCaptureStrategy(FFrameRate InFrameRate)
+	: FrameRate(InFrameRate)
 {
 }
 
 void FFixedTimeStepCaptureStrategy::OnWarmup()
 {
-	FApp::SetFixedDeltaTime(1.0 / TargetFPS);
+	FApp::SetFixedDeltaTime(FrameRate.AsInterval());
 	FApp::SetUseFixedTimeStep(true);
 }
 
 void FFixedTimeStepCaptureStrategy::OnStart()
 {
-	FApp::SetFixedDeltaTime(1.0 / TargetFPS);
+	FApp::SetFixedDeltaTime(FrameRate.AsInterval());
 	FApp::SetUseFixedTimeStep(true);
 }
 
@@ -479,8 +487,8 @@ int32 FFixedTimeStepCaptureStrategy::GetDroppedFrames(double CurrentTimeSeconds,
 	return 0;
 }
 
-FRealTimeCaptureStrategy::FRealTimeCaptureStrategy(uint32 InTargetFPS)
-	: NextPresentTimeS(0), FrameLength(1.0 / InTargetFPS)
+FRealTimeCaptureStrategy::FRealTimeCaptureStrategy(FFrameRate InFrameRate)
+	: NextPresentTimeS(0), FrameLength(InFrameRate.AsInterval())
 {
 }
 
diff --git a/Engine/Source/Runtime/MovieSceneCapture/Private/MovieSceneCaptureHelpers.cpp b/Engine/Source/Runtime/MovieSceneCapture/Private/MovieSceneCaptureHelpers.cpp
--- a/Engine/Source/Runtime/MovieSceneCapture/Private/MovieSceneCaptureHelpers.cpp
+++ b/Engine/Source/Runtime/MovieSceneCapture/Private/MovieSceneCaptureHelpers.cpp
@@ -11,6 +11,8 @@
 #include "Sections/MovieSceneCinematicShotSection.h"
 #include "Tracks/MovieSceneCinematicShotTrack.h"
 #include "AssetRegistryModule.h"
+#include "FrameRate.h"
+#include "MovieSceneTimeHelpers.h"
 
 /* MovieSceneCaptureHelpers
  *****************************************************************************/
@@ -35,69 +37,70 @@ struct FShotData
 		ET_None
 	};
 
-	FShotData(const FString& InElementName, const FString& InElementPath, ETrackType InTrackType, EEditType InEditType, float InSourceInTime, float InSourceOutTime, float InEditInTime, float InEditOutTime, bool bInWithinPlaybackRange) : 
+	FShotData(const FString& InElementName, const FString& InElementPath, ETrackType InTrackType, EEditType InEditType, FFrameNumber InSourceInFrame, FFrameNumber InSourceOutFrame, FFrameNumber InEditInFrame, FFrameNumber InEditOutFrame, bool bInWithinPlaybackRange) : 
 		  ElementName(InElementName)
 		, ElementPath(InElementPath)
 		, TrackType(InTrackType)
 		, EditType(InEditType)
-		, SourceInTime(InSourceInTime)
-		, SourceOutTime(InSourceOutTime)
-		, EditInTime(InEditInTime)
-		, EditOutTime(InEditOutTime)
+		, SourceInFrame(InSourceInFrame)
+		, SourceOutFrame(InSourceOutFrame)
+		, EditInFrame(InEditInFrame)
+		, EditOutFrame(InEditOutFrame)
 		, bWithinPlaybackRange(bInWithinPlaybackRange) {}
 
-	bool operator<(const FShotData& Other) const { return EditInTime < Other.EditInTime; }
+	bool operator<(const FShotData& Other) const { return EditInFrame < Other.EditInFrame; }
 
 	FString ElementName;
 	FString ElementPath;
 	ETrackType TrackType;
 	EEditType EditType;
-	float SourceInTime;
-	float SourceOutTime;
-	float EditInTime;
-	float EditOutTime;
+	FFrameNumber SourceInFrame;
+	FFrameNumber SourceOutFrame;
+	FFrameNumber EditInFrame;
+	FFrameNumber EditOutFrame;
 	bool bWithinPlaybackRange;
 };
 
-void SMPTEToTime(FString SMPTE, float FrameRate, float& OutTime)
+FFrameNumber SMPTEToFrame(FString SMPTE, FFrameRate FrameRate)
 {
 	TArray<FString> OutArray;
 	SMPTE.ParseIntoArray(OutArray, TEXT(":"));
 	
 	if (OutArray.Num() == 4)
 	{
-		float Hours = FCString::Atof(*OutArray[0]);
-		float Minutes = FCString::Atof(*OutArray[1]);
-		float Seconds = FCString::Atof(*OutArray[2]);
-		float Frames = FCString::Atof(*OutArray[3]);
+		const int64 Hours   = FCString::Atoi(*OutArray[0]);
+		const int64 Minutes = FCString::Atoi(*OutArray[1]);
+		const int64 Seconds = FCString::Atoi(*OutArray[2]);
+		const int32 Frames  = FCString::Atoi(*OutArray[3]);
 
-		OutTime = Hours * 3600.f + Minutes * 60.f  + Seconds + (Frames / FrameRate);
+		int64 TotalSeconds = int64(Seconds) + Minutes*60 + Hours*3600;
+
+		// @todo sequencer-timecode Redo this to take FTimeCode structs and consolidate duplication.
+		return Frames + int32(TotalSeconds * FrameRate.Denominator / FrameRate.Numerator);
 	}
 	// The edl is in frames
 	else
 	{
-		float Frames = FCString::Atof(*SMPTE);
-		OutTime = Frames / FrameRate;
+		return FCString::Atoi(*SMPTE);
 	}
 }
 
-FString TimeToSMPTE(float InTime, float FrameRate)
+FString TimeToSMPTE(FFrameNumber InTime, FFrameRate FrameRate)
 {
-	float Hours, Minutes, Seconds, Frames;
-	Frames = 0.5f + FrameRate * (InTime - FMath::FloorToInt(InTime));
+	// Divide by the framerate to get into seconds
+	const int64 TotalSeconds = int64(InTime.Value) * FrameRate.Denominator / FrameRate.Numerator;
+	const int32 MaxFrame     = FrameRate.Numerator / FrameRate.Denominator;
 
-	Hours = FMath::FloorToInt(InTime / (60.f * 60.f));
-	InTime -= Hours * 60.f * 60.f;
-	Minutes = FMath::FloorToInt(InTime / 60.f);
-	InTime -= Minutes * 60.f;
-	Seconds = InTime;
+	int32 Hours   =  (TotalSeconds       ) / 3600;
+	int32 Minutes =  (TotalSeconds % 3600) / 60;
+	int32 Seconds =  (TotalSeconds % 60  );
 
-	FString OutputString = FString::Printf(TEXT("%02d:%02d:%02d:%02d"), (int)Hours, (int)Minutes, (int)Seconds, (int)Frames);
+	FString OutputString = FString::Printf(TEXT("%02d:%02d:%02d:%02d"), Hours, Minutes, Seconds, InTime.Value % MaxFrame);
 
 	return OutputString;
 }
 
-void ParseFromEDL(const FString& InputString, float FrameRate, TArray<FShotData>& OutShotData)
+void ParseFromEDL(const FString& InputString, FFrameRate FrameRate, TArray<FShotData>& OutShotData)
 {
 	TArray<FString> InputLines;
 	InputString.ParseIntoArray(InputLines, TEXT("\n"));
@@ -108,10 +111,11 @@ void ParseFromEDL(const FString& InputString, float FrameRate, TArray<FShotData>
 	FString ReelName;
 	FShotData::ETrackType TrackType = FShotData::ETrackType::TT_None;
 	FShotData::EEditType EditType = FShotData::EEditType::ET_None;
-	float SourceInTime = 0.f;
-	float SourceOutTime = 0.f;
-	float EditInTime = 0.f;
-	float EditOutTime = 0.f;
+
+	FFrameNumber SourceInFrame  = 0;
+	FFrameNumber SourceOutFrame = 0;
+	FFrameNumber EditInFrame    = 0;
+	FFrameNumber EditOutFrame   = 0;
 
 	for (FString InputLine : InputLines)
 	{
@@ -165,10 +169,10 @@ void ParseFromEDL(const FString& InputString, float FrameRate, TArray<FShotData>
 				if (TrackType != FShotData::ETrackType::TT_None &&
 					EditType != FShotData::EEditType::ET_None)
 				{
-					SMPTEToTime(InputChars[4], FrameRate, SourceInTime);
-					SMPTEToTime(InputChars[5], FrameRate, SourceOutTime);
-					SMPTEToTime(InputChars[6], FrameRate, EditInTime);
-					SMPTEToTime(InputChars[7], FrameRate, EditOutTime);
+					SourceInFrame  = SMPTEToFrame(InputChars[4], FrameRate);
+					SourceOutFrame = SMPTEToFrame(InputChars[5], FrameRate);
+					EditInFrame    = SMPTEToFrame(InputChars[6], FrameRate);
+					EditOutFrame   = SMPTEToFrame(InputChars[7], FrameRate);
 
 					bFoundEventLine = true;
 					continue; // Go to the next line
@@ -195,7 +199,7 @@ void ParseFromEDL(const FString& InputString, float FrameRate, TArray<FShotData>
 				FString ElementPath = ElementName;
 
 				// If everything checks out add to OutShotData
-				OutShotData.Add(FShotData(ElementName, ElementPath, TrackType, EditType, SourceInTime, SourceOutTime, EditInTime, EditOutTime, true));
+				OutShotData.Add(FShotData(ElementName, ElementPath, TrackType, EditType, SourceInFrame, SourceOutFrame, EditInFrame, EditOutFrame, true));
 				bFoundEventLine = false; // Reset and go to next line to look for element line
 				continue;
 			}
@@ -203,7 +207,7 @@ void ParseFromEDL(const FString& InputString, float FrameRate, TArray<FShotData>
 	}
 }
 
-void FormatForEDL(FString& OutputString, const FString& SequenceName, float FrameRate, const TArray<FShotData>& InShotData)
+void FormatForEDL(FString& OutputString, const FString& SequenceName, FFrameRate FrameRate, const TArray<FShotData>& InShotData)
 {
 	OutputString += TEXT("TITLE: ") + SequenceName + TEXT("\n");
 	OutputString += TEXT("FCM: NON-DROP FRAME\n\n");
@@ -213,16 +217,16 @@ void FormatForEDL(FString& OutputString, const FString& SequenceName, float Fram
 	FString SourceSMPTEIn, SourceSMPTEOut, EditSMPTEIn, EditSMPTEOut;
 
 	// Insert blank if doesn't start at 0.
-	if (InShotData[0].EditInTime != 0)
+	if (InShotData[0].EditInFrame != 0)
 	{
 		EventName = FString::Printf(TEXT("%03d"), ++EventIndex);
 		TypeName = TEXT("V");
 		EditName = TEXT("C");
 
-		SourceSMPTEIn = TimeToSMPTE(0.f, FrameRate);
-		SourceSMPTEOut = TimeToSMPTE(InShotData[0].EditInTime, FrameRate);
-		EditSMPTEIn = TimeToSMPTE(0.f, FrameRate);
-		EditSMPTEOut = TimeToSMPTE(InShotData[0].EditInTime, FrameRate);
+		SourceSMPTEIn  = TimeToSMPTE(0, FrameRate);
+		SourceSMPTEOut = TimeToSMPTE(InShotData[0].EditInFrame, FrameRate);
+		EditSMPTEIn    = TimeToSMPTE(0, FrameRate);
+		EditSMPTEOut   = TimeToSMPTE(InShotData[0].EditInFrame, FrameRate);
 
 		OutputString += EventName + TEXT(" ") + TEXT("BL ") + TypeName + TEXT(" ") + EditName + TEXT(" ");
 		OutputString += SourceSMPTEIn + TEXT(" ") + SourceSMPTEOut + TEXT(" ") + EditSMPTEIn + TEXT(" ") + EditSMPTEOut + TEXT("\n\n");
@@ -268,10 +272,10 @@ void FormatForEDL(FString& OutputString, const FString& SequenceName, float Fram
 			EditName = TEXT("K");
 		}
 
-		SourceSMPTEIn = TimeToSMPTE(InShotData[ShotIndex].SourceInTime, FrameRate);
-		SourceSMPTEOut = TimeToSMPTE(InShotData[ShotIndex].SourceOutTime, FrameRate);
-		EditSMPTEIn = TimeToSMPTE(InShotData[ShotIndex].EditInTime, FrameRate);
-		EditSMPTEOut = TimeToSMPTE(InShotData[ShotIndex].EditOutTime, FrameRate);
+		SourceSMPTEIn  = TimeToSMPTE(InShotData[ShotIndex].SourceInFrame, FrameRate);
+		SourceSMPTEOut = TimeToSMPTE(InShotData[ShotIndex].SourceOutFrame, FrameRate);
+		EditSMPTEIn    = TimeToSMPTE(InShotData[ShotIndex].EditInFrame, FrameRate);
+		EditSMPTEOut   = TimeToSMPTE(InShotData[ShotIndex].EditOutFrame, FrameRate);
 
 		OutputString += EventName + TEXT(" ") + TEXT("AX ") + TypeName + TEXT(" ") + EditName + TEXT(" ");
 		OutputString += SourceSMPTEIn + TEXT(" ") + SourceSMPTEOut + TEXT(" ") + EditSMPTEIn + TEXT(" ") + EditSMPTEOut + TEXT("\n");
@@ -279,7 +283,7 @@ void FormatForEDL(FString& OutputString, const FString& SequenceName, float Fram
 	}
 }
 
-void FormatForRV(FString& OutputString, const FString& SequenceName, float FrameRate, const TArray<FShotData>& InShotData)
+void FormatForRV(FString& OutputString, const FString& SequenceName, FFrameRate FrameRate, const TArray<FShotData>& InShotData)
 {
 	// Header
 	OutputString += TEXT("GTOa (3)\n\n");
@@ -287,7 +291,7 @@ void FormatForRV(FString& OutputString, const FString& SequenceName, float Frame
 	OutputString += TEXT("{\n");
 	OutputString += TEXT("\tsession\n");
 	OutputString += TEXT("\t{\n");
-	OutputString += TEXT("\t\tfloat fps = ") + FString::Printf(TEXT("%f"), FrameRate) + TEXT("\n");
+	OutputString += TEXT("\t\tfloat fps = ") + FString::Printf(TEXT("%f"), FrameRate.AsDecimal()) + TEXT("\n");
 	OutputString += TEXT("\t\tint realtime = 1\n");
 	OutputString += TEXT("\t}\n\n");
 	OutputString += TEXT("\twriter\n");
@@ -307,8 +311,8 @@ void FormatForRV(FString& OutputString, const FString& SequenceName, float Frame
 
 		FString SourceName = FString::Printf(TEXT("sourceGroup%06d"), EventIndex);
 
-		int32 SourceInFrame = FMath::RoundToInt(InShotData[EventIndex].SourceInTime * FrameRate);
-		int32 SourceOutFrame = FMath::RoundToInt(InShotData[EventIndex].SourceOutTime * FrameRate);
+		FFrameNumber SourceInFrame  = InShotData[EventIndex].SourceInFrame;
+		FFrameNumber SourceOutFrame = InShotData[EventIndex].SourceOutFrame;
 
 		OutputString += SourceName + TEXT(" : RVSourceGroup (1)\n");
 		OutputString += TEXT("{\n");
@@ -322,8 +326,8 @@ void FormatForRV(FString& OutputString, const FString& SequenceName, float Frame
 		OutputString += TEXT("{\n");
 		OutputString += TEXT("\tcut\n");
 		OutputString += TEXT("\t{\n");
-		OutputString += TEXT("\t\tint in = ") + FString::FromInt(SourceInFrame) + TEXT("\n");
-		OutputString += TEXT("\t\tint out = ") + FString::FromInt(SourceOutFrame) + TEXT("\n");
+		OutputString += TEXT("\t\tint in = ") + FString::FromInt(SourceInFrame.Value) + TEXT("\n");
+		OutputString += TEXT("\t\tint out = ") + FString::FromInt(SourceOutFrame.Value) + TEXT("\n");
 		OutputString += TEXT("\t}\n\n");
 
 		OutputString += TEXT("\tgroup\n");
@@ -340,9 +344,9 @@ void FormatForRV(FString& OutputString, const FString& SequenceName, float Frame
 	}
 }
 
-void FormatForRVBat(FString& OutputString, const FString& SequenceName, float FrameRate, const TArray<FShotData>& InShotData)
+void FormatForRVBat(FString& OutputString, const FString& SequenceName, FFrameRate FrameRate, const TArray<FShotData>& InShotData)
 {
-	OutputString += TEXT("rv -nomb -fullscreen -noBorders -fps " + FString::Printf(TEXT("%f"), FrameRate));
+	OutputString += TEXT("rv -nomb -fullscreen -noBorders -fps ") + FString::Printf(TEXT("%f"), FrameRate.AsInterval());
 	for (int32 EventIndex = 0; EventIndex < InShotData.Num(); ++EventIndex)
 	{
 		if (InShotData[EventIndex].bWithinPlaybackRange)
@@ -353,7 +357,7 @@ void FormatForRVBat(FString& OutputString, const FString& SequenceName, float Fr
 }
 
 
-bool MovieSceneCaptureHelpers::ImportEDL(UMovieScene* InMovieScene, float InFrameRate, FString InFilename)
+bool MovieSceneCaptureHelpers::ImportEDL(UMovieScene* InMovieScene, FFrameRate InFrameRate, FString InFilename)
 {
 	FString InputString;
 	if (!FFileHelper::LoadFileToString(InputString, *InFilename))
@@ -413,16 +417,15 @@ bool MovieSceneCaptureHelpers::ImportEDL(UMovieScene* InMovieScene, float InFram
 				}
 
 				CinematicShotTrack->Modify();
-				ShotSection = Cast<UMovieSceneCinematicShotSection>(CinematicShotTrack->AddSequence(SequenceToAdd, ShotData.EditInTime, ShotData.EditOutTime - ShotData.EditInTime));
+				ShotSection = Cast<UMovieSceneCinematicShotSection>(CinematicShotTrack->AddSequence(SequenceToAdd, ShotData.EditInFrame, (ShotData.EditOutFrame - ShotData.EditInFrame).Value));
 			}
 
 			// Conform this shot section
 			if (ShotSection)
 			{
 				ShotSection->Modify();
-				ShotSection->Parameters.StartOffset = ShotData.SourceInTime;
-				ShotSection->SetStartTime(ShotData.EditInTime);
-				ShotSection->SetEndTime(ShotData.EditOutTime);
+				ShotSection->Parameters.SetStartFrameOffset(ShotData.SourceInFrame.Value);
+				ShotSection->SetRange(TRange<FFrameNumber>(ShotData.EditInFrame, ShotData.EditOutFrame));
 			}
 		}
 	}
@@ -430,7 +433,7 @@ bool MovieSceneCaptureHelpers::ImportEDL(UMovieScene* InMovieScene, float InFram
 	return true;
 }
 
-bool MovieSceneCaptureHelpers::ExportEDL(const UMovieScene* InMovieScene, float InFrameRate, FString InSaveFilename, const int32 InHandleFrames)
+bool MovieSceneCaptureHelpers::ExportEDL(const UMovieScene* InMovieScene, FFrameRate InFrameRate, FString InSaveFilename, const int32 InHandleFrames)
 {
 	FString SequenceName = InMovieScene->GetOuter()->GetName();
 	FString SaveBasename = FPaths::GetPath(InSaveFilename) / FPaths::GetBaseFilename(InSaveFilename);
@@ -445,20 +448,19 @@ bool MovieSceneCaptureHelpers::ExportEDL(const UMovieScene* InMovieScene, float
 
 	TArray<FShotData> ShotDataArray;
 
-	float EditTime = 0.f;
-
-	for (auto MasterTrack : InMovieScene->GetMasterTracks())
+	for (UMovieSceneTrack* MasterTrack : InMovieScene->GetMasterTracks())
 	{
-		TRange<float> PlaybackRange = InMovieScene->GetPlaybackRange();
+		// @todo: sequencer-timecode: deal with framerate differences??
+		TRange<FFrameNumber> PlaybackRange = InMovieScene->GetPlaybackRange();
 		if (MasterTrack->IsA(UMovieSceneCinematicShotTrack::StaticClass()))
 		{
 			UMovieSceneCinematicShotTrack* CinematicShotTrack = Cast<UMovieSceneCinematicShotTrack>(MasterTrack);
 
-			for (auto ShotSection : CinematicShotTrack->GetAllSections())
+			for (UMovieSceneSection* ShotSection : CinematicShotTrack->GetAllSections())
 			{
 				UMovieSceneCinematicShotSection* CinematicShotSection = Cast<UMovieSceneCinematicShotSection>(ShotSection);
 
-				if ( CinematicShotSection->GetSequence() == nullptr )
+				if ( CinematicShotSection->GetSequence() == nullptr || !CinematicShotSection->HasStartFrame() || !CinematicShotSection->HasEndFrame() )
 				{
 					// If the shot doesn't have a valid sequence skip it.  This is currently the case for filler sections.
 					// TODO: Handle this properly in the edl output.
@@ -468,12 +470,11 @@ bool MovieSceneCaptureHelpers::ExportEDL(const UMovieScene* InMovieScene, float
 				FString ShotName = CinematicShotSection->GetShotDisplayName();
 				FString ShotPath = CinematicShotSection->GetSequence()->GetMovieScene()->GetOuter()->GetPathName();
 
-				float HandleFrameTime = (float)InHandleFrames / (float)InFrameRate;
-				float SourceInTime = HandleFrameTime;
-				float SourceOutTime = SourceInTime + CinematicShotSection->GetTimeSize();
-				
-				float EditInTime = CinematicShotSection->GetStartTime();
-				float EditOutTime = CinematicShotSection->GetEndTime();
+				FFrameNumber SourceInFrame  = InHandleFrames;
+				FFrameNumber SourceOutFrame = InHandleFrames + MovieScene::DiscreteSize(CinematicShotSection->GetRange());
+
+				FFrameNumber EditInFrame    = CinematicShotSection->GetInclusiveStartFrame();
+				FFrameNumber EditOutFrame   = CinematicShotSection->GetExclusiveEndFrame();
 
 				//@todo can't assume avi's written out
 				ShotName += ".avi";
@@ -481,13 +482,12 @@ bool MovieSceneCaptureHelpers::ExportEDL(const UMovieScene* InMovieScene, float
 				//@todo shotpath should really be moviefile path
 				ShotPath = ShotName;
 
-				TRange<float> EditRange(EditInTime, EditOutTime);
-				TRange<float> Intersection = TRange<float>::Intersection(PlaybackRange, EditRange);
-				
-				bool bWithinPlaybackRange = Intersection.Size<float>() > (1.0f / (float)InFrameRate);
+				TRange<FFrameNumber> EditRange    = TRange<FFrameNumber>(EditInFrame, EditOutFrame);
+				TRange<FFrameNumber> Intersection = TRange<FFrameNumber>::Intersection(PlaybackRange, EditRange);
+
+				bool bWithinPlaybackRange = EditRange.Overlaps(PlaybackRange);
 
-				ShotDataArray.Add(FShotData(ShotName, ShotPath, FShotData::ETrackType::TT_Video, FShotData::EEditType::ET_Cut, SourceInTime, SourceOutTime, EditInTime, EditOutTime, bWithinPlaybackRange));
-				EditTime = EditOutTime;
+				ShotDataArray.Add(FShotData(ShotName, ShotPath, FShotData::ETrackType::TT_Video, FShotData::EEditType::ET_Cut, SourceInFrame, SourceOutFrame, EditInFrame, EditOutFrame, bWithinPlaybackRange));
 			}
 		}
 		else if (MasterTrack->IsA(UMovieSceneAudioTrack::StaticClass()))
