Sequencer: More robust compilation of sequence templates during cook

  - We now compile templates as part of PreSave rather than Serialize as this prevents us from repretedly compiling and wiping templates in Serialize, which was also leading to a disparity in the cooked data with FMovieSceneEvaluationField::Ranges becoming empty.

#tests Ran a cooked soak test of paragon and FMovieSceneCompiler::CompileTime no longer gets called at runtime, where before it did. Preflighted with "Test Orion PS4" which now passes

diff --git a/Engine/Source/Runtime/MovieScene/MovieScene.Build.cs b/Engine/Source/Runtime/MovieScene/MovieScene.Build.cs
--- a/Engine/Source/Runtime/MovieScene/MovieScene.Build.cs
+++ b/Engine/Source/Runtime/MovieScene/MovieScene.Build.cs
@@ -8,7 +8,13 @@ public class MovieScene : ModuleRules
 	{
 		PrivateIncludePaths.Add("Runtime/MovieScene/Private");
 
-		PublicDependencyModuleNames.AddRange(
+        PrivateIncludePathModuleNames.AddRange(
+            new string[] {
+                "TargetPlatform"
+            }
+        );
+
+        PublicDependencyModuleNames.AddRange(
 			new string[] {
 				"Core",
 				"CoreUObject",
diff --git a/Engine/Source/Runtime/MovieScene/Private/MovieSceneSequence.cpp b/Engine/Source/Runtime/MovieScene/Private/MovieSceneSequence.cpp
--- a/Engine/Source/Runtime/MovieScene/Private/MovieSceneSequence.cpp
+++ b/Engine/Source/Runtime/MovieScene/Private/MovieSceneSequence.cpp
@@ -8,6 +8,7 @@
 #include "Tracks/MovieSceneSubTrack.h"
 #include "Sections/MovieSceneSubSection.h"
 #include "Compilation/MovieSceneCompiler.h"
+#include "Interfaces/ITargetPlatform.h"
 
 UMovieSceneSequence::UMovieSceneSequence(const FObjectInitializer& Init)
 	: Super(Init)
@@ -38,23 +39,30 @@ void UMovieSceneSequence::PostLoad()
 	Super::PostLoad();
 }
 
-void UMovieSceneSequence::Serialize(FArchive& Ar)
+void UMovieSceneSequence::PreSave(const ITargetPlatform* TargetPlatform)
 {
-	Ar.UsingCustomVersion(FMovieSceneEvaluationCustomVersion::GUID);
-	Ar.UsingCustomVersion(FEditorObjectVersion::GUID);
-
-#if WITH_EDITORONLY_DATA
-	if (Ar.IsCooking() && !HasAnyFlags(RF_ClassDefaultObject))
+#if WITH_EDITOR
+	if (!HasAnyFlags(RF_ClassDefaultObject))
 	{
-		FMovieSceneSequencePrecompiledTemplateStore Store;
-		FMovieSceneCompiler::Compile(*this, Store);
-	}
-	else if (Ar.IsSaving())
-	{
-		// Don't save template data unless we're cooking
-		PrecompiledEvaluationTemplate = FMovieSceneEvaluationTemplate();
+		if (TargetPlatform && TargetPlatform->RequiresCookedData())
+		{
+			FMovieSceneSequencePrecompiledTemplateStore Store;
+			FMovieSceneCompiler::Compile(*this, Store);
+		}
+		else
+		{
+			// Don't save template data unless we're cooking
+			PrecompiledEvaluationTemplate = FMovieSceneEvaluationTemplate();
+		}
 	}
 #endif
+	Super::PreSave(TargetPlatform);
+}
+
+void UMovieSceneSequence::Serialize(FArchive& Ar)
+{
+	Ar.UsingCustomVersion(FMovieSceneEvaluationCustomVersion::GUID);
+	Ar.UsingCustomVersion(FEditorObjectVersion::GUID);
 	
 	Super::Serialize(Ar);
 }
diff --git a/Engine/Source/Runtime/MovieScene/Public/MovieSceneSection.h b/Engine/Source/Runtime/MovieScene/Public/MovieSceneSection.h
--- a/Engine/Source/Runtime/MovieScene/Public/MovieSceneSection.h
+++ b/Engine/Source/Runtime/MovieScene/Public/MovieSceneSection.h
@@ -173,7 +173,9 @@ public:
 	 */
 	void SetRange(const TRange<FFrameNumber>& NewRange)
 	{
-		if (TryModify())
+		// Do not modify for objects that still need initialization (i.e. we're in the object's constructor)
+		bool bCanSetRange = HasAnyFlags(RF_NeedInitialization) || TryModify();
+		if (bCanSetRange)
 		{
 			check(NewRange.GetLowerBound().IsOpen() || NewRange.GetUpperBound().IsOpen() || NewRange.GetLowerBoundValue() <= NewRange.GetUpperBoundValue());
 			SectionRange.Value = NewRange;
diff --git a/Engine/Source/Runtime/MovieScene/Public/MovieSceneSequence.h b/Engine/Source/Runtime/MovieScene/Public/MovieSceneSequence.h
--- a/Engine/Source/Runtime/MovieScene/Public/MovieSceneSequence.h
+++ b/Engine/Source/Runtime/MovieScene/Public/MovieSceneSequence.h
@@ -141,6 +141,7 @@ public:
 public:
 
 	MOVIESCENE_API virtual void PostLoad() override;
+	MOVIESCENE_API virtual void PreSave(const ITargetPlatform* TargetPlatform) override;
 
 	MOVIESCENE_API virtual void Serialize(FArchive& Ar) override;
 
diff --git a/Engine/Source/Runtime/UMG/Private/UserWidget.cpp b/Engine/Source/Runtime/UMG/Private/UserWidget.cpp
--- a/Engine/Source/Runtime/UMG/Private/UserWidget.cpp
+++ b/Engine/Source/Runtime/UMG/Private/UserWidget.cpp
@@ -23,6 +23,7 @@
 #include "UMGPrivate.h"
 #include "UObject/UObjectHash.h"
 #include "UObject/PropertyPortFlags.h"
+#include "Compilation/MovieSceneCompiler.h"
 #include "TimerManager.h"
 
 DECLARE_CYCLE_STAT(TEXT("UserWidget Create"), STAT_CreateWidget, STATGROUP_Slate);
@@ -149,6 +150,10 @@ void UUserWidget::TemplateInitInner()
 
 			if ( DuplicatedAnimation->GetMovieScene() )
 			{
+				// Compile the animation template here since PreSave will not get called on the duplicated animation that was itself created during PreSave
+				FMovieSceneSequencePrecompiledTemplateStore Store;
+				FMovieSceneCompiler::Compile(*DuplicatedAnimation, Store);
+
 				// Find property with the same name as the template and assign the new widget to it.
 				UObjectPropertyBase* Prop = FindField<UObjectPropertyBase>(WidgetClass, DuplicatedAnimation->GetMovieScene()->GetFName());
 				if ( Prop )
