Copying //UE4/Dev-Sequencer to //UE4/Dev-Main (Source: //UE4/Dev-Sequencer @ 4044339)

diff --git a/Engine/Source/Runtime/LevelSequence/Private/LevelSequence.cpp b/Engine/Source/Runtime/LevelSequence/Private/LevelSequence.cpp
--- a/Engine/Source/Runtime/LevelSequence/Private/LevelSequence.cpp
+++ b/Engine/Source/Runtime/LevelSequence/Private/LevelSequence.cpp
@@ -10,6 +10,12 @@
 #include "MovieScene.h"
 #include "UObject/Package.h"
 #include "UObject/UObjectHash.h"
+#include "Animation/AnimInstance.h"
+
+#if WITH_EDITOR
+	#include "UObject/SequencerObjectVersion.h"
+	#include "UObject/ObjectRedirector.h"
+#endif
 
 DEFINE_LOG_CATEGORY_STATIC(LogLevelSequence, Log, All);
 
@@ -19,10 +25,10 @@ static TAutoConsoleVariable<int32> CVarDefaultEvaluationType(
 	TEXT("0: Playback locked to playback frames\n1: Unlocked playback with sub frame interpolation"),
 	ECVF_Default);
 
-static TAutoConsoleVariable<FString> CVarDefaultFrameResolution(
-	TEXT("LevelSequence.DefaultFrameResolution"),
+static TAutoConsoleVariable<FString> CVarDefaultTickResolution(
+	TEXT("LevelSequence.DefaultTickResolution"),
 	TEXT("24000fps"),
-	TEXT("Specifies default a frame resolution for newly created level sequences. Examples: 30 fps, 120/1 (120 fps), 30000/1001 (29.97), 0.01s (10ms)."),
+	TEXT("Specifies default a tick resolution for newly created level sequences. Examples: 30 fps, 120/1 (120 fps), 30000/1001 (29.97), 0.01s (10ms)."),
 	ECVF_Default);
 
 static TAutoConsoleVariable<FString> CVarDefaultDisplayRate(
@@ -47,13 +53,13 @@ void ULevelSequence::Initialize()
 
 	MovieScene->SetEvaluationType( bFrameLocked ? EMovieSceneEvaluationType::FrameLocked : EMovieSceneEvaluationType::WithSubFrames );
 
-	FFrameRate FrameResolution(60000, 1);
-	TryParseString(FrameResolution, *CVarDefaultFrameResolution.GetValueOnGameThread());
-	MovieScene->SetFrameResolutionDirectly(FrameResolution);
+	FFrameRate TickResolution(60000, 1);
+	TryParseString(TickResolution, *CVarDefaultTickResolution.GetValueOnGameThread());
+	MovieScene->SetTickResolutionDirectly(TickResolution);
 
-	FFrameRate PlaybackFrameRate(30, 1);
-	TryParseString(PlaybackFrameRate, *CVarDefaultDisplayRate.GetValueOnGameThread());
-	MovieScene->SetPlaybackFrameRate(PlaybackFrameRate);
+	FFrameRate DisplayRate(30, 1);
+	TryParseString(DisplayRate, *CVarDefaultDisplayRate.GetValueOnGameThread());
+	MovieScene->SetDisplayRate(DisplayRate);
 }
 
 UObject* ULevelSequence::MakeSpawnableTemplateFromInstance(UObject& InSourceObject, FName ObjectName)
@@ -62,6 +68,7 @@ UObject* ULevelSequence::MakeSpawnableTemplateFromInstance(UObject& InSourceObje
 
 	UEngine::FCopyPropertiesForUnrelatedObjectsParams CopyParams;
 	CopyParams.bNotifyObjectReplacement = false;
+	CopyParams.bPreserveRootComponent = false;
 	UEngine::CopyPropertiesForUnrelatedObjects(&InSourceObject, NewInstance, CopyParams);
 
 	AActor* Actor = CastChecked<AActor>(NewInstance);
@@ -77,8 +84,66 @@ UObject* ULevelSequence::MakeSpawnableTemplateFromInstance(UObject& InSourceObje
 
 bool ULevelSequence::CanAnimateObject(UObject& InObject) const 
 {
-	return InObject.IsA<AActor>() || InObject.IsA<UActorComponent>();
+	return InObject.IsA<AActor>() || InObject.IsA<UActorComponent>() || InObject.IsA<UAnimInstance>();
+}
+
+#if WITH_EDITOR
+
+void PurgeLegacyBlueprints(UObject* InObject, UPackage* Package)
+{
+	if (UBlueprint* BP = Cast<UBlueprint>(InObject))
+	{
+		UPackage* TransientPackage = GetTransientPackage();
+
+		{
+			FString OldName = BP->GetName();
+
+			BP->ClearFlags(RF_Public);
+			BP->SetFlags(RF_Transient);
+			BP->RemoveFromRoot();
+
+			FName NewName = MakeUniqueObjectName(TransientPackage, UBlueprint::StaticClass(), *FString::Printf(TEXT("DEAD_SPAWNABLE_BLUEPRINT_%s"), *BP->GetName()));
+			BP->Rename(*NewName.ToString(), GetTransientPackage(), (REN_NonTransactional|REN_ForceNoResetLoaders|REN_DoNotDirty));
+
+			UE_LOG(LogLevelSequence, Log, TEXT("Discarding blueprint '%s' from package '%s'."), *OldName, *Package->GetName());
+		}
+
+		if (BP->GeneratedClass)
+		{
+			FName    OldName    = BP->GeneratedClass->GetFName();
+			UObject* OldOuter   = BP->GeneratedClass->GetOuter();
+			UClass*  SuperClass = BP->GeneratedClass->GetSuperClass();
+
+			if( BP->GeneratedClass->ClassDefaultObject )
+			{
+				BP->GeneratedClass->ClassDefaultObject->ClearFlags(RF_Public);
+				BP->GeneratedClass->ClassDefaultObject->SetFlags(RF_Transient);
+				BP->GeneratedClass->ClassDefaultObject->RemoveFromRoot();
+			}
+
+			BP->GeneratedClass->ClearFlags(RF_Public);
+			BP->GeneratedClass->SetFlags(RF_Transient);
+			BP->GeneratedClass->ClassFlags |= CLASS_Deprecated;
+			BP->GeneratedClass->RemoveFromRoot();
+
+			FName NewName = MakeUniqueObjectName(TransientPackage, BP->GeneratedClass, *FString::Printf(TEXT("DEAD_SPAWNABLE_BP_CLASS_%s_C"), *BP->GeneratedClass->ClassGeneratedBy->GetName()));
+			BP->GeneratedClass->Rename(*NewName.ToString(), GetTransientPackage(), (REN_DoNotDirty|REN_NonTransactional|REN_ForceNoResetLoaders));
+
+			if (SuperClass)
+			{
+				UObjectRedirector* Redirector = NewObject<UObjectRedirector>(OldOuter, OldName);
+				Redirector->DestinationObject = SuperClass;
+
+				UE_LOG(LogLevelSequence, Log, TEXT("Discarding generated class '%s' from package '%s'. Replacing with redirector to '%s'"), *OldName.ToString(), *Package->GetName(), *SuperClass->GetName());
+			}
+			else
+			{
+				UE_LOG(LogLevelSequence, Log, TEXT("Discarding generated class '%s' from package '%s'. Unable to create redirector due to no super class."), *OldName.ToString(), *Package->GetName());
+			}
+		}
+	}
 }
+#endif
 
 void ULevelSequence::PostLoad()
 {
@@ -96,7 +161,7 @@ void ULevelSequence::PostLoad()
 			{
 				const FName TemplateName = MakeUniqueObjectName(MovieScene, UObject::StaticClass(), Spawnable.GeneratedClass_DEPRECATED->ClassGeneratedBy->GetFName());
 
-				UObject* NewTemplate = NewObject<UObject>(MovieScene, Spawnable.GeneratedClass_DEPRECATED, TemplateName);
+				UObject* NewTemplate = NewObject<UObject>(MovieScene, Spawnable.GeneratedClass_DEPRECATED->GetSuperClass(), TemplateName);
 				if (NewTemplate)
 				{
 					Spawnable.CopyObjectTemplate(*NewTemplate, *this);
@@ -115,6 +180,27 @@ void ULevelSequence::PostLoad()
 	{
 		MovieScene->RemoveSpawnable(ID);
 	}
+
+	if (GetLinkerCustomVersion(FSequencerObjectVersion::GUID) < FSequencerObjectVersion::PurgeSpawnableBlueprints)
+	{
+		// Remove any old generated classes from the package that will have been left behind from when we used blueprints for spawnables
+		UPackage* Package = GetOutermost();
+		ForEachObjectWithOuter(MovieScene, [=](UObject* InObject) { PurgeLegacyBlueprints(InObject, Package); }, false);
+
+		// Remove any invalid object bindings
+		TSet<FGuid> ValidObjectBindings;
+		for (int32 Index = 0; Index < MovieScene->GetSpawnableCount(); ++Index)
+		{
+			ValidObjectBindings.Add(MovieScene->GetSpawnable(Index).GetGuid());
+		}
+		for (int32 Index = 0; Index < MovieScene->GetPossessableCount(); ++Index)
+		{
+			ValidObjectBindings.Add(MovieScene->GetPossessable(Index).GetGuid());
+		}
+
+		BindingReferences.RemoveInvalidBindings(ValidObjectBindings);
+	}
+
 #endif
 }
 
@@ -149,7 +235,7 @@ void ULevelSequence::BindPossessableObject(const FGuid& ObjectId, UObject& Posse
 
 bool ULevelSequence::CanPossessObject(UObject& Object, UObject* InPlaybackContext) const
 {
-	return Object.IsA<AActor>() || Object.IsA<UActorComponent>();
+	return Object.IsA<AActor>() || Object.IsA<UActorComponent>() || Object.IsA<UAnimInstance>();
 }
 
 void ULevelSequence::LocateBoundObjects(const FGuid& ObjectId, UObject* Context, TArray<UObject*, TInlineAllocator<1>>& OutObjects) const
@@ -164,6 +250,21 @@ void ULevelSequence::LocateBoundObjects(const FGuid& ObjectId, UObject* Context,
 	BindingReferences.ResolveBinding(ObjectId, Context, OutObjects);
 }
 
+void ULevelSequence::GatherExpiredObjects(const FMovieSceneObjectCache& InObjectCache, TArray<FGuid>& OutInvalidIDs) const
+{
+	for (const FGuid& ObjectId : BindingReferences.GetBoundAnimInstances())
+	{
+		for (TWeakObjectPtr<> WeakObject : InObjectCache.IterateBoundObjects(ObjectId))
+		{
+			UAnimInstance* AnimInstance = Cast<UAnimInstance>(WeakObject.Get());
+			if (!AnimInstance || !AnimInstance->GetOwningComponent() || AnimInstance->GetOwningComponent()->GetAnimInstance() != AnimInstance)
+			{
+				OutInvalidIDs.Add(ObjectId);
+			}
+		}
+	}
+}
+
 UMovieScene* ULevelSequence::GetMovieScene() const
 {
 	return MovieScene;
@@ -171,13 +272,16 @@ UMovieScene* ULevelSequence::GetMovieScene() const
 
 UObject* ULevelSequence::GetParentObject(UObject* Object) const
 {
-	UActorComponent* Component = Cast<UActorComponent>(Object);
-
-	if (Component != nullptr)
+	if (UActorComponent* Component = Cast<UActorComponent>(Object))
 	{
 		return Component->GetOwner();
 	}
 
+	if (UAnimInstance* AnimInstance = Cast<UAnimInstance>(Object))
+	{
+		return AnimInstance->GetOwningComponent();
+	}
+
 	return nullptr;
 }
 
diff --git a/Engine/Source/Runtime/LevelSequence/Private/LevelSequenceActorSpawner.cpp b/Engine/Source/Runtime/LevelSequence/Private/LevelSequenceActorSpawner.cpp
--- a/Engine/Source/Runtime/LevelSequence/Private/LevelSequenceActorSpawner.cpp
+++ b/Engine/Source/Runtime/LevelSequence/Private/LevelSequenceActorSpawner.cpp
@@ -7,6 +7,8 @@
 #include "Engine/World.h"
 #include "Particles/ParticleSystemComponent.h"
 #include "Engine/Engine.h"
+#include "Engine/LevelStreaming.h"
+#include "LevelUtils.h"
 
 static const FName SequencerActorTag(TEXT("SequencerActor"));
 
@@ -20,6 +22,25 @@ UClass* FLevelSequenceActorSpawner::GetSupportedTemplateType() const
 	return AActor::StaticClass();
 }
 
+ULevelStreaming* GetLevelStreaming(FString SafeLevelName, UWorld& World)
+{
+	if (FPackageName::IsShortPackageName(SafeLevelName))
+	{
+		// Make sure MyMap1 and Map1 names do not resolve to a same streaming level
+		SafeLevelName = TEXT("/") + SafeLevelName;
+	}
+
+	for (ULevelStreaming* LevelStreaming : World.GetStreamingLevels())
+	{
+		if (LevelStreaming && LevelStreaming->GetWorldAssetPackageName().EndsWith(SafeLevelName, ESearchCase::IgnoreCase))
+		{
+			return LevelStreaming;
+		}
+	}
+
+	return nullptr;
+}
+
 UObject* FLevelSequenceActorSpawner::SpawnObject(FMovieSceneSpawnable& Spawnable, FMovieSceneSequenceIDRef TemplateID, IMovieScenePlayer& Player)
 {
 	AActor* ObjectTemplate = Cast<AActor>(Spawnable.GetObjectTemplate());
@@ -28,9 +49,6 @@ UObject* FLevelSequenceActorSpawner::SpawnObject(FMovieSceneSpawnable& Spawnable
 		return nullptr;
 	}
 
-	// @todo sequencer: We should probably spawn these in a specific sub-level!
-	// World->CurrentLevel = ???;
-
 	const EObjectFlags ObjectFlags = RF_Transient;
 
 	// @todo sequencer livecapture: Consider using SetPlayInEditorWorld() and RestoreEditorWorld() here instead
@@ -43,6 +61,17 @@ UObject* FLevelSequenceActorSpawner::SpawnObject(FMovieSceneSpawnable& Spawnable
 		WorldContext = GWorld;
 	}
 
+	ULevelStreaming* LevelStreaming = GetLevelStreaming(Spawnable.GetLevelName().ToString(), *WorldContext);
+	if (LevelStreaming && LevelStreaming->GetWorldAsset().IsValid())
+	{
+		WorldContext = LevelStreaming->GetWorldAsset().Get();
+	}
+	else if (Spawnable.GetLevelName() != NAME_None)
+	{
+		UE_LOG(LogMovieScene, Warning, TEXT("Can't find sublevel '%s' to spawn '%s' into"), *Spawnable.GetLevelName().ToString(), *Spawnable.GetName());
+	}
+
+
 	// Spawn the puppet actor
 	FActorSpawnParameters SpawnInfo;
 	{
diff --git a/Engine/Source/Runtime/LevelSequence/Private/LevelSequenceBindingReference.cpp b/Engine/Source/Runtime/LevelSequence/Private/LevelSequenceBindingReference.cpp
--- a/Engine/Source/Runtime/LevelSequence/Private/LevelSequenceBindingReference.cpp
+++ b/Engine/Source/Runtime/LevelSequence/Private/LevelSequenceBindingReference.cpp
@@ -7,6 +7,8 @@
 #include "MovieSceneFwd.h"
 #include "Misc/PackageName.h"
 #include "Engine/World.h"
+#include "Animation/AnimInstance.h"
+#include "Components/SkeletalMeshComponent.h"
 
 FLevelSequenceBindingReference::FLevelSequenceBindingReference(UObject* InObject, UObject* InContext)
 {
@@ -183,17 +185,25 @@ bool FLevelSequenceObjectReferenceMap::Serialize(FArchive& Ar)
 
 bool FLevelSequenceBindingReferences::HasBinding(const FGuid& ObjectId) const
 {
-	return BindingIdToReferences.Contains(ObjectId);
+	return BindingIdToReferences.Contains(ObjectId) || AnimSequenceInstances.Contains(ObjectId);
 }
 
 void FLevelSequenceBindingReferences::AddBinding(const FGuid& ObjectId, UObject* InObject, UObject* InContext)
 {
-	BindingIdToReferences.FindOrAdd(ObjectId).References.Emplace(InObject, InContext);
+	if (InObject->IsA<UAnimInstance>())
+	{
+		AnimSequenceInstances.Add(ObjectId);
+	}
+	else
+	{
+		BindingIdToReferences.FindOrAdd(ObjectId).References.Emplace(InObject, InContext);
+	}
 }
 
 void FLevelSequenceBindingReferences::RemoveBinding(const FGuid& ObjectId)
 {
 	BindingIdToReferences.Remove(ObjectId);
+	AnimSequenceInstances.Remove(ObjectId);
 }
 
 void FLevelSequenceBindingReferences::ResolveBinding(const FGuid& ObjectId, UObject* InContext, TArray<UObject*, TInlineAllocator<1>>& OutObjects) const
@@ -201,6 +211,13 @@ void FLevelSequenceBindingReferences::ResolveBinding(const FGuid& ObjectId, UObj
 	const FLevelSequenceBindingReferenceArray* ReferenceArray = BindingIdToReferences.Find(ObjectId);
 	if (!ReferenceArray)
 	{
+		// If the object ID exists in the AnimSequenceInstances set, then this binding relates to an anim instance on a skeletal mesh component
+		USkeletalMeshComponent* SkeletalMesh = Cast<USkeletalMeshComponent>(InContext);
+		if (SkeletalMesh && AnimSequenceInstances.Contains(ObjectId) && SkeletalMesh->GetAnimInstance())
+		{
+			OutObjects.Add(SkeletalMesh->GetAnimInstance());
+		}
+
 		return;
 	}
 
@@ -217,6 +234,18 @@ void FLevelSequenceBindingReferences::ResolveBinding(const FGuid& ObjectId, UObj
 }
 
 
+void FLevelSequenceBindingReferences::RemoveInvalidBindings(const TSet<FGuid>& ValidBindingIDs)
+{
+	for (auto It = BindingIdToReferences.CreateIterator(); It; ++It)
+	{
+		if (!ValidBindingIDs.Contains(It->Key))
+		{
+			It.RemoveCurrent();
+		}
+	}
+}
+
+
 UObject* FLevelSequenceObjectReferenceMap::ResolveBinding(const FGuid& ObjectId, UObject* InContext) const
 {
 	const FLevelSequenceLegacyObjectReference* Reference = Map.Find(ObjectId);
diff --git a/Engine/Source/Runtime/LevelSequence/Private/LevelSequencePlayer.cpp b/Engine/Source/Runtime/LevelSequence/Private/LevelSequencePlayer.cpp
--- a/Engine/Source/Runtime/LevelSequence/Private/LevelSequencePlayer.cpp
+++ b/Engine/Source/Runtime/LevelSequence/Private/LevelSequencePlayer.cpp
@@ -343,7 +343,7 @@ void ULevelSequencePlayer::TakeFrameSnapshot(FLevelSequencePlayerSnapshot& OutSn
 			// Assume that shots with no sequence start at 0.
 			FMovieSceneSequenceTransform OuterToInnerTransform = ActiveShot->OuterToInnerTransform();
 			UMovieSceneSequence*         InnerSequence         = ActiveShot->GetSequence();
-			FFrameRate                   InnerFrameRate        = InnerSequence ? InnerSequence->GetMovieScene()->GetPlaybackFrameRate() : PlayPosition.GetInputRate();
+			FFrameRate                   InnerFrameRate        = InnerSequence ? InnerSequence->GetMovieScene()->GetDisplayRate() : PlayPosition.GetInputRate();
 
 			OutSnapshot.CurrentShotName      = ActiveShot->GetShotDisplayName();
 			OutSnapshot.CurrentShotLocalTime = FQualifiedFrameTime(CurrentPlayTime * OuterToInnerTransform, InnerFrameRate);
diff --git a/Engine/Source/Runtime/LevelSequence/Public/DefaultLevelSequenceInstanceData.h b/Engine/Source/Runtime/LevelSequence/Public/DefaultLevelSequenceInstanceData.h
--- a/Engine/Source/Runtime/LevelSequence/Public/DefaultLevelSequenceInstanceData.h
+++ b/Engine/Source/Runtime/LevelSequence/Public/DefaultLevelSequenceInstanceData.h
@@ -10,11 +10,13 @@
 #include "DefaultLevelSequenceInstanceData.generated.h"
 
 /** Default instance data class that level sequences understand. Implements IMovieSceneTransformOrigin. */
-UCLASS()
+UCLASS(BlueprintType)
 class UDefaultLevelSequenceInstanceData
 	: public UObject
 	, public IMovieSceneTransformOrigin
 {
+public:
+
 	GENERATED_BODY()
 
 	UDefaultLevelSequenceInstanceData(const FObjectInitializer& Init)
@@ -27,10 +29,10 @@ class UDefaultLevelSequenceInstanceData
 	virtual FTransform NativeGetTransformOrigin() const override { return TransformOriginActor ? TransformOriginActor->ActorToWorld() : TransformOrigin; }
 
 	/** When set, this actor's world position will be used as the transform origin for all absolute transform sections */
-	UPROPERTY(EditAnywhere, Category="General")
+	UPROPERTY(EditAnywhere, BlueprintReadWrite, Category="General")
 	AActor* TransformOriginActor;
 
 	/** Specifies a transform from which all absolute transform sections inside the sequence should be added to. Scale is ignored. */
-	UPROPERTY(EditAnywhere, Category="General")
+	UPROPERTY(EditAnywhere, BlueprintReadWrite, Category="General")
 	FTransform TransformOrigin;
 };
diff --git a/Engine/Source/Runtime/LevelSequence/Public/LevelSequence.h b/Engine/Source/Runtime/LevelSequence/Public/LevelSequence.h
--- a/Engine/Source/Runtime/LevelSequence/Public/LevelSequence.h
+++ b/Engine/Source/Runtime/LevelSequence/Public/LevelSequence.h
@@ -41,6 +41,7 @@ public:
 	virtual void BindPossessableObject(const FGuid& ObjectId, UObject& PossessedObject, UObject* Context) override;
 	virtual bool CanPossessObject(UObject& Object, UObject* InPlaybackContext) const override;
 	virtual void LocateBoundObjects(const FGuid& ObjectId, UObject* Context, TArray<UObject*, TInlineAllocator<1>>& OutObjects) const override;
+	virtual void GatherExpiredObjects(const FMovieSceneObjectCache& InObjectCache, TArray<FGuid>& OutInvalidIDs) const override;
 	virtual UMovieScene* GetMovieScene() const override;
 	virtual UObject* GetParentObject(UObject* Object) const override;
 	virtual void UnbindPossessableObjects(const FGuid& ObjectId) override;
diff --git a/Engine/Source/Runtime/LevelSequence/Public/LevelSequenceBindingReference.h b/Engine/Source/Runtime/LevelSequence/Public/LevelSequenceBindingReference.h
--- a/Engine/Source/Runtime/LevelSequence/Public/LevelSequenceBindingReference.h
+++ b/Engine/Source/Runtime/LevelSequence/Public/LevelSequenceBindingReference.h
@@ -120,9 +120,28 @@ struct FLevelSequenceBindingReferences
 	 */
 	void ResolveBinding(const FGuid& ObjectId, UObject* InContext, TArray<UObject*, TInlineAllocator<1>>& OutObjects) const;
 
+	/**
+	 * Const accessor for the currently bound anim instance IDs
+	 */
+	const TSet<FGuid>& GetBoundAnimInstances() const
+	{
+		return AnimSequenceInstances;
+	}
+
+	/**
+	 * Filter out any bindings that do not match the specified set of GUIDs
+	 *
+	 * @param ValidBindingIDs A set of GUIDs that are considered valid. Anything references not matching these will be removed.
+	 */
+	void RemoveInvalidBindings(const TSet<FGuid>& ValidBindingIDs);
+
 private:
 
 	/** The map from object binding ID to an array of references that pertain to that ID */
 	UPROPERTY()
 	TMap<FGuid, FLevelSequenceBindingReferenceArray> BindingIdToReferences;
+
+	/** A set of object binding IDs that relate to anim sequence instances (must be a child of USkeletalMeshComponent) */
+	UPROPERTY()
+	TSet<FGuid> AnimSequenceInstances;
 };
diff --git a/Engine/Source/Runtime/Media/Private/MediaModule.cpp b/Engine/Source/Runtime/Media/Private/MediaModule.cpp
--- a/Engine/Source/Runtime/Media/Private/MediaModule.cpp
+++ b/Engine/Source/Runtime/Media/Private/MediaModule.cpp
@@ -6,6 +6,7 @@
 #include "HAL/RunnableThread.h"
 #include "Misc/ScopeLock.h"
 #include "Modules/ModuleManager.h"
+#include "Stats/Stats.h"
 
 #include "IMediaCaptureSupport.h"
 #include "IMediaPlayerFactory.h"
@@ -84,11 +85,20 @@ public:
 
 	virtual void TickPostEngine() override
 	{
-		Clock.TickFetch();
+		{
+			QUICK_SCOPE_CYCLE_COUNTER(STAT_Media_TickFetch);
+			Clock.TickFetch();
+		}
+
+		{
+			QUICK_SCOPE_CYCLE_COUNTER(STAT_Media_TickRender);
+			Clock.TickRender();
+		}
 	}
 
 	virtual void TickPostRender() override
 	{
+		QUICK_SCOPE_CYCLE_COUNTER(STAT_Media_TickOutput);
 		Clock.TickOutput();
 	}
 
@@ -99,6 +109,7 @@ public:
 			Clock.UpdateTimecode(TimeSource->GetTimecode(), TimecodeLocked);
 		}
 
+		QUICK_SCOPE_CYCLE_COUNTER(STAT_Media_TickInput);
 		Clock.TickInput();
 
 		OnTickPreEngineCompleted.Broadcast();
@@ -106,7 +117,7 @@ public:
 
 	virtual void TickPreSlate() override
 	{
-		Clock.TickRender();
+		// currently not used
 	}
 
 	virtual void UnregisterCaptureSupport(IMediaCaptureSupport& Support) override
diff --git a/Engine/Source/Runtime/Media/Public/IMediaControls.h b/Engine/Source/Runtime/Media/Public/IMediaControls.h
--- a/Engine/Source/Runtime/Media/Public/IMediaControls.h
+++ b/Engine/Source/Runtime/Media/Public/IMediaControls.h
@@ -13,6 +13,9 @@
  */
 enum class EMediaControl
 {
+	/** Block on fetch. */
+	BlockOnFetch,
+
 	/** Pause playback. */
 	Pause,
 
diff --git a/Engine/Source/Runtime/Media/Public/IMediaTimeSource.h b/Engine/Source/Runtime/Media/Public/IMediaTimeSource.h
--- a/Engine/Source/Runtime/Media/Public/IMediaTimeSource.h
+++ b/Engine/Source/Runtime/Media/Public/IMediaTimeSource.h
@@ -2,6 +2,8 @@
 
 #pragma once
 
+#include "Misc/Timespan.h"
+
 
 /**
  * Interface for media clock time sources.
diff --git a/Engine/Source/Runtime/MediaAssets/Private/Assets/MediaPlayer.cpp b/Engine/Source/Runtime/MediaAssets/Private/Assets/MediaPlayer.cpp
--- a/Engine/Source/Runtime/MediaAssets/Private/Assets/MediaPlayer.cpp
+++ b/Engine/Source/Runtime/MediaAssets/Private/Assets/MediaPlayer.cpp
@@ -524,7 +524,7 @@ bool UMediaPlayer::Rewind()
 
 bool UMediaPlayer::Seek(const FTimespan& Time)
 {
-	UE_LOG(LogMediaAssets, VeryVerbose, TEXT("%s.Seek %s"), *GetFName().ToString(), *Time.ToString());
+	UE_LOG(LogMediaAssets, VeryVerbose, TEXT("%s.Seek %s"), *GetFName().ToString(), *Time.ToString(TEXT("%h:%m:%s.%t")));
 	return PlayerFacade->Seek(Time);
 }
 
@@ -536,6 +536,13 @@ bool UMediaPlayer::SelectTrack(EMediaPlayerTrack TrackType, int32 TrackIndex)
 }
 
 
+void UMediaPlayer::SetBlockOnTime(const FTimespan& Time)
+{
+	UE_LOG(LogMediaAssets, VeryVerbose, TEXT("%s.SetBlockOnTime %s"), *GetFName().ToString(), *Time.ToString(TEXT("%h:%m:%s.%t")));
+	return PlayerFacade->SetBlockOnTime(Time);
+}
+
+
 void UMediaPlayer::SetDesiredPlayerName(FName PlayerName)
 {
 	UE_LOG(LogMediaAssets, Verbose, TEXT("%s.SetDesiredPlayerName %s"), *GetFName().ToString(), *PlayerName.ToString());
@@ -708,13 +715,13 @@ void UMediaPlayer::PostInitProperties()
 	}
 }
 
+
 void UMediaPlayer::PostLoad()
 {
 	Super::PostLoad();
 
 	if (!HasAnyFlags(RF_ClassDefaultObject))
 	{
-		// Set the player GUID - required for UMediaPlayer assets
 		PlayerFacade->SetGuid(PlayerGuid);
 	}
 }
diff --git a/Engine/Source/Runtime/MediaAssets/Private/Assets/MediaSoundComponent.cpp b/Engine/Source/Runtime/MediaAssets/Private/Assets/MediaSoundComponent.cpp
--- a/Engine/Source/Runtime/MediaAssets/Private/Assets/MediaSoundComponent.cpp
+++ b/Engine/Source/Runtime/MediaAssets/Private/Assets/MediaSoundComponent.cpp
@@ -1,6 +1,7 @@
 // Copyright 1998-2018 Epic Games, Inc. All Rights Reserved.
 
 #include "MediaSoundComponent.h"
+#include "MediaAssetsPrivate.h"
 
 #include "Components/BillboardComponent.h"
 #include "Engine/Texture2D.h"
@@ -20,6 +21,9 @@
 #endif
 
 
+#define MEDIASOUNDCOMPONENT_TRACE_RATEADJUSTMENT 0
+
+
 /* Static initialization
  *****************************************************************************/
 
@@ -32,8 +36,12 @@ USoundClass* UMediaSoundComponent::DefaultMediaSoundClassObject = nullptr;
 UMediaSoundComponent::UMediaSoundComponent(const FObjectInitializer& ObjectInitializer)
 	: Super(ObjectInitializer)
 	, Channels(EMediaSoundChannels::Stereo)
+	, DynamicRateAdjustment(true)
+	, RateAdjustmentFactor(0.00000001f)
+	, RateAdjustmentRange(FFloatRange(0.995f, 1.005f))
 	, CachedRate(0.0f)
 	, CachedTime(FTimespan::Zero())
+	, RateAdjustment(1.0f)
 	, Resampler(new FMediaAudioResampler)
 {
 	PrimaryComponentTick.bCanEverTick = true;
@@ -271,6 +279,7 @@ bool UMediaSoundComponent::Init(int32& SampleRate)
 	}*/
 
 	Resampler->Initialize(NumChannels, SampleRate);
+
 	return true;
 }
 
@@ -285,8 +294,55 @@ void UMediaSoundComponent::OnGenerateAudio(float* OutAudio, int32 NumSamples)
 
 	if (PinnedSampleQueue.IsValid() && (CachedRate != 0.0f))
 	{
-		const uint32 FramesRequested = NumSamples / NumChannels;
-		const uint32 FramesWritten = Resampler->Generate(OutAudio, FramesRequested, CachedRate, CachedTime, *PinnedSampleQueue);
+		const float Rate = CachedRate.Load();
+		const FTimespan Time = CachedTime.Load();
+
+		FTimespan OutTime = FTimespan::Zero();
+
+		while (true)
+		{
+			const uint32 FramesRequested = NumSamples / NumChannels;
+			const uint32 FramesWritten = Resampler->Generate(OutAudio, OutTime, FramesRequested, Rate * RateAdjustment, Time, *PinnedSampleQueue);
+
+			if (FramesWritten == 0)
+			{
+				return; // no samples available
+			}
+
+			if (DynamicRateAdjustment)
+			{
+				RateAdjustment = 1.0f + (CachedTime.Load().GetTicks() - OutTime.GetTicks()) * RateAdjustmentFactor;
+			}
+
+			if (RateAdjustmentRange.IsEmpty() || RateAdjustmentRange.Contains(RateAdjustment))
+			{
+				break; // valid sample
+			}
+
+			// drop sample (clocks are too out of sync)
+			RateAdjustment = 1.0f;
+
+			#if MEDIASOUNDCOMPONENT_TRACE_RATEADJUSTMENT
+				UE_LOG(LogMediaAssets, Verbose, TEXT("MediaSoundComponent %p: Sample dropped, Rate %f, Time %s, OutTime %s, Queue %i"),
+					this,
+					Rate,
+					*Time.ToString(TEXT("%h:%m:%s.%t")),
+					*OutTime.ToString(TEXT("%h:%m:%s.%t")),
+					PinnedSampleQueue->Num()
+				);
+			#endif
+		}
+
+		#if MEDIASOUNDCOMPONENT_TRACE_RATEADJUSTMENT
+			UE_LOG(LogMediaAssets, Verbose, TEXT("MediaSoundComponent %p: Sample rendered, Rate %f, Time %s, OutTime %s, RateAdjustment %f, Queue %i"),
+				this,
+				Rate,
+				*Time.ToString(TEXT("%h:%m:%s.%t")),
+				*OutTime.ToString(TEXT("%h:%m:%s.%t")),
+				RateAdjustment,
+				PinnedSampleQueue->Num()
+			);
+		#endif
 	}
 	else
 	{
diff --git a/Engine/Source/Runtime/MediaAssets/Private/Misc/MediaTextureResource.cpp b/Engine/Source/Runtime/MediaAssets/Private/Misc/MediaTextureResource.cpp
--- a/Engine/Source/Runtime/MediaAssets/Private/Misc/MediaTextureResource.cpp
+++ b/Engine/Source/Runtime/MediaAssets/Private/Misc/MediaTextureResource.cpp
@@ -161,46 +161,73 @@ void FMediaTextureResource::Render(const FRenderParams& Params)
 	{
 		// get the most current sample to be rendered
 		TSharedPtr<IMediaTextureSample, ESPMode::ThreadSafe> Sample;
-		bool SampleValid = false;
+		bool UseSample = false;
 		
 		while (SampleSource->Peek(Sample) && Sample.IsValid())
 		{
 			const FTimespan StartTime = Sample->GetTime();
 			const FTimespan EndTime = StartTime + Sample->GetDuration();
 
-			if (((Params.Rate > 0.0f) && (StartTime >= Params.Time)) ||
-				((Params.Rate < 0.0f) && (EndTime <= Params.Time)))
+			if ((Params.Rate >= 0.0f) && (Params.Time < StartTime))
 			{
-				break; // future sample
+				break; // future sample (forward play)
 			}
 
-			SampleValid = SampleSource->Dequeue(Sample);
+			if ((Params.Rate <= 0.0f) && (Params.Time >= EndTime))
+			{
+				break; // future sample (reverse play)
+			}
+
+			UseSample = SampleSource->Dequeue(Sample);
+
+#if MEDIATEXTURERESOURCE_TRACE_RENDER
+			if (!UseSample && Sample.IsValid())
+			{
+				UE_LOG(LogMediaAssets, VeryVerbose, TEXT("TextureResource %p: Sample with time %s got flushed at time %s"),
+					this,
+					*Sample->GetTime().ToString(TEXT("%h:%m:%s.%t")),
+					*Params.Time.ToString(TEXT("%h:%m:%s.%t"))
+				);
+			}
+#endif
 		}
 
-		if (SampleValid)
+		if (UseSample)
 		{
 			// render the sample
 			if (Sample->GetOutputDim().GetMin() <= 0)
 			{
-				#if MEDIATEXTURERESOURCE_TRACE_RENDER
-					UE_LOG(LogMediaAssets, VeryVerbose, TEXT("TextureResource %p: Corrupt sample with time %s at time %s"), this, *Sample->GetTime().ToString(), *Params.Time.ToString());
-				#endif
+#if MEDIATEXTURERESOURCE_TRACE_RENDER
+				UE_LOG(LogMediaAssets, VeryVerbose, TEXT("TextureResource %p: Corrupt sample with time %s at time %s"),
+					this,
+					*Sample->GetTime().ToString(TEXT("%h:%m:%s.%t")),
+					*Params.Time.ToString(TEXT("%h:%m:%s.%t"))
+				);
+#endif
 
 				ClearTexture(FLinearColor::Red, Params.SrgbOutput); // mark corrupt sample
 			}
 			else if (MediaTextureResource::RequiresConversion(Sample, Params.SrgbOutput))
 			{
-				#if MEDIATEXTURERESOURCE_TRACE_RENDER
-					UE_LOG(LogMediaAssets, VeryVerbose, TEXT("TextureResource %p: Converting sample with time %s at time %s"), this, *Sample->GetTime().ToString(), *Params.Time.ToString());
-				#endif
+#if MEDIATEXTURERESOURCE_TRACE_RENDER
+				UE_LOG(LogMediaAssets, VeryVerbose, TEXT("TextureResource %p: Converting sample with time %s at time %s"),
+					this,
+					*Sample->GetTime().ToString(TEXT("%h:%m:%s.%t")),
+					*Params.Time.ToString(TEXT("%h:%m:%s.%t"))
+				);
+#endif
 
 				ConvertSample(Sample, Params.ClearColor, Params.SrgbOutput);
 			}
 			else
 			{
-				#if MEDIATEXTURERESOURCE_TRACE_RENDER
-					UE_LOG(LogMediaAssets, VeryVerbose, TEXT("TextureResource %p: Copying sample with time %s at time %s"), this, *Sample->GetTime().ToString(), *Params.Time.ToString());
-				#endif
+#if MEDIATEXTURERESOURCE_TRACE_RENDER
+				UE_LOG(LogMediaAssets, VeryVerbose, TEXT("TextureResource %p: Copying sample with time %s at time %s"),
+					this,
+					*Sample->GetTime().ToString(TEXT("%h:%m:%s.%t")),
+					*Params.Time.ToString(TEXT("%h:%m:%s.%t"))
+				);
+#endif
 
 				CopySample(Sample, Params.ClearColor, Params.SrgbOutput);
 			}
@@ -208,14 +235,34 @@ void FMediaTextureResource::Render(const FRenderParams& Params)
 			Rotation = Sample->GetScaleRotation();
 			Offset = Sample->GetOffset();
 		}
+#if MEDIATEXTURERESOURCE_TRACE_RENDER
+		else if (Sample.IsValid())
+		{
+			UE_LOG(LogMediaAssets, VeryVerbose, TEXT("TextureResource %p: Sample with time %s cannot be used at time %s"),
+				this,
+				*Sample->GetTime().ToString(TEXT("%h:%m:%s.%t")),
+				*Params.Time.ToString(TEXT("%h:%m:%s.%t"))
+			);
+		}
+		else
+		{
+			UE_LOG(LogMediaAssets, VeryVerbose, TEXT("TextureResource %p: No valid sample available at time %s"),
+				this,
+				*Params.Time.ToString(TEXT("%h:%m:%s.%t"))
+			);
+		}
+#endif
 	}
 	else if (Params.CanClear)
 	{
 		if (!Cleared || (Params.ClearColor != CurrentClearColor))
 		{
-			#if MEDIATEXTURERESOURCE_TRACE_RENDER
-				UE_LOG(LogMediaAssets, VeryVerbose, TEXT("TextureResource %p: Clearing texture at time %s"), this, *Params.Time.ToString());
-			#endif
+#if MEDIATEXTURERESOURCE_TRACE_RENDER
+			UE_LOG(LogMediaAssets, VeryVerbose, TEXT("TextureResource %p: Clearing texture at time %s"),
+				this,
+				*Params.Time.ToString(TEXT("%h:%m:%s.%t"))
+			);
+#endif
 
 			ClearTexture(Params.ClearColor, Params.SrgbOutput);
 		}
diff --git a/Engine/Source/Runtime/MediaAssets/Public/MediaPlayer.h b/Engine/Source/Runtime/MediaAssets/Public/MediaPlayer.h
--- a/Engine/Source/Runtime/MediaAssets/Public/MediaPlayer.h
+++ b/Engine/Source/Runtime/MediaAssets/Public/MediaPlayer.h
@@ -646,6 +646,17 @@ public:
 	UFUNCTION(BlueprintCallable, Category="Media|MediaPlayer")
 	bool SelectTrack(EMediaPlayerTrack TrackType, int32 TrackIndex);
 
+	/**
+	 * Set the time on which to block.
+	 *
+	 * If set, this player will block in TickFetch until the video sample
+	 * for the specified time are actually available.
+	 *
+	 * @param Time The time to block on, or FTimespan::MinValue to disable.
+	 */
+	UFUNCTION(BlueprintCallable, Category = "Media|MediaPlayer")
+	void SetBlockOnTime(const FTimespan& Time);
+
 	/**
 	 * Set the name of the desired native player.
 	 *
diff --git a/Engine/Source/Runtime/MediaAssets/Public/MediaSoundComponent.h b/Engine/Source/Runtime/MediaAssets/Public/MediaSoundComponent.h
--- a/Engine/Source/Runtime/MediaAssets/Public/MediaSoundComponent.h
+++ b/Engine/Source/Runtime/MediaAssets/Public/MediaSoundComponent.h
@@ -52,9 +52,32 @@ class MEDIAASSETS_API UMediaSoundComponent
 public:
 
 	/** Media sound channel type. */
-	UPROPERTY(EditAnywhere, Category="MediaSoundComponent")
+	UPROPERTY(EditAnywhere, Category="Media")
 	EMediaSoundChannels Channels;
 
+	/** Dynamically adjust the sample rate if audio and media clock desynchronize. */
+	UPROPERTY(EditAnywhere, Category="Media", AdvancedDisplay)
+	bool DynamicRateAdjustment;
+
+	/**
+	 * Factor for calculating the sample rate adjustment.
+	 *
+	 * If dynamic rate adjustment is enabled, this number is multiplied with the drift
+	 * between the audio and media clock (in 100ns ticks) to determine the adjustment.
+	 * that is to be multiplied into the current playrate.
+	 */
+	UPROPERTY(EditAnywhere, Category="Media", AdvancedDisplay)
+	float RateAdjustmentFactor;
+
+	/**
+	 * The allowed range of dynamic rate adjustment.
+	 *
+	 * If dynamic rate adjustment is enabled, and the necessary adjustment
+	 * falls outside of this range, audio samples will be dropped.
+	 */
+	UPROPERTY(EditAnywhere, Category="Media", AdvancedDisplay)
+	FFloatRange RateAdjustmentRange;
+
 public:
 
 	/**
@@ -176,6 +199,9 @@ private:
 	/** The player facade that's currently providing texture samples. */
 	TWeakPtr<FMediaPlayerFacade, ESPMode::ThreadSafe> CurrentPlayerFacade;
 
+	/** Adjusts the output sample rate to synchronize audio and media clock. */
+	float RateAdjustment;
+
 	/** The audio resampler. */
 	FMediaAudioResampler* Resampler;
 
diff --git a/Engine/Source/Runtime/MediaUtils/MediaUtils.Build.cs b/Engine/Source/Runtime/MediaUtils/MediaUtils.Build.cs
--- a/Engine/Source/Runtime/MediaUtils/MediaUtils.Build.cs
+++ b/Engine/Source/Runtime/MediaUtils/MediaUtils.Build.cs
@@ -13,6 +13,11 @@ namespace UnrealBuildTool.Rules
 					"Media",
 				});
 
+			PrivateDependencyModuleNames.AddRange(
+				new string[] {
+					"ImageWrapper",
+				});
+
 			PrivateIncludePaths.AddRange(
 				new string[] {
 					"Runtime/MediaUtils/Private",
diff --git a/Engine/Source/Runtime/MediaUtils/Private/AppMediaTimeSource.cpp b/Engine/Source/Runtime/MediaUtils/Private/AppMediaTimeSource.cpp
new file mode 100644
--- /dev/null
+++ b/Engine/Source/Runtime/MediaUtils/Private/AppMediaTimeSource.cpp
@@ -0,0 +1,21 @@
+// Copyright 1998-2018 Epic Games, Inc. All Rights Reserved.
+
+#include "AppMediaTimeSource.h"
+#include "MediaUtilsPrivate.h"
+
+#include "Misc/App.h"
+#include "Misc/Timespan.h"
+
+
+/* IMediaTimeSource interface 
+ *****************************************************************************/
+
+FTimespan FAppMediaTimeSource::GetTimecode()
+{
+	const double CurrentTime = FApp::GetCurrentTime();
+	const FTimespan Timecode = FTimespan::FromSeconds(CurrentTime);
+
+	UE_LOG(LogMediaUtils, VeryVerbose, TEXT("AppMediaTimeSource: Time %.10f, Delta %.10f, Timecode %s"), CurrentTime, FApp::GetDeltaTime(), *Timecode.ToString(TEXT("%h:%m:%s.%t")));
+
+	return Timecode;
+}
diff --git a/Engine/Source/Runtime/MediaUtils/Private/MediaAudioResampler.cpp b/Engine/Source/Runtime/MediaUtils/Private/MediaAudioResampler.cpp
--- a/Engine/Source/Runtime/MediaUtils/Private/MediaAudioResampler.cpp
+++ b/Engine/Source/Runtime/MediaUtils/Private/MediaAudioResampler.cpp
@@ -334,7 +334,7 @@ void FMediaAudioResampler::Flush()
 }
 
 
-uint32 FMediaAudioResampler::Generate(float* Output, const uint32 FramesRequested, float Rate, FTimespan Time, FMediaAudioSampleSource& SampleSource)
+uint32 FMediaAudioResampler::Generate(float* Output, FTimespan& OutTime, const uint32 FramesRequested, float Rate, FTimespan Time, FMediaAudioSampleSource& SampleSource)
 {
 	if ((FramesRequested == 0) || (Output == nullptr) || (OutputSampleRate == 0))
 	{
@@ -363,7 +363,7 @@ uint32 FMediaAudioResampler::Generate(float* Output, const uint32 FramesRequeste
 				}
 
 				NextSampleTime -= FTimespan(1); // point into the previous sample
-				NextSampleTime -= FTimespan(1); // accommodate for duration rounding errors
+				NextSampleTime -= FTimespan(1); // account for duration rounding errors
 			}
 			else
 			{
@@ -374,7 +374,7 @@ uint32 FMediaAudioResampler::Generate(float* Output, const uint32 FramesRequeste
 				else
 				{
 					NextSampleTime = InputTime + InputDuration; // point into next sample
-					NextSampleTime += FTimespan(1); // accommodate for duration rounding errors
+					NextSampleTime += FTimespan(1); // account for duration rounding errors
 				}
 			}
 
@@ -405,6 +405,10 @@ uint32 FMediaAudioResampler::Generate(float* Output, const uint32 FramesRequeste
 			continue;
 		}
 
+		// calculate output time
+		const float InputOffset = (FrameIndex + FrameAlpha) / InputFrames;
+		OutTime = InputTime + InputOffset * InputDuration;
+
 		// invert buffer if rate reversed
 		if ((Rate * LastRate) < 0.0f)
 		{
diff --git a/Engine/Source/Runtime/MediaUtils/Private/MediaPlayerFacade.cpp b/Engine/Source/Runtime/MediaUtils/Private/MediaPlayerFacade.cpp
--- a/Engine/Source/Runtime/MediaUtils/Private/MediaPlayerFacade.cpp
+++ b/Engine/Source/Runtime/MediaUtils/Private/MediaPlayerFacade.cpp
@@ -4,6 +4,7 @@
 #include "MediaUtilsPrivate.h"
 
 #include "HAL/PlatformMath.h"
+#include "HAL/PlatformProcess.h"
 #include "IMediaCache.h"
 #include "IMediaControls.h"
 #include "IMediaModule.h"
@@ -11,6 +12,7 @@
 #include "IMediaPlayer.h"
 #include "IMediaPlayerFactory.h"
 #include "IMediaSamples.h"
+#include "IMediaTextureSample.h"
 #include "IMediaTracks.h"
 #include "IMediaView.h"
 #include "Math/NumericLimits.h"
@@ -22,6 +24,7 @@
 #include "MediaSampleCache.h"
 
 
+#define MEDIAPLAYERFACADE_DISABLE_BLOCKING 0
 #define MEDIAPLAYERFACADE_TRACE_SINKOVERFLOWS 0
 
 
@@ -71,8 +74,10 @@ namespace MediaPlayerFacade
 
 FMediaPlayerFacade::FMediaPlayerFacade()
 	: TimeDelay(FTimespan::Zero())
+	, BlockOnTime(FTimespan::MinValue())
 	, Cache(new FMediaSampleCache)
 	, LastRate(0.0f)
+	, NextVideoSampleTime(FTimespan::MinValue())
 { }
 
 
@@ -189,6 +194,7 @@ void FMediaPlayerFacade::Close()
 		Player->Close();
 	}
 
+	BlockOnTime = FTimespan::MinValue();
 	Cache->Empty();
 	CurrentUrl.Empty();
 	LastRate = 0.0f;
@@ -521,6 +527,12 @@ bool FMediaPlayerFacade::SelectTrack(EMediaTrackType TrackType, int32 TrackIndex
 }
 
 
+void FMediaPlayerFacade::SetBlockOnTime(const FTimespan& Time)
+{
+	BlockOnTime = Time;
+}
+
+
 void FMediaPlayerFacade::SetCacheWindow(FTimespan Ahead, FTimespan Behind)
 {
 	Cache->SetCacheWindow(Ahead, Behind);
@@ -589,6 +601,45 @@ bool FMediaPlayerFacade::SupportsRate(float Rate, bool Unthinned) const
 /* FMediaPlayerFacade implementation
 *****************************************************************************/
 
+bool FMediaPlayerFacade::BlockOnFetch() const
+{
+#if MEDIAPLAYERFACADE_DISABLE_BLOCKING
+	return false;
+#endif
+
+	check(Player.IsValid());
+
+	if (BlockOnTime == FTimespan::MinValue())
+	{
+		return false; // no blocking requested
+	}
+
+	if (!Player->GetControls().CanControl(EMediaControl::BlockOnFetch))
+	{
+		return false; // not supported by player plug-in
+	}
+
+	if (IsPreparing())
+	{
+		return true; // block on media opening
+	}
+
+	if (!IsPlaying() || (GetRate() < 0.0f))
+	{
+		return false; // block only in forward play
+	}
+
+	const bool VideoReady = (VideoSampleSinks.Num() == 0) || (BlockOnTime < NextVideoSampleTime);
+
+	if (VideoReady)
+	{
+		return false; // video is ready
+	}
+
+	return true;
+}
+
+
 void FMediaPlayerFacade::FlushSinks()
 {
 	UE_LOG(LogMediaUtils, Verbose, TEXT("PlayerFacade %p: Flushing sinks"), this);
@@ -605,6 +656,8 @@ void FMediaPlayerFacade::FlushSinks()
 	{
 		Player->GetSamples().FlushSamples();
 	}
+
+	NextVideoSampleTime = FTimespan::MinValue();
 }
 
 
@@ -849,27 +902,61 @@ void FMediaPlayerFacade::TickFetch(FTimespan DeltaTime, FTimespan Timecode)
 	// determine range of valid samples
 	TRange<FTimespan> TimeRange;
 
-	const FTimespan Time = GetTime();
+	const FTimespan CurrentTime = GetTime();
 	
 	if (Rate > 0.0f)
 	{
-		TimeRange = TRange<FTimespan>::AtMost(Time);
+		TimeRange = TRange<FTimespan>::AtMost(CurrentTime);
 	}
 	else if (Rate < 0.0f)
 	{
-		TimeRange = TRange<FTimespan>::AtLeast(Time);
+		TimeRange = TRange<FTimespan>::AtLeast(CurrentTime);
 	}
 	else
 	{
-		TimeRange = TRange<FTimespan>(Time);
+		TimeRange = TRange<FTimespan>(CurrentTime);
 	}
 
 	// process samples in range
 	IMediaSamples& Samples = Player->GetSamples();
 
-	ProcessCaptionSamples(Samples, TimeRange);
-	ProcessSubtitleSamples(Samples, TimeRange);
-	ProcessVideoSamples(Samples, TimeRange);
+	bool Blocked = false;
+	FDateTime BlockedTime;
+
+	while (true)
+	{
+		ProcessCaptionSamples(Samples, TimeRange);
+		ProcessSubtitleSamples(Samples, TimeRange);
+		ProcessVideoSamples(Samples, TimeRange);
+
+		if (!BlockOnFetch())
+		{
+			break;
+		}
+
+		if (Blocked)
+		{
+			if ((FDateTime::UtcNow() - BlockedTime) >= FTimespan::FromSeconds(MEDIAUTILS_MAX_BLOCKONFETCH_SECONDS))
+			{
+				UE_LOG(LogMediaUtils, Verbose, TEXT("PlayerFacade %p: Aborted block on fetch %s after %i seconds"),
+					this,
+					*BlockOnTime.ToString(TEXT("%h:%m:%s.%t")),
+					MEDIAUTILS_MAX_BLOCKONFETCH_SECONDS
+				);
+
+				break;
+			}
+		}
+		else
+		{
+			UE_LOG(LogMediaUtils, VeryVerbose, TEXT("PlayerFacade %p: Blocking on fetch %s"), this, *BlockOnTime.ToString(TEXT("%h:%m:%s.%t")));
+
+			Blocked = true;
+			BlockedTime = FDateTime::UtcNow();
+		}
+
+		FPlatformProcess::Sleep(0.0f);
+	}
 }
 
 
@@ -944,16 +1031,6 @@ void FMediaPlayerFacade::TickTickable()
 }
 
 
-/* IMediaEventSink interface
-*****************************************************************************/
-
-void FMediaPlayerFacade::ReceiveMediaEvent(EMediaEvent Event)
-{
-	UE_LOG(LogMediaUtils, VeryVerbose, TEXT("PlayerFacade %p: Received media event %s"), this, *MediaUtils::EventToString(Event));
-	QueuedEvents.Enqueue(Event);
-}
-
-
 /* FMediaPlayerFacade implementation
 *****************************************************************************/
 
@@ -1027,7 +1104,22 @@ void FMediaPlayerFacade::ProcessVideoSamples(IMediaSamples& Samples, TRange<FTim
 
 	while (Samples.FetchVideo(TimeRange, Sample))
 	{
-		if (Sample.IsValid() && !VideoSampleSinks.Enqueue(Sample.ToSharedRef(), MediaPlayerFacade::MaxVideoSinkDepth))
+		if (!Sample.IsValid())
+		{
+			continue;
+		}
+
+		UE_LOG(LogMediaUtils, VeryVerbose, TEXT("PlayerFacade %p: Fetched video sample %s"), this, *Sample->GetTime().ToString(TEXT("%h:%m:%s.%t")));
+
+		if (VideoSampleSinks.Enqueue(Sample.ToSharedRef(), MediaPlayerFacade::MaxVideoSinkDepth))
+		{
+			if (GetRate() >= 0.0f)
+			{
+				NextVideoSampleTime = Sample->GetTime() + Sample->GetDuration();
+				UE_LOG(LogMediaUtils, VeryVerbose, TEXT("PlayerFacade %p: Next video sample time %s"), this, *NextVideoSampleTime.ToString(TEXT("%h:%m:%s.%t")));
+			}
+		}
+		else
 		{
 			#if MEDIAPLAYERFACADE_TRACE_SINKOVERFLOWS
 				UE_LOG(LogMediaUtils, VeryVerbose, TEXT("PlayerFacade %p: Video sample sink overflow"), this);
@@ -1035,3 +1127,13 @@ void FMediaPlayerFacade::ProcessVideoSamples(IMediaSamples& Samples, TRange<FTim
 		}
 	}
 }
+
+
+/* IMediaEventSink interface
+*****************************************************************************/
+
+void FMediaPlayerFacade::ReceiveMediaEvent(EMediaEvent Event)
+{
+	UE_LOG(LogMediaUtils, VeryVerbose, TEXT("PlayerFacade %p: Received media event %s"), this, *MediaUtils::EventToString(Event));
+	QueuedEvents.Enqueue(Event);
+}
diff --git a/Engine/Source/Runtime/MediaUtils/Private/MediaSamples.cpp b/Engine/Source/Runtime/MediaUtils/Private/MediaSamples.cpp
--- a/Engine/Source/Runtime/MediaUtils/Private/MediaSamples.cpp
+++ b/Engine/Source/Runtime/MediaUtils/Private/MediaSamples.cpp
@@ -70,5 +70,6 @@ void FMediaSamples::FlushSamples()
 	AudioSampleQueue.RequestFlush();
 	MetadataSampleQueue.RequestFlush();
 	CaptionSampleQueue.RequestFlush();
+	SubtitleSampleQueue.RequestFlush();
 	VideoSampleQueue.RequestFlush();
 }
diff --git a/Engine/Source/Runtime/MediaUtils/Private/MediaUtilsPrivate.h b/Engine/Source/Runtime/MediaUtils/Private/MediaUtilsPrivate.h
--- a/Engine/Source/Runtime/MediaUtils/Private/MediaUtilsPrivate.h
+++ b/Engine/Source/Runtime/MediaUtils/Private/MediaUtilsPrivate.h
@@ -6,5 +6,9 @@
 #include "Stats/Stats.h"
 
 
+/** Maximum time that TickFetch will block waiting for samples (in seconds). */
+#define MEDIAUTILS_MAX_BLOCKONFETCH_SECONDS 5
+
+
 /** Declares a log category for this module. */
 DECLARE_LOG_CATEGORY_EXTERN(LogMediaUtils, Log, All);
diff --git a/Engine/Source/Runtime/MediaUtils/Public/AppMediaTimeSource.h b/Engine/Source/Runtime/MediaUtils/Public/AppMediaTimeSource.h
--- a/Engine/Source/Runtime/MediaUtils/Public/AppMediaTimeSource.h
+++ b/Engine/Source/Runtime/MediaUtils/Public/AppMediaTimeSource.h
@@ -3,14 +3,12 @@
 #pragma once
 
 #include "IMediaTimeSource.h"
-#include "Misc/App.h"
-#include "Misc/Timespan.h"
 
 
 /**
  * Implements the a media time source that derives its time from the application's global time.
  */
-class FAppMediaTimeSource
+class MEDIAUTILS_API FAppMediaTimeSource
 	: public IMediaTimeSource
 {
 public:
@@ -22,8 +20,5 @@ public:
 
 	//~ IMediaTimeSource interface
 
-	virtual FTimespan GetTimecode() override
-	{
-		return FTimespan::FromSeconds(FApp::GetCurrentTime());
-	}
+	virtual FTimespan GetTimecode() override;
 };
diff --git a/Engine/Source/Runtime/MediaUtils/Public/MediaAudioResampler.h b/Engine/Source/Runtime/MediaUtils/Public/MediaAudioResampler.h
--- a/Engine/Source/Runtime/MediaUtils/Public/MediaAudioResampler.h
+++ b/Engine/Source/Runtime/MediaUtils/Public/MediaAudioResampler.h
@@ -30,13 +30,14 @@ public:
 	 * Generate the next frame of audio.
 	 *
 	 * @param Output The output sample buffer (will be incremented by the number of sample values written).
+	 * @param OutTime The time of the first output frame that was returned.
 	 * @param FramesRequested The maximum number of frames to get.
 	 * @param Rate The current play rate.
 	 * @param Time The current play time.
 	 * @param SampleSource The object that provides more audio samples if needed.
 	 * @return The actual number of frames returned.
 	 */
-	uint32 Generate(float* Output, const uint32 FramesRequested, float Rate, FTimespan Time, FMediaAudioSampleSource& SampleSource);
+	uint32 Generate(float* Output, FTimespan& OutTime, const uint32 FramesRequested, float Rate, FTimespan Time, FMediaAudioSampleSource& SampleSource);
 
 	/**
 	 * Initialize the resampler.
diff --git a/Engine/Source/Runtime/MediaUtils/Public/MediaPlayerFacade.h b/Engine/Source/Runtime/MediaUtils/Public/MediaPlayerFacade.h
--- a/Engine/Source/Runtime/MediaUtils/Public/MediaPlayerFacade.h
+++ b/Engine/Source/Runtime/MediaUtils/Public/MediaPlayerFacade.h
@@ -68,14 +68,44 @@ public:
 
 public:
 
+	/**
+	 * Add the given audio sample sink to this player.
+	 *
+	 * @param SampleSink The sink to receive audio samples.
+	 * @see  AddCaptionSampleSink, AddMetadataSampleSink, AddSubtitleSampleSink, AddVideoSampleSink
+	 */
 	void AddAudioSampleSink(const TSharedRef<FMediaAudioSampleSink, ESPMode::ThreadSafe>& SampleSink);
 
+	/**
+	 * Add the given audio sample sink to this player.
+	 *
+	 * @param SampleSink The sink to receive caption samples.
+	 * @see AddAudioSampleSink, AddMetadataSampleSink, AddSubtitleSampleSink, AddVideoSampleSink
+	 */
 	void AddCaptionSampleSink(const TSharedRef<FMediaOverlaySampleSink, ESPMode::ThreadSafe>& SampleSink);
 
+	/**
+	 * Add the given audio sample sink to this player.
+	 *
+	 * @param SampleSink The sink to receive metadata samples.
+	 * @see AddAudioSampleSink, AddCaptionSampleSink, AddSubtitleSampleSink, AddVideoSampleSink
+	 */
 	void AddMetadataSampleSink(const TSharedRef<FMediaBinarySampleSink, ESPMode::ThreadSafe>& SampleSink);
 
+	/**
+	 * Add the given audio sample sink to this player.
+	 *
+	 * @param SampleSink The sink to receive subtitle samples.
+	 * @see AddAudioSampleSink, AddCaptionSampleSink, AddMetadataSampleSink, AddVideoSampleSink
+	 */
 	void AddSubtitleSampleSink(const TSharedRef<FMediaOverlaySampleSink, ESPMode::ThreadSafe>& SampleSink);
 
+	/**
+	 * Add the given audio sample sink to this player.
+	 *
+	 * @param SampleSink The sink to receive video samples.
+	 * @see AddAudioSampleSink, AddCaptionSampleSink, AddMetadataSampleSink, AddSubtitleSampleSink
+	 */
 	void AddVideoSampleSink(const TSharedRef<FMediaTextureSampleSink, ESPMode::ThreadSafe>& SampleSink);
 
 	/**
@@ -489,6 +519,17 @@ public:
 	 */
 	bool SelectTrack(EMediaTrackType TrackType, int32 TrackIndex);
 
+	/**
+	 * Set the time on which to block.
+	 *
+	 * If set, this player will block in TickFetch until the video sample
+	 * for the specified time are actually available.
+	 *
+	 * @param Time The time to block on, or FTimespan::MinValue to disable.
+	 * @see TickFetch
+	 */
+	void SetBlockOnTime(const FTimespan& Time);
+
 	/**
 	 * Set sample caching options.
 	 *
@@ -604,6 +645,13 @@ public:
 
 protected:
 
+	/**
+	 * Whether sample fetching should block.
+	 *
+	 * @return true if sample fetching should block, false otherwise.
+	 */
+	bool BlockOnFetch() const;
+
 	/** Flush all media sample sinks. */
 	void FlushSinks();
 
@@ -648,12 +696,6 @@ protected:
 	/** Select the default media tracks. */
 	void SelectDefaultTracks();
 
-protected:
-
-	//~ IMediaEventSink interface
-
-	void ReceiveMediaEvent(EMediaEvent Event) override;
-
 protected:
 
 	/** Fetch audio samples from the player and forward them to the registered sinks. */
@@ -671,6 +713,12 @@ protected:
 	/** Fetch video samples from the player and forward them to the registered sinks. */
 	void ProcessVideoSamples(IMediaSamples& Samples, TRange<FTimespan> TimeRange);
 
+protected:
+
+	//~ IMediaEventSink interface
+
+	void ReceiveMediaEvent(EMediaEvent Event) override;
+
 private:
 
 	/** Audio sample sinks. */
@@ -690,13 +738,16 @@ private:
 
 private:
 
+	/** The time to block on sample fetching. */
+	FTimespan BlockOnTime;
+
 	/** Media sample cache. */
 	FMediaSampleCache* Cache;
 
 	/** Synchronizes access to Player. */
 	FCriticalSection CriticalSection;
 
-	/** Holds the URL of the currently loaded media source. */
+	/** The URL of the currently loaded media source. */
 	FString CurrentUrl;
 
 	/** The last used non-zero play rate (zero if playback never started). */
@@ -705,6 +756,9 @@ private:
 	/** An event delegate that is invoked when a media event occurred. */
 	FOnMediaEvent MediaEvent;
 
+	/** Time of the next expected video sample (used for block on fetch). */
+	FTimespan NextVideoSampleTime;
+
 	/** The low-level player used to play the media source. */
 	TSharedPtr<IMediaPlayer, ESPMode::ThreadSafe> Player;
 
diff --git a/Engine/Source/Runtime/MediaUtils/Public/MediaRecorder.cpp b/Engine/Source/Runtime/MediaUtils/Public/MediaRecorder.cpp
new file mode 100644
--- /dev/null
+++ b/Engine/Source/Runtime/MediaUtils/Public/MediaRecorder.cpp
@@ -0,0 +1,173 @@
+// Copyright 1998-2018 Epic Games, Inc. All Rights Reserved.
+
+#include "MediaRecorder.h"
+
+#include "Async/Async.h"
+#include "Containers/UnrealString.h"
+#include "HAL/FileManager.h"
+#include "IImageWrapper.h"
+#include "IImageWrapperModule.h"
+#include "IMediaClock.h"
+#include "IMediaClockSink.h"
+#include "IMediaModule.h"
+#include "Modules/ModuleManager.h"
+#include "Serialization/Archive.h"
+
+#include "MediaPlayerFacade.h"
+
+
+/* Local helpers
+ *****************************************************************************/
+
+/**
+ * Media clock sink for media textures.
+ */
+class FMediaRecorderClockSink
+	: public IMediaClockSink
+{
+public:
+
+	FMediaRecorderClockSink(FMediaRecorder& InOwner)
+		: Owner(InOwner)
+	{ }
+
+	virtual ~FMediaRecorderClockSink() { }
+
+public:
+
+	virtual void TickOutput(FTimespan DeltaTime, FTimespan Timecode) override
+	{
+		Owner.TickRecording(Timecode);
+	}
+
+private:
+
+	FMediaRecorder& Owner;
+};
+
+
+/* FMediaRecorder structors
+ *****************************************************************************/
+
+FMediaRecorder::FMediaRecorder()
+	: FrameCount(0)
+	, Recording(false)
+{ }
+
+
+/* FMediaRecorder interface
+ *****************************************************************************/
+
+void FMediaRecorder::StartRecording(const TSharedRef<FMediaPlayerFacade, ESPMode::ThreadSafe>& PlayerFacade)
+{
+	if (Recording)
+	{
+		StopRecording();
+	}
+
+	IMediaModule* MediaModule = FModuleManager::LoadModulePtr<IMediaModule>("Media");
+
+	if (MediaModule == nullptr)
+	{
+		return;
+	}
+
+	SampleQueue = MakeShared<FMediaTextureSampleQueue, ESPMode::ThreadSafe>();
+	PlayerFacade->AddVideoSampleSink(SampleQueue.ToSharedRef());
+
+	ClockSink = MakeShared<FMediaRecorderClockSink, ESPMode::ThreadSafe>(*this);
+	MediaModule->GetClock().AddSink(ClockSink.ToSharedRef());
+
+	FrameCount = 0;
+	Recording = true;
+}
+
+
+void FMediaRecorder::StopRecording()
+{
+	Recording = false;
+
+	ClockSink.Reset();
+	SampleQueue.Reset();
+}
+
+
+/* FMediaRecorder interface
+ *****************************************************************************/
+
+void FMediaRecorder::TickRecording(FTimespan Timecode)
+{
+	if (!Recording)
+	{
+		return; // not recording
+	}
+
+	check(SampleQueue.IsValid());
+
+	IImageWrapperModule* ImageWrapperModule = FModuleManager::LoadModulePtr<IImageWrapperModule>(FName("ImageWrapper"));
+
+	if (ImageWrapperModule == nullptr)
+	{
+		while (SampleQueue->Pop());
+	}
+	else
+	{
+		TSharedPtr<IMediaTextureSample, ESPMode::ThreadSafe> Sample;
+
+		while (SampleQueue->Dequeue(Sample) && Sample.IsValid())
+		{
+			const void* SampleBuffer = Sample->GetBuffer();
+
+			if (SampleBuffer == nullptr)
+			{
+				continue; // only raw samples supported right now
+			}
+
+			if (Sample->GetDim().GetMin() == 0)
+			{
+				continue; // nothing to save
+			}
+
+			if (Sample->GetFormat() != EMediaTextureSampleFormat::CharBGRA)
+			{
+				continue; // only supports BGRA for now
+			}
+
+			TSharedPtr<IImageWrapper> ImageWrapper = ImageWrapperModule->CreateImageWrapper(EImageFormat::PNG);
+
+			if (!ImageWrapper.IsValid())
+			{
+				continue; // failed to create image wrapper
+			}
+
+			const int32 FrameNumber = FrameCount;
+
+			// convert, compress and write frame asynchronously
+			TFunction<void()> AsyncWrite = [FrameNumber, ImageWrapper, Sample]()
+			{
+				const FIntPoint BufferDim = Sample->GetDim();
+				const SIZE_T BufferSize = BufferDim.X * BufferDim.Y * 4;
+				const FString Filename = FString::Printf(TEXT("%08d.png"), FrameNumber);
+
+				IFileManager* FileManager = &IFileManager::Get();
+				FArchive* Archive = FileManager->CreateFileWriter(*Filename);
+
+				if (Archive != nullptr)
+				{
+					ImageWrapper->SetRaw(Sample->GetBuffer(), BufferSize, BufferDim.X, BufferDim.Y, ERGBFormat::BGRA, 32);
+
+					const TArray<uint8>& CompressedData = ImageWrapper->GetCompressed((int32)EImageCompressionQuality::Default);
+					const int32 CompressedSize = CompressedData.Num();
+
+					Archive->Serialize((void*)CompressedData.GetData(), CompressedSize);
+
+					delete Archive;
+				}
+			};
+
+			Async(EAsyncExecution::ThreadPool, AsyncWrite);
+
+			++FrameCount;
+		}
+	}
+}
diff --git a/Engine/Source/Runtime/MediaUtils/Public/MediaRecorder.h b/Engine/Source/Runtime/MediaUtils/Public/MediaRecorder.h
new file mode 100644
--- /dev/null
+++ b/Engine/Source/Runtime/MediaUtils/Public/MediaRecorder.h
@@ -0,0 +1,67 @@
+// Copyright 1998-2018 Epic Games, Inc. All Rights Reserved.
+
+#pragma once
+
+#include "CoreTypes.h"
+#include "MediaSampleQueue.h"
+#include "Misc/Timespan.h"
+#include "Templates/SharedPointer.h"
+
+class FMediaPlayerFacade;
+class FMediaRecorderClockSink;
+
+
+/**
+ * Records samples from a media player.
+ *
+ * Currently only records texture samples.
+ */
+class FMediaRecorder
+{
+public:
+
+	/** Default constructor. */
+	FMediaRecorder();
+
+public:
+
+	/**
+	 * Start recording samples from a given media player.
+	 *
+	 * @param PlayerFacade The media player facade to record from.
+	 * @see StopRecording
+	 */
+	void StartRecording(const TSharedRef<FMediaPlayerFacade, ESPMode::ThreadSafe>& PlayerFacade);
+
+	/**
+	 * Stop recording media samples.
+	 *
+	 * @see StartRecording
+	 */
+	void StopRecording();
+
+protected:
+
+	/**
+	 * Tick the recorder.
+	 *
+	 * @param Timecode The current timecode.
+	 */
+	void TickRecording(FTimespan Timecode);
+
+private:
+
+	friend class FMediaRecorderClockSink;
+
+	/** The recorder's media clock sink. */
+	TSharedPtr<FMediaRecorderClockSink, ESPMode::ThreadSafe> ClockSink;
+
+	/** Number of frames recorded. */
+	int32 FrameCount;
+
+	/** Whether recording is in progress. */
+	bool Recording;
+
+	/** Texture sample queue. */
+	TSharedPtr<FMediaTextureSampleQueue, ESPMode::ThreadSafe> SampleQueue;
+};
diff --git a/Engine/Source/Runtime/MediaUtils/Public/MediaSampleSinks.h b/Engine/Source/Runtime/MediaUtils/Public/MediaSampleSinks.h
--- a/Engine/Source/Runtime/MediaUtils/Public/MediaSampleSinks.h
+++ b/Engine/Source/Runtime/MediaUtils/Public/MediaSampleSinks.h
@@ -23,7 +23,7 @@ public:
 	 * Add the given media sample sink to the collection.
 	 *
 	 * @param SampleSink The sink to add.
-	 * @see Remove
+	 * @see Num, Remove
 	 */
 	void Add(const TSharedRef<TMediaSampleSink<SampleType>, ESPMode::ThreadSafe>& SampleSink)
 	{
@@ -92,11 +92,22 @@ public:
 		}
 	}
 
+	/**
+	 * Get the number of sinks in this collection.
+	 *
+	 * @return Number of sinks.
+	 * @see Add, Remove
+	 */
+	int32 Num() const
+	{
+		return Sinks.Num();
+	}
+
 	/**
 	 * Remove the given media sample sink from the collection.
 	 *
 	 * @param SampleSink The sink to remove.
-	 * @see Add
+	 * @see Add, Num
 	 */
 	void Remove(const TSharedRef<TMediaSampleSink<SampleType>, ESPMode::ThreadSafe>& SampleSink)
 	{
diff --git a/Engine/Source/Runtime/Messaging/Private/Bus/MessageRouter.cpp b/Engine/Source/Runtime/Messaging/Private/Bus/MessageRouter.cpp
--- a/Engine/Source/Runtime/Messaging/Private/Bus/MessageRouter.cpp
+++ b/Engine/Source/Runtime/Messaging/Private/Bus/MessageRouter.cpp
@@ -17,7 +17,7 @@ FMessageRouter::FMessageRouter()
 	, Tracer(MakeShareable(new FMessageTracer()))
 {
 	ActiveSubscriptions.FindOrAdd(NAME_All);
-	WorkEvent = FPlatformProcess::GetSynchEventFromPool(true);
+	WorkEvent = FPlatformProcess::GetSynchEventFromPool();
 }
 
 
@@ -45,17 +45,14 @@ bool FMessageRouter::Init()
 
 uint32 FMessageRouter::Run()
 {
-	CurrentTime = FDateTime::UtcNow();
-
 	while (!Stopping)
 	{
-		if (WorkEvent->Wait(CalculateWaitTime()))
-		{
-			ProcessCommands();
-			WorkEvent->Reset();
-		}
+		CurrentTime = FDateTime::UtcNow();
 
+		ProcessCommands();
 		ProcessDelayedMessages();
+
+		WorkEvent->Wait(CalculateWaitTime());
 	}
 
 	return 0;
@@ -208,7 +205,6 @@ void FMessageRouter::FilterSubscriptions(
 
 void FMessageRouter::ProcessCommands()
 {
-	CurrentTime = FDateTime::UtcNow();
 	CommandDelegate Command;
 
 	while (Commands.Dequeue(Command))
@@ -235,8 +231,10 @@ void FMessageRouter::ProcessDelayedMessages()
 
 void FMessageRouter::Tick()
 {
-	ProcessCommands();
+	CurrentTime = FDateTime::UtcNow();
+
 	ProcessDelayedMessages();
+	ProcessCommands();
 }
 
 
@@ -353,8 +351,7 @@ void FMessageRouter::HandleRouteMessage(TSharedRef<IMessageContext, ESPMode::Thr
 	}
 
 	// dispatch the message
-	// @todo gmp: implement time synchronization between networked message endpoints
-	if (false) //(Context->GetTimeSent() > CurrentTime)
+	if (Context->GetTimeSent() > CurrentTime)
 	{
 		DelayedMessages.HeapPush(FDelayedMessage(Context, ++DelayedMessagesSequence));
 	}
diff --git a/Engine/Source/Runtime/MovieScene/Private/Channels/MovieSceneFloatChannel.cpp b/Engine/Source/Runtime/MovieScene/Private/Channels/MovieSceneFloatChannel.cpp
--- a/Engine/Source/Runtime/MovieScene/Private/Channels/MovieSceneFloatChannel.cpp
+++ b/Engine/Source/Runtime/MovieScene/Private/Channels/MovieSceneFloatChannel.cpp
@@ -272,7 +272,105 @@ bool FMovieSceneFloatChannel::EvaluateExtrapolation(FFrameTime InTime, float& Ou
 	return false;
 }
 
-bool FMovieSceneFloatChannel::Evaluate(FFrameTime InTime, float& OutValue) const
+
+/* Solve Cubic Euqation using Cardano's forumla
+* Adopted from Graphic Gems 1
+* [REDACTED_PROFILE_URL]/GraphicsGems/blob/master/gems/Roots3And4.c
+*  Solve cubic of form
+*
+* @param Coeff Coefficient parameters of form  Coeff[0] + Coeff[1]*x + Coeff[2]*x^2 + Coeff[3]*x^3 + Coeff[4]*x^4 = 0
+* @param Solution Up to 3 real solutions. We don't include imaginary solutions, would need a complex number objecct
+* @return Returns the number of real solutions returned in the Solution array.
+*/
+static int SolveCubic(double Coeff[4], double Solution[3])
+{
+	auto cbrt = [](double x) -> double
+	{
+		return ((x) > 0.0 ? pow((x), 1.0 / 3.0) : ((x) < 0.0 ? -pow((double)-(x), 1.0 / 3.0) : 0.0));
+	};
+	int     NumSolutions = 0;
+
+	/* normal form: x^3 + Ax^2 + Bx + C = 0 */
+
+	double A = Coeff[2] / Coeff[3];
+	double B = Coeff[1] / Coeff[3];
+	double C = Coeff[0] / Coeff[3];
+
+	/*  substitute x = y - A/3 to eliminate quadric term:
+	x^3 +px + q = 0 */
+
+	double SqOfA = A * A;
+	double P = 1.0 / 3 * (-1.0 / 3 * SqOfA + B);
+	double Q = 1.0 / 2 * (2.0 / 27 * A * SqOfA - 1.0 / 3 * A * B + C);
+
+	/* use Cardano's formula */
+
+	double CubeOfP = P * P * P;
+	double D = Q * Q + CubeOfP;
+
+	if (FMath::IsNearlyZero(D))
+	{
+		if (FMath::IsNearlyZero(Q)) /* one triple solution */
+		{
+			Solution[0] = 0;
+			NumSolutions = 1;
+		}
+		else /* one single and one double solution */
+		{
+			double u = cbrt(-Q);
+			Solution[0] = 2 * u;
+			Solution[1] = -u;
+			NumSolutions = 2;
+		}
+	}
+	else if (D < 0) /* Casus irreducibilis: three real solutions */
+	{
+		double phi = 1.0 / 3 * acos(-Q / sqrt(-CubeOfP));
+		double t = 2 * sqrt(-P);
+
+		Solution[0] = t * cos(phi);
+		Solution[1] = -t * cos(phi + PI / 3);
+		Solution[2] = -t * cos(phi - PI / 3);
+		NumSolutions = 3;
+	}
+	else /* one real solution */
+	{
+		double sqrt_D = sqrt(D);
+		double u = cbrt(sqrt_D - Q);
+		double v = -cbrt(sqrt_D + Q);
+
+		Solution[0] = u + v;
+		NumSolutions = 1;
+	}
+
+	/* resubstitute */
+
+	double Sub = 1.0 / 3 * A;
+
+	for (int i = 0; i < NumSolutions; ++i)
+		Solution[i] -= Sub;
+
+	return NumSolutions;
+}
+
+/*
+*   Convert the control values for a polynomial defined in the Bezier
+*		basis to a polynomial defined in the power basis (t^3 t^2 t 1).
+*/
+static void BezierToPower(	double A1, double B1, double C1, double D1,
+	double *A2, double *B2, double *C2, double *D2)
+{
+	double A = B1 - A1;
+	double B = C1 - B1;
+	double C = D1 - C1;
+	double D = B - A;
+	*A2 = C- B - D;
+	*B2 = 3.0 * D;
+	*C2 = 3.0 * A;
+	*D2 = A1;
+}
+
+bool FMovieSceneFloatChannel::Evaluate(FFrameTime InTime,  float& OutValue) const
 {
 	const int32 NumKeys = Times.Num();
 
@@ -351,14 +449,85 @@ bool FMovieSceneFloatChannel::Evaluate(FFrameTime InTime, float& OutValue) const
 		{
 		case RCIM_Cubic:
 		{
-			const float OneThird = 1.0f / 3.0f;
-			const int32 Diff = Times[Index2].Value - Times[Index1].Value;
-			const float P0 = Key1.Value;
-			const float P1 = P0 + (Key1.Tangent.LeaveTangent * Diff * OneThird);
-			const float P3 = Key2.Value;
-			const float P2 = P3 - (Key2.Tangent.ArriveTangent * Diff * OneThird);
-
-			OutValue = Params.ValueOffset + BezierInterp(P0, P1, P2, P3, Interp);
+			if ((Key1.Tangent.TangentWeightMode == RCTWM_WeightedNone || Key1.Tangent.TangentWeightMode == RCTWM_WeightedArrive)
+				&& (Key2.Tangent.TangentWeightMode == RCTWM_WeightedNone || Key2.Tangent.TangentWeightMode == RCTWM_WeightedLeave))
+			{
+				const float OneThird = 1.0f / 3.0f;
+				const int32 Diff = Times[Index2].Value - Times[Index1].Value;
+				const float P0 = Key1.Value;
+				const float P1 = P0 + (Key1.Tangent.LeaveTangent * Diff * OneThird);
+				const float P3 = Key2.Value;
+				const float P2 = P3 - (Key2.Tangent.ArriveTangent * Diff * OneThird);
+
+				OutValue = Params.ValueOffset + BezierInterp(P0, P1, P2, P3, Interp);
+				break;
+			}
+			else //its weighted
+			{
+				float ToSeconds = TickResolution.AsDecimal();
+
+				const double Time1 = TickResolution.AsSeconds(Times[Index1].Value);
+				const double Time2 = TickResolution.AsSeconds(Times[Index2].Value);
+				float CosAngle, SinAngle;
+				float Angle = FMath::Atan(Key1.Tangent.LeaveTangent * ToSeconds);
+				FMath::SinCos(&SinAngle, &CosAngle, Angle);
+				const float Key1TanX = CosAngle * Key1.Tangent.LeaveTangentWeight + Time1;
+				const float Key1TanY = SinAngle * Key1.Tangent.LeaveTangentWeight + Key1.Value;
+
+				Angle = FMath::Atan(Key2.Tangent.ArriveTangent * ToSeconds);
+				FMath::SinCos(&SinAngle, &CosAngle, Angle);
+				const float Key2TanX = -CosAngle * Key2.Tangent.ArriveTangentWeight + Time2;
+				const float Key2TanY = -SinAngle * Key2.Tangent.ArriveTangentWeight + Key2.Value;
+
+				//Normalize the Time Range
+				const float RangeX = Time2 - Time1;
+
+				const float Dx1 = Key1TanX - Time1;
+				const float Dx2 = Key2TanX - Time1;
+
+				// Normalize values
+				const float NormalizedX1 = Dx1 / RangeX;
+				const float NormalizedX2 = Dx2 / RangeX;
+				
+				double Coeff[4];
+				double Results[3];
+
+				//Convert Bezier to Power basis, also float to double for precision for root finding.
+				BezierToPower(
+					0.0, NormalizedX1, NormalizedX2, 1.0,
+					&(Coeff[3]), &(Coeff[2]), &(Coeff[1]), &(Coeff[0])
+				);
+
+				Coeff[0] = Coeff[0] - Interp;
+				
+				int NumResults = SolveCubic(Coeff, Results);
+				float NewInterp = Interp;
+				if (NumResults == 1)
+				{
+					NewInterp = Results[0];
+				}
+				else
+				{
+					NewInterp = TNumericLimits<float>::Lowest(); //just need to be out of range
+					for (double Result : Results)
+					{
+						if ((Result >= 0.0f) && (Result <= 1.0f))
+						{
+							if (NewInterp < 0.0f || Result > NewInterp)
+							{
+								NewInterp = Result;
+							}
+						}
+					}
+				}
+				//now use NewInterp and adjusted tangents plugged into the Y (Value) part of the graph.
+				const float P0 = Key1.Value;
+				const float P1 = Key1TanY;
+				const float P3 = Key2.Value;
+				const float P2 = Key2TanY;
+
+				OutValue = Params.ValueOffset + BezierInterp(P0, P1, P2, P3,  NewInterp);
+			}
 			break;
 		}
 
@@ -387,6 +556,7 @@ void FMovieSceneFloatChannel::AutoSetTangents(float Tension)
 		if (FirstValue.InterpMode == RCIM_Cubic && FirstValue.TangentMode == RCTM_Auto)
 		{
 			FirstValue.Tangent.LeaveTangent = FirstValue.Tangent.ArriveTangent = 0.0f;
+			FirstValue.Tangent.TangentWeightMode = RCTWM_WeightedNone; 
 		}
 	}
 
@@ -395,6 +565,7 @@ void FMovieSceneFloatChannel::AutoSetTangents(float Tension)
 		if (LastValue.InterpMode == RCIM_Cubic && LastValue.TangentMode == RCTM_Auto)
 		{
 			LastValue.Tangent.LeaveTangent = LastValue.Tangent.ArriveTangent = 0.0f;
+			LastValue.Tangent.TangentWeightMode = RCTWM_WeightedNone;
 		}
 	}
 
@@ -415,6 +586,7 @@ void FMovieSceneFloatChannel::AutoSetTangents(float Tension)
 
 			// In 'auto' mode, arrive and leave tangents are always the same
 			ThisKey.Tangent.LeaveTangent = ThisKey.Tangent.ArriveTangent = NewTangent;
+			ThisKey.Tangent.TangentWeightMode = RCTWM_WeightedNone;
 		}
 		else if (PrevKey.InterpMode == RCIM_Constant || ThisKey.InterpMode == RCIM_Constant)
 		{
@@ -428,10 +600,10 @@ void FMovieSceneFloatChannel::AutoSetTangents(float Tension)
 	}
 }
 
-void FMovieSceneFloatChannel::PopulateCurvePoints(double StartTimeSeconds, double EndTimeSeconds, double TimeThreshold, float ValueThreshold, FFrameRate FrameResolution, TArray<TTuple<double, double>>& InOutPoints) const
+void FMovieSceneFloatChannel::PopulateCurvePoints(double StartTimeSeconds, double EndTimeSeconds, double TimeThreshold, float ValueThreshold, FFrameRate InTickResolution, TArray<TTuple<double, double>>& InOutPoints) const
 {
-	const FFrameNumber StartFrame = (StartTimeSeconds * FrameResolution).FloorToFrame();
-	const FFrameNumber EndFrame   = (EndTimeSeconds   * FrameResolution).CeilToFrame();
+	const FFrameNumber StartFrame = (StartTimeSeconds * InTickResolution).FloorToFrame();
+	const FFrameNumber EndFrame   = (EndTimeSeconds   * InTickResolution).CeilToFrame();
 
 	const int32 StartingIndex = Algo::UpperBound(Times, StartFrame);
 	const int32 EndingIndex   = Algo::LowerBound(Times, EndFrame);
@@ -440,31 +612,31 @@ void FMovieSceneFloatChannel::PopulateCurvePoints(double StartTimeSeconds, doubl
 	float EvaluatedValue = 0.f;
 	if (Evaluate(StartFrame, EvaluatedValue))
 	{
-		InOutPoints.Add(MakeTuple(StartFrame / FrameResolution, double(EvaluatedValue)));
+		InOutPoints.Add(MakeTuple(StartFrame / InTickResolution, double(EvaluatedValue)));
 	}
 
 	// Add all keys in-between
 	for (int32 KeyIndex = StartingIndex; KeyIndex < EndingIndex; ++KeyIndex)
 	{
-		InOutPoints.Add(MakeTuple(Times[KeyIndex] / FrameResolution, double(Values[KeyIndex].Value)));
+		InOutPoints.Add(MakeTuple(Times[KeyIndex] / InTickResolution, double(Values[KeyIndex].Value)));
 	}
 
 	// Add the upper bound of the visible space
 	if (Evaluate(EndFrame, EvaluatedValue))
 	{
-		InOutPoints.Add(MakeTuple(EndFrame / FrameResolution, double(EvaluatedValue)));
+		InOutPoints.Add(MakeTuple(EndFrame / InTickResolution, double(EvaluatedValue)));
 	}
 
 	int32 OldSize = InOutPoints.Num();
 	do
 	{
 		OldSize = InOutPoints.Num();
-		RefineCurvePoints(FrameResolution, TimeThreshold, ValueThreshold, InOutPoints);
+		RefineCurvePoints(InTickResolution, TimeThreshold, ValueThreshold, InOutPoints);
 	}
 	while(OldSize != InOutPoints.Num());
 }
 
-void FMovieSceneFloatChannel::RefineCurvePoints(FFrameRate FrameResolution, double TimeThreshold, float ValueThreshold, TArray<TTuple<double, double>>& InOutPoints) const
+void FMovieSceneFloatChannel::RefineCurvePoints(FFrameRate InTickResolution, double TimeThreshold, float ValueThreshold, TArray<TTuple<double, double>>& InOutPoints) const
 {
 	const float InterpTimes[] = { 0.25f, 0.5f, 0.6f };
 
@@ -486,7 +658,7 @@ void FMovieSceneFloatChannel::RefineCurvePoints(FFrameRate FrameResolution, doub
 				EvalTime = FMath::Lerp(Lower.Get<0>(), Upper.Get<0>(), InterpTimes[InterpIndex]);
 
 				float Value = 0.f;
-				Evaluate(EvalTime * FrameResolution, Value);
+				Evaluate(EvalTime * InTickResolution, Value);
 
 				const float LinearValue = FMath::Lerp(Lower.Get<1>(), Upper.Get<1>(), InterpTimes[InterpIndex]);
 				if (bSegmentIsLinear)
@@ -555,7 +727,7 @@ FKeyHandle AddKeyToChannel(FMovieSceneFloatChannel* Channel, FFrameNumber InFram
 	return Channel->GetInterface().GetHandle(ExistingIndex);
 }
 
-void Optimize(FMovieSceneFloatChannel* InChannel, const FKeyDataOptimizationParams& Params)
+void Optimize(FMovieSceneFloatChannel* InChannel,  const FKeyDataOptimizationParams& Params)
 {
 	auto ChannelInterface = InChannel->GetInterface();
 	if (ChannelInterface.GetTimes().Num() > 1)
@@ -569,7 +741,7 @@ void Optimize(FMovieSceneFloatChannel* InChannel, const FKeyDataOptimizationPara
 			EndIndex   = Params.Range.GetUpperBound().IsClosed() ? Algo::UpperBound(Times, Params.Range.GetUpperBoundValue()) : Times.Num();
 		}
 
-		for (int32 Index = StartIndex; Index < EndIndex; ++Index)
+		for (int32 Index = StartIndex; Index < EndIndex && Index < ChannelInterface.GetTimes().Num(); ++Index)
 		{
 			// Reget times and values as they may be reallocated
 			FFrameNumber Time                   = ChannelInterface.GetTimes()[Index];
@@ -604,7 +776,7 @@ void Dilate(FMovieSceneFloatChannel* InChannel, FFrameNumber Origin, float Dilat
 	InChannel->AutoSetTangents();
 }
 
-void ChangeFrameResolution(FMovieSceneFloatChannel* InChannel, FFrameRate SourceRate, FFrameRate DestinationRate)
+void ChangeTickResolution(FMovieSceneFloatChannel* InChannel, FFrameRate SourceRate, FFrameRate DestinationRate)
 {
 	TArrayView<FFrameNumber>          Times  = InChannel->GetInterface().GetTimes();
 	TArrayView<FMovieSceneFloatValue> Values = InChannel->GetInterface().GetValues();
@@ -620,4 +792,5 @@ void ChangeFrameResolution(FMovieSceneFloatChannel* InChannel, FFrameRate Source
 		Value.Tangent.ArriveTangent *= IntervalFactor;
 		Value.Tangent.LeaveTangent  *= IntervalFactor;
 	}
+	InChannel->SetTickResolution(DestinationRate);
 }
\ No newline at end of file
diff --git a/Engine/Source/Runtime/MovieScene/Private/Compilation/MovieSceneCompiler.cpp b/Engine/Source/Runtime/MovieScene/Private/Compilation/MovieSceneCompiler.cpp
--- a/Engine/Source/Runtime/MovieScene/Private/Compilation/MovieSceneCompiler.cpp
+++ b/Engine/Source/Runtime/MovieScene/Private/Compilation/MovieSceneCompiler.cpp
@@ -523,10 +523,10 @@ void FMovieSceneCompiler::PopulateMetaData(FCompiledGroupResult& OutResult, cons
 
 			UMovieSceneSequence* Sequence = SubData->GetSequence();
 
-			FGuid TemplateID = Sequence ? TemplateStore.AccessTemplate(*Sequence).TemplateSignature : FGuid();
+			uint32 TemplateSerialNumber = Sequence ? TemplateStore.AccessTemplate(*Sequence).TemplateSerialNumber.GetValue() : 0;
 
 			OutResult.MetaData.ActiveSequences.Add(SequenceID);
-			OutResult.MetaData.SubTemplateSignatures.Add(SequenceID, TemplateID);
+			OutResult.MetaData.SubTemplateSerialNumbers.Add(SequenceID, TemplateSerialNumber);
 		}
 
 		OutResult.MetaData.ActiveSequences.Sort();
diff --git a/Engine/Source/Runtime/MovieScene/Private/Compilation/MovieSceneEvaluationTemplateGenerator.cpp b/Engine/Source/Runtime/MovieScene/Private/Compilation/MovieSceneEvaluationTemplateGenerator.cpp
--- a/Engine/Source/Runtime/MovieScene/Private/Compilation/MovieSceneEvaluationTemplateGenerator.cpp
+++ b/Engine/Source/Runtime/MovieScene/Private/Compilation/MovieSceneEvaluationTemplateGenerator.cpp
@@ -53,7 +53,7 @@ void FMovieSceneEvaluationTemplateGenerator::Generate()
 	Template.RemoveStaleData(CompiledSignatures);
 
 	Template.SequenceSignature = SourceSequence.GetSignature();
-	Template.TemplateSignature = FGuid::NewGuid();
+	Template.TemplateSerialNumber.Increment();
 }
 
 void FMovieSceneEvaluationTemplateGenerator::ProcessTrack(const UMovieSceneTrack& Track, const FGuid& ObjectBindingId)
diff --git a/Engine/Source/Runtime/MovieScene/Private/Evaluation/MovieSceneEvaluationField.cpp b/Engine/Source/Runtime/MovieScene/Private/Evaluation/MovieSceneEvaluationField.cpp
--- a/Engine/Source/Runtime/MovieScene/Private/Evaluation/MovieSceneEvaluationField.cpp
+++ b/Engine/Source/Runtime/MovieScene/Private/Evaluation/MovieSceneEvaluationField.cpp
@@ -254,7 +254,7 @@ bool FMovieSceneEvaluationMetaData::IsDirty(UMovieSceneSequence& RootSequence, c
 {
 	bool bDirty = false;
 
-	for (const TTuple<FMovieSceneSequenceID, FGuid>& Pair : SubTemplateSignatures)
+	for (const TTuple<FMovieSceneSequenceID, uint32>& Pair : SubTemplateSerialNumbers)
 	{
 		// Sequence IDs at this point are relative to the root override template
 		const FMovieSceneSubSequenceData* SubData = RootHierarchy.FindSubData(Pair.Key);
@@ -265,7 +265,7 @@ bool FMovieSceneEvaluationMetaData::IsDirty(UMovieSceneSequence& RootSequence, c
 		{
 			FMovieSceneEvaluationTemplate& Template = TemplateStore.AccessTemplate(*SubSequence);
 
-			bThisSequenceIsDirty = Template.TemplateSignature != Pair.Value || Template.SequenceSignature != SubSequence->GetSignature();
+			bThisSequenceIsDirty = Template.TemplateSerialNumber.GetValue() != Pair.Value || Template.SequenceSignature != SubSequence->GetSignature();
 		}
 
 		if (bThisSequenceIsDirty)
diff --git a/Engine/Source/Runtime/MovieScene/Private/Evaluation/MovieSceneEvaluationState.cpp b/Engine/Source/Runtime/MovieScene/Private/Evaluation/MovieSceneEvaluationState.cpp
--- a/Engine/Source/Runtime/MovieScene/Private/Evaluation/MovieSceneEvaluationState.cpp
+++ b/Engine/Source/Runtime/MovieScene/Private/Evaluation/MovieSceneEvaluationState.cpp
@@ -45,6 +45,23 @@ TArrayView<TWeakObjectPtr<>> FMovieSceneObjectCache::FindBoundObjects(const FGui
 	return TArrayView<TWeakObjectPtr<>>();
 }
 
+TArrayView<const TWeakObjectPtr<>> FMovieSceneObjectCache::IterateBoundObjects(const FGuid& InBindingID) const
+{
+	MOVIESCENE_DETAILED_SCOPE_CYCLE_COUNTER(MovieSceneEval_IterateBoundObjects)
+
+	const FBoundObjects* Bindings = BoundObjects.Find(InBindingID);
+	if (Bindings && Bindings->bUpToDate)
+	{
+		return TArrayView<const TWeakObjectPtr<>>(
+			Bindings->Objects.GetData(),
+			Bindings->Objects.Num()
+			);
+	}
+
+	// Just return nothing
+	return TArrayView<TWeakObjectPtr<>>();
+}
+
 FGuid FMovieSceneObjectCache::FindObjectId(UObject& InObject, IMovieScenePlayer& Player)
 {
 	UMovieSceneSequence* Sequence = WeakSequence.Get();
@@ -101,6 +118,17 @@ void FMovieSceneObjectCache::InvalidateExpiredObjects()
 			}
 		}
 	}
+
+	if (UMovieSceneSequence* Sequence = WeakSequence.Get())
+	{
+		TArray<FGuid> InvalidObjectIDs;
+		Sequence->GatherExpiredObjects(*this, InvalidObjectIDs);
+
+		for (const FGuid& ObjectID : InvalidObjectIDs)
+		{
+			Invalidate(ObjectID);
+		}
+	}
 }
 
 void FMovieSceneObjectCache::Invalidate(const FGuid& InGuid)
@@ -158,7 +186,7 @@ void FMovieSceneObjectCache::UpdateBindings(const FGuid& InGuid, IMovieScenePlay
 
 	// Find the sequence for the current
 	UMovieSceneSequence* Sequence = WeakSequence.Get();
-	if (!ensure(Sequence))
+	if (!Sequence)
 	{
 		return;
 	}
diff --git a/Engine/Source/Runtime/MovieScene/Private/Evaluation/MovieSceneEvaluationTemplateInstance.cpp b/Engine/Source/Runtime/MovieScene/Private/Evaluation/MovieSceneEvaluationTemplateInstance.cpp
--- a/Engine/Source/Runtime/MovieScene/Private/Evaluation/MovieSceneEvaluationTemplateInstance.cpp
+++ b/Engine/Source/Runtime/MovieScene/Private/Evaluation/MovieSceneEvaluationTemplateInstance.cpp
@@ -293,7 +293,7 @@ void FMovieSceneRootEvaluationTemplateInstance::EvaluateGroup(const FMovieSceneE
 				SubContext = Context;
 				if (Instance.SubData)
 				{
-					SubContext = Context.Transform(Instance.SubData->RootToSequenceTransform, Instance.SubData->FrameResolution);
+					SubContext = Context.Transform(Instance.SubData->RootToSequenceTransform, Instance.SubData->TickResolution);
 
 					// Hittest against the sequence's pre and postroll ranges
 					SubContext.ReportOuterSectionRanges(Instance.SubData->PreRollRange.Value, Instance.SubData->PostRollRange.Value);
@@ -335,7 +335,7 @@ void FMovieSceneRootEvaluationTemplateInstance::EvaluateGroup(const FMovieSceneE
 				SubContext = Context;
 				if (Instance.SubData)
 				{
-					SubContext = Context.Transform(Instance.SubData->RootToSequenceTransform, Instance.SubData->FrameResolution);
+					SubContext = Context.Transform(Instance.SubData->RootToSequenceTransform, Instance.SubData->TickResolution);
 
 					// Hittest against the sequence's pre and postroll ranges
 					SubContext.ReportOuterSectionRanges(Instance.SubData->PreRollRange.Value, Instance.SubData->PostRollRange.Value);
diff --git a/Engine/Source/Runtime/MovieScene/Private/Evaluation/MovieSceneSequenceHierarchy.cpp b/Engine/Source/Runtime/MovieScene/Private/Evaluation/MovieSceneSequenceHierarchy.cpp
--- a/Engine/Source/Runtime/MovieScene/Private/Evaluation/MovieSceneSequenceHierarchy.cpp
+++ b/Engine/Source/Runtime/MovieScene/Private/Evaluation/MovieSceneSequenceHierarchy.cpp
@@ -26,7 +26,7 @@ FMovieSceneSubSequenceData::FMovieSceneSubSequenceData(const UMovieSceneSubSecti
 
 	checkf(MovieScenePtr, TEXT("Attempting to construct sub sequence data with a null sequence."));
 
-	FrameResolution = MovieScenePtr->GetFrameResolution();
+	TickResolution = MovieScenePtr->GetTickResolution();
 
 	TRange<FFrameNumber> SubRange = InSubSection.GetRange();
 	checkf(SubRange.GetLowerBound().IsClosed() && SubRange.GetUpperBound().IsClosed(), TEXT("Use of open (infinite) bounds with sub sections is not supported."));
diff --git a/Engine/Source/Runtime/MovieScene/Private/MovieScene.cpp b/Engine/Source/Runtime/MovieScene/Private/MovieScene.cpp
--- a/Engine/Source/Runtime/MovieScene/Private/MovieScene.cpp
+++ b/Engine/Source/Runtime/MovieScene/Private/MovieScene.cpp
@@ -37,7 +37,7 @@ UMovieScene::UMovieScene(const FObjectInitializer& ObjectInitializer)
 
 	if (!HasAnyFlags(RF_ClassDefaultObject) && GetLinkerCustomVersion(FSequencerObjectVersion::GUID) < FSequencerObjectVersion::FloatToIntConversion)
 	{
-		FrameResolution = GetLegacyConversionFrameRate();
+		TickResolution = GetLegacyConversionFrameRate();
 	}
 
 #if WITH_EDITORONLY_DATA
@@ -85,20 +85,20 @@ void UMovieScene::Serialize( FArchive& Ar )
 		}
 
 		// Legacy fixed frame interval conversion to integer play rates
-		if      (FixedFrameInterval_DEPRECATED == 1 / 15.0f)  { PlayRate = FCommonFrameRates::FPS_15();  }
-		else if (FixedFrameInterval_DEPRECATED == 1 / 24.0f)  { PlayRate = FCommonFrameRates::FPS_24();  }
-		else if (FixedFrameInterval_DEPRECATED == 1 / 25.0f)  { PlayRate = FCommonFrameRates::FPS_25();  }
-		else if (FixedFrameInterval_DEPRECATED == 1 / 29.97)  { PlayRate = FCommonFrameRates::NTSC_30(); }
-		else if (FixedFrameInterval_DEPRECATED == 1 / 30.0f)  { PlayRate = FCommonFrameRates::FPS_30();  }
-		else if (FixedFrameInterval_DEPRECATED == 1 / 48.0f)  { PlayRate = FCommonFrameRates::FPS_48();  }
-		else if (FixedFrameInterval_DEPRECATED == 1 / 50.0f)  { PlayRate = FCommonFrameRates::FPS_50();  }
-		else if (FixedFrameInterval_DEPRECATED == 1 / 59.94)  { PlayRate = FCommonFrameRates::NTSC_60(); }
-		else if (FixedFrameInterval_DEPRECATED == 1 / 60.0f)  { PlayRate = FCommonFrameRates::FPS_60();  }
-		else if (FixedFrameInterval_DEPRECATED == 1 / 120.0f) { PlayRate = FCommonFrameRates::FPS_120(); }
+		if      (FixedFrameInterval_DEPRECATED == 1 / 15.0f)  { DisplayRate = FCommonFrameRates::FPS_15();  }
+		else if (FixedFrameInterval_DEPRECATED == 1 / 24.0f)  { DisplayRate = FCommonFrameRates::FPS_24();  }
+		else if (FixedFrameInterval_DEPRECATED == 1 / 25.0f)  { DisplayRate = FCommonFrameRates::FPS_25();  }
+		else if (FixedFrameInterval_DEPRECATED == 1 / 29.97)  { DisplayRate = FCommonFrameRates::NTSC_30(); }
+		else if (FixedFrameInterval_DEPRECATED == 1 / 30.0f)  { DisplayRate = FCommonFrameRates::FPS_30();  }
+		else if (FixedFrameInterval_DEPRECATED == 1 / 48.0f)  { DisplayRate = FCommonFrameRates::FPS_48();  }
+		else if (FixedFrameInterval_DEPRECATED == 1 / 50.0f)  { DisplayRate = FCommonFrameRates::FPS_50();  }
+		else if (FixedFrameInterval_DEPRECATED == 1 / 59.94)  { DisplayRate = FCommonFrameRates::NTSC_60(); }
+		else if (FixedFrameInterval_DEPRECATED == 1 / 60.0f)  { DisplayRate = FCommonFrameRates::FPS_60();  }
+		else if (FixedFrameInterval_DEPRECATED == 1 / 120.0f) { DisplayRate = FCommonFrameRates::FPS_120(); }
 		else if (FixedFrameInterval_DEPRECATED != 0.f)
 		{
 			uint32 Numerator = FMath::RoundToInt(1.f / FixedFrameInterval_DEPRECATED);
-			PlayRate = FFrameRate(Numerator, 1);
+			DisplayRate = FFrameRate(Numerator, 1);
 		}
 		else
 		{
@@ -108,17 +108,17 @@ void UMovieScene::Serialize( FArchive& Ar )
 			if (Outer && Outer->GetClass()->GetFName() == "WidgetAnimation")
 			{
 				// Widget animations defaulted to 0.05s
-				PlayRate = FFrameRate(20, 1);
+				DisplayRate = FFrameRate(20, 1);
 			}
 			else if (Outer && Outer->GetClass()->GetFName() == "ActorSequence")
 			{
 				// Actor sequences defaulted to 0.1s
-				PlayRate = FFrameRate(10, 1);
+				DisplayRate = FFrameRate(10, 1);
 			}
 			else
 			{
 				// Level sequences defaulted to 30fps - this is the fallback default for anything else
-				PlayRate = FFrameRate(30, 1);
+				DisplayRate = FFrameRate(30, 1);
 			}
 		}
 	}
@@ -447,8 +447,8 @@ void UMovieScene::SetPlaybackRange(const TRange<FFrameNumber>& NewRange, bool bA
 
 #if WITH_EDITORONLY_DATA
 	// Update the working and view ranges to encompass the new range
-	const double RangeStartSeconds = NewRange.GetLowerBoundValue() / FrameResolution;
-	const double RangeEndSeconds   = NewRange.GetUpperBoundValue() / FrameResolution;
+	const double RangeStartSeconds = NewRange.GetLowerBoundValue() / TickResolution;
+	const double RangeEndSeconds   = NewRange.GetUpperBoundValue() / TickResolution;
 
 	// Initialize the working and view range with a little bit more space
 	const double OutputChange      = (RangeEndSeconds - RangeStartSeconds) * 0.1;
@@ -742,9 +742,9 @@ void UMovieScene::UpgradeTimeRanges()
 	{
 		// Finite range already defined in old data
 		PlaybackRange.Value = TRange<FFrameNumber>(
-			FrameResolution.AsFrameNumber(InTime_DEPRECATED),
+			TickResolution.AsFrameNumber(InTime_DEPRECATED),
 			// Prefer exclusive upper bounds for playback ranges so we stop at the next frame
-			++FrameResolution.AsFrameNumber(OutTime_DEPRECATED)
+			++TickResolution.AsFrameNumber(OutTime_DEPRECATED)
 			);
 		bFiniteRangeDefined = true;
 	}
@@ -814,14 +814,14 @@ void UMovieScene::UpgradeTimeRanges()
 	}
 	else if (EditorData.WorkStart >= EditorData.WorkEnd)
 	{
-		EditorData.WorkStart = PlaybackRange.Value.GetLowerBoundValue() / FrameResolution;
-		EditorData.WorkEnd   = PlaybackRange.Value.GetUpperBoundValue() / FrameResolution;
+		EditorData.WorkStart = PlaybackRange.Value.GetLowerBoundValue() / TickResolution;
+		EditorData.WorkEnd   = PlaybackRange.Value.GetUpperBoundValue() / TickResolution;
 	}
 
 	if (EditorData.ViewStart >= EditorData.ViewEnd)
 	{
-		EditorData.ViewStart = PlaybackRange.Value.GetLowerBoundValue() / FrameResolution;
-		EditorData.ViewEnd   = PlaybackRange.Value.GetUpperBoundValue() / FrameResolution;
+		EditorData.ViewStart = PlaybackRange.Value.GetLowerBoundValue() / TickResolution;
+		EditorData.ViewEnd   = PlaybackRange.Value.GetUpperBoundValue() / TickResolution;
 	}
 
 	if (SelectionRange.Value.GetLowerBound().IsOpen() || SelectionRange.Value.GetUpperBound().IsOpen())
@@ -945,3 +945,54 @@ void UMovieScene::ReplaceBinding(const FGuid& OldGuid, const FGuid& NewGuid, con
 		}
 	}
 }
+
+
+void UMovieScene::MoveBindingContents(const FGuid& SourceBindingId, const FGuid& DestinationBindingId)
+{
+	FMovieSceneBinding* SourceBinding = nullptr;
+	FMovieSceneBinding* DestinationBinding = nullptr;
+
+	for (FMovieSceneBinding& Binding : ObjectBindings)
+	{
+		if (Binding.GetObjectGuid() == SourceBindingId)
+		{
+			SourceBinding = &Binding;
+		}
+		else if (Binding.GetObjectGuid() == DestinationBindingId)
+		{
+			DestinationBinding = &Binding;
+		}
+
+		if (SourceBinding && DestinationBinding)
+		{
+			break;
+		}
+	}
+
+	if (SourceBinding && DestinationBinding)
+	{
+		// Swap the tracks round
+		DestinationBinding->SetTracks(SourceBinding->StealTracks());
+
+		// Changing a binding guid invalidates any tracks contained within the binding
+		// Make sure they are written into the transaction buffer by calling modify
+		for (UMovieSceneTrack* Track : DestinationBinding->GetTracks())
+		{
+			Track->Modify();
+		}
+	}
+
+	FMovieSceneSpawnable* DestinationSpawnable = FindSpawnable(DestinationBindingId);
+
+	for (FMovieScenePossessable& Possessable : Possessables)
+	{
+		if (Possessable.GetParent() == SourceBindingId)
+		{
+			Possessable.SetParent(DestinationBindingId);
+			if (DestinationSpawnable)
+			{
+				DestinationSpawnable->AddChildPossessable(Possessable.GetGuid());
+			}
+		}
+	}
+}
diff --git a/Engine/Source/Runtime/MovieScene/Private/MovieSceneCommonHelpers.cpp b/Engine/Source/Runtime/MovieScene/Private/MovieSceneCommonHelpers.cpp
--- a/Engine/Source/Runtime/MovieScene/Private/MovieSceneCommonHelpers.cpp
+++ b/Engine/Source/Runtime/MovieScene/Private/MovieSceneCommonHelpers.cpp
@@ -7,8 +7,9 @@
 #include "KeyParams.h"
 #include "MovieSceneSection.h"
 #include "Algo/Sort.h"
-
-
+#include "Sound/SoundWave.h"
+#include "Sound/SoundCue.h"
+#include "Sound/SoundNodeWavePlayer.h"
 
 UMovieSceneSection* MovieSceneHelpers::FindSectionAtTime( const TArray<UMovieSceneSection*>& Sections, FFrameNumber Time )
 {
@@ -200,6 +201,36 @@ UCameraComponent* MovieSceneHelpers::CameraComponentFromRuntimeObject(UObject* R
 	return nullptr;
 }
 
+float MovieSceneHelpers::GetSoundDuration(USoundBase* Sound)
+{
+	USoundWave* SoundWave = nullptr;
+
+	if (Sound->IsA<USoundWave>())
+	{
+		SoundWave = Cast<USoundWave>(Sound);
+	}
+	else if (Sound->IsA<USoundCue>())
+	{
+#if WITH_EDITORONLY_DATA
+		USoundCue* SoundCue = Cast<USoundCue>(Sound);
+
+		// @todo Sequencer - Right now for sound cues, we just use the first sound wave in the cue
+		// In the future, it would be better to properly generate the sound cue's data after forcing determinism
+		const TArray<USoundNode*>& AllNodes = SoundCue->AllNodes;
+		for (int32 i = 0; i < AllNodes.Num() && SoundWave == nullptr; ++i)
+		{
+			if (AllNodes[i]->IsA<USoundNodeWavePlayer>())
+			{
+				SoundWave = Cast<USoundNodeWavePlayer>(AllNodes[i])->GetSoundWave();
+			}
+		}
+#endif
+	}
+
+	const float Duration = (SoundWave ? SoundWave->GetDuration() : 0.f);
+	return Duration == INDEFINITELY_LOOPING_DURATION ? SoundWave->Duration : Duration;
+}
+
 FTrackInstancePropertyBindings::FTrackInstancePropertyBindings( FName InPropertyName, const FString& InPropertyPath, const FName& InFunctionName, const FName& InNotifyFunctionName )
     : PropertyPath( InPropertyPath )
 	, NotifyFunctionName(InNotifyFunctionName)
diff --git a/Engine/Source/Runtime/MovieScene/Private/MovieSceneModule.cpp b/Engine/Source/Runtime/MovieScene/Private/MovieSceneModule.cpp
--- a/Engine/Source/Runtime/MovieScene/Private/MovieSceneModule.cpp
+++ b/Engine/Source/Runtime/MovieScene/Private/MovieSceneModule.cpp
@@ -14,7 +14,7 @@ DEFINE_LOG_CATEGORY(LogMovieScene);
 TAutoConsoleVariable<FString> CVarLegacyConversionFrameRate(
 	TEXT("MovieScene.LegacyConversionFrameRate"),
 	TEXT("60000fps"),
-	TEXT("Specifies default frame resolution for UMovieScene data saved before 4.20 (default: 60000fps). Examples: 60000 fps, 120/1 (120 fps), 30000/1001 (29.97), 0.01s (10ms)."),
+	TEXT("Specifies default tick resolution for UMovieScene data saved before 4.20 (default: 60000fps). Examples: 60000 fps, 120/1 (120 fps), 30000/1001 (29.97), 0.01s (10ms)."),
 	ECVF_Default);
 
 struct FCachedLegacyConversionFrameRate
diff --git a/Engine/Source/Runtime/MovieScene/Private/MovieSceneSection.cpp b/Engine/Source/Runtime/MovieScene/Private/MovieSceneSection.cpp
--- a/Engine/Source/Runtime/MovieScene/Private/MovieSceneSection.cpp
+++ b/Engine/Source/Runtime/MovieScene/Private/MovieSceneSection.cpp
@@ -184,6 +184,27 @@ TRange<FFrameNumber> UMovieSceneSection::ComputeEffectiveRange() const
 	return TRange<FFrameNumber>::Intersection(EffectiveRange, SectionRange.Value);
 }
 
+
+TOptional<TRange<FFrameNumber> > UMovieSceneSection::GetAutoSizeRange() const
+{
+	if (ChannelProxy.IsValid())
+	{
+		TRange<FFrameNumber> EffectiveRange = TRange<FFrameNumber>::Empty();
+	
+		for (const FMovieSceneChannelEntry& Entry : ChannelProxy->GetAllEntries())
+		{
+			EffectiveRange = TRange<FFrameNumber>::Hull(EffectiveRange, Entry.GetBatchChannelInterface().ComputeEffectiveRange_Batch(Entry.GetChannels()));
+		}
+
+		if (!EffectiveRange.IsEmpty())
+		{
+			return EffectiveRange;
+		}
+	}
+
+	return TOptional<TRange<FFrameNumber> >();
+}
+
 FMovieSceneBlendTypeField UMovieSceneSection::GetSupportedBlendTypes() const
 {
 	UMovieSceneTrack* Track = GetTypedOuter<UMovieSceneTrack>();
@@ -367,9 +388,9 @@ void UMovieSceneSection::InitialPlacementOnRow(const TArray<UMovieSceneSection*>
 	}
 }
 
-UMovieSceneSection* UMovieSceneSection::SplitSection(FFrameNumber SplitTime)
+UMovieSceneSection* UMovieSceneSection::SplitSection(FQualifiedFrameTime SplitTime)
 {
-	if (!SectionRange.Value.Contains(SplitTime))
+	if (!SectionRange.Value.Contains(SplitTime.Time.GetFrame()))
 	{
 		return nullptr;
 	}
@@ -379,8 +400,8 @@ UMovieSceneSection* UMovieSceneSection::SplitSection(FFrameNumber SplitTime)
 	if (TryModify())
 	{
 		TRange<FFrameNumber> StartingRange  = SectionRange.Value;
-		TRange<FFrameNumber> LeftHandRange  = TRange<FFrameNumber>(StartingRange.GetLowerBound(), TRangeBound<FFrameNumber>::Exclusive(SplitTime));
-		TRange<FFrameNumber> RightHandRange = TRange<FFrameNumber>(TRangeBound<FFrameNumber>::Inclusive(SplitTime), StartingRange.GetUpperBound());
+		TRange<FFrameNumber> LeftHandRange  = TRange<FFrameNumber>(StartingRange.GetLowerBound(), TRangeBound<FFrameNumber>::Exclusive(SplitTime.Time.GetFrame()));
+		TRange<FFrameNumber> RightHandRange = TRange<FFrameNumber>(TRangeBound<FFrameNumber>::Inclusive(SplitTime.Time.GetFrame()), StartingRange.GetUpperBound());
 
 		// Trim off the right
 		SectionRange = LeftHandRange;
@@ -402,20 +423,20 @@ UMovieSceneSection* UMovieSceneSection::SplitSection(FFrameNumber SplitTime)
 }
 
 
-void UMovieSceneSection::TrimSection(FFrameNumber TrimTime, bool bTrimLeft)
+void UMovieSceneSection::TrimSection(FQualifiedFrameTime TrimTime, bool bTrimLeft)
 {
-	if (SectionRange.Value.Contains(TrimTime))
+	if (SectionRange.Value.Contains(TrimTime.Time.GetFrame()))
 	{
 		SetFlags(RF_Transactional);
 		if (TryModify())
 		{
 			if (bTrimLeft)
 			{
-				SectionRange.Value.SetLowerBound(TRangeBound<FFrameNumber>::Inclusive(TrimTime));
+				SectionRange.Value.SetLowerBound(TRangeBound<FFrameNumber>::Inclusive(TrimTime.Time.GetFrame()));
 			}
 			else
 			{
-				SectionRange.Value.SetUpperBound(TRangeBound<FFrameNumber>::Exclusive(TrimTime));
+				SectionRange.Value.SetUpperBound(TRangeBound<FFrameNumber>::Exclusive(TrimTime.Time.GetFrame()));
 			}
 		}
 	}
