Copying //UE4/Dev-Sequencer to //UE4/Dev-Main (Source: //UE4/Dev-Sequencer @ 4044339)

diff --git a/Engine/Source/Runtime/MovieScene/Private/MovieSceneSequencePlayer.cpp b/Engine/Source/Runtime/MovieScene/Private/MovieSceneSequencePlayer.cpp
--- a/Engine/Source/Runtime/MovieScene/Private/MovieSceneSequencePlayer.cpp
+++ b/Engine/Source/Runtime/MovieScene/Private/MovieSceneSequencePlayer.cpp
@@ -120,7 +120,7 @@ void UMovieSceneSequencePlayer::PlayInternal()
 
 		if (!PlayPosition.GetLastPlayEvalPostition().IsSet() || PlayPosition.GetLastPlayEvalPostition() != PlayPosition.GetCurrentPosition())
 		{
-			UpdateMovieSceneInstance(PlayPosition.JumpTo(StartTime), Status);
+			UpdateMovieSceneInstance(PlayPosition.PlayTo(PlayPosition.GetCurrentPosition()), EMovieScenePlayerStatus::Playing);
 		}
 
 		if (MovieSceneSequence)
@@ -249,7 +249,7 @@ void UMovieSceneSequencePlayer::Stop()
 
 void UMovieSceneSequencePlayer::GoToEndAndStop()
 {
-	JumpToFrame(GetFrameDuration()-1);
+	JumpToFrame(GetLastValidTime());
 	Stop();
 }
 
@@ -274,9 +274,9 @@ void UMovieSceneSequencePlayer::SetFrameRate(FFrameRate FrameRate)
 	UMovieScene* MovieScene = Sequence ? Sequence->GetMovieScene() : nullptr;
 	if (MovieScene)
 	{
-		if (MovieScene->GetEvaluationType() == EMovieSceneEvaluationType::FrameLocked && !FrameRate.IsMultipleOf(MovieScene->GetFrameResolution()))
+		if (MovieScene->GetEvaluationType() == EMovieSceneEvaluationType::FrameLocked && !FrameRate.IsMultipleOf(MovieScene->GetTickResolution()))
 		{
-			UE_LOG(LogMovieScene, Warning, TEXT("Attempting to play back a sequence with resolution of %f ticks per second frame locked to %f fps, which is not a multiple of the resolution."), MovieScene->GetFrameResolution().AsDecimal(), FrameRate.AsDecimal());
+			UE_LOG(LogMovieScene, Warning, TEXT("Attempting to play back a sequence with tick resolution of %f ticks per second frame locked to %f fps, which is not a multiple of the resolution."), MovieScene->GetTickResolution().AsDecimal(), FrameRate.AsDecimal());
 		}
 	}
 
@@ -372,6 +372,11 @@ bool UMovieSceneSequencePlayer::IsPaused() const
 	return Status == EMovieScenePlayerStatus::Paused;
 }
 
+bool UMovieSceneSequencePlayer::IsReversed() const
+{
+	return bReversePlayback;
+}
+
 float UMovieSceneSequencePlayer::GetLength() const
 {
 	return GetDuration().AsSeconds();
@@ -428,14 +433,14 @@ void UMovieSceneSequencePlayer::Initialize(UMovieSceneSequence* InSequence, cons
 	if (MovieScene)
 	{
 		EMovieSceneEvaluationType EvaluationType    = MovieScene->GetEvaluationType();
-		FFrameRate                FrameResolution   = MovieScene->GetFrameResolution();
-		FFrameRate                PlaybackFrameRate = MovieScene->GetPlaybackFrameRate();
+		FFrameRate                TickResolution    = MovieScene->GetTickResolution();
+		FFrameRate                DisplayRate       = MovieScene->GetDisplayRate();
 
-		UE_LOG(LogMovieScene, Verbose, TEXT("Initialize - MovieSceneSequence: %s, FrameResolution: %f, PlaybackFrameRate: %d, CurrentTime: %d"), *InSequence->GetName(), FrameResolution.Numerator, PlaybackFrameRate.Numerator);
+		UE_LOG(LogMovieScene, Verbose, TEXT("Initialize - MovieSceneSequence: %s, TickResolution: %f, DisplayRate: %d, CurrentTime: %d"), *InSequence->GetName(), TickResolution.Numerator, DisplayRate.Numerator);
 
-		// We set the play position in terms of the playback frame rate,
-		// but want evaluation ranges in the moviescene's frame resolution
-		PlayPosition.SetTimeBase(PlaybackFrameRate, FrameResolution, EvaluationType);
+		// We set the play position in terms of the display rate,
+		// but want evaluation ranges in the moviescene's tick resolution
+		PlayPosition.SetTimeBase(DisplayRate, TickResolution, EvaluationType);
 
 		{
 			// Set up the default frame range from the sequence's play range
@@ -444,14 +449,14 @@ void UMovieSceneSequencePlayer::Initialize(UMovieSceneSequence* InSequence, cons
 			const FFrameNumber SrcStartFrame = MovieScene::DiscreteInclusiveLower(PlaybackRange);
 			const FFrameNumber SrcEndFrame   = MovieScene::DiscreteExclusiveUpper(PlaybackRange);
 
-			const FFrameNumber StartingFrame = ConvertFrameTime(SrcStartFrame, FrameResolution, PlaybackFrameRate).FloorToFrame();
-			const FFrameNumber EndingFrame   = ConvertFrameTime(SrcEndFrame,   FrameResolution, PlaybackFrameRate).FloorToFrame();
+			const FFrameNumber StartingFrame = ConvertFrameTime(SrcStartFrame, TickResolution, DisplayRate).FloorToFrame();
+			const FFrameNumber EndingFrame   = ConvertFrameTime(SrcEndFrame,   TickResolution, DisplayRate).FloorToFrame();
 
 			SetFrameRange(StartingFrame.Value, (EndingFrame - StartingFrame).Value);
 		}
 
 		// Reset the play position based on the user-specified start offset, or a random time
-		FFrameTime SpecifiedStartOffset = PlaybackSettings.StartTime * PlaybackFrameRate;
+		FFrameTime SpecifiedStartOffset = PlaybackSettings.StartTime * DisplayRate;
 
 		// Setup the starting time
 		FFrameTime StartingTimeOffset = PlaybackSettings.bRandomStartTime
@@ -511,9 +516,30 @@ void UMovieSceneSequencePlayer::UpdateTimeCursorPosition(FFrameTime NewPosition,
 	}
 }
 
+EMovieScenePlayerStatus::Type UpdateMethodToStatus(EUpdatePositionMethod Method)
+{
+	switch(Method)
+	{
+	case EUpdatePositionMethod::Scrub: return EMovieScenePlayerStatus::Scrubbing;
+	case EUpdatePositionMethod::Jump:  return EMovieScenePlayerStatus::Stopped;
+	case EUpdatePositionMethod::Play:  return EMovieScenePlayerStatus::Playing;
+	default:                           return EMovieScenePlayerStatus::Stopped;
+	}
+}
+
+FMovieSceneEvaluationRange UpdatePlayPosition(FMovieScenePlaybackPosition& InOutPlayPosition, FFrameTime NewTime, EUpdatePositionMethod Method)
+{
+	if (Method == EUpdatePositionMethod::Play)
+	{
+		return InOutPlayPosition.PlayTo(NewTime);
+	}
+
+	return InOutPlayPosition.JumpTo(NewTime);
+}
+
 void UMovieSceneSequencePlayer::UpdateTimeCursorPosition_Internal(FFrameTime NewPosition, EUpdatePositionMethod Method)
 {
-	EMovieScenePlayerStatus::Type StatusOverride = Method == EUpdatePositionMethod::Scrub ? EMovieScenePlayerStatus::Scrubbing : (EMovieScenePlayerStatus::Type)Status;
+	EMovieScenePlayerStatus::Type StatusOverride = UpdateMethodToStatus(Method);
 
 	const int32 Duration = DurationFrames;
 	if (Duration == 0)
@@ -524,7 +550,7 @@ void UMovieSceneSequencePlayer::UpdateTimeCursorPosition_Internal(FFrameTime New
 
 	FFrameTime PositionRelativeToStart = NewPosition.FrameNumber - StartTime;
 
-	if (ShouldStopOrLoop(NewPosition))
+	if (Method == EUpdatePositionMethod::Play && ShouldStopOrLoop(NewPosition))
 	{
 		const int32 NumTimesLooped    = FMath::Abs(PositionRelativeToStart.FrameNumber.Value / Duration);
 		const bool  bLoopIndefinitely = PlaybackSettings.LoopCount < 0;
@@ -563,7 +589,7 @@ void UMovieSceneSequencePlayer::UpdateTimeCursorPosition_Internal(FFrameTime New
 			// Clamp the position to the duration
 			NewPosition = FMath::Clamp(NewPosition, FFrameTime(StartTime), GetLastValidTime());
 
-			FMovieSceneEvaluationRange Range = PlayPosition.PlayTo(NewPosition);
+			FMovieSceneEvaluationRange Range = UpdatePlayPosition(PlayPosition, NewPosition, Method);
 			UpdateMovieSceneInstance(Range, StatusOverride);
 
 			if (PlaybackSettings.bPauseAtEnd)
@@ -589,7 +615,7 @@ void UMovieSceneSequencePlayer::UpdateTimeCursorPosition_Internal(FFrameTime New
 	else
 	{
 		// Just update the time and sequence
-		FMovieSceneEvaluationRange Range = PlayPosition.PlayTo(NewPosition);
+		FMovieSceneEvaluationRange Range = UpdatePlayPosition(PlayPosition, NewPosition, Method);
 		UpdateMovieSceneInstance(Range, StatusOverride);
 	}
 }
diff --git a/Engine/Source/Runtime/MovieScene/Private/MovieSceneSpawnRegister.cpp b/Engine/Source/Runtime/MovieScene/Private/MovieSceneSpawnRegister.cpp
--- a/Engine/Source/Runtime/MovieScene/Private/MovieSceneSpawnRegister.cpp
+++ b/Engine/Source/Runtime/MovieScene/Private/MovieSceneSpawnRegister.cpp
@@ -38,7 +38,7 @@ UObject* FMovieSceneSpawnRegister::SpawnObject(const FGuid& BindingId, UMovieSce
 	}
 
 	return SpawnedActor;
-};
+}
 
 bool FMovieSceneSpawnRegister::DestroySpawnedObject(const FGuid& BindingId, FMovieSceneSequenceIDRef TemplateID, IMovieScenePlayer& Player)
 {
diff --git a/Engine/Source/Runtime/MovieScene/Private/MovieSceneSpawnable.cpp b/Engine/Source/Runtime/MovieScene/Private/MovieSceneSpawnable.cpp
--- a/Engine/Source/Runtime/MovieScene/Private/MovieSceneSpawnable.cpp
+++ b/Engine/Source/Runtime/MovieScene/Private/MovieSceneSpawnable.cpp
@@ -43,7 +43,5 @@ void FMovieSceneSpawnable::CopyObjectTemplate(UObject& InSourceObject, UMovieSce
 	check(ObjectTemplate);
 
 	MarkSpawnableTemplate(*ObjectTemplate);
-
-	// @todo: this will mark the package as dirty whenever a spawnable is destroyed. We should diff the duplicated object, and only mark dirty where the spawnable has actually changed
 	MovieSceneSequence.MarkPackageDirty();
 }
diff --git a/Engine/Source/Runtime/MovieScene/Private/MovieSceneTrack.cpp b/Engine/Source/Runtime/MovieScene/Private/MovieSceneTrack.cpp
--- a/Engine/Source/Runtime/MovieScene/Private/MovieSceneTrack.cpp
+++ b/Engine/Source/Runtime/MovieScene/Private/MovieSceneTrack.cpp
@@ -81,6 +81,14 @@ void UMovieSceneTrack::UpdateEasing()
 				UMovieSceneSection* Other = RowSections[OtherIndex];
 				TRange<FFrameNumber> OtherSectionRange = Other->GetRange();
 
+				if (!OtherSectionRange.HasLowerBound() && !OtherSectionRange.HasUpperBound())
+				{
+					// If we're testing against an infinite range we want to use the PlayRange of the sequence
+					// instead so that blends stop at the end of a clip instead of a quarter of the length.
+					UMovieScene* OuterScene = OuterTrack->GetTypedOuter<UMovieScene>();
+					OtherSectionRange = OuterScene->GetPlaybackRange();
+				}
+
 				bIsEntirelyUnderlapped = OtherSectionRange.Contains(CurrentSectionRange);
 
 				// Check the lower bound of the current section against the other section's upper bound
@@ -94,7 +102,7 @@ void UMovieSceneTrack::UpdateEasing()
 
 				if (bSectionRangeContainsOtherLowerBound &&!bSectionRangeContainsOtherUpperBound)
 				{
-					const int32 Difference = MovieScene::DiscreteSize(TRange<FFrameNumber>(OtherSectionRange.GetLowerBoundValue(), CurrentSectionRange.GetUpperBoundValue()));
+					const int32 Difference = MovieScene::DiscreteSize(TRange<FFrameNumber>(OtherSectionRange.GetLowerBound(), CurrentSectionRange.GetUpperBound()));
 					MaxEaseOut = FMath::Max(MaxEaseOut, Difference);
 				}
 			}
diff --git a/Engine/Source/Runtime/MovieScene/Private/Sections/MovieSceneSubSection.cpp b/Engine/Source/Runtime/MovieScene/Private/Sections/MovieSceneSubSection.cpp
--- a/Engine/Source/Runtime/MovieScene/Private/Sections/MovieSceneSubSection.cpp
+++ b/Engine/Source/Runtime/MovieScene/Private/Sections/MovieSceneSubSection.cpp
@@ -6,6 +6,7 @@
 #include "MovieSceneSequence.h"
 #include "MovieSceneTimeHelpers.h"
 #include "Evaluation/MovieSceneEvaluationTemplate.h"
+#include <FrameRate.h>
 
 TWeakObjectPtr<UMovieSceneSubSection> UMovieSceneSubSection::TheRecordingSection;
 
@@ -24,6 +25,11 @@ UMovieSceneSubSection::UMovieSceneSubSection()
 FMovieSceneSequenceTransform UMovieSceneSubSection::OuterToInnerTransform() const
 {
 	UMovieSceneSequence* SequencePtr   = GetSequence();
+	if (!SequencePtr)
+	{
+		return FMovieSceneSequenceTransform();
+	}
+
 	UMovieScene*         MovieScenePtr = SequencePtr->GetMovieScene();
 
 	TRange<FFrameNumber> SubRange = GetRange();
@@ -35,8 +41,8 @@ FMovieSceneSequenceTransform UMovieSceneSubSection::OuterToInnerTransform() cons
 	const FFrameNumber InnerStartTime = MovieScene::DiscreteInclusiveLower(MovieScenePtr->GetPlaybackRange()) + Parameters.GetStartFrameOffset();
 	const FFrameNumber OuterStartTime = MovieScene::DiscreteInclusiveLower(SubRange);
 
-	const FFrameRate   InnerFrameRate = MovieScenePtr->GetFrameResolution();
-	const FFrameRate   OuterFrameRate = GetTypedOuter<UMovieScene>()->GetFrameResolution();
+	const FFrameRate   InnerFrameRate = MovieScenePtr->GetTickResolution();
+	const FFrameRate   OuterFrameRate = GetTypedOuter<UMovieScene>()->GetTickResolution();
 
 	const float        FrameRateScale = (OuterFrameRate == InnerFrameRate) ? 1.f : (InnerFrameRate / OuterFrameRate).AsDecimal();
 
@@ -196,22 +202,28 @@ void UMovieSceneSubSection::PostEditChangeProperty(FPropertyChangedEvent& Proper
 }
 #endif
 
-UMovieSceneSection* UMovieSceneSubSection::SplitSection( FFrameNumber SplitTime )
+UMovieSceneSection* UMovieSceneSubSection::SplitSection( FQualifiedFrameTime SplitTime )
 {
+	// GetRange is in owning sequence resolution so we check against the incoming SplitTime without converting it.
 	TRange<FFrameNumber> InitialRange = GetRange();
-	if ( !InitialRange.Contains( SplitTime ) )
+	if ( !InitialRange.Contains(SplitTime.Time.FrameNumber) )
 	{
 		return nullptr;
 	}
 
 	FFrameNumber InitialStartOffset = Parameters.GetStartFrameOffset();
 
-	UMovieSceneSubSection* NewSection = Cast<UMovieSceneSubSection>( UMovieSceneSection::SplitSection( SplitTime ) );
+	UMovieSceneSubSection* NewSection = Cast<UMovieSceneSubSection>( UMovieSceneSection::SplitSection(SplitTime));
 	if ( NewSection )
 	{
 		if (InitialRange.GetLowerBound().IsClosed())
 		{
-			FFrameNumber NewStartOffset = ( SplitTime - InitialRange.GetLowerBoundValue() ) / Parameters.TimeScale;
+			// Sections need their offsets calculated in their local resolution. Different sequences can have different tick resolutions 
+			// so we need to transform from the parent resolution to the local one before splitting them.
+			FFrameRate LocalTickResolution = GetSequence()->GetMovieScene()->GetTickResolution();
+			FFrameNumber LocalResolutionStartOffset = FFrameRate::TransformTime(SplitTime.Time.GetFrame() - MovieScene::DiscreteInclusiveLower(InitialRange), SplitTime.Rate, LocalTickResolution).FrameNumber;
+
+			FFrameNumber NewStartOffset = LocalResolutionStartOffset / Parameters.TimeScale;
 			NewStartOffset += InitialStartOffset;
 
 			if (NewStartOffset >= 0)
@@ -226,10 +238,26 @@ UMovieSceneSection* UMovieSceneSubSection::SplitSection( FFrameNumber SplitTime
 	return nullptr;
 }
 
-void UMovieSceneSubSection::TrimSection( FFrameNumber TrimTime, bool bTrimLeft )
+TOptional<TRange<FFrameNumber> > UMovieSceneSubSection::GetAutoSizeRange() const
+{
+	if (SubSequence && SubSequence->GetMovieScene())
+	{
+		FMovieSceneSequenceTransform InnerToOuter = OuterToInnerTransform().Inverse();
+		UMovieScene* InnerMovieScene = SubSequence->GetMovieScene();
+
+		FFrameTime IncAutoStartTime = FFrameTime(MovieScene::DiscreteInclusiveLower(InnerMovieScene->GetPlaybackRange())) * InnerToOuter;
+		FFrameTime ExcAutoEndTime   = FFrameTime(MovieScene::DiscreteExclusiveUpper(InnerMovieScene->GetPlaybackRange())) * InnerToOuter;
+
+		return TRange<FFrameNumber>(GetInclusiveStartFrame(), GetInclusiveStartFrame() + (ExcAutoEndTime.RoundToFrame() - IncAutoStartTime.RoundToFrame()));
+	}
+
+	return Super::GetAutoSizeRange();
+}
+
+void UMovieSceneSubSection::TrimSection( FQualifiedFrameTime TrimTime, bool bTrimLeft )
 {
 	TRange<FFrameNumber> InitialRange = GetRange();
-	if ( !InitialRange.Contains( TrimTime ) )
+	if ( !InitialRange.Contains( TrimTime.Time.GetFrame() ) )
 	{
 		return;
 	}
@@ -241,7 +269,13 @@ void UMovieSceneSubSection::TrimSection( FFrameNumber TrimTime, bool bTrimLeft )
 	// If trimming off the left, set the offset of the shot
 	if ( bTrimLeft && InitialRange.GetLowerBound().IsClosed() )
 	{
-		FFrameNumber NewStartOffset = ( TrimTime - InitialRange.GetLowerBoundValue() ) / Parameters.TimeScale;
+		// Sections need their offsets calculated in their local resolution. Different sequences can have different tick resolutions 
+		// so we need to transform from the parent resolution to the local one before splitting them.
+		FFrameRate LocalTickResolution = GetSequence()->GetMovieScene()->GetTickResolution();
+		FFrameNumber LocalResolutionStartOffset = FFrameRate::TransformTime(TrimTime.Time.GetFrame() - MovieScene::DiscreteInclusiveLower(InitialRange), TrimTime.Rate, LocalTickResolution).FrameNumber;
+
+
+		FFrameNumber NewStartOffset = LocalResolutionStartOffset / Parameters.TimeScale;
 		NewStartOffset += InitialStartOffset;
 
 		// Ensure start offset is not less than 0
diff --git a/Engine/Source/Runtime/MovieScene/Public/Channels/BatchChannelInterface.h b/Engine/Source/Runtime/MovieScene/Public/Channels/BatchChannelInterface.h
--- a/Engine/Source/Runtime/MovieScene/Public/Channels/BatchChannelInterface.h
+++ b/Engine/Source/Runtime/MovieScene/Public/Channels/BatchChannelInterface.h
@@ -11,7 +11,7 @@ struct IBatchChannelInterface
 {
 	virtual ~IBatchChannelInterface() {}
 
-	virtual void ChangeFrameResolution_Batch(TArrayView<void* const> Ptrs, FFrameRate SourceRate, FFrameRate DestinationRate) const = 0;
+	virtual void ChangeTickResolution_Batch(TArrayView<void* const> Ptrs, FFrameRate SourceRate, FFrameRate DestinationRate) const = 0;
 	virtual TRange<FFrameNumber> ComputeEffectiveRange_Batch(TArrayView<void* const> Ptrs) const = 0;
 	virtual int32 GetNumKeys_Batch(TArrayView<void* const> Ptrs) const = 0;
 	virtual void Reset_Batch(TArrayView<void* const> Ptrs) const = 0;
@@ -39,14 +39,14 @@ struct TBatchChannelInterface : IBatchChannelInterface
 		return Range;
 	}
 
-	virtual void ChangeFrameResolution_Batch(TArrayView<void* const> Ptrs, FFrameRate SourceRate, FFrameRate DestinationRate) const override
+	virtual void ChangeTickResolution_Batch(TArrayView<void* const> Ptrs, FFrameRate SourceRate, FFrameRate DestinationRate) const override
 	{
 		using namespace MovieScene;
 		for (void* Ptr : Ptrs)
 		{
 			if (Ptr)
 			{
-				ChangeFrameResolution(static_cast<ChannelType*>(Ptr), SourceRate, DestinationRate);
+				ChangeTickResolution(static_cast<ChannelType*>(Ptr), SourceRate, DestinationRate);
 			}
 		}
 	}
diff --git a/Engine/Source/Runtime/MovieScene/Public/Channels/MovieSceneChannelTraits.h b/Engine/Source/Runtime/MovieScene/Public/Channels/MovieSceneChannelTraits.h
--- a/Engine/Source/Runtime/MovieScene/Public/Channels/MovieSceneChannelTraits.h
+++ b/Engine/Source/Runtime/MovieScene/Public/Channels/MovieSceneChannelTraits.h
@@ -129,14 +129,14 @@ namespace MovieScene
 	}
 
 	/**
-	 * Convert the frame resolution of a movie scene channel by moving the key times to the equivalent frame time
+	 * Convert the tick resolution of a movie scene channel by moving the key times to the equivalent frame time
 	 *
 	 * @param InChannel       The channel to compute the range for
 	 * @param SourceRate      The frame rate the channel is currently in
 	 * @param DestinationRate The new frame rate to convert the channel to
 	 */
 	template<typename ChannelType>
-	void ChangeFrameResolution(ChannelType* InChannel, FFrameRate SourceRate, FFrameRate DestinationRate)
+	void ChangeTickResolution(ChannelType* InChannel, FFrameRate SourceRate, FFrameRate DestinationRate)
 	{
 		TArrayView<FFrameNumber> Times = InChannel->GetInterface().GetTimes();
 		for (int32 Index = 0; Index < Times.Num(); ++Index)
diff --git a/Engine/Source/Runtime/MovieScene/Public/Channels/MovieSceneFloatChannel.h b/Engine/Source/Runtime/MovieScene/Public/Channels/MovieSceneFloatChannel.h
--- a/Engine/Source/Runtime/MovieScene/Public/Channels/MovieSceneFloatChannel.h
+++ b/Engine/Source/Runtime/MovieScene/Public/Channels/MovieSceneFloatChannel.h
@@ -21,6 +21,9 @@ struct FMovieSceneTangentData
 	FMovieSceneTangentData()
 		: ArriveTangent(0.f)
 		, LeaveTangent(0.f)
+		, TangentWeightMode(RCTWM_WeightedNone)
+		, ArriveTangentWeight(0.f)
+		, LeaveTangentWeight(0.f)
 	{}
 
 	/** If RCIM_Cubic, the arriving tangent at this key */
@@ -30,6 +33,19 @@ struct FMovieSceneTangentData
 	/** If RCIM_Cubic, the leaving tangent at this key */
 	UPROPERTY(EditAnywhere, Category="Key")
 	float LeaveTangent;
+
+	/** If RCIM_Cubic, the tangent weight mode */
+	UPROPERTY(EditAnywhere, Category = "Key")
+	TEnumAsByte<ERichCurveTangentWeightMode> TangentWeightMode;
+
+	/** If RCTWM_WeightedArrive or RCTWM_WeightedBoth, the weight of the left tangent */
+	UPROPERTY(EditAnywhere, Category="Key")
+	float ArriveTangentWeight;
+
+	/** If RCTWM_WeightedLeave or RCTWM_WeightedBoth, the weight of the right tangent */
+	UPROPERTY(EditAnywhere, Category="Key")
+	float LeaveTangentWeight;
+
 };
 
 USTRUCT()
@@ -110,15 +126,21 @@ struct FMovieSceneFloatChannel
 		return Values;
 	}
 
+
+
+
 	/**
-	 * Evaluate this channel
-	 *
-	 * @param InTime     The time to evaluate at
-	 * @param OutValue   A value to receive the result
-	 * @return true if the channel was evaluated successfully, false otherwise
-	 */
+	* Evaluate this channel with the frame resolution 
+	*
+	* @param InTime     The time to evaluate at
+	* @param InTime     The Frame Rate the time to evaluate at
+	* @param OutValue   A value to receive the result
+	* @return true if the channel was evaluated successfully, false otherwise
+	*/
 	MOVIESCENE_API bool Evaluate(FFrameTime InTime, float& OutValue) const;
 
+
+
 	/**
 	 * Set the channel's times and values to the requested values
 	 */
@@ -187,10 +209,10 @@ public:
 	 * @param EndTimeSeconds        The last time in seconds to include in the resulting array
 	 * @param TimeThreshold         A small time threshold in seconds below which we should stop adding new points
 	 * @param ValueThreshold        A small value threshold below which we should stop adding new points where the linear interpolation would suffice
-	 * @param FrameResolution       The frame resolution with which to interpret this channel's times
+	 * @param TickResolution        The tick resolution with which to interpret this channel's times
 	 * @param InOutPoints           An array to populate with the evaluated points
 	 */
-	MOVIESCENE_API void PopulateCurvePoints(double StartTimeSeconds, double EndTimeSeconds, double TimeThreshold, float ValueThreshold, FFrameRate FrameResolution, TArray<TTuple<double, double>>& InOutPoints) const;	
+	MOVIESCENE_API void PopulateCurvePoints(double StartTimeSeconds, double EndTimeSeconds, double TimeThreshold, float ValueThreshold, FFrameRate TickResolution, TArray<TTuple<double, double>>& InOutPoints) const;	
 
 public:
 
@@ -218,12 +240,12 @@ private:
 	/**
 	 * Adds median points between each of the supplied points if their evaluated value is significantly different than the linear interpolation of those points
 	 *
-	 * @param FrameResolution       The frame resolution with which to interpret this channel's times
+	 * @param TickResolution        The tick resolution with which to interpret this channel's times
 	 * @param TimeThreshold         A small time threshold in seconds below which we should stop adding new points
 	 * @param ValueThreshold        A small value threshold below which we should stop adding new points where the linear interpolation would suffice
 	 * @param InOutPoints           An array to populate with the evaluated points
 	 */
-	void RefineCurvePoints(FFrameRate FrameResolution, double TimeThreshold, float ValueThreshold, TArray<TTuple<double, double>>& InOutPoints) const;
+	void RefineCurvePoints(FFrameRate TickResolution, double TimeThreshold, float ValueThreshold, TArray<TTuple<double, double>>& InOutPoints) const;
 
 	UPROPERTY()
 	TArray<FFrameNumber> Times;
@@ -238,6 +260,16 @@ private:
 	bool bHasDefaultValue;
 
 	FMovieSceneKeyHandleMap KeyHandles;
+
+	UPROPERTY()
+	FFrameRate TickResolution;
+
+public:
+	//Set it's frame resolution
+	void SetTickResolution(FFrameRate InTickSolution)
+	{
+		TickResolution = InTickSolution;
+	}
 };
 
 template<>
@@ -268,12 +300,12 @@ inline bool ValueExistsAtTime(const FMovieSceneFloatChannel* Channel, FFrameNumb
 	const FFrameTime FrameTime(InFrameNumber);
 
 	float ExistingValue = 0.f;
-	return Channel->Evaluate(FrameTime, ExistingValue) && FMath::IsNearlyEqual(ExistingValue, Value, Tolerance);
+	return Channel->Evaluate(FrameTime,  ExistingValue) && FMath::IsNearlyEqual(ExistingValue, Value, Tolerance);
 }
 
 inline bool ValueExistsAtTime(const FMovieSceneFloatChannel* Channel, FFrameNumber InFrameNumber, const FMovieSceneFloatValue& InValue, float Tolerance = KINDA_SMALL_NUMBER)
 {
-	return ValueExistsAtTime(Channel, InFrameNumber, InValue.Value, Tolerance);
+	return ValueExistsAtTime(Channel, InFrameNumber,  InValue.Value, Tolerance);
 }
 
 inline void AssignValue(FMovieSceneFloatChannel* InChannel, int32 ValueIndex, float InValue)
@@ -300,6 +332,6 @@ MOVIESCENE_API void Dilate(FMovieSceneFloatChannel* InChannel, FFrameNumber Orig
 
 
 /**
- * Overload for Converting a float channel from one frame rate to another. Ensures that both key times and tangents are migrated correctly. See MovieScene::ChangeFrameResolution for default implementation
+ * Overload for Converting a float channel from one frame rate to another. Ensures that both key times and tangents are migrated correctly. See MovieScene::ChangeTickResolution for default implementation
  */
-MOVIESCENE_API void ChangeFrameResolution(FMovieSceneFloatChannel* InChannel, FFrameRate SourceRate, FFrameRate DestinationRate);
\ No newline at end of file
+MOVIESCENE_API void ChangeTickResolution(FMovieSceneFloatChannel* InChannel, FFrameRate SourceRate, FFrameRate DestinationRate);
\ No newline at end of file
diff --git a/Engine/Source/Runtime/MovieScene/Public/Evaluation/MovieSceneEvaluationField.h b/Engine/Source/Runtime/MovieScene/Public/Evaluation/MovieSceneEvaluationField.h
--- a/Engine/Source/Runtime/MovieScene/Public/Evaluation/MovieSceneEvaluationField.h
+++ b/Engine/Source/Runtime/MovieScene/Public/Evaluation/MovieSceneEvaluationField.h
@@ -206,9 +206,9 @@ struct FMovieSceneEvaluationMetaData
 	UPROPERTY()
 	TArray<FMovieSceneOrderedEvaluationKey> ActiveEntities;
 
-	/** Map of sub sequence IDs to FMovieSceneEvaluationTemplate::TemplateSignature that this meta data was generated with (not including root). */
+	/** Map of sub sequence IDs to FMovieSceneEvaluationTemplate::TemplateSerialNumber that this meta data was generated with (not including root). */
 	UPROPERTY()
-	TMap<FMovieSceneSequenceID, FGuid> SubTemplateSignatures;
+	TMap<FMovieSceneSequenceID, uint32> SubTemplateSerialNumbers;
 };
 
 /**
diff --git a/Engine/Source/Runtime/MovieScene/Public/Evaluation/MovieSceneEvaluationState.h b/Engine/Source/Runtime/MovieScene/Public/Evaluation/MovieSceneEvaluationState.h
--- a/Engine/Source/Runtime/MovieScene/Public/Evaluation/MovieSceneEvaluationState.h
+++ b/Engine/Source/Runtime/MovieScene/Public/Evaluation/MovieSceneEvaluationState.h
@@ -28,6 +28,15 @@ struct FMovieSceneObjectCache
 	 */
 	MOVIESCENE_API TArrayView<TWeakObjectPtr<>> FindBoundObjects(const FGuid& InBindingID, IMovieScenePlayer& Player);
 
+	/**
+	 * Find all objects that are bound to the specified binding ID
+	 * @note Does not update bindings if they are out of date, or invalid
+	 *
+	 * @param InBindingID		The object binding GUID for a spawnable or posessable in a UMovieScene
+	 * @return An iterable type of all objects bound to the specified ID.
+	 */
+	MOVIESCENE_API TArrayView<const TWeakObjectPtr<>> IterateBoundObjects(const FGuid& InBindingID) const;
+
 	/**
 	 * Set the sequence that this cache applies to
 	 *
diff --git a/Engine/Source/Runtime/MovieScene/Public/Evaluation/MovieSceneEvaluationTemplate.h b/Engine/Source/Runtime/MovieScene/Public/Evaluation/MovieSceneEvaluationTemplate.h
--- a/Engine/Source/Runtime/MovieScene/Public/Evaluation/MovieSceneEvaluationTemplate.h
+++ b/Engine/Source/Runtime/MovieScene/Public/Evaluation/MovieSceneEvaluationTemplate.h
@@ -125,20 +125,61 @@ template<> struct TStructOpsTypeTraits<FMovieSceneSubSectionFieldData> : public
 
 
 /**
- * Template that is used for efficient runtime evaluation of a movie scene sequence. Potentially serialized into the asset.
+ * Sereal number used to identify evaluation template state that can only ever increase over its lifetime.
+ * Only to be stored internally on FMovieSceneEvaluationTemplate.
  */
 USTRUCT()
-struct FMovieSceneEvaluationTemplate
+struct FMovieSceneEvaluationTemplateSerialNumber
 {
 	GENERATED_BODY()
 
+	FMovieSceneEvaluationTemplateSerialNumber()
+		: Value(0)
+	{}
+
 	/**
-	 * Default construction
+	 * Access this serial number's value
 	 */
-	FMovieSceneEvaluationTemplate()
+	uint32 GetValue() const
+	{
+		return Value;
+	}
+
+	/**
+	* Increment this serial number
+	*/
+	void Increment()
 	{
+		++Value;
 	}
 
+private:
+	
+	friend struct FMovieSceneEvaluationTemplate;
+
+	/**
+	 * Copy/move semantics only ever initialize to zero, or passthrough, to ensure that FMovieSceneEvaluationTemplate::TemplateSerialNumber cannot ever move backwards.
+	 */
+	FMovieSceneEvaluationTemplateSerialNumber(const FMovieSceneEvaluationTemplateSerialNumber&) : Value(0) {}
+	FMovieSceneEvaluationTemplateSerialNumber(FMovieSceneEvaluationTemplateSerialNumber&&)      : Value(0) {}
+
+	FMovieSceneEvaluationTemplateSerialNumber& operator=(const FMovieSceneEvaluationTemplateSerialNumber&) { return *this; }
+	FMovieSceneEvaluationTemplateSerialNumber& operator=(FMovieSceneEvaluationTemplateSerialNumber&&)      { return *this; }
+
+	/** The internal value of the serial number */
+	UPROPERTY()
+	uint32 Value;
+};
+template<> struct TStructOpsTypeTraits<FMovieSceneEvaluationTemplateSerialNumber> : public TStructOpsTypeTraitsBase2<FMovieSceneEvaluationTemplateSerialNumber> { enum { WithCopy = false }; };
+
+/**
+ * Template that is used for efficient runtime evaluation of a movie scene sequence. Potentially serialized into the asset.
+ */
+USTRUCT()
+struct FMovieSceneEvaluationTemplate
+{
+	GENERATED_BODY()
+
 public:
 
 	/**
@@ -293,9 +334,9 @@ public:
 	UPROPERTY()
 	FGuid SequenceSignature;
 
-	/** Signature that is re-generated every time this template is re-generated through FMovieSceneEvaluationTemplateGenerator */
+	/** Serial number that is incremented every time this template is re-generated through FMovieSceneEvaluationTemplateGenerator */
 	UPROPERTY()
-	FGuid TemplateSignature;
+	FMovieSceneEvaluationTemplateSerialNumber TemplateSerialNumber;
 
 private:
 
diff --git a/Engine/Source/Runtime/MovieScene/Public/Evaluation/MovieScenePlayback.h b/Engine/Source/Runtime/MovieScene/Public/Evaluation/MovieScenePlayback.h
--- a/Engine/Source/Runtime/MovieScene/Public/Evaluation/MovieScenePlayback.h
+++ b/Engine/Source/Runtime/MovieScene/Public/Evaluation/MovieScenePlayback.h
@@ -496,10 +496,10 @@ private:
 
 private:
 
-	/** The framerate to be used when interpreting frame time values provided to this class (i.e. moviescene playback frame rate) */
+	/** The framerate to be used when interpreting frame time values provided to this class (i.e. moviescene display rate) */
 	FFrameRate InputRate;
 
-	/** The framerate to be used when returning frame time values from this class (i.e. moviescene frame resolution) */
+	/** The framerate to be used when returning frame time values from this class (i.e. moviescene tick resolution) */
 	FFrameRate OutputRate;
 
 	/** The type of evaluation to use */
diff --git a/Engine/Source/Runtime/MovieScene/Public/Evaluation/MovieSceneSequenceHierarchy.h b/Engine/Source/Runtime/MovieScene/Public/Evaluation/MovieSceneSequenceHierarchy.h
--- a/Engine/Source/Runtime/MovieScene/Public/Evaluation/MovieSceneSequenceHierarchy.h
+++ b/Engine/Source/Runtime/MovieScene/Public/Evaluation/MovieSceneSequenceHierarchy.h
@@ -54,9 +54,9 @@ struct FMovieSceneSubSequenceData
 	UPROPERTY()
 	FMovieSceneSequenceTransform RootToSequenceTransform;
 
-	/** The frame resolution of the inner sequence. */
+	/** The tick resolution of the inner sequence. */
 	UPROPERTY()
-	FFrameRate FrameResolution;
+	FFrameRate TickResolution;
 
 	/** This sequence's deterministic sequence ID. Used in editor to reduce the risk of collisions on recompilation. */ 
 	UPROPERTY()
diff --git a/Engine/Source/Runtime/MovieScene/Public/MovieScene.h b/Engine/Source/Runtime/MovieScene/Public/MovieScene.h
--- a/Engine/Source/Runtime/MovieScene/Public/MovieScene.h
+++ b/Engine/Source/Runtime/MovieScene/Public/MovieScene.h
@@ -416,6 +416,14 @@ public:
 	 */
 	bool RemoveMasterTrack(UMovieSceneTrack& Track);
 
+	/**
+	 * Move all the contents (tracks, child bindings) of the specified binding ID onto another
+	 *
+	 * @param SourceBindingId The identifier of the binding ID to move all tracks and children from
+	 * @param DestinationBindingId The identifier of the binding ID to move the contents to
+	 */
+	void MoveBindingContents(const FGuid& SourceBindingId, const FGuid& DestinationBindingId);
+
 public:
 
 	// @todo sequencer: the following methods really shouldn't be here
@@ -476,35 +484,35 @@ public:
 	}
 
 	/**
-	 * Retrieve the frame resolution at which all frame numbers within this movie scene are defined
+	 * Retrieve the tick resolution at which all frame numbers within this movie scene are defined
 	 */
-	FFrameRate GetFrameResolution() const
+	FFrameRate GetTickResolution() const
 	{
-		return FrameResolution;
+		return TickResolution;
 	}
 
 	/**
-	 * Directly set the frame resolution for this movie scene without applying any conversion whatsoever, or modifying the data
+	 * Directly set the tick resolution for this movie scene without applying any conversion whatsoever, or modifying the data
 	 */
-	void SetFrameResolutionDirectly(FFrameRate InFrameResolution)
+	void SetTickResolutionDirectly(FFrameRate InTickResolution)
 	{
-		FrameResolution = InFrameResolution;
+		TickResolution = InTickResolution;
 	}
 
 	/**
-	 * Retrieve the frame resolution at which all frame numbers within this movie scene are defined
+	 * Retrieve the display frame rate for this data, in which frame numbers should be displayed on UI, and interacted with in movie scene players
 	 */
-	FFrameRate GetPlaybackFrameRate() const
+	FFrameRate GetDisplayRate() const
 	{
-		return PlayRate;
+		return DisplayRate;
 	}
 
 	/**
 	 * Set the play rate for this movie scene
 	 */
-	void SetPlaybackFrameRate(FFrameRate InPlayRate)
+	void SetDisplayRate(FFrameRate InDisplayRate)
 	{
-		PlayRate = InPlayRate;
+		DisplayRate = InDisplayRate;
 	}
 
 	/**
@@ -703,11 +711,11 @@ private:
 
 	/** The resolution at which all frame numbers within this movie-scene data are stored */
 	UPROPERTY()
-	FFrameRate FrameResolution;
+	FFrameRate TickResolution;
 
-	/** The rate at which to playback this moviescene data */
+	/** The rate at which we should interact with this moviescene data on UI, and to movie scene players. Also defines the frame locked frame rate. */
 	UPROPERTY()
-	FFrameRate PlayRate;
+	FFrameRate DisplayRate;
 
 	/** The type of evaluation to use when playing back this sequence */
 	UPROPERTY()
diff --git a/Engine/Source/Runtime/MovieScene/Public/MovieSceneBinding.h b/Engine/Source/Runtime/MovieScene/Public/MovieSceneBinding.h
--- a/Engine/Source/Runtime/MovieScene/Public/MovieSceneBinding.h
+++ b/Engine/Source/Runtime/MovieScene/Public/MovieSceneBinding.h
@@ -111,6 +111,24 @@ struct FMovieSceneBinding
 		return Tracks;
 	}
 
+	/**
+	 * Reset all tracks in this binding, returning the previous array of tracks
+	 */
+	TArray<UMovieSceneTrack*> StealTracks()
+	{
+		TArray<UMovieSceneTrack*> Empty;
+		Swap(Empty, Tracks);
+		return Empty;
+	}
+
+	/**
+	 * Assign all tracks in this binding
+	 */
+	void SetTracks(TArray<UMovieSceneTrack*>&& InTracks)
+	{
+		Tracks = MoveTemp(InTracks);
+	}
+
 #if WITH_EDITOR
 	/**
 	 * Perform cook-time optimization on this object binding
diff --git a/Engine/Source/Runtime/MovieScene/Public/MovieSceneClipboard.h b/Engine/Source/Runtime/MovieScene/Public/MovieSceneClipboard.h
--- a/Engine/Source/Runtime/MovieScene/Public/MovieSceneClipboard.h
+++ b/Engine/Source/Runtime/MovieScene/Public/MovieSceneClipboard.h
@@ -268,8 +268,8 @@ struct FMovieSceneClipboardEnvironment
 	/** The date/time at which the copy operation was performed */
 	FDateTime DateTime;
 
-	/** The frame resolution within which all clipboard data is stored */
-	FFrameRate FrameResolution;
+	/** The tick resolution within which all clipboard data is stored */
+	FFrameRate TickResolution;
 };
 
 /** A clipboard representing serializable copied data for a movie scene */
diff --git a/Engine/Source/Runtime/MovieScene/Public/MovieSceneCommonHelpers.h b/Engine/Source/Runtime/MovieScene/Public/MovieSceneCommonHelpers.h
--- a/Engine/Source/Runtime/MovieScene/Public/MovieSceneCommonHelpers.h
+++ b/Engine/Source/Runtime/MovieScene/Public/MovieSceneCommonHelpers.h
@@ -13,6 +13,7 @@ class AActor;
 class UCameraComponent;
 class UMovieSceneSection;
 class USceneComponent;
+class USoundBase;
 struct FRichCurve;
 enum class EMovieSceneKeyInterpolation : uint8;
 
@@ -82,6 +83,14 @@ public:
 	 */
 	static void SetRuntimeObjectMobility(UObject* Object, EComponentMobility::Type ComponentMobility = EComponentMobility::Movable);
 
+	/*
+	 * Get the duration for the given sound
+
+	 * @param Sound The sound to get the duration for
+	 * @return The duration in seconds
+	 */
+	static float GetSoundDuration(USoundBase* Sound);
+
 	/**
 	 * Sort predicate that sorts lower bounds of a range
 	 */
diff --git a/Engine/Source/Runtime/MovieScene/Public/MovieSceneFwd.h b/Engine/Source/Runtime/MovieScene/Public/MovieSceneFwd.h
--- a/Engine/Source/Runtime/MovieScene/Public/MovieSceneFwd.h
+++ b/Engine/Source/Runtime/MovieScene/Public/MovieSceneFwd.h
@@ -26,7 +26,7 @@ namespace EMovieScenePlayerStatus
 UENUM()
 enum class EMovieSceneEvaluationType : uint8
 {
-	/** Play the sequence frame-locked to its playback rate (snapped to the frame resolution - no sub-frames) */
+	/** Play the sequence frame-locked to its playback rate (snapped to the tick resolution - no sub-frames) */
 	FrameLocked,
 
 	/** Play the sequence in real-time, with sub-frame interpolation if necessary */
diff --git a/Engine/Source/Runtime/MovieScene/Public/MovieSceneKeyProxy.h b/Engine/Source/Runtime/MovieScene/Public/MovieSceneKeyProxy.h
new file mode 100644
--- /dev/null
+++ b/Engine/Source/Runtime/MovieScene/Public/MovieSceneKeyProxy.h
@@ -0,0 +1,110 @@
+// Copyright 1998-2018 Epic Games, Inc. All Rights Reserved.
+
+#pragma once
+
+#include "UObject/ObjectMacros.h"
+#include "UObject/Interface.h"
+#include "MovieSceneSection.h"
+#include "Channels/MovieSceneChannelHandle.h"
+#include "Channels/MovieSceneChannel.h"
+#include "Curves/KeyHandle.h"
+#include "MovieSceneKeyProxy.generated.h"
+
+UINTERFACE()
+class MOVIESCENE_API UMovieSceneKeyProxy : public UInterface
+{
+	GENERATED_BODY()
+};
+
+
+/**
+ * Interface that can be implemented by any object that is used as a key editing proxy for a moviescene channel.
+ * When used, UpdateValuesFromRawData should be called every frame to optionally retrieve the current values
+ * of the key for this proxy.
+ */
+class MOVIESCENE_API IMovieSceneKeyProxy
+{
+public:
+
+	GENERATED_BODY()
+
+	/**
+	 * To be called by the edit interface to update this instance's properties with the underlying raw data
+	 */
+	virtual void UpdateValuesFromRawData() = 0;
+
+protected:
+
+	/**
+	 * Implementation function that sets the underlying key time/value to the specified values if possible.
+	 * If the section is locked, InOutValue and InOutTime will be reset back to the current key's value
+	 *
+	 * @param InChannelHandle      The channel on which the underlying key value resides
+	 * @param InSection            The section that owns the channel
+	 * @param InKeyHandle          A handle to the key that we are reflecting
+	 * @param InOutValue           Value to assign to the underlying key. If the section is locked, this will receive the existing key's value without changing the underlying key value.
+	 * @param InOutTime            Time to move the underlying key to. If the section is locked, this will receive the existing key's time without changing the underlying time.
+	 */
+	template<typename ChannelType, typename ValueType>
+	void OnProxyValueChanged(TMovieSceneChannelHandle<ChannelType> InChannelHandle, UMovieSceneSection* InSection, FKeyHandle InKeyHandle, ValueType& InOutValue, FFrameNumber& InOutTime);
+
+	/**
+	 * Implementation function that retrieves the underlying key time/value and applies then to the specified value and time parameters. Normally called once per tick.
+	 *
+	 * @param InChannelHandle      The channel on which the underlying key value resides
+	 * @param InKeyHandle          A handle to the key that we are reflecting
+	 * @param OutValue             (Out) Value to receive the underlying key's value
+	 * @param OutTime              (Out) Time to receive the underlying key time
+	 */
+	template<typename ChannelType, typename ValueType>
+	void RefreshCurrentValue(TMovieSceneChannelHandle<ChannelType> InChannelHandle, FKeyHandle InKeyHandle, ValueType& OutValue, FFrameNumber& OutTime);
+};
+
+
+
+template<typename ChannelType, typename ValueType>
+void IMovieSceneKeyProxy::OnProxyValueChanged(TMovieSceneChannelHandle<ChannelType> InChannelHandle, UMovieSceneSection* InSection, FKeyHandle InKeyHandle, ValueType& InOutValue, FFrameNumber& InOutTime)
+{
+	auto* Channel = InChannelHandle.Get();
+	if (!Channel || !InSection)
+	{
+		return;
+	}
+
+	auto ChannelInterface = Channel->GetInterface();
+	int32 KeyIndex = ChannelInterface.GetIndex(InKeyHandle);
+	if (KeyIndex != INDEX_NONE)
+	{
+		// If we have no section, or it's locked, don't let the user change the value
+		if (!InSection || !InSection->TryModify())
+		{
+			InOutTime  = ChannelInterface.GetTimes()[KeyIndex];
+			InOutValue = ChannelInterface.GetValues()[KeyIndex];
+		}
+		else
+		{
+			ChannelInterface.GetValues()[KeyIndex] = InOutValue;
+
+			ChannelInterface.MoveKey(KeyIndex, InOutTime);
+			InSection->ExpandToFrame(InOutTime);
+		}
+	}
+}
+
+
+
+template<typename ChannelType, typename ValueType>
+void IMovieSceneKeyProxy::RefreshCurrentValue(TMovieSceneChannelHandle<ChannelType> InChannelHandle, FKeyHandle InKeyHandle, ValueType& OutValue, FFrameNumber& OutTime)
+{
+	auto* Channel = InChannelHandle.Get();
+	if (Channel)
+	{
+		auto ChannelInterface = Channel->GetInterface();
+		int32 KeyIndex = ChannelInterface.GetIndex(InKeyHandle);
+		if (KeyIndex != INDEX_NONE)
+		{
+			OutValue = ChannelInterface.GetValues()[KeyIndex];
+			OutTime  = ChannelInterface.GetTimes()[KeyIndex];
+		}
+	}
+}
\ No newline at end of file
diff --git a/Engine/Source/Runtime/MovieScene/Public/MovieSceneSection.h b/Engine/Source/Runtime/MovieScene/Public/MovieSceneSection.h
--- a/Engine/Source/Runtime/MovieScene/Public/MovieSceneSection.h
+++ b/Engine/Source/Runtime/MovieScene/Public/MovieSceneSection.h
@@ -11,6 +11,7 @@
 #include "Evaluation/Blending/MovieSceneBlendType.h"
 #include "Generators/MovieSceneEasingFunction.h"
 #include "MovieSceneFrameMigration.h"
+#include "QualifiedFrameTime.h"
 #include "MovieSceneSection.generated.h"
 
 class FStructOnScope;
@@ -247,6 +248,14 @@ public:
 		return SectionRange.Value.Contains(Position);
 	}
 
+	/*
+	 * Returns the range to auto size this section to, if there is one. This defaults to the 
+	 * range of all the keys.
+	 *
+	 * @return the range of this section to auto size to
+	 */
+	MOVIESCENE_API virtual TOptional<TRange<FFrameNumber> > GetAutoSizeRange() const;
+
 	/**
 	 * Gets this section's blend type
 	 */
@@ -291,7 +300,7 @@ public:
 	 * @param SplitTime The time at which to split
 	 * @return The newly created split section
 	 */
-	MOVIESCENE_API virtual UMovieSceneSection* SplitSection(FFrameNumber SplitTime);
+	MOVIESCENE_API virtual UMovieSceneSection* SplitSection(FQualifiedFrameTime SplitTime);
 
 	/**
 	 * Trim a section at the trim time
@@ -299,7 +308,7 @@ public:
 	 * @param TrimTime The time at which to trim
 	 * @param bTrimLeft Whether to trim left or right
 	 */
-	MOVIESCENE_API virtual void TrimSection(FFrameNumber TrimTime, bool bTrimLeft);
+	MOVIESCENE_API virtual void TrimSection(FQualifiedFrameTime TrimTime, bool bTrimLeft);
 
 	/**
 	 * Get the data structure representing the specified keys.
@@ -440,6 +449,13 @@ public:
 
 	/** Does this movie section support infinite ranges for evaluation */
 	MOVIESCENE_API bool GetSupportsInfiniteRange() const { return bSupportsInfiniteRange; }
+
+	/**
+	*  Whether or not we draw a curve for a particular channel owned by this section.
+	*  Defaults to true.
+	*/
+	virtual bool ShowCurveForChannel(const void *Channel) const  { return true; }
+
 protected:
 
 	//~ UObject interface
diff --git a/Engine/Source/Runtime/MovieScene/Public/MovieSceneSequence.h b/Engine/Source/Runtime/MovieScene/Public/MovieSceneSequence.h
--- a/Engine/Source/Runtime/MovieScene/Public/MovieSceneSequence.h
+++ b/Engine/Source/Runtime/MovieScene/Public/MovieSceneSequence.h
@@ -13,6 +13,7 @@
 class ITargetPlatform;
 class UMovieScene;
 struct FMovieScenePossessable;
+struct FMovieSceneObjectCache;
 
 /**
  * Abstract base class for movie scene animations (C++ version).
@@ -80,6 +81,14 @@ public:
 	 */
 	MOVIESCENE_API FGuid FindPossessableObjectId(UObject& Object, UObject* Context) const;
 
+	/**
+	 * Called to validate the specified object cache by removing anything that should be deemed out of date
+	 * 
+	 * @param InObjectCache 		The object cache container that contains all the objects currently animated by this sequence
+	 * @param OutInvalidIDs 		(Out) Array to populate with any object bindings that should be invalidated
+	 */
+	virtual void GatherExpiredObjects(const FMovieSceneObjectCache& InObjectCache, TArray<FGuid>& OutInvalidIDs) const {}
+
 	/**
 	 * Get the movie scene that controls this animation.
 	 *
diff --git a/Engine/Source/Runtime/MovieScene/Public/MovieSceneSequencePlayer.h b/Engine/Source/Runtime/MovieScene/Public/MovieSceneSequencePlayer.h
--- a/Engine/Source/Runtime/MovieScene/Public/MovieSceneSequencePlayer.h
+++ b/Engine/Source/Runtime/MovieScene/Public/MovieSceneSequencePlayer.h
@@ -201,7 +201,7 @@ public:
 	 */
 	DEPRECATED(4.20, "Please use PlayToFrame instead")
 	UFUNCTION(BlueprintCallable, Category="Game|Cinematic")
-	void SetPlaybackPosition(float NewPlaybackPosition) { PlayToSeconds(NewPlaybackPosition + StartTime / PlayPosition.GetInputRate()); }
+	void SetPlaybackPosition(float NewPlaybackPosition) { Status == EMovieScenePlayerStatus::Playing ? PlayToSeconds(NewPlaybackPosition + StartTime / PlayPosition.GetInputRate()) : JumpToSeconds(NewPlaybackPosition + StartTime / PlayPosition.GetInputRate()); }
 
 	/**
 	 * Sets the range in time to be played back by this player, overriding the default range stored in the asset
@@ -244,7 +244,7 @@ public:
 	int32 GetFrameDuration() const;
 
 	/**
-	 * Get this sequence's playback frame rate.
+	 * Get this sequence's display rate.
 	 */
 	UFUNCTION(BlueprintCallable, Category="Game|Cinematic")
 	FFrameRate GetFrameRate() const { return PlayPosition.GetInputRate(); }
@@ -348,6 +348,10 @@ public:
 	UFUNCTION(BlueprintCallable, Category="Game|Cinematic")
 	bool IsPaused() const;
 
+	/** Check whether playback is reversed. */
+	UFUNCTION(BlueprintCallable, Category="Game|Cinematic")
+	bool IsReversed() const;
+
 	/** Get the playback rate of this player. */
 	UFUNCTION(BlueprintCallable, Category="Game|Cinematic")
 	float GetPlayRate() const;
diff --git a/Engine/Source/Runtime/MovieScene/Public/MovieSceneSpawnable.h b/Engine/Source/Runtime/MovieScene/Public/MovieSceneSpawnable.h
--- a/Engine/Source/Runtime/MovieScene/Public/MovieSceneSpawnable.h
+++ b/Engine/Source/Runtime/MovieScene/Public/MovieSceneSpawnable.h
@@ -188,6 +188,26 @@ public:
 	UPROPERTY(EditAnywhere, AdvancedDisplay, Category = Actor)
 	FTransform SpawnTransform;
 
+	/**
+	 * Get the name of the level to spawn into.
+	 *
+	 * @return Level name.
+	 */
+	FName GetLevelName() const
+	{
+		return LevelName;
+	}
+
+	/**
+	 * Set the name of the level to spawn into.
+	 *
+	 * @InLevelName The desired spawnable level name.
+	 */
+	void SetLevelName(FName InLevelName)
+	{
+		LevelName = InLevelName;
+	}
+
 private:
 
 	/** Unique identifier of the spawnable object. */
@@ -219,4 +239,8 @@ public:
 	UPROPERTY()
 	UClass* GeneratedClass_DEPRECATED;
 #endif
+
+	/** Name of level to spawn into */
+	UPROPERTY()
+	FName LevelName;
 };
diff --git a/Engine/Source/Runtime/MovieScene/Public/MovieSceneTimeHelpers.h b/Engine/Source/Runtime/MovieScene/Public/MovieSceneTimeHelpers.h
--- a/Engine/Source/Runtime/MovieScene/Public/MovieSceneTimeHelpers.h
+++ b/Engine/Source/Runtime/MovieScene/Public/MovieSceneTimeHelpers.h
@@ -152,8 +152,8 @@ inline TRange<T> TranslateRange(const TRange<T>& InRange, T Amount)
  */
 inline FFrameTime ClampToDiscreteRange(FFrameTime InTime, const TRange<FFrameNumber>& InRange)
 {
-	FFrameTime MinTime = DiscreteInclusiveLower(InRange);
-	FFrameTime MaxTime = FFrameTime(DiscreteExclusiveUpper(InRange)-1, 0.99999994f);
+	FFrameTime MinTime = InRange.GetLowerBound().IsClosed() ? DiscreteInclusiveLower(InRange) : FFrameTime(TNumericLimits<int32>::Lowest());
+	FFrameTime MaxTime = FFrameTime(InRange.GetUpperBound().IsClosed() ? DiscreteExclusiveUpper(InRange)-1 : TNumericLimits<int32>::Max(), 0.99999994f);
 
 	return FMath::Clamp(InTime, MinTime, MaxTime);
 }
diff --git a/Engine/Source/Runtime/MovieScene/Public/Sections/MovieSceneSubSection.h b/Engine/Source/Runtime/MovieScene/Public/Sections/MovieSceneSubSection.h
--- a/Engine/Source/Runtime/MovieScene/Public/Sections/MovieSceneSubSection.h
+++ b/Engine/Source/Runtime/MovieScene/Public/Sections/MovieSceneSubSection.h
@@ -133,9 +133,9 @@ public:
 public:
 
 	//~ UMovieSceneSection interface
-
-	virtual UMovieSceneSection* SplitSection( FFrameNumber SplitTime ) override;
-	virtual void TrimSection( FFrameNumber TrimTime, bool bTrimLeft ) override;
+	virtual TOptional<TRange<FFrameNumber> > GetAutoSizeRange() const override;
+	virtual UMovieSceneSection* SplitSection( FQualifiedFrameTime SplitTime ) override;
+	virtual void TrimSection( FQualifiedFrameTime TrimTime, bool bTrimLeft ) override;
 	virtual TOptional<FFrameTime> GetOffsetTime() const override { return TOptional<FFrameTime>(FFrameTime(Parameters.GetStartFrameOffset())); }
 
 public:
diff --git a/Engine/Source/Runtime/MovieSceneCapture/Private/MovieSceneCaptureModule.cpp b/Engine/Source/Runtime/MovieSceneCapture/Private/MovieSceneCaptureModule.cpp
--- a/Engine/Source/Runtime/MovieSceneCapture/Private/MovieSceneCaptureModule.cpp
+++ b/Engine/Source/Runtime/MovieSceneCapture/Private/MovieSceneCaptureModule.cpp
@@ -4,6 +4,7 @@
 #include "Misc/CoreDelegates.h"
 #include "UObject/UObjectGlobals.h"
 #include "UObject/Class.h"
+#include "UObject/CoreRedirects.h"
 #include "Misc/CommandLine.h"
 #include "Misc/FileHelper.h"
 #include "EngineGlobals.h"
@@ -175,7 +176,18 @@ private:
 					UClass* Class = FindObject<UClass>(nullptr, *TypeName);
 					if (!Class)
 					{
-						return nullptr;
+						const FCoreRedirect* ValueRedirect = nullptr;
+						FCoreRedirectObjectName NewRedirectName;
+
+						if (FCoreRedirects::RedirectNameAndValues(ECoreRedirectFlags::Type_Class, TypeName, NewRedirectName, &ValueRedirect))
+						{
+							Class = FindObject<UClass>(nullptr, *NewRedirectName.ToString());
+						}
+
+						if (!Class)
+						{
+							return nullptr;
+						}
 					}
 
 					Capture = NewObject<UMovieSceneCapture>(GetTransientPackage(), Class);
@@ -208,7 +220,18 @@ private:
 			UClass* Class = FindObject<UClass>( nullptr, *TypeName );
 			if( !Class )
 			{
-				return nullptr;
+				const FCoreRedirect* ValueRedirect = nullptr;
+				FCoreRedirectObjectName NewRedirectName;
+
+				if (FCoreRedirects::RedirectNameAndValues(ECoreRedirectFlags::Type_Class, TypeName, NewRedirectName, &ValueRedirect))
+				{
+					Class = FindObject<UClass>(nullptr, *NewRedirectName.ToString());
+				}
+
+				if (!Class)
+				{
+					return nullptr;
+				}
 			}
 
 			Capture = NewObject<UMovieSceneCapture>( GetTransientPackage(), Class );
diff --git a/Engine/Source/Runtime/MovieSceneCapture/Public/MovieSceneCaptureSettings.h b/Engine/Source/Runtime/MovieSceneCapture/Public/MovieSceneCaptureSettings.h
--- a/Engine/Source/Runtime/MovieSceneCapture/Public/MovieSceneCaptureSettings.h
+++ b/Engine/Source/Runtime/MovieSceneCapture/Public/MovieSceneCaptureSettings.h
@@ -59,7 +59,7 @@ struct MOVIESCENECAPTURE_API FMovieSceneCaptureSettings
 	UPROPERTY(config, EditAnywhere, Category=General, AdvancedDisplay )
 	bool bOverwriteExisting;
 
-	/** True if frame numbers in the output files should be relative to zero, rather than the actual frame numbers in the originating animation content */
+	/** True if frame numbers in the output files should be relative to zero, rather than the actual frame numbers in the originating animation content. */
 	UPROPERTY(config, EditAnywhere, Category=General, AdvancedDisplay)
 	bool bUseRelativeFrameNumbers;
 
diff --git a/Engine/Source/Runtime/MovieSceneTracks/Private/Evaluation/MovieScene3DTransformTemplate.cpp b/Engine/Source/Runtime/MovieSceneTracks/Private/Evaluation/MovieScene3DTransformTemplate.cpp
--- a/Engine/Source/Runtime/MovieSceneTracks/Private/Evaluation/MovieScene3DTransformTemplate.cpp
+++ b/Engine/Source/Runtime/MovieSceneTracks/Private/Evaluation/MovieScene3DTransformTemplate.cpp
@@ -214,6 +214,7 @@ void FMovieSceneComponentTransformSectionTemplate::Interrogate(const FMovieScene
 FMovieScene3DTransformTemplateData::FMovieScene3DTransformTemplateData(const UMovieScene3DTransformSection& Section)
 	: BlendType(Section.GetBlendType().Get())
 	, Mask(Section.GetMask())
+	, bUseQuaternionInterpolation(Section.GetUseQuaternionInterpolation())
 {
 	EMovieSceneTransformChannel MaskChannels = Mask.GetChannels();
 	TArrayView<FMovieSceneFloatChannel*> FloatChannels = Section.GetChannelProxy().GetChannels<FMovieSceneFloatChannel>();
@@ -252,9 +253,131 @@ MovieScene::TMultiChannelValue<float, 9> FMovieScene3DTransformTemplateData::Eva
 	EvalChannel(1, EMovieSceneTransformChannel::TranslationY, TranslationCurve[1]);
 	EvalChannel(2, EMovieSceneTransformChannel::TranslationZ, TranslationCurve[2]);
 
-	EvalChannel(3, EMovieSceneTransformChannel::RotationX, RotationCurve[0]);
-	EvalChannel(4, EMovieSceneTransformChannel::RotationY, RotationCurve[1]);
-	EvalChannel(5, EMovieSceneTransformChannel::RotationZ, RotationCurve[2]);
+	if (!bUseQuaternionInterpolation)
+	{
+		EvalChannel(3, EMovieSceneTransformChannel::RotationX, RotationCurve[0]);
+		EvalChannel(4, EMovieSceneTransformChannel::RotationY, RotationCurve[1]);
+		EvalChannel(5, EMovieSceneTransformChannel::RotationZ, RotationCurve[2]);
+	}
+	else
+	{
+		//Use Quaternion Interpolation. This is complicated since unlike Matinee we may not have
+		//a key for each Euler angle at the same times. So we need a way to calculate which 'range'
+		//to perform the interpolation. To do this if finds the exclusive closest range of keys to encomposs
+		//the passed in time
+		auto SetFrameRange = [Time](TRange<FFrameNumber> &FrameRange, const TArrayView<const FFrameNumber> &Times)
+		{
+			int32 Index1, Index2;
+			Index2 = 0;
+			Index2 = Algo::UpperBound(Times, Time.FrameNumber);
+			Index1 = Index2 - 1;
+			Index1 = Index1 >= 0 ? Index1 : INDEX_NONE;
+			Index2 = Index2 < Times.Num() ? Index2 : INDEX_NONE;
+			if (Index1 != INDEX_NONE && Index2 != INDEX_NONE)
+			{
+				if (Times[Index1] != Time.FrameNumber && Times[Index1] > FrameRange.GetLowerBoundValue()) 
+				{
+					FrameRange.SetLowerBoundValue(Times[Index1]);
+				}
+				if (Times[Index2] != Time.FrameNumber && Times[Index1] < FrameRange.GetUpperBoundValue())
+				{
+					FrameRange.SetUpperBoundValue(Times[Index2]);
+				}
+			}
+		};
+
+		TRange<FFrameNumber> FrameRange(TNumericLimits<FFrameNumber>::Min(), TNumericLimits<FFrameNumber>::Max());
+		if (EnumHasAllFlags(ChannelMask, EMovieSceneTransformChannel::RotationX))
+		{
+			SetFrameRange(FrameRange, RotationCurve[0].GetTimes());
+		}
+		if (EnumHasAllFlags(ChannelMask, EMovieSceneTransformChannel::RotationY))
+		{
+			SetFrameRange(FrameRange, RotationCurve[1].GetTimes());
+		}
+		if (EnumHasAllFlags(ChannelMask, EMovieSceneTransformChannel::RotationZ))
+		{
+			SetFrameRange(FrameRange, RotationCurve[2].GetTimes());
+		}
+		FFrameNumber LowerBound = FrameRange.GetLowerBoundValue();
+		FFrameNumber UpperBound = FrameRange.GetUpperBoundValue();
+		if (LowerBound != TNumericLimits<FFrameNumber>::Min() && UpperBound != TNumericLimits<FFrameNumber>::Max())
+		{
+			float Value;
+			FVector FirstRot(0.0f, 0.0f, 0.0f);
+			FVector SecondRot(0.0f, 0.0f, 0.0f);
+			double U = (Time.AsDecimal() - (double) FrameRange.GetLowerBoundValue().Value) /
+				double(FrameRange.GetUpperBoundValue().Value - FrameRange.GetLowerBoundValue().Value);
+			U = FMath::Clamp(U, 0.0, 1.0);
+
+			if (EnumHasAllFlags(ChannelMask, EMovieSceneTransformChannel::RotationX))
+			{
+				if (RotationCurve[0].Evaluate(LowerBound, Value))
+				{
+					FirstRot[0] = Value;
+				}
+				if (RotationCurve[0].Evaluate(UpperBound, Value))
+				{
+					SecondRot[0] = Value;
+				}
+			}
+			if (EnumHasAllFlags(ChannelMask, EMovieSceneTransformChannel::RotationY))
+			{
+				if (RotationCurve[1].Evaluate(LowerBound, Value))
+				{
+					FirstRot[1] = Value;
+				}
+				if (RotationCurve[1].Evaluate(UpperBound, Value))
+				{
+					SecondRot[1] = Value;
+				}
+			}
+			if (EnumHasAllFlags(ChannelMask, EMovieSceneTransformChannel::RotationZ))
+			{
+				if (RotationCurve[2].Evaluate(LowerBound, Value))
+				{
+					FirstRot[2] = Value;
+				}
+				if (RotationCurve[2].Evaluate(UpperBound, Value))
+				{
+					SecondRot[2] = Value;
+				}
+			}
+
+			const FQuat Key1Quat = FQuat::MakeFromEuler(FirstRot);
+			const FQuat Key2Quat = FQuat::MakeFromEuler(SecondRot);
+
+			const FQuat SlerpQuat = FQuat::Slerp(Key1Quat, Key2Quat, U);
+			FVector Euler = FRotator(SlerpQuat).Euler();
+			AnimatedData.Set(3, Euler[0]);
+			AnimatedData.Set(4, Euler[1]);
+			AnimatedData.Set(5, Euler[2]);
+		}
+
+		else  //no range found default to regular, but still do RotToQuat
+		{
+			float Value;
+			FVector CurrentRot(0.0f, 0.0f, 0.0f);
+			FRotator Rotator;
+			if (EnumHasAllFlags(ChannelMask, EMovieSceneTransformChannel::RotationX) && RotationCurve[0].Evaluate(Time, Value))
+			{
+				CurrentRot[0] = Value;
+			}
+			if (EnumHasAllFlags(ChannelMask, EMovieSceneTransformChannel::RotationY) && RotationCurve[1].Evaluate(Time, Value))
+			{
+				CurrentRot[1] = Value;
+			}
+			if (EnumHasAllFlags(ChannelMask, EMovieSceneTransformChannel::RotationZ) && RotationCurve[2].Evaluate(Time, Value))
+			{
+				CurrentRot[2] = Value;
+			}
+			FQuat Quat = FQuat::MakeFromEuler(CurrentRot);
+			FVector Euler = FRotator(Quat).Euler();
+			AnimatedData.Set(3, Euler[0]);
+			AnimatedData.Set(4, Euler[1]);
+			AnimatedData.Set(5, Euler[2]);
+		}
+	}
 
 	EvalChannel(6, EMovieSceneTransformChannel::ScaleX, ScaleCurve[0]);
 	EvalChannel(7, EMovieSceneTransformChannel::ScaleY, ScaleCurve[1]);
diff --git a/Engine/Source/Runtime/MovieSceneTracks/Private/Evaluation/MovieSceneAudioTemplate.cpp b/Engine/Source/Runtime/MovieSceneTracks/Private/Evaluation/MovieSceneAudioTemplate.cpp
--- a/Engine/Source/Runtime/MovieSceneTracks/Private/Evaluation/MovieSceneAudioTemplate.cpp
+++ b/Engine/Source/Runtime/MovieSceneTracks/Private/Evaluation/MovieSceneAudioTemplate.cpp
@@ -282,7 +282,7 @@ FMovieSceneAudioSectionTemplateData::FMovieSceneAudioSectionTemplateData(const U
 	, OnAudioFinished(Section.GetOnAudioFinished())
 	, OnAudioPlaybackPercent(Section.GetOnAudioPlaybackPercent())
 {
-	SectionStartTimeSeconds = (Section.HasStartFrame() ? Section.GetInclusiveStartFrame() : 0) / Section.GetTypedOuter<UMovieScene>()->GetFrameResolution();
+	SectionStartTimeSeconds = (Section.HasStartFrame() ? Section.GetInclusiveStartFrame() : 0) / Section.GetTypedOuter<UMovieScene>()->GetTickResolution();
 }
 
 void FMovieSceneAudioSectionTemplateData::EnsureAudioIsPlaying(UAudioComponent& AudioComponent, FPersistentEvaluationData& PersistentData, const FMovieSceneContext& Context, bool bAllowSpatialization, IMovieScenePlayer& Player) const
diff --git a/Engine/Source/Runtime/MovieSceneTracks/Private/Evaluation/MovieSceneEventTemplate.cpp b/Engine/Source/Runtime/MovieSceneTracks/Private/Evaluation/MovieSceneEventTemplate.cpp
--- a/Engine/Source/Runtime/MovieSceneTracks/Private/Evaluation/MovieSceneEventTemplate.cpp
+++ b/Engine/Source/Runtime/MovieSceneTracks/Private/Evaluation/MovieSceneEventTemplate.cpp
@@ -128,18 +128,18 @@ struct FEventTrackExecutionToken
 				{
 					if (!ParamInstanceIt)
 					{
-						UE_LOG(LogMovieScene, Warning, TEXT("Sequencer Event Track: Parameter count mistatch for event '%s'. Required parameter of type '%s' at index '%d'."), *Event.Payload.EventName.ToString(), *ParamIt->GetName(), NumParams);
+						UE_LOG(LogMovieScene, Warning, TEXT("Sequencer Event Track: Parameter count mismatch for event '%s'. Required parameter of type '%s' at index '%d'."), *Event.Payload.EventName.ToString(), *ParamIt->GetName(), NumParams);
 						return;
 					}
 					else if (!ParamIt)
 					{
 						// Mismatch (too many params)
-						UE_LOG(LogMovieScene, Warning, TEXT("Sequencer Event Track: Parameter count mistatch for event '%s'. Parameter struct contains too many parameters ('%s' is superfluous at index '%d'."), *Event.Payload.EventName.ToString(), *ParamInstanceIt->GetName(), NumParams);
+						UE_LOG(LogMovieScene, Warning, TEXT("Sequencer Event Track: Parameter count mismatch for event '%s'. Parameter struct contains too many parameters ('%s' is superfluous at index '%d'."), *Event.Payload.EventName.ToString(), *ParamInstanceIt->GetName(), NumParams);
 						return;
 					}
 					else if (!ParamInstanceIt->SameType(*ParamIt) || ParamInstanceIt->GetOffset_ForUFunction() != ParamIt->GetOffset_ForUFunction() || ParamInstanceIt->GetSize() != ParamIt->GetSize())
 					{
-						UE_LOG(LogMovieScene, Warning, TEXT("Sequencer Event Track: Parameter type mistatch for event '%s' ('%s' != '%s')."),
+						UE_LOG(LogMovieScene, Warning, TEXT("Sequencer Event Track: Parameter type mismatch for event '%s' ('%s' != '%s')."),
 							*Event.Payload.EventName.ToString(),
 							*ParamInstanceIt->GetClass()->GetName(),
 							*ParamIt->GetClass()->GetName()
diff --git a/Engine/Source/Runtime/MovieSceneTracks/Private/Evaluation/MovieSceneSkeletalAnimationTemplate.cpp b/Engine/Source/Runtime/MovieSceneTracks/Private/Evaluation/MovieSceneSkeletalAnimationTemplate.cpp
--- a/Engine/Source/Runtime/MovieSceneTracks/Private/Evaluation/MovieSceneSkeletalAnimationTemplate.cpp
+++ b/Engine/Source/Runtime/MovieSceneTracks/Private/Evaluation/MovieSceneSkeletalAnimationTemplate.cpp
@@ -114,13 +114,14 @@ struct FPreAnimatedAnimationTokenProducer : IMovieScenePreAnimatedTokenProducer
 
 struct FMinimalAnimParameters
 {
-	FMinimalAnimParameters(UAnimSequenceBase* InAnimation, float InEvalTime, float InBlendWeight, const FMovieSceneEvaluationScope& InScope, FName InSlotName, FObjectKey InSection)
+	FMinimalAnimParameters(UAnimSequenceBase* InAnimation, float InEvalTime, float InBlendWeight, const FMovieSceneEvaluationScope& InScope, FName InSlotName, FObjectKey InSection, bool InSkipAnimationNotifiers)
 		: Animation(InAnimation)
 		, EvalTime(InEvalTime)
 		, BlendWeight(InBlendWeight)
 		, EvaluationScope(InScope)
 		, SlotName(InSlotName)
 		, Section(InSection)
+		, bSkipAnimNotifiers(InSkipAnimationNotifiers)
 	{}
 	
 	UAnimSequenceBase* Animation;
@@ -129,6 +130,7 @@ struct FMinimalAnimParameters
 	FMovieSceneEvaluationScope EvaluationScope;
 	FName SlotName;
 	FObjectKey Section;
+	bool bSkipAnimNotifiers;
 };
 
 /** Montage player per section data */
@@ -196,6 +198,7 @@ namespace MovieScene
 			const EMovieScenePlayerStatus::Type PlayerStatus = Player.GetPlaybackStatus();
 
 			// If the playback status is jumping, ie. one such occurrence is setting the time for thumbnail generation, disable anim notifies updates because it could fire audio
+			// We now layer this with the passed in notify toggleto force a disable in this case.
 			const bool bFireNotifies = !bPreviewPlayback || (PlayerStatus != EMovieScenePlayerStatus::Jumping && PlayerStatus != EMovieScenePlayerStatus::Stopped);
 
 			// When jumping from one cut to another cut, the delta time should be 0 so that anim notifies before the current position are not evaluated. Note, anim notifies at the current time should still be evaluated.
@@ -215,13 +218,13 @@ namespace MovieScene
 				{
 					PreviewSetAnimPosition(PersistentData, Player, SkeletalMeshComponent,
 						AnimParams.SlotName, AnimParams.Section, AnimParams.Animation, AnimParams.EvalTime, AnimParams.BlendWeight,
-						bLooping, bFireNotifies, DeltaTime, Player.GetPlaybackStatus() == EMovieScenePlayerStatus::Playing, bResetDynamics);
+						bLooping, bFireNotifies && !AnimParams.bSkipAnimNotifiers, DeltaTime, Player.GetPlaybackStatus() == EMovieScenePlayerStatus::Playing, bResetDynamics);
 				}
 				else
 				{
 					SetAnimPosition(PersistentData, Player, SkeletalMeshComponent,
 						AnimParams.SlotName, AnimParams.Section, AnimParams.Animation, AnimParams.EvalTime, AnimParams.BlendWeight,
-						bLooping, Player.GetPlaybackStatus() == EMovieScenePlayerStatus::Playing, bFireNotifies);
+						bLooping, Player.GetPlaybackStatus() == EMovieScenePlayerStatus::Playing, bFireNotifies && !AnimParams.bSkipAnimNotifiers);
 				}
 			}
 
@@ -372,8 +375,8 @@ void FMovieSceneSkeletalAnimationSectionTemplate::Evaluate(const FMovieSceneEval
 
 		// Add the blendable to the accumulator
 		FMinimalAnimParameters AnimParams(
-			Params.Animation, EvalTime, Weight, ExecutionTokens.GetCurrentScope(), Params.SlotName, GetSourceSection()
-			);
+			Params.Animation, EvalTime, Weight, ExecutionTokens.GetCurrentScope(), Params.SlotName, GetSourceSection(), Params.bSkipAnimNotifiers
+		);
 		ExecutionTokens.BlendToken(ActuatorTypeID, TBlendableToken<MovieScene::FBlendedAnimation>(AnimParams, BlendType.Get(), 1.f));
 	}
 }
diff --git a/Engine/Source/Runtime/MovieSceneTracks/Private/Sections/MovieScene3DTransformSection.cpp b/Engine/Source/Runtime/MovieSceneTracks/Private/Sections/MovieScene3DTransformSection.cpp
--- a/Engine/Source/Runtime/MovieSceneTracks/Private/Sections/MovieScene3DTransformSection.cpp
+++ b/Engine/Source/Runtime/MovieSceneTracks/Private/Sections/MovieScene3DTransformSection.cpp
@@ -275,6 +275,7 @@ void FMovieScene3DTransformKeyStruct::PropagateChanges(const FPropertyChangedEve
 
 UMovieScene3DTransformSection::UMovieScene3DTransformSection(const FObjectInitializer& ObjectInitializer)
 	: Super(ObjectInitializer)
+	, bUseQuaternionInterpolation(false)
 #if WITH_EDITORONLY_DATA
 	, Show3DTrajectory(EShow3DTrajectory::EST_OnlyWhenSelected)
 #endif
@@ -315,6 +316,60 @@ void UMovieScene3DTransformSection::SetMask(FMovieSceneTransformMask NewMask)
 	UpdateChannelProxy();
 }
 
+FMovieSceneTransformMask UMovieScene3DTransformSection::GetMaskByName(const FName& InName) const
+{
+	if (InName == TEXT("Location"))
+	{
+		return EMovieSceneTransformChannel::Translation;
+	}
+	else if (InName == TEXT("Location.X"))
+	{
+		return EMovieSceneTransformChannel::TranslationX;
+	}
+	else if (InName == TEXT("Location.Y"))
+	{
+		return EMovieSceneTransformChannel::TranslationY;
+	}
+	else if (InName == TEXT("Location.Z"))
+	{
+		return EMovieSceneTransformChannel::TranslationZ;
+	}
+	else if (InName == TEXT("Rotation"))
+	{
+		return EMovieSceneTransformChannel::Rotation;
+	}
+	else if (InName == TEXT("Rotation.X"))
+	{
+		return EMovieSceneTransformChannel::RotationX;
+	}
+	else if (InName == TEXT("Rotation.Y"))
+	{
+		return EMovieSceneTransformChannel::RotationY;
+	}
+	else if (InName == TEXT("Rotation.Z"))
+	{
+		return EMovieSceneTransformChannel::RotationZ;
+	}
+	else if (InName == TEXT("Scale"))
+	{
+		return EMovieSceneTransformChannel::Scale;
+	}
+	else if (InName == TEXT("Scale.X"))
+	{
+		return EMovieSceneTransformChannel::ScaleX;
+	}
+	else if (InName == TEXT("Scale.Y"))
+	{
+		return EMovieSceneTransformChannel::ScaleY;
+	}
+	else if (InName == TEXT("Scale.Z"))
+	{
+		return EMovieSceneTransformChannel::ScaleZ;
+	}
+
+	return EMovieSceneTransformChannel::All;
+}
+
 void UMovieScene3DTransformSection::UpdateChannelProxy()
 {
 	if (ProxyChannels == TransformMask.GetChannels())
@@ -465,3 +520,22 @@ FMovieSceneEvalTemplatePtr UMovieScene3DTransformSection::GenerateTemplate() con
 {
 	return FMovieSceneComponentTransformSectionTemplate(*this);
 }
+
+
+bool UMovieScene3DTransformSection::GetUseQuaternionInterpolation() const
+{
+	return bUseQuaternionInterpolation;
+}
+
+bool UMovieScene3DTransformSection::ShowCurveForChannel(const void *ChannelPtr) const
+{
+	if (GetUseQuaternionInterpolation())
+	{
+		TArrayView<FMovieSceneFloatChannel*> FloatChannels = ChannelProxy->GetChannels<FMovieSceneFloatChannel>();
+		if (static_cast<void *>(FloatChannels[3]) == ChannelPtr || static_cast<void *>(FloatChannels[4]) == ChannelPtr || static_cast<void *>(FloatChannels[5]) == ChannelPtr)
+		{
+			return false;
+		}
+	}
+	return true;
+}
diff --git a/Engine/Source/Runtime/MovieSceneTracks/Private/Sections/MovieSceneAudioSection.cpp b/Engine/Source/Runtime/MovieSceneTracks/Private/Sections/MovieSceneAudioSection.cpp
--- a/Engine/Source/Runtime/MovieSceneTracks/Private/Sections/MovieSceneAudioSection.cpp
+++ b/Engine/Source/Runtime/MovieSceneTracks/Private/Sections/MovieSceneAudioSection.cpp
@@ -6,6 +6,7 @@
 #include "UObject/SequencerObjectVersion.h"
 #include "Channels/MovieSceneChannelProxy.h"
 #include "MovieScene.h"
+#include "MovieSceneCommonHelpers.h"
 
 #if WITH_EDITOR
 
@@ -67,7 +68,7 @@ FMovieSceneEvalTemplatePtr UMovieSceneAudioSection::GenerateTemplate() const
 
 TOptional<FFrameTime> UMovieSceneAudioSection::GetOffsetTime() const
 {
-	return TOptional<FFrameTime>(StartOffset * GetTypedOuter<UMovieScene>()->GetFrameResolution());
+	return TOptional<FFrameTime>(StartOffset * GetTypedOuter<UMovieScene>()->GetTickResolution());
 }
 
 void UMovieSceneAudioSection::PostLoad()
@@ -94,16 +95,56 @@ void UMovieSceneAudioSection::PostLoad()
 		// clip at the start of the section evaluation as such: Section Start Time - Start Time. 
 		if (AudioStartTime_DEPRECATED != 0.f && HasStartFrame())
 		{
-			StartOffset = GetInclusiveStartFrame() / GetTypedOuter<UMovieScene>()->GetFrameResolution() - AudioStartTime_DEPRECATED;
+			StartOffset = GetInclusiveStartFrame() / GetTypedOuter<UMovieScene>()->GetTickResolution() - AudioStartTime_DEPRECATED;
 		}
 		AudioStartTime_DEPRECATED = FLT_MAX;
 	}
 }
 
+float GetStartOffsetAtTrimTime(FQualifiedFrameTime TrimTime, float StartOffset, FFrameNumber StartFrame)
+{
+	return StartOffset + (TrimTime.Time - StartFrame) / TrimTime.Rate;
+}
+
+	
+TOptional<TRange<FFrameNumber> > UMovieSceneAudioSection::GetAutoSizeRange() const
+{
+	float SoundDuration = MovieSceneHelpers::GetSoundDuration(Sound);
+
+	FFrameRate FrameRate = GetTypedOuter<UMovieScene>()->GetTickResolution();
+
+	// determine initial duration
+	// @todo Once we have infinite sections, we can remove this
+	// @todo ^^ Why? Infinte sections would mean there's no starting time?
+	FFrameTime DurationToUse = 1.f * FrameRate; // if all else fails, use 1 second duration
+
+	if (SoundDuration != INDEFINITELY_LOOPING_DURATION)
+	{
+		DurationToUse = SoundDuration * FrameRate;
+	}
+
+	return TRange<FFrameNumber>(GetInclusiveStartFrame(), GetInclusiveStartFrame() + DurationToUse.FrameNumber);
+}
+
+	
+void UMovieSceneAudioSection::TrimSection(FQualifiedFrameTime TrimTime, bool bTrimLeft)
+{
+	SetFlags(RF_Transactional);
+
+	if (TryModify())
+	{
+		if (bTrimLeft)
+		{
+			StartOffset = HasStartFrame() ? GetStartOffsetAtTrimTime(TrimTime, StartOffset, GetInclusiveStartFrame()) : 0;
+		}
+
+		Super::TrimSection(TrimTime, bTrimLeft);
+	}
+}
 
-UMovieSceneSection* UMovieSceneAudioSection::SplitSection(FFrameNumber SplitTime)
+UMovieSceneSection* UMovieSceneAudioSection::SplitSection(FQualifiedFrameTime SplitTime)
 {
-	const float NewOffset = HasStartFrame() ? StartOffset + (SplitTime - GetInclusiveStartFrame()) / GetTypedOuter<UMovieScene>()->GetFrameResolution() : 0;
+	const float NewOffset = HasStartFrame() ? GetStartOffsetAtTrimTime(SplitTime, StartOffset, GetInclusiveStartFrame()) : 0;
 
 	UMovieSceneSection* NewSection = Super::SplitSection(SplitTime);
 	if (NewSection != nullptr)
diff --git a/Engine/Source/Runtime/MovieSceneTracks/Private/Sections/MovieSceneBoolSection.cpp b/Engine/Source/Runtime/MovieSceneTracks/Private/Sections/MovieSceneBoolSection.cpp
--- a/Engine/Source/Runtime/MovieSceneTracks/Private/Sections/MovieSceneBoolSection.cpp
+++ b/Engine/Source/Runtime/MovieSceneTracks/Private/Sections/MovieSceneBoolSection.cpp
@@ -11,6 +11,7 @@ UMovieSceneBoolSection::UMovieSceneBoolSection( const FObjectInitializer& Object
 	bSupportsInfiniteRange = true;
 #if WITH_EDITORONLY_DATA
 	bIsInfinite_DEPRECATED = true;
+	bIsExternallyInverted = false;
 #endif
 
 	SetRange(TRange<FFrameNumber>::All());
@@ -22,9 +23,48 @@ UMovieSceneBoolSection::UMovieSceneBoolSection( const FObjectInitializer& Object
 			EMovieSceneCompletionMode::RestoreState : 
 			EMovieSceneCompletionMode::ProjectDefault);
 
+	ReconstructChannelProxy();
+}
+
+#if WITH_EDITORONLY_DATA
+
+void UMovieSceneBoolSection::SetIsExternallyInverted(bool bInIsExternallyInverted)
+{
+	bIsExternallyInverted = bInIsExternallyInverted;
+	ReconstructChannelProxy();
+}
+
+void UMovieSceneBoolSection::Serialize(FArchive& Ar)
+{
+	Super::Serialize(Ar);
+
+	if (Ar.IsLoading())
+	{
+		ReconstructChannelProxy();
+	}
+}
+
+#endif
+
+
+void UMovieSceneBoolSection::ReconstructChannelProxy()
+{
 #if WITH_EDITOR
 
-	ChannelProxy = MakeShared<FMovieSceneChannelProxy>(BoolCurve, FMovieSceneChannelEditorData(), TMovieSceneExternalValue<bool>::Make());
+	// set up the external value to retrieve the inverted value if necessary. This is used by visibility tracks that
+	// are bound to "Actor Hidden in Game" properties, but displayed as "Visibility"
+	struct FGetInvertedBool
+	{
+		static TOptional<bool> GetValue(UObject& InObject, FTrackInstancePropertyBindings* Bindings)
+		{
+			return Bindings ? !Bindings->GetCurrentValue<bool>(InObject) : TOptional<bool>();
+		}
+	};
+
+	TMovieSceneExternalValue<bool> ExternalValue;
+	ExternalValue.OnGetExternalValue = bIsExternallyInverted ? FGetInvertedBool::GetValue :TMovieSceneExternalValue<bool>::GetValue;
+
+	ChannelProxy = MakeShared<FMovieSceneChannelProxy>(BoolCurve, FMovieSceneChannelEditorData(), ExternalValue);
 
 #else
 
@@ -33,7 +73,6 @@ UMovieSceneBoolSection::UMovieSceneBoolSection( const FObjectInitializer& Object
 #endif
 }
 
-
 void UMovieSceneBoolSection::PostLoad()
 {
 	if (!BoolCurve.GetDefault().IsSet() && DefaultValue_DEPRECATED)
diff --git a/Engine/Source/Runtime/MovieSceneTracks/Private/Sections/MovieSceneCameraCutSection.cpp b/Engine/Source/Runtime/MovieSceneTracks/Private/Sections/MovieSceneCameraCutSection.cpp
--- a/Engine/Source/Runtime/MovieSceneTracks/Private/Sections/MovieSceneCameraCutSection.cpp
+++ b/Engine/Source/Runtime/MovieSceneTracks/Private/Sections/MovieSceneCameraCutSection.cpp
@@ -31,7 +31,7 @@ FMovieSceneEvalTemplatePtr UMovieSceneCameraCutSection::GenerateTemplate() const
 				{
 					// Extract the transform
 					FMovieSceneEvaluationTrack TransformTrackTemplate = TransformTrack->GenerateTrackTemplate();
-					FMovieSceneContext Context = FMovieSceneEvaluationRange(GetInclusiveStartFrame(), MovieScene->GetFrameResolution());
+					FMovieSceneContext Context = FMovieSceneEvaluationRange(GetInclusiveStartFrame(), MovieScene->GetTickResolution());
 
 					FMovieSceneInterrogationData Container;
 					TransformTrackTemplate.Interrogate(Context, Container);
diff --git a/Engine/Source/Runtime/MovieSceneTracks/Private/Sections/MovieSceneSkeletalAnimationSection.cpp b/Engine/Source/Runtime/MovieSceneTracks/Private/Sections/MovieSceneSkeletalAnimationSection.cpp
--- a/Engine/Source/Runtime/MovieSceneTracks/Private/Sections/MovieSceneSkeletalAnimationSection.cpp
+++ b/Engine/Source/Runtime/MovieSceneTracks/Private/Sections/MovieSceneSkeletalAnimationSection.cpp
@@ -25,6 +25,7 @@ FMovieSceneSkeletalAnimationParams::FMovieSceneSkeletalAnimationParams()
 	bReverse = false;
 	SlotName = DefaultSlotName;
 	Weight.SetDefault(1.f);
+	bSkipAnimNotifiers = false;
 }
 
 UMovieSceneSkeletalAnimationSection::UMovieSceneSkeletalAnimationSection( const FObjectInitializer& ObjectInitializer )
@@ -61,7 +62,7 @@ UMovieSceneSkeletalAnimationSection::UMovieSceneSkeletalAnimationSection( const
 
 TOptional<FFrameTime> UMovieSceneSkeletalAnimationSection::GetOffsetTime() const
 {
-	return TOptional<FFrameTime>(Params.StartOffset * GetTypedOuter<UMovieScene>()->GetFrameResolution());
+	return TOptional<FFrameTime>(Params.StartOffset * GetTypedOuter<UMovieScene>()->GetTickResolution());
 }
 
 void UMovieSceneSkeletalAnimationSection::Serialize(FArchive& Ar)
@@ -148,15 +149,48 @@ FMovieSceneEvalTemplatePtr UMovieSceneSkeletalAnimationSection::GenerateTemplate
 	return FMovieSceneSkeletalAnimationSectionTemplate(*this);
 }
 
-UMovieSceneSection* UMovieSceneSkeletalAnimationSection::SplitSection(FFrameNumber SplitTime)
+float GetStartOffsetAtTrimTime(FQualifiedFrameTime TrimTime, const FMovieSceneSkeletalAnimationParams& Params, FFrameNumber StartFrame)
 {
 	float AnimPlayRate = FMath::IsNearlyZero(Params.PlayRate) ? 1.0f : Params.PlayRate;
-	float AnimPosition = (SplitTime - GetRange().GetLowerBoundValue()) / GetTypedOuter<UMovieScene>()->GetFrameResolution() * AnimPlayRate;
+	float AnimPosition = (TrimTime.Time - StartFrame) / TrimTime.Rate * AnimPlayRate;
 	float SeqLength = Params.GetSequenceLength() - (Params.StartOffset + Params.EndOffset);
 
 	float NewOffset = FMath::Fmod(AnimPosition, SeqLength);
 	NewOffset += Params.StartOffset;
 
+	return NewOffset;
+}
+
+
+TOptional<TRange<FFrameNumber> > UMovieSceneSkeletalAnimationSection::GetAutoSizeRange() const
+{
+	FFrameRate FrameRate = GetTypedOuter<UMovieScene>()->GetTickResolution();
+
+	FFrameTime AnimationLength = Params.GetSequenceLength() * FrameRate;
+
+	return TRange<FFrameNumber>(GetInclusiveStartFrame(), GetInclusiveStartFrame() + AnimationLength.FrameNumber);
+}
+
+
+void UMovieSceneSkeletalAnimationSection::TrimSection(FQualifiedFrameTime TrimTime, bool bTrimLeft)
+{
+	SetFlags(RF_Transactional);
+
+	if (TryModify())
+	{
+		if (bTrimLeft)
+		{
+			Params.StartOffset = HasStartFrame() ? GetStartOffsetAtTrimTime(TrimTime, Params, GetInclusiveStartFrame()) : 0;
+		}
+
+		Super::TrimSection(TrimTime, bTrimLeft);
+	}
+}
+
+UMovieSceneSection* UMovieSceneSkeletalAnimationSection::SplitSection(FQualifiedFrameTime SplitTime)
+{
+	const float NewOffset = HasStartFrame() ? GetStartOffsetAtTrimTime(SplitTime, Params, GetInclusiveStartFrame()) : 0;
+
 	UMovieSceneSection* NewSection = Super::SplitSection(SplitTime);
 	if (NewSection != nullptr)
 	{
@@ -171,7 +205,7 @@ void UMovieSceneSkeletalAnimationSection::GetSnapTimes(TArray<FFrameNumber>& Out
 {
 	Super::GetSnapTimes(OutSnapTimes, bGetSectionBorders);
 
-	const FFrameRate   FrameRate  = GetTypedOuter<UMovieScene>()->GetFrameResolution();
+	const FFrameRate   FrameRate  = GetTypedOuter<UMovieScene>()->GetTickResolution();
 	const FFrameNumber StartFrame = GetInclusiveStartFrame();
 	const FFrameNumber EndFrame   = GetExclusiveEndFrame() - 1; // -1 because we don't need to add the end frame twice
 
@@ -191,6 +225,13 @@ void UMovieSceneSkeletalAnimationSection::GetSnapTimes(TArray<FFrameNumber>& Out
 	}
 }
 
+float UMovieSceneSkeletalAnimationSection::MapTimeToAnimation(FFrameTime InPosition, FFrameRate InFrameRate) const
+{
+	FMovieSceneSkeletalAnimationSectionTemplateParameters TemplateParams(Params, GetInclusiveStartFrame(), GetExclusiveEndFrame());
+	return TemplateParams.MapTimeToAnimation(InPosition, InFrameRate);
+}
+
+
 #if WITH_EDITOR
 void UMovieSceneSkeletalAnimationSection::PreEditChange(UProperty* PropertyAboutToChange)
 {
@@ -210,9 +251,9 @@ void UMovieSceneSkeletalAnimationSection::PostEditChangeProperty(FPropertyChange
 
 		if (!FMath::IsNearlyZero(NewPlayRate))
 		{
-			float CurrentDuration = MovieScene::DiscreteSize(GetRange()) * PreviousPlayRate;
+			float CurrentDuration = MovieScene::DiscreteSize(GetRange());
 			float NewDuration = CurrentDuration * (PreviousPlayRate / NewPlayRate);
-			SetEndFrame( GetInclusiveStartFrame() + FMath::FloorToInt(NewDuration * NewPlayRate) );
+			SetEndFrame( GetInclusiveStartFrame() + FMath::FloorToInt(NewDuration) );
 
 			PreviousPlayRate = NewPlayRate;
 		}
diff --git a/Engine/Source/Runtime/MovieSceneTracks/Private/Tracks/MovieSceneAudioTrack.cpp b/Engine/Source/Runtime/MovieSceneTracks/Private/Tracks/MovieSceneAudioTrack.cpp
--- a/Engine/Source/Runtime/MovieSceneTracks/Private/Tracks/MovieSceneAudioTrack.cpp
+++ b/Engine/Source/Runtime/MovieSceneTracks/Private/Tracks/MovieSceneAudioTrack.cpp
@@ -3,16 +3,14 @@
 #include "Tracks/MovieSceneAudioTrack.h"
 #include "Audio.h"
 #include "Sound/SoundWave.h"
-#include "Sound/SoundCue.h"
 #include "MovieScene.h"
 #include "Sections/MovieSceneAudioSection.h"
-#include "Sound/SoundNodeWavePlayer.h"
 #include "Kismet/GameplayStatics.h"
 #include "AudioDecompress.h"
 #include "Evaluation/MovieSceneSegment.h"
 #include "Compilation/MovieSceneSegmentCompiler.h"
 #include "Compilation/MovieSceneCompilerRules.h"
-
+#include "MovieSceneCommonHelpers.h"
 
 #define LOCTEXT_NAMESPACE "MovieSceneAudioTrack"
 
@@ -68,48 +66,19 @@ bool UMovieSceneAudioTrack::IsEmpty() const
 	return AudioSections.Num() == 0;
 }
 
-float UMovieSceneAudioTrack::GetSoundDuration(USoundBase* Sound)
-{
-	USoundWave* SoundWave = nullptr;
-
-	if (Sound->IsA<USoundWave>())
-	{
-		SoundWave = Cast<USoundWave>(Sound);
-	}
-	else if (Sound->IsA<USoundCue>())
-	{
-#if WITH_EDITORONLY_DATA
-		USoundCue* SoundCue = Cast<USoundCue>(Sound);
-
-		// @todo Sequencer - Right now for sound cues, we just use the first sound wave in the cue
-		// In the future, it would be better to properly generate the sound cue's data after forcing determinism
-		const TArray<USoundNode*>& AllNodes = SoundCue->AllNodes;
-		for (int32 i = 0; i < AllNodes.Num() && SoundWave == nullptr; ++i)
-		{
-			if (AllNodes[i]->IsA<USoundNodeWavePlayer>())
-			{
-				SoundWave = Cast<USoundNodeWavePlayer>(AllNodes[i])->GetSoundWave();
-			}
-		}
-#endif
-	}
-
-	const float Duration = (SoundWave ? SoundWave->GetDuration() : 0.f);
-	return Duration == INDEFINITELY_LOOPING_DURATION ? SoundWave->Duration : Duration;
-}
 
 UMovieSceneSection* UMovieSceneAudioTrack::AddNewSoundOnRow(USoundBase* Sound, FFrameNumber Time, int32 RowIndex)
 {
 	check(Sound);
 	
-	FFrameRate FrameRate = GetTypedOuter<UMovieScene>()->GetFrameResolution();
+	FFrameRate FrameRate = GetTypedOuter<UMovieScene>()->GetTickResolution();
 
 	// determine initial duration
 	// @todo Once we have infinite sections, we can remove this
 	// @todo ^^ Why? Infinte sections would mean there's no starting time?
 	FFrameTime DurationToUse = 1.f * FrameRate; // if all else fails, use 1 second duration
 
-	float SoundDuration = GetSoundDuration(Sound);
+	float SoundDuration = MovieSceneHelpers::GetSoundDuration(Sound);
 	if (SoundDuration != INDEFINITELY_LOOPING_DURATION)
 	{
 		DurationToUse = SoundDuration * FrameRate;
diff --git a/Engine/Source/Runtime/MovieSceneTracks/Private/Tracks/MovieSceneCameraAnimTrack.cpp b/Engine/Source/Runtime/MovieSceneTracks/Private/Tracks/MovieSceneCameraAnimTrack.cpp
--- a/Engine/Source/Runtime/MovieSceneTracks/Private/Tracks/MovieSceneCameraAnimTrack.cpp
+++ b/Engine/Source/Runtime/MovieSceneTracks/Private/Tracks/MovieSceneCameraAnimTrack.cpp
@@ -19,7 +19,7 @@ UMovieSceneSection* UMovieSceneCameraAnimTrack::AddNewCameraAnim(FFrameNumber Ke
 	UMovieSceneCameraAnimSection* const NewSection = Cast<UMovieSceneCameraAnimSection>(CreateNewSection());
 	if (NewSection)
 	{
-		FFrameTime AnimDurationFrames = CameraAnim->AnimLength * GetTypedOuter<UMovieScene>()->GetFrameResolution();
+		FFrameTime AnimDurationFrames = CameraAnim->AnimLength * GetTypedOuter<UMovieScene>()->GetTickResolution();
 		NewSection->InitialPlacement(CameraAnimSections, KeyTime, AnimDurationFrames.FrameNumber.Value, SupportsMultipleRows());
 		NewSection->AnimData.CameraAnim = CameraAnim;
 
diff --git a/Engine/Source/Runtime/MovieSceneTracks/Private/Tracks/MovieSceneCameraCutTrack.cpp b/Engine/Source/Runtime/MovieSceneTracks/Private/Tracks/MovieSceneCameraCutTrack.cpp
--- a/Engine/Source/Runtime/MovieSceneTracks/Private/Tracks/MovieSceneCameraCutTrack.cpp
+++ b/Engine/Source/Runtime/MovieSceneTracks/Private/Tracks/MovieSceneCameraCutTrack.cpp
@@ -134,7 +134,7 @@ FFrameNumber UMovieSceneCameraCutTrack::FindEndTimeForCameraCut( FFrameNumber St
 	if( StartTime == ExclusiveEndTime )
 	{
 		// Give the CameraCut a reasonable length of time to start out with.  A 0 time CameraCut is not usable
-		ExclusiveEndTime = (StartTime + .5f * OwnerScene->GetFrameResolution()).FrameNumber;
+		ExclusiveEndTime = (StartTime + .5f * OwnerScene->GetTickResolution()).FrameNumber;
 	}
 
 	return ExclusiveEndTime;
diff --git a/Engine/Source/Runtime/MovieSceneTracks/Private/Tracks/MovieSceneCameraShakeTrack.cpp b/Engine/Source/Runtime/MovieSceneTracks/Private/Tracks/MovieSceneCameraShakeTrack.cpp
--- a/Engine/Source/Runtime/MovieSceneTracks/Private/Tracks/MovieSceneCameraShakeTrack.cpp
+++ b/Engine/Source/Runtime/MovieSceneTracks/Private/Tracks/MovieSceneCameraShakeTrack.cpp
@@ -18,7 +18,7 @@ UMovieSceneSection* UMovieSceneCameraShakeTrack::AddNewCameraShake(FFrameNumber
 	if (NewSection)
 	{
 		// #fixme get length
-		FFrameTime Duration = 5.0 * GetTypedOuter<UMovieScene>()->GetFrameResolution();
+		FFrameTime Duration = 5.0 * GetTypedOuter<UMovieScene>()->GetTickResolution();
 		NewSection->InitialPlacement(CameraShakeSections, KeyTime, Duration.FrameNumber.Value, SupportsMultipleRows());
 		NewSection->ShakeData.ShakeClass = ShakeClass;
 		
diff --git a/Engine/Source/Runtime/MovieSceneTracks/Private/Tracks/MovieSceneParticleTrack.cpp b/Engine/Source/Runtime/MovieSceneTracks/Private/Tracks/MovieSceneParticleTrack.cpp
--- a/Engine/Source/Runtime/MovieSceneTracks/Private/Tracks/MovieSceneParticleTrack.cpp
+++ b/Engine/Source/Runtime/MovieSceneTracks/Private/Tracks/MovieSceneParticleTrack.cpp
@@ -58,9 +58,6 @@ void UMovieSceneParticleTrack::AddNewSection( FFrameNumber SectionTime )
 	if ( MovieSceneHelpers::FindSectionAtTime( ParticleSections, SectionTime ) == nullptr )
 	{
 		UMovieSceneParticleSection* NewSection = Cast<UMovieSceneParticleSection>( CreateNewSection() );
-
-		NewSection->SetRange(TRange<FFrameNumber>::Inclusive(SectionTime, SectionTime));
-
 		ParticleSections.Add(NewSection);
 	}
 }
diff --git a/Engine/Source/Runtime/MovieSceneTracks/Private/Tracks/MovieScenePropertyTrack.cpp b/Engine/Source/Runtime/MovieSceneTracks/Private/Tracks/MovieScenePropertyTrack.cpp
--- a/Engine/Source/Runtime/MovieSceneTracks/Private/Tracks/MovieScenePropertyTrack.cpp
+++ b/Engine/Source/Runtime/MovieSceneTracks/Private/Tracks/MovieScenePropertyTrack.cpp
@@ -174,8 +174,6 @@ UMovieSceneSection* UMovieScenePropertyTrack::FindOrAddSection(FFrameNumber Time
 		return FoundSection;
 	}
 
-	check(Sections.Num() == 0);
-
 	// Add a new section that starts and ends at the same time
 	UMovieSceneSection* NewSection = CreateNewSection();
 	NewSection->SetFlags(RF_Transactional);
diff --git a/Engine/Source/Runtime/MovieSceneTracks/Private/Tracks/MovieSceneSkeletalAnimationTrack.cpp b/Engine/Source/Runtime/MovieSceneTracks/Private/Tracks/MovieSceneSkeletalAnimationTrack.cpp
--- a/Engine/Source/Runtime/MovieSceneTracks/Private/Tracks/MovieSceneSkeletalAnimationTrack.cpp
+++ b/Engine/Source/Runtime/MovieSceneTracks/Private/Tracks/MovieSceneSkeletalAnimationTrack.cpp
@@ -36,7 +36,7 @@ UMovieSceneSection* UMovieSceneSkeletalAnimationTrack::AddNewAnimationOnRow(FFra
 {
 	UMovieSceneSkeletalAnimationSection* NewSection = Cast<UMovieSceneSkeletalAnimationSection>(CreateNewSection());
 	{
-		FFrameTime AnimationLength = AnimSequence->SequenceLength * GetTypedOuter<UMovieScene>()->GetFrameResolution();
+		FFrameTime AnimationLength = AnimSequence->SequenceLength * GetTypedOuter<UMovieScene>()->GetTickResolution();
 		NewSection->InitialPlacementOnRow(AnimationSections, KeyTime, AnimationLength.FrameNumber.Value, RowIndex);
 		NewSection->Params.Animation = AnimSequence;
 	}
diff --git a/Engine/Source/Runtime/MovieSceneTracks/Public/Evaluation/MovieScene3DTransformTemplate.h b/Engine/Source/Runtime/MovieSceneTracks/Public/Evaluation/MovieScene3DTransformTemplate.h
--- a/Engine/Source/Runtime/MovieSceneTracks/Public/Evaluation/MovieScene3DTransformTemplate.h
+++ b/Engine/Source/Runtime/MovieSceneTracks/Public/Evaluation/MovieScene3DTransformTemplate.h
@@ -19,7 +19,7 @@ struct FMovieScene3DTransformTemplateData
 {
 	GENERATED_BODY()
 
-	FMovieScene3DTransformTemplateData(){}
+	FMovieScene3DTransformTemplateData(): bUseQuaternionInterpolation(false) {}
 	FMovieScene3DTransformTemplateData(const UMovieScene3DTransformSection& Section);
 
 	MovieScene::TMultiChannelValue<float, 9> Evaluate(FFrameTime InTime) const;
@@ -41,6 +41,9 @@ struct FMovieScene3DTransformTemplateData
 
 	UPROPERTY()
 	FMovieSceneTransformMask Mask;
+
+	UPROPERTY()
+	bool bUseQuaternionInterpolation;
 };
 
 USTRUCT()
diff --git a/Engine/Source/Runtime/MovieSceneTracks/Public/Sections/MovieScene3DTransformSection.h b/Engine/Source/Runtime/MovieSceneTracks/Public/Sections/MovieScene3DTransformSection.h
--- a/Engine/Source/Runtime/MovieSceneTracks/Public/Sections/MovieScene3DTransformSection.h
+++ b/Engine/Source/Runtime/MovieSceneTracks/Public/Sections/MovieScene3DTransformSection.h
@@ -244,6 +244,12 @@ class UMovieScene3DTransformSection
 {
 	GENERATED_UCLASS_BODY()
 
+public:
+
+	/* From UMovieSection*/
+	
+	virtual bool ShowCurveForChannel(const void *Channel) const override;
+
 public:
 
 	/**
@@ -256,6 +262,16 @@ public:
 	 */
 	MOVIESCENETRACKS_API void SetMask(FMovieSceneTransformMask NewMask);
 
+	/**
+	 * Get the mask by name
+	 */
+	MOVIESCENETRACKS_API FMovieSceneTransformMask GetMaskByName(const FName& InName) const;
+
+	/**
+	* Get whether we should use quaternion interpolation for our rotations.
+	*/
+	MOVIESCENETRACKS_API bool GetUseQuaternionInterpolation() const;
+
 protected:
 
 	virtual void Serialize(FArchive& Ar) override;
@@ -288,6 +304,10 @@ private:
 	/** Unserialized mask that defines the mask of the current channel proxy so we don't needlessly re-create it on post-undo */
 	EMovieSceneTransformChannel ProxyChannels;
 
+	/** Whether to use a quaternion linear interpolation between keys. This finds the 'shortest' distance between keys */
+	UPROPERTY(EditAnywhere, DisplayName = "Use Quaternion Interpolation", Category = "Rotation")
+	bool bUseQuaternionInterpolation;
+
 #if WITH_EDITORONLY_DATA
 
 public:
diff --git a/Engine/Source/Runtime/MovieSceneTracks/Public/Sections/MovieSceneAudioSection.h b/Engine/Source/Runtime/MovieSceneTracks/Public/Sections/MovieSceneAudioSection.h
--- a/Engine/Source/Runtime/MovieSceneTracks/Public/Sections/MovieSceneAudioSection.h
+++ b/Engine/Source/Runtime/MovieSceneTracks/Public/Sections/MovieSceneAudioSection.h
@@ -139,7 +139,9 @@ public:
 public:
 
 	//~ UMovieSceneSection interface
-	virtual UMovieSceneSection* SplitSection(FFrameNumber SplitTime) override;
+	virtual TOptional<TRange<FFrameNumber> > GetAutoSizeRange() const override;
+	virtual void TrimSection(FQualifiedFrameTime TrimTime, bool bTrimLeft) override;
+	virtual UMovieSceneSection* SplitSection(FQualifiedFrameTime SplitTime) override;
 	virtual TOptional<FFrameTime> GetOffsetTime() const override;
 	virtual FMovieSceneEvalTemplatePtr GenerateTemplate() const override;
 
@@ -166,11 +168,11 @@ private:
 	float AudioVolume_DEPRECATED;
 
 	/** The volume the sound will be played with. */
-	UPROPERTY( EditAnywhere, Category = "Audio" )
+	UPROPERTY( )
 	FMovieSceneFloatChannel SoundVolume;
 
 	/** The pitch multiplier the sound will be played with. */
-	UPROPERTY( EditAnywhere, Category = "Audio" )
+	UPROPERTY( )
 	FMovieSceneFloatChannel PitchMultiplier;
 
 	UPROPERTY( EditAnywhere, Category="Audio" )
diff --git a/Engine/Source/Runtime/MovieSceneTracks/Public/Sections/MovieSceneBoolSection.h b/Engine/Source/Runtime/MovieSceneTracks/Public/Sections/MovieSceneBoolSection.h
--- a/Engine/Source/Runtime/MovieSceneTracks/Public/Sections/MovieSceneBoolSection.h
+++ b/Engine/Source/Runtime/MovieSceneTracks/Public/Sections/MovieSceneBoolSection.h
@@ -36,7 +36,32 @@ public:
 
 protected:
 
+	/**
+	 * Update the channel proxy to ensure it has the correct flags and pointers
+	 */
+	void ReconstructChannelProxy();
+
 	/** Ordered curve data */
 	UPROPERTY()
 	FMovieSceneBoolChannel BoolCurve;
+
+#if WITH_EDITORONLY_DATA
+
+public:
+
+	/**
+	 * Set a flag indicating that the actual property that this bool represents is the opposite of the values stored in this section
+	 */
+	void SetIsExternallyInverted(bool bInIsExternallyInverted);
+
+protected:
+
+	/** Overloaded serializer to ensure that the channel proxy is updated correctly on load and duplicate */
+	virtual void Serialize(FArchive& Ar) override;
+
+	/** True if this section represents a property that is the inversion of the values stored on this channel */
+	UPROPERTY()
+	bool bIsExternallyInverted;
+
+#endif
 };
diff --git a/Engine/Source/Runtime/MovieSceneTracks/Public/Sections/MovieSceneSkeletalAnimationSection.h b/Engine/Source/Runtime/MovieSceneTracks/Public/Sections/MovieSceneSkeletalAnimationSection.h
--- a/Engine/Source/Runtime/MovieSceneTracks/Public/Sections/MovieSceneSkeletalAnimationSection.h
+++ b/Engine/Source/Runtime/MovieSceneTracks/Public/Sections/MovieSceneSkeletalAnimationSection.h
@@ -47,8 +47,13 @@ struct FMovieSceneSkeletalAnimationParams
 	FName SlotName;
 
 	/** The weight curve for this animation section */
-	UPROPERTY( EditAnywhere, Category = "Animation" )
+	UPROPERTY( )
 	FMovieSceneFloatChannel Weight;
+
+	/** If on will skip sending animation notifies */
+	UPROPERTY(EditAnywhere, Category = "Animation")
+	bool bSkipAnimNotifiers;
+
 };
 
 /**
@@ -65,10 +70,15 @@ public:
 	UPROPERTY(EditAnywhere, Category = "Animation", meta=(ShowOnlyInnerProperties))
 	FMovieSceneSkeletalAnimationParams Params;
 
+	/** Get Frame Time as Animation Time*/
+	MOVIESCENETRACKS_API float MapTimeToAnimation(FFrameTime InPosition, FFrameRate InFrameRate) const;
+
 protected:
 
 	//~ UMovieSceneSection interface
-	virtual UMovieSceneSection* SplitSection(FFrameNumber SplitTime) override;
+	virtual TOptional<TRange<FFrameNumber> > GetAutoSizeRange() const override;
+	virtual void TrimSection(FQualifiedFrameTime TrimTime, bool bTrimLeft) override;
+	virtual UMovieSceneSection* SplitSection(FQualifiedFrameTime SplitTime) override;
 	virtual void GetSnapTimes(TArray<FFrameNumber>& OutSnapTimes, bool bGetSectionBorders) const override;
 	virtual TOptional<FFrameTime> GetOffsetTime() const override;
 	virtual FMovieSceneEvalTemplatePtr GenerateTemplate() const override;
diff --git a/Engine/Source/Runtime/MovieSceneTracks/Public/Tracks/MovieSceneAudioTrack.h b/Engine/Source/Runtime/MovieSceneTracks/Public/Tracks/MovieSceneAudioTrack.h
--- a/Engine/Source/Runtime/MovieSceneTracks/Public/Tracks/MovieSceneAudioTrack.h
+++ b/Engine/Source/Runtime/MovieSceneTracks/Public/Tracks/MovieSceneAudioTrack.h
@@ -41,8 +41,6 @@ public:
 	/** @return true if this is a master audio track */
 	bool IsAMasterTrack() const;
 
-	float GetSoundDuration(USoundBase* Sound);
-
 public:
 
 	// UMovieSceneTrack interface
diff --git a/Engine/Source/Runtime/MovieSceneTracks/Public/Tracks/MovieSceneVisibilityTrack.h b/Engine/Source/Runtime/MovieSceneTracks/Public/Tracks/MovieSceneVisibilityTrack.h
--- a/Engine/Source/Runtime/MovieSceneTracks/Public/Tracks/MovieSceneVisibilityTrack.h
+++ b/Engine/Source/Runtime/MovieSceneTracks/Public/Tracks/MovieSceneVisibilityTrack.h
@@ -20,6 +20,8 @@ public:
 
 	// UMovieSceneTrack interface
 	virtual FMovieSceneEvalTemplatePtr CreateTemplateForSection(const UMovieSceneSection& InSection) const override;
+	virtual void PostLoad() override;
+	virtual UMovieSceneSection* CreateNewSection() override;
 
 #if WITH_EDITORONLY_DATA
 	virtual FText GetDisplayName() const override;
diff --git a/Engine/Source/Runtime/Sockets/Public/Sockets.h b/Engine/Source/Runtime/Sockets/Public/Sockets.h
--- a/Engine/Source/Runtime/Sockets/Public/Sockets.h
+++ b/Engine/Source/Runtime/Sockets/Public/Sockets.h
@@ -142,21 +142,31 @@ public:
 	/**
 	 * Reads a chunk of data from the socket and gathers the source address.
 	 *
+	 * A return value of 'true' does not necessarily mean that data was returned.
+	 * Callers must check the 'BytesRead' parameter for the actual amount of data
+	 * returned. A value of zero indicates that there was no data available for reading.
+	 *
 	 * @param Data The buffer to read into.
 	 * @param BufferSize The max size of the buffer.
 	 * @param BytesRead Will indicate how many bytes were read from the socket.
 	 * @param Source Will contain the receiving the address of the sender of the data.
 	 * @param Flags The receive flags.
+	 * @return true on success, false in case of a closed socket or an unrecoverable error.
 	 */
 	virtual bool RecvFrom(uint8* Data, int32 BufferSize, int32& BytesRead, FInternetAddr& Source, ESocketReceiveFlags::Type Flags = ESocketReceiveFlags::None);
 
 	/**
 	 * Reads a chunk of data from a connected socket
 	 *
+	 * A return value of 'true' does not necessarily mean that data was returned.
+	 * Callers must check the 'BytesRead' parameter for the actual amount of data
+	 * returned. A value of zero indicates that there was no data available for reading.
+	 *
 	 * @param Data The buffer to read into
 	 * @param BufferSize The max size of the buffer
 	 * @param BytesRead Will indicate how many bytes were read from the socket
 	 * @param Flags the receive flags
+	 * @return true on success, false in case of a closed socket or an unrecoverable error.
 	 */
 	virtual bool Recv(uint8* Data, int32 BufferSize, int32& BytesRead, ESocketReceiveFlags::Type Flags = ESocketReceiveFlags::None);
 
