Sequencer: Check for a valid segment ID in EvaluateStatic

diff --git a/Engine/Source/Runtime/MovieScene/Private/Evaluation/MovieSceneEvaluationTrack.cpp b/Engine/Source/Runtime/MovieScene/Private/Evaluation/MovieSceneEvaluationTrack.cpp
--- a/Engine/Source/Runtime/MovieScene/Private/Evaluation/MovieSceneEvaluationTrack.cpp
+++ b/Engine/Source/Runtime/MovieScene/Private/Evaluation/MovieSceneEvaluationTrack.cpp
@@ -442,6 +442,12 @@ void FMovieSceneEvaluationTrack::DefaultEvaluate(FMovieSceneSegmentIdentifier Se
 
 void FMovieSceneEvaluationTrack::EvaluateStatic(FMovieSceneSegmentIdentifier SegmentID, const FMovieSceneEvaluationOperand& Operand, FMovieSceneContext Context, const FPersistentEvaluationData& PersistentData, FMovieSceneExecutionTokens& ExecutionTokens) const
 {
+	int32 SortedIndex = Segments.GetSortedIndex(SegmentID);
+	if (SortedIndex == INDEX_NONE)
+	{
+		return;
+	}
+
 	for (const FSectionEvaluationData& EvalData : GetSegment(SegmentID).Impls)
 	{
 		const FMovieSceneEvalTemplate& Template = GetChildTemplate(EvalData.ImplIndex);
diff --git a/Engine/Source/Runtime/MovieScene/Public/Evaluation/MovieSceneEvaluationTrack.h b/Engine/Source/Runtime/MovieScene/Public/Evaluation/MovieSceneEvaluationTrack.h
--- a/Engine/Source/Runtime/MovieScene/Public/Evaluation/MovieSceneEvaluationTrack.h
+++ b/Engine/Source/Runtime/MovieScene/Public/Evaluation/MovieSceneEvaluationTrack.h
@@ -68,7 +68,7 @@ struct FMovieSceneEvaluationTrackSegments
 	/** Access the sorted index of the specified valid identifier */
 	FORCEINLINE int32 GetSortedIndex(FMovieSceneSegmentIdentifier ID) const
 	{
-		return SegmentIdentifierToIndex[ID.GetIndex()];
+		return SegmentIdentifierToIndex.IsValidIndex(ID.GetIndex()) ? SegmentIdentifierToIndex[ID.GetIndex()] : INDEX_NONE;
 	}
 
 	/** Access the sorted array of segments */
