Sequencer: Fixed RootToSequence transform not being used when compiling segments with "Evaluate Sub Sequences in Isolation" enabled

diff --git a/Engine/Source/Runtime/MovieScene/Private/Evaluation/MovieSceneEvaluationTemplateInstance.cpp b/Engine/Source/Runtime/MovieScene/Private/Evaluation/MovieSceneEvaluationTemplateInstance.cpp
--- a/Engine/Source/Runtime/MovieScene/Private/Evaluation/MovieSceneEvaluationTemplateInstance.cpp
+++ b/Engine/Source/Runtime/MovieScene/Private/Evaluation/MovieSceneEvaluationTemplateInstance.cpp
@@ -194,8 +194,16 @@ const FMovieSceneEvaluationGroup* FMovieSceneRootEvaluationTemplateInstance::Set
 		FMovieSceneEvaluationTemplateGenerator(*RootOverrideInstance.Sequence, *RootOverrideInstance.Template).Generate();
 	}
 
+	FMovieSceneSequenceTransform MasterToRootOverrideTransform;
+	if (RootOverrideInstance.SubData)
+	{
+		MasterToRootOverrideTransform = RootOverrideInstance.SubData->RootToSequenceTransform;
+	}
+
+	FFrameNumber RootOverrideTime = (Context.GetTime() * MasterToRootOverrideTransform).FloorToFrame();
+
 	// First off, attempt to find the evaluation group in the existing evaluation field data from the template
-	int32 TemplateFieldIndex = RootOverrideInstance.Template->EvaluationField.GetSegmentFromTime(Context.GetTime().FrameNumber);
+	int32 TemplateFieldIndex = RootOverrideInstance.Template->EvaluationField.GetSegmentFromTime(RootOverrideTime);
 
 	if (TemplateFieldIndex != INDEX_NONE)
 	{
@@ -222,17 +230,17 @@ const FMovieSceneEvaluationGroup* FMovieSceneRootEvaluationTemplateInstance::Set
 		if (bFullCompile)
 		{
 			FMovieSceneCompiler::Compile(*RootOverrideInstance.Sequence, *TemplateStore);
-			TemplateFieldIndex = RootOverrideInstance.Template->EvaluationField.GetSegmentFromTime(Context.GetTime().FrameNumber);
+			TemplateFieldIndex = RootOverrideInstance.Template->EvaluationField.GetSegmentFromTime(RootOverrideTime);
 		}
 		else
 		{
-			TOptional<FCompiledGroupResult> CompileResult = FMovieSceneCompiler::CompileTime(Context.GetTime().FrameNumber, *RootOverrideInstance.Sequence, *TemplateStore);
+			TOptional<FCompiledGroupResult> CompileResult = FMovieSceneCompiler::CompileTime(RootOverrideTime, *RootOverrideInstance.Sequence, *TemplateStore);
 
 			if (CompileResult.IsSet())
 			{
 				TRange<FFrameNumber> FieldRange = CompileResult->Range;
 				TemplateFieldIndex = RootOverrideInstance.Template->EvaluationField.Insert(
-					Context.GetTime().FrameNumber,
+					RootOverrideTime,
 					FieldRange,
 					MoveTemp(CompileResult->Group),
 					MoveTemp(CompileResult->MetaData)
