weighted, non-weighted fix typo fix for issue found by AndrewP

diff --git a/Engine/Source/Runtime/MovieScene/Private/Channels/MovieSceneFloatChannel.cpp b/Engine/Source/Runtime/MovieScene/Private/Channels/MovieSceneFloatChannel.cpp
--- a/Engine/Source/Runtime/MovieScene/Private/Channels/MovieSceneFloatChannel.cpp
+++ b/Engine/Source/Runtime/MovieScene/Private/Channels/MovieSceneFloatChannel.cpp
@@ -486,7 +486,7 @@ bool FMovieSceneFloatChannel::Evaluate(FFrameTime InTime,  float& OutValue) cons
 				float ArriveWeight;
 				if (Key2.Tangent.TangentWeightMode == RCTWM_WeightedNone || Key2.Tangent.TangentWeightMode == RCTWM_WeightedLeave)
 				{
-					const float ArriveTangentNormalized = Key1.Tangent.ArriveTangent / (TimeInterval);
+					const float ArriveTangentNormalized = Key2.Tangent.ArriveTangent / (TimeInterval);
 					const float Y = ArriveTangentNormalized * X;
 					ArriveWeight = FMath::Sqrt(X*X + Y * Y) * OneThird;
 				}
