Sequencer: Fix uninitialized default completion mode

diff --git a/Engine/Source/Runtime/MovieScene/Private/MovieSceneTrack.cpp b/Engine/Source/Runtime/MovieScene/Private/MovieSceneTrack.cpp
--- a/Engine/Source/Runtime/MovieScene/Private/MovieSceneTrack.cpp
+++ b/Engine/Source/Runtime/MovieScene/Private/MovieSceneTrack.cpp
@@ -2,6 +2,8 @@
 
 #include "MovieSceneTrack.h"
 #include "MovieScene.h"
+#include "MovieSceneSequence.h"
+
 #include "MovieSceneTimeHelpers.h"
 #include "Evaluation/MovieSceneSegment.h"
 #include "Compilation/MovieSceneSegmentCompiler.h"
@@ -201,7 +203,11 @@ FMovieSceneEvaluationTrack UMovieSceneTrack::GenerateTrackTemplate() const
 		virtual void AddOwnedTrack(FMovieSceneEvaluationTrack&& InTrackTemplate, const UMovieSceneTrack& SourceTrack) override {}
 	} Generator;
 
-	Compile(TrackTemplate, FMovieSceneTrackCompilerArgs(Generator));
+	FMovieSceneTrackCompilerArgs Args(Generator);
+
+	Args.DefaultCompletionMode = GetTypedOuter<UMovieSceneSequence>()->DefaultCompletionMode;
+
+	Compile(TrackTemplate, Args);
 
 	return TrackTemplate;
 }
