Sequencer: Fix uninitialized default completion mode

diff --git a/Engine/Source/Runtime/MovieScene/Private/MovieSceneTrack.cpp b/Engine/Source/Runtime/MovieScene/Private/MovieSceneTrack.cpp
--- a/Engine/Source/Runtime/MovieScene/Private/MovieSceneTrack.cpp
+++ b/Engine/Source/Runtime/MovieScene/Private/MovieSceneTrack.cpp
@@ -204,8 +204,10 @@ FMovieSceneEvaluationTrack UMovieSceneTrack::GenerateTrackTemplate() const
 	} Generator;
 
 	FMovieSceneTrackCompilerArgs Args(Generator);
-
-	Args.DefaultCompletionMode = GetTypedOuter<UMovieSceneSequence>()->DefaultCompletionMode;
+	if (GetTypedOuter<UMovieSceneSequence>())
+	{
+		Args.DefaultCompletionMode = GetTypedOuter<UMovieSceneSequence>()->DefaultCompletionMode;
+	}
 
 	Compile(TrackTemplate, Args);
 
