Sequencer: Fixed crash caused by re-allocation of the same named camera cuts track

diff --git a/Engine/Source/Runtime/MovieScene/Private/MovieScene.cpp b/Engine/Source/Runtime/MovieScene/Private/MovieScene.cpp
--- a/Engine/Source/Runtime/MovieScene/Private/MovieScene.cpp
+++ b/Engine/Source/Runtime/MovieScene/Private/MovieScene.cpp
@@ -701,7 +701,7 @@ UMovieSceneTrack* UMovieScene::AddCameraCutTrack( TSubclassOf<UMovieSceneTrack>
 	if( !CameraCutTrack )
 	{
 		Modify();
-		CameraCutTrack = NewObject<UMovieSceneTrack>(this, TrackClass, FName("Camera Cuts"), RF_Transactional);
+		CameraCutTrack = NewObject<UMovieSceneTrack>(this, TrackClass, NAME_None, RF_Transactional);
 	}
 
 	return CameraCutTrack;
