Sequencer: Fixed edits to sub-sections not invalidating their sub sequence data for deep hierarchical sequence structures

  - Full hierarchical sub-sequence data is stored in the outermost sequence hierarchy as part of the compilation, but such data was not necessarily being invalidated for deep hierarchies (sub-sub-sections) when an outer sub section/sequence was modified. When an outer sub section modification is encountered, the whole child hierarchy is now invalidated.
  - Additionally, inner sub sequence data that is cached in grandparent sequence hierarchies had no method of being invalidated if the grandchild sub-section was changed in some way. The sub-section signature and outer to inner transform now exists on the sub data to check for such situations as part of the compilation steps.

diff --git a/Engine/Source/Runtime/MovieScene/Private/Compilation/MovieSceneCompiler.cpp b/Engine/Source/Runtime/MovieScene/Private/Compilation/MovieSceneCompiler.cpp
--- a/Engine/Source/Runtime/MovieScene/Private/Compilation/MovieSceneCompiler.cpp
+++ b/Engine/Source/Runtime/MovieScene/Private/Compilation/MovieSceneCompiler.cpp
@@ -396,11 +396,16 @@ const FMovieSceneSubSequenceData* FMovieSceneCompiler::GetOrCreateSubSequenceDat
 {
 	// Find/add sub data in the root template
 	const FMovieSceneSubSequenceData* SubData = InOutHierarchy.FindSubData(InnerSequenceID);
-	if (SubData)
+	if (SubData && !SubData->IsDirty(SubSection))
 	{
 		return SubData;
 	}
-	
+
+	// Ensure that any ((great)grand)child sub sequences have their sub data regenerated
+	// by removing this whole sequence branch from the hierarchy (if it exists). This is
+	// necessary as all children will depend on this sequences's transform
+	InOutHierarchy.Remove(MakeArrayView(&InnerSequenceID, 1));
+
 	FSubSequenceInstanceDataParams InstanceParams{ InnerSequenceID, FMovieSceneEvaluationOperand(ParentSequenceID, InObjectBindingId) };
 	FMovieSceneSubSequenceData NewSubData = SubSection.GenerateSubSequenceData(InstanceParams);
 
diff --git a/Engine/Source/Runtime/MovieScene/Private/Evaluation/MovieSceneEvaluationTemplate.cpp b/Engine/Source/Runtime/MovieScene/Private/Evaluation/MovieSceneEvaluationTemplate.cpp
--- a/Engine/Source/Runtime/MovieScene/Private/Evaluation/MovieSceneEvaluationTemplate.cpp
+++ b/Engine/Source/Runtime/MovieScene/Private/Evaluation/MovieSceneEvaluationTemplate.cpp
@@ -64,16 +64,6 @@ void FMovieSceneEvaluationTemplate::AddSubSectionRange(UMovieSceneSubSection& In
 	// Don't need to do anything else if the section was already generated
 	if (!TemplateLedger.ContainsSubSection(InSubSection.GetSignature()))
 	{
-		if (InSubSection.GetSequence())
-		{
-			FMovieSceneSequenceID SubSequenceID = InSubSection.GetSequenceID();
-			FSubSequenceInstanceDataParams InstanceParams{ SubSequenceID, FMovieSceneEvaluationOperand(MovieSceneSequenceID::Root, InObjectBindingId) };
-
-			FMovieSceneSubSequenceData NewSubData = InSubSection.GenerateSubSequenceData(InstanceParams);
-
-			Hierarchy.Add(NewSubData, SubSequenceID, MovieSceneSequenceID::Root);
-		}
-
 		TRange<FFrameNumber> EntireSectionRange = InSubSection.GetRange();
 		EntireSectionRange.SetLowerBoundValue(EntireSectionRange.GetLowerBoundValue() - InSubSection.GetPreRollFrames());
 		EntireSectionRange.SetUpperBoundValue(EntireSectionRange.GetUpperBoundValue() + InSubSection.GetPostRollFrames());
@@ -83,6 +73,11 @@ void FMovieSceneEvaluationTemplate::AddSubSectionRange(UMovieSceneSubSection& In
 
 		// Invalidate the overlapping field
 		EvaluationField.Invalidate(EntireSectionRange);
+
+		// Invalidate the sequence hierarchy for the leaf starting at this sub section.
+		// The hierarchy will be populated by the compiler when this template is compiled
+		FMovieSceneSequenceID SubSequenceID = InSubSection.GetSequenceID();
+		Hierarchy.Remove(MakeArrayView(&SubSequenceID, 1));
 	}
 }
 
diff --git a/Engine/Source/Runtime/MovieScene/Private/Evaluation/MovieSceneSequenceHierarchy.cpp b/Engine/Source/Runtime/MovieScene/Private/Evaluation/MovieSceneSequenceHierarchy.cpp
--- a/Engine/Source/Runtime/MovieScene/Private/Evaluation/MovieSceneSequenceHierarchy.cpp
+++ b/Engine/Source/Runtime/MovieScene/Private/Evaluation/MovieSceneSequenceHierarchy.cpp
@@ -18,6 +18,7 @@ FMovieSceneSubSequenceData::FMovieSceneSubSequenceData(const UMovieSceneSubSecti
 #if WITH_EDITORONLY_DATA
 	, SectionPath(*InSubSection.GetPathNameInMovieScene())
 #endif
+	, SubSectionSignature(InSubSection.GetSignature())
 {
 	PreRollRange.Value  = TRange<FFrameNumber>::Empty();
 	PostRollRange.Value = TRange<FFrameNumber>::Empty();
@@ -35,7 +36,7 @@ FMovieSceneSubSequenceData::FMovieSceneSubSequenceData(const UMovieSceneSubSecti
 	const FFrameNumber OuterStartTime = MovieScene::DiscreteInclusiveLower(SubRange);
 	const FFrameNumber OuterEndTime   = MovieScene::DiscreteExclusiveUpper(SubRange);
 
-	RootToSequenceTransform = InSubSection.OuterToInnerTransform();
+	OuterToInnerTransform = RootToSequenceTransform = InSubSection.OuterToInnerTransform();
 
 	PlayRange.Value = SubRange * RootToSequenceTransform;
 
@@ -67,6 +68,11 @@ UMovieSceneSequence* FMovieSceneSubSequenceData::GetLoadedSequence() const
 	return CachedSequence.Get();
 }
 
+bool FMovieSceneSubSequenceData::IsDirty(const UMovieSceneSubSection& InSubSection) const
+{
+	return InSubSection.GetSignature() != SubSectionSignature || InSubSection.OuterToInnerTransform() != OuterToInnerTransform;
+}
+
 void FMovieSceneSequenceHierarchy::Add(const FMovieSceneSubSequenceData& Data, FMovieSceneSequenceIDRef ThisSequenceID, FMovieSceneSequenceIDRef ParentID)
 {
 	check(ParentID != MovieSceneSequenceID::Invalid);
diff --git a/Engine/Source/Runtime/MovieScene/Public/Evaluation/MovieSceneSequenceHierarchy.h b/Engine/Source/Runtime/MovieScene/Public/Evaluation/MovieSceneSequenceHierarchy.h
--- a/Engine/Source/Runtime/MovieScene/Public/Evaluation/MovieSceneSequenceHierarchy.h
+++ b/Engine/Source/Runtime/MovieScene/Public/Evaluation/MovieSceneSequenceHierarchy.h
@@ -46,6 +46,11 @@ struct FMovieSceneSubSequenceData
 	 */
 	MOVIESCENE_API UMovieSceneSequence* GetLoadedSequence() const;
 
+	/**
+	 * Check whether this structure is dirty and should be reconstructed
+	 */
+	MOVIESCENE_API bool IsDirty(const UMovieSceneSubSection& InSubSection) const;
+
 	/** The sequence that the sub section references */
 	UPROPERTY(meta=(AllowedClasses="MovieSceneSequence"))
 	FSoftObjectPath Sequence;
@@ -93,6 +98,14 @@ private:
 
 	/** Cached version of the sequence to avoid resolving it every time */
 	mutable TWeakObjectPtr<UMovieSceneSequence> CachedSequence;
+
+	/** The sub section's signature at the time this structure was populated. */
+	UPROPERTY()
+	FGuid SubSectionSignature;
+
+	/** The transform from this sub sequence's parent to its own play space. */
+	UPROPERTY()
+	FMovieSceneSequenceTransform OuterToInnerTransform;
 };
 
 /**
diff --git a/Engine/Source/Runtime/MovieScene/Public/Evaluation/MovieSceneSequenceTransform.h b/Engine/Source/Runtime/MovieScene/Public/Evaluation/MovieSceneSequenceTransform.h
--- a/Engine/Source/Runtime/MovieScene/Public/Evaluation/MovieSceneSequenceTransform.h
+++ b/Engine/Source/Runtime/MovieScene/Public/Evaluation/MovieSceneSequenceTransform.h
@@ -41,6 +41,16 @@ struct FMovieSceneSequenceTransform
 		, Offset(InOffset)
 	{}
 
+	friend bool operator==(const FMovieSceneSequenceTransform& A, const FMovieSceneSequenceTransform& B)
+	{
+		return A.TimeScale == B.TimeScale && A.Offset == B.Offset;
+	}
+
+	friend bool operator!=(const FMovieSceneSequenceTransform& A, const FMovieSceneSequenceTransform& B)
+	{
+		return A.TimeScale != B.TimeScale || A.Offset != B.Offset;
+	}
+
 	/**
 	 * Retrieve the inverse of this transform
 	 */
