Sequencer: Fix reverse playback so loop condition is based on End = StartTime + Duration

Tested reverse playback and also tested with number of loops and overplay.

diff --git a/Engine/Source/Runtime/MovieScene/Private/MovieSceneSequencePlayer.cpp b/Engine/Source/Runtime/MovieScene/Private/MovieSceneSequencePlayer.cpp
--- a/Engine/Source/Runtime/MovieScene/Private/MovieSceneSequencePlayer.cpp
+++ b/Engine/Source/Runtime/MovieScene/Private/MovieSceneSequencePlayer.cpp
@@ -547,11 +547,14 @@ void UMovieSceneSequencePlayer::UpdateTimeCursorPosition_Internal(FFrameTime New
 		UE_LOG(LogMovieScene, Warning, TEXT("Attempting to play back a sequence with zero duration"));
 		return;
 	}
-
-	FFrameTime PositionRelativeToStart = NewPosition.FrameNumber - StartTime;
-
+	
 	if (Method == EUpdatePositionMethod::Play && ShouldStopOrLoop(NewPosition))
 	{
+		// The actual start time taking into account reverse playback
+		FFrameNumber StartTimeWithReversed = bReversePlayback ? StartTime + Duration : StartTime;
+
+		FFrameTime PositionRelativeToStart = NewPosition.FrameNumber - StartTimeWithReversed;
+
 		const int32 NumTimesLooped    = FMath::Abs(PositionRelativeToStart.FrameNumber.Value / Duration);
 		const bool  bLoopIndefinitely = PlaybackSettings.LoopCount < 0;
 
@@ -561,7 +564,16 @@ void UMovieSceneSequencePlayer::UpdateTimeCursorPosition_Internal(FFrameTime New
 			CurrentNumLoops += NumTimesLooped;
 
 			const FFrameTime Overplay       = FFrameTime(PositionRelativeToStart.FrameNumber.Value % Duration, PositionRelativeToStart.GetSubFrame());
-			const FFrameTime NewFrameOffset = (Overplay < 0) ? FFrameTime(Duration) + Overplay : Overplay;
+			FFrameTime NewFrameOffset;
+			
+			if (bReversePlayback)
+			{
+				NewFrameOffset = (Overplay > 0) ?  FFrameTime(Duration) + Overplay : Overplay;
+			}
+			else
+			{
+				NewFrameOffset = (Overplay < 0) ? FFrameTime(Duration) + Overplay : Overplay;
+			}
 
 			if (SpawnRegister.IsValid())
 			{
@@ -569,15 +581,22 @@ void UMovieSceneSequencePlayer::UpdateTimeCursorPosition_Internal(FFrameTime New
 			}
 
 			// Reset the play position, and generate a new range that gets us to the new frame time
-			PlayPosition.Reset(Overplay < 0 ? GetLastValidTime() : StartTime);
+			if (bReversePlayback)
+			{
+				PlayPosition.Reset(Overplay > 0 ? GetLastValidTime() : StartTimeWithReversed);
+			}
+			else
+			{
+				PlayPosition.Reset(Overplay < 0 ? GetLastValidTime() : StartTimeWithReversed);
+			}
 
-			FMovieSceneEvaluationRange Range = PlayPosition.PlayTo(StartTime + NewFrameOffset);
+			FMovieSceneEvaluationRange Range = PlayPosition.PlayTo(StartTimeWithReversed + NewFrameOffset);
 
 			const bool bHasJumped = true;
 			UpdateMovieSceneInstance(Range, StatusOverride, bHasJumped);
 
 			// Use the exact time here rather than a frame locked time to ensure we don't skip the amount that was overplayed in the time controller
-			FQualifiedFrameTime ExactCurrentTime(StartTime + NewFrameOffset, PlayPosition.GetInputRate());
+			FQualifiedFrameTime ExactCurrentTime(StartTimeWithReversed + NewFrameOffset, PlayPosition.GetInputRate());
 			PlaybackSettings.TimeController->Reset(ExactCurrentTime);
 
 			OnLooped();
