Sequencer: Handle null pre-animated state tokens being returned from token producers

  - Some producers (such as the MPC or particle parameter tracks) will return a null cached state token when the parameter name is no longer applicable to the source object.
  - Previously we'd assumed them to always be non-null, but now we catch such cases, and do not add any entry to the pre animated state cache.

Fixes crashreporter bugg 2448098.

diff --git a/Engine/Source/Runtime/MovieScene/Private/Evaluation/MovieScenePreAnimatedState.cpp b/Engine/Source/Runtime/MovieScene/Private/Evaluation/MovieScenePreAnimatedState.cpp
--- a/Engine/Source/Runtime/MovieScene/Private/Evaluation/MovieScenePreAnimatedState.cpp
+++ b/Engine/Source/Runtime/MovieScene/Private/Evaluation/MovieScenePreAnimatedState.cpp
@@ -80,32 +80,38 @@ void TMovieSceneSavedTokens<TokenType>::OnPreAnimated(ECapturePreAnimatedState C
 {
 	MOVIESCENE_DETAILED_SCOPE_CYCLE_COUNTER(MovieSceneEval_SavePreAnimatedState)
 
+	// If we're not capturing any state, return immediately
 	if (CaptureState == ECapturePreAnimatedState::None)
 	{
 		return;
 	}
 
-	if (CaptureState == ECapturePreAnimatedState::Entity)
+	// If the entity key and anim type combination already exists in the animated entities array,
+	// we've already saved state for this entity and this type, and can just return immediately
+	FMovieSceneEntityAndAnimTypeID EntityAndTypeID{AssociatedKey, InAnimTypeID};
+	if (CaptureState == ECapturePreAnimatedState::Entity && AnimatedEntities.Contains(EntityAndTypeID))
 	{
-		FMovieSceneEntityAndAnimTypeID EntityAndTypeID{AssociatedKey, InAnimTypeID};
-
-		// If the entity key and anim type combination already exists in the animated entities array, we've already got a preanimated token reference
-		if (AnimatedEntities.Contains(EntityAndTypeID))
-		{
-			return;
-		}
-
-		AnimatedEntities.Add(EntityAndTypeID);
+		return;
 	}
 
 	auto ResolvedPayload = Payload.Get();
 
+	// Attempt to locate an existing animated state token for this type ID
 	int32 TokenIndex = AllAnimatedTypeIDs.IndexOfByKey(InAnimTypeID);
 	if (TokenIndex == INDEX_NONE)
 	{
-		// Create the token, and update the arrays
+		auto NewlyCachedState = MovieSceneImpl::CacheExistingState(Producer, ResolvedPayload);
+
+		// If the producer returned a null state token, there's no point saving anything.
+		// Return immediately without mutating anything in this class.
+		if (!NewlyCachedState.IsValid())
+		{
+			return;
+		}
+
+		// Record this type ID as being animated, and push the new state token onto the array
 		AllAnimatedTypeIDs.Add(InAnimTypeID);
-		PreAnimatedTokens.Add(TokenType(MovieSceneImpl::CacheExistingState(Producer, ResolvedPayload)));
+		PreAnimatedTokens.Add(TokenType(MoveTemp(NewlyCachedState)));
 
 		// If we're capturing for the entity as well, increment the ref count
 		if (CaptureState == ECapturePreAnimatedState::Entity)
@@ -119,20 +125,37 @@ void TMovieSceneSavedTokens<TokenType>::OnPreAnimated(ECapturePreAnimatedState C
 	}
 	else if (CaptureState == ECapturePreAnimatedState::Entity)
 	{
-		// We already have a token animated
+		// We already have a token animated, either with Restore State, or Keep State.
 		TPreAnimatedToken<TokenType>& Token = PreAnimatedTokens[TokenIndex];
 
 		if (Token.EntityRefCount == 0)
 		{
-			// If the ref count is 0, a previous entity must have animated, but been set to 'keep state'. In this case, we need to define an additional token to ensure we restore to the correct (current) value.
+			// If the ref count is 0, a previous entity must have animated, but been set to 'keep state'.
+			// In this case, we need to define an additional token to ensure we restore to the correct (current) value when this entity restores.
 			// Don't call InitializeForAnimation here, as we've clearly already done so (a token exists for it)
-			Token.OptionalEntityToken = MovieSceneImpl::CacheExistingState(Producer, ResolvedPayload);
+			auto NewlyCachedState = MovieSceneImpl::CacheExistingState(Producer, ResolvedPayload);
+
+			// If the producer returned a null state token, there's no point saving anything.
+			// Return immediately without mutating anything in this class.
+			if (!NewlyCachedState.IsValid())
+			{
+				return;
+			}
+
+			Token.OptionalEntityToken = MoveTemp(NewlyCachedState);
 		}
 
 		// Increment the reference count regardless of whether we just created the token or not (we always need a reference)
 		++Token.EntityRefCount;
 		MovieSceneImpl::EntityHasAnimated(AssociatedKey, Parent, ResolvedPayload);
 	}
+
+	// If we're capturing at the entity level (ie, this entity is restore state), add it to the list of animated entites.
+	// We know by this point in the function that the entity was not previously animated, and a valid restore-state token has been added
+	if (CaptureState == ECapturePreAnimatedState::Entity)
+	{
+		AnimatedEntities.Add(EntityAndTypeID);
+	}
 }
 
 template<typename TokenType>
