Sequencer: Fixed sub-sub sequences not being clamped by their outermost section's range
  - Sub section ranges were only being propagated one level deep, which meant that sub-sub sections could break out of their outermost section's range
  - All clamp ranges are now performed in root space to remove any conversion inaacuracies between the two when transformed from local to root.

#tests Tested problem case described in the jira, along with my own broad test content, plus all the automation test relating to sequencer still pass.

diff --git a/Engine/Source/Runtime/MovieScene/Private/Compilation/MovieSceneCompiler.cpp b/Engine/Source/Runtime/MovieScene/Private/Compilation/MovieSceneCompiler.cpp
--- a/Engine/Source/Runtime/MovieScene/Private/Compilation/MovieSceneCompiler.cpp
+++ b/Engine/Source/Runtime/MovieScene/Private/Compilation/MovieSceneCompiler.cpp
@@ -46,6 +46,12 @@ struct FGatherParameters
 		LocalClampRange = InNewRootClampRange * RootToSequenceTransform;
 	}
 
+	/** Clamp the specified range to the current clamp range (in root space) */
+	TRange<FFrameNumber> ClampRoot(const TRange<FFrameNumber>& InRootRange) const
+	{
+		return TRange<FFrameNumber>::Intersection(RootClampRange, InRootRange);
+	}
+
 	/** Path from root to current sequence */
 	FMovieSceneRootOverridePath& RootPath;
 	/** Hierarchy for the root sequence template */
@@ -318,17 +324,17 @@ void FMovieSceneCompiler::GatherCompileOnTheFlyData(UMovieSceneSequence& InSeque
 	FMovieSceneEvaluationTreeRangeIterator SubSectionIt(SubSectionField.IterateFromLowerBound(CompileClampIntersection.GetLowerBound()));
 
 	// Intersect the unique range in the tree with the current overlapping empty range to constrict the resulting compile range in the case where this is a gap between sub sections
-	OutData.EmptyOverlappingRange = TRange<FFrameNumber>::Intersection(OutData.EmptyOverlappingRange, SubSectionIt.Range() * Params.RootToSequenceTransform.Inverse());
+	OutData.EmptyOverlappingRange = TRange<FFrameNumber>::Intersection(OutData.EmptyOverlappingRange, Params.ClampRoot(SubSectionIt.Range() * Params.RootToSequenceTransform.Inverse()));
 
 	for ( ; SubSectionIt && SubSectionIt.Range().Overlaps(CompileClampIntersection); ++SubSectionIt)
 	{
-		TRange<FFrameNumber> ThisSegmentRange = SubSectionIt.Range() * Params.RootToSequenceTransform.Inverse();
-		if (ThisSegmentRange.IsEmpty())
+		TRange<FFrameNumber> ThisSegmentRangeRoot = Params.ClampRoot(SubSectionIt.Range() * Params.RootToSequenceTransform.Inverse());
+		if (ThisSegmentRangeRoot.IsEmpty())
 		{
 			continue;
 		}
 
-		SubSectionGatherParams.SetClampRange(ThisSegmentRange);
+		SubSectionGatherParams.SetClampRange(ThisSegmentRangeRoot);
 
 		// Iterate all sub sections in the current range
 		for (const FMovieSceneSubSectionData& SubSectionData : SubSectionField.GetAllData(SubSectionIt.Node()))
@@ -444,7 +450,7 @@ void FMovieSceneCompiler::GatherCompileDataForTrack(FMovieSceneEvaluationTrack&
 	{
 		// No segment at this time, so just report the time range of the empty space.
 		TRange<FFrameNumber> EmptyTrackSpace = Track.GetUniqueRangeFromLowerBound(Params.LocalCompileRange.GetLowerBound());
-		TRange<FFrameNumber> ClampedEmptyTrackSpaceRoot = TRange<FFrameNumber>::Intersection(Params.LocalClampRange, EmptyTrackSpace) * SequenceToRootTransform;
+		TRange<FFrameNumber> ClampedEmptyTrackSpaceRoot = Params.ClampRoot(EmptyTrackSpace * SequenceToRootTransform);
 
 		OutData.EmptyOverlappingRange = TRange<FFrameNumber>::Intersection(ClampedEmptyTrackSpaceRoot, OutData.EmptyOverlappingRange);
 	}
@@ -460,7 +466,7 @@ void FMovieSceneCompiler::GatherCompileDataForTrack(FMovieSceneEvaluationTrack&
 		Data.Track = &Track;
 		Data.bRequiresInit = ThisSegment.Impls.ContainsByPredicate(RequiresInit);
 
-		TRange<FFrameNumber> IntersectionRange = TRange<FFrameNumber>::Intersection(Params.LocalClampRange, ThisSegment.Range) * SequenceToRootTransform;
+		TRange<FFrameNumber> IntersectionRange = Params.ClampRoot(ThisSegment.Range * SequenceToRootTransform);
 		if (!IntersectionRange.IsEmpty())
 		{
 			OutData.Tracks.Add(IntersectionRange, Data);
