Sequencer: Remove edit condition because it ends up editing bRandomStartTime which causes the sequence to start at a random start time.

diff --git a/Engine/Source/Runtime/MovieScene/Public/MovieSceneSequencePlayer.h b/Engine/Source/Runtime/MovieScene/Public/MovieSceneSequencePlayer.h
--- a/Engine/Source/Runtime/MovieScene/Public/MovieSceneSequencePlayer.h
+++ b/Engine/Source/Runtime/MovieScene/Public/MovieSceneSequencePlayer.h
@@ -64,7 +64,7 @@ struct FMovieSceneSequencePlaybackSettings
 	bool bRandomStartTime;
 
 	/** Start playback at the specified offset from the start of the sequence's playback range */
-	UPROPERTY(EditAnywhere, BlueprintReadWrite, Category="Playback", DisplayName="Start Offset", meta=(EditCondition="!bRandomStartTime", Units=s))
+	UPROPERTY(EditAnywhere, BlueprintReadWrite, Category="Playback", DisplayName="Start Offset", meta=(Units=s))
 	float StartTime;
 
 	/** Flag used to specify whether actor states should be restored on stop */
