Sequencer: Resolve some merge conflicts with ReplaceBinding

diff --git a/Engine/Source/Editor/Sequencer/Private/Sequencer.cpp b/Engine/Source/Editor/Sequencer/Private/Sequencer.cpp
--- a/Engine/Source/Editor/Sequencer/Private/Sequencer.cpp
+++ b/Engine/Source/Editor/Sequencer/Private/Sequencer.cpp
@@ -5823,23 +5823,13 @@ void FSequencer::PasteCopiedTracks()
 				// and update it.
 				FGuid NewGuid = MakeNewSpawnable(*CopyableBinding->SpawnableObjectTemplate);
 				FMovieSceneBinding NewBinding(NewGuid, CopyableBinding->Binding.GetName(), CopyableBinding->Tracks);
-
 				FMovieSceneSpawnable* Spawnable = MovieScene->FindSpawnable(NewGuid);
+
 				// Copy the name of the original spawnable too.
 				Spawnable->SetName(CopyableBinding->Spawnable.GetName());
 
 				// Replace the auto-generated binding with our deserialized bindings (which has our tracks)
-				FMovieSceneBinding* Binding = MovieScene->FindBinding(NewGuid);
-				if (Binding)
-				{
-					*Binding = NewBinding;
-
-					// We also need to change the track's owners to be the MovieScene. Typically MovieScene::AddSpawnable/MovieScene::AddSpawnable/AddPossessable does this.
-					for (auto Track : Binding->GetTracks())
-					{
-						Track->Rename(nullptr, MovieScene);
-					}
-				}
+				MovieScene->ReplaceBinding(NewGuid, NewBinding);
 
 				OldToNewGuidMap.Add(CopyableBinding->Spawnable.GetGuid(), NewGuid);
 
diff --git a/Engine/Source/Runtime/MovieScene/Private/MovieScene.cpp b/Engine/Source/Runtime/MovieScene/Private/MovieScene.cpp
--- a/Engine/Source/Runtime/MovieScene/Private/MovieScene.cpp
+++ b/Engine/Source/Runtime/MovieScene/Private/MovieScene.cpp
@@ -951,6 +951,20 @@ void UMovieScene::ReplaceBinding(const FGuid& OldGuid, const FGuid& NewGuid, con
 	}
 }
 
+void UMovieScene::ReplaceBinding(const FGuid& BindingToReplaceGuid, const FMovieSceneBinding& NewBinding)
+{
+	FMovieSceneBinding* Binding = ObjectBindings.FindByPredicate([BindingToReplaceGuid](const FMovieSceneBinding& CheckedBinding) { return CheckedBinding.GetObjectGuid() == BindingToReplaceGuid; });
+	if (Binding)
+	{
+		*Binding = NewBinding;
+
+		// We also need to change the track's owners to be the MovieScene.
+		for (auto Track : Binding->GetTracks())
+		{
+			Track->Rename(nullptr, this);
+		}
+	}
+}
 
 void UMovieScene::MoveBindingContents(const FGuid& SourceBindingId, const FGuid& DestinationBindingId)
 {
diff --git a/Engine/Source/Runtime/MovieScene/Public/MovieScene.h b/Engine/Source/Runtime/MovieScene/Public/MovieScene.h
--- a/Engine/Source/Runtime/MovieScene/Public/MovieScene.h
+++ b/Engine/Source/Runtime/MovieScene/Public/MovieScene.h
@@ -451,16 +451,6 @@ public:
 	 */
 	void MoveBindingContents(const FGuid& SourceBindingId, const FGuid& DestinationBindingId);
 
-	/**
-	 * Tries to find an FMovieSceneBinding for the specified Guid.
-	 * 
-	 * @param ForGuid	The binding's Guid to look for.
-	 * @return			Pointer to the binding, otherwise nullptr.
-	 */
-	FMovieSceneBinding* FindBinding(const FGuid& ForGuid)
-	{
-		return ObjectBindings.FindByPredicate([ForGuid](const FMovieSceneBinding& Binding) { return Binding.GetObjectGuid() == ForGuid; });
-	}
 public:
 
 	// @todo sequencer: the following methods really shouldn't be here
@@ -594,6 +584,14 @@ public:
 	*/
 	void ReplaceBinding(const FGuid& OldGuid, const FGuid& NewGuid, const FString& Name);
 
+	/*
+	* Replace an existing binding with another. Assumes ownership of any
+	* tracks listed in the binding. Does nothing if no binding can be found.
+	* @param BindingToReplaceGuid	Binding Guid that should be replaced
+	* @param NewBinding				Binding Data that should replace the original one specified by BindingToReplaceGuid.
+	*/
+	void ReplaceBinding(const FGuid& BindingToReplaceGuid, const FMovieSceneBinding& NewBinding);
+
 #if WITH_EDITORONLY_DATA
 	/**
 	 */
