Fix for evaluation a section of time when one key is non-weighted and the other is weighted.  What we do is evaluate them both as being weighted, but we don't have the weight value for the non-weighted tangent. The weight of the non-weighted tangent is implicilty 1/3rd the distance between the two points, so we just calculate that if needed.

diff --git a/Engine/Source/Runtime/MovieScene/Private/Channels/MovieSceneFloatChannel.cpp b/Engine/Source/Runtime/MovieScene/Private/Channels/MovieSceneFloatChannel.cpp
--- a/Engine/Source/Runtime/MovieScene/Private/Channels/MovieSceneFloatChannel.cpp
+++ b/Engine/Source/Runtime/MovieScene/Private/Channels/MovieSceneFloatChannel.cpp
@@ -449,10 +449,10 @@ bool FMovieSceneFloatChannel::Evaluate(FFrameTime InTime,  float& OutValue) cons
 		{
 		case RCIM_Cubic:
 		{
+			const float OneThird = 1.0f / 3.0f;
 			if ((Key1.Tangent.TangentWeightMode == RCTWM_WeightedNone || Key1.Tangent.TangentWeightMode == RCTWM_WeightedArrive)
 				&& (Key2.Tangent.TangentWeightMode == RCTWM_WeightedNone || Key2.Tangent.TangentWeightMode == RCTWM_WeightedLeave))
 			{
-				const float OneThird = 1.0f / 3.0f;
 				const int32 Diff = Times[Index2].Value - Times[Index1].Value;
 				const float P0 = Key1.Value;
 				const float P1 = P0 + (Key1.Tangent.LeaveTangent * Diff * OneThird);
@@ -464,20 +464,44 @@ bool FMovieSceneFloatChannel::Evaluate(FFrameTime InTime,  float& OutValue) cons
 			}
 			else //its weighted
 			{
-				float ToSeconds = TickResolution.AsDecimal();
+				const float TimeInterval = TickResolution.AsInterval();
+				const float ToSeconds = 1.0f / TimeInterval;
 
 				const double Time1 = TickResolution.AsSeconds(Times[Index1].Value);
 				const double Time2 = TickResolution.AsSeconds(Times[Index2].Value);
+				const float X = Time2 - Time1;
 				float CosAngle, SinAngle;
 				float Angle = FMath::Atan(Key1.Tangent.LeaveTangent * ToSeconds);
 				FMath::SinCos(&SinAngle, &CosAngle, Angle);
-				const float Key1TanX = CosAngle * Key1.Tangent.LeaveTangentWeight + Time1;
-				const float Key1TanY = SinAngle * Key1.Tangent.LeaveTangentWeight + Key1.Value;
+				float LeaveWeight;
+				if (Key1.Tangent.TangentWeightMode == RCTWM_WeightedNone || Key1.Tangent.TangentWeightMode == RCTWM_WeightedArrive)
+				{
+					const float LeaveTangentNormalized = Key1.Tangent.LeaveTangent / (TimeInterval);
+					const float Y = LeaveTangentNormalized * X;
+					LeaveWeight = FMath::Sqrt(X*X + Y * Y) * OneThird;
+				}
+				else
+				{
+					LeaveWeight = Key1.Tangent.LeaveTangentWeight;
+				}
+				const float Key1TanX = CosAngle * LeaveWeight + Time1;
+				const float Key1TanY = SinAngle * LeaveWeight + Key1.Value;
 
 				Angle = FMath::Atan(Key2.Tangent.ArriveTangent * ToSeconds);
 				FMath::SinCos(&SinAngle, &CosAngle, Angle);
-				const float Key2TanX = -CosAngle * Key2.Tangent.ArriveTangentWeight + Time2;
-				const float Key2TanY = -SinAngle * Key2.Tangent.ArriveTangentWeight + Key2.Value;
+				float ArriveWeight;
+				if (Key2.Tangent.TangentWeightMode == RCTWM_WeightedNone || Key2.Tangent.TangentWeightMode == RCTWM_WeightedLeave)
+				{
+					const float ArriveTangentNormalized = Key1.Tangent.ArriveTangent / (TimeInterval);
+					const float Y = ArriveTangentNormalized * X;
+					ArriveWeight = FMath::Sqrt(X*X + Y * Y) * OneThird;
+				}
+				else
+				{
+					ArriveWeight =  Key2.Tangent.ArriveTangentWeight;
+				}
+				const float Key2TanX = -CosAngle * ArriveWeight + Time2;
+				const float Key2TanY = -SinAngle * ArriveWeight + Key2.Value;
 
 				//Normalize the Time Range
 				const float RangeX = Time2 - Time1;
