Sequencer: Fix bOverridesDefault tooltip so that it reads correctly.

diff --git a/Engine/Source/Runtime/MovieScene/Public/MovieSceneBindingOverrides.h b/Engine/Source/Runtime/MovieScene/Public/MovieSceneBindingOverrides.h
--- a/Engine/Source/Runtime/MovieScene/Public/MovieSceneBindingOverrides.h
+++ b/Engine/Source/Runtime/MovieScene/Public/MovieSceneBindingOverrides.h
@@ -30,7 +30,7 @@ struct FMovieSceneBindingOverrideData
 	UPROPERTY(EditAnywhere, Category="Binding")
 	TWeakObjectPtr<UObject> Object;
 
-	/** Specifies whether the default assignment should remain bound (false) or if this should completely override the default binding (false). */
+	/** Specifies whether the default assignment should remain bound (false) or if this should completely override the default binding (true). */
 	UPROPERTY(EditAnywhere, Category="Binding")
 	bool bOverridesDefault;
 };
