Sequencer: Add GetObjectBindings() to retrieve the object bindings for the passed in UObject.

diff --git a/Engine/Source/Runtime/MovieScene/Private/MovieSceneSequencePlayer.cpp b/Engine/Source/Runtime/MovieScene/Private/MovieSceneSequencePlayer.cpp
--- a/Engine/Source/Runtime/MovieScene/Private/MovieSceneSequencePlayer.cpp
+++ b/Engine/Source/Runtime/MovieScene/Private/MovieSceneSequencePlayer.cpp
@@ -685,6 +685,24 @@ TArray<UObject*> UMovieSceneSequencePlayer::GetBoundObjects(FMovieSceneObjectBin
 	return Objects;
 }
 
+TArray<FMovieSceneObjectBindingID> UMovieSceneSequencePlayer::GetObjectBindings(UObject* InObject)
+{
+	TArray<FMovieSceneObjectBindingID> ObjectBindings;
+
+	for (FMovieSceneSequenceIDRef SequenceID : GetEvaluationTemplate().GetThisFrameMetaData().ActiveSequences)
+	{
+		FGuid ObjectGuid = FindObjectId(*InObject, SequenceID);
+		if (ObjectGuid.IsValid())
+		{
+			FMovieSceneObjectBindingID ObjectBinding(ObjectGuid, SequenceID);
+
+			ObjectBindings.Add(ObjectBinding);
+		}
+	}
+
+	return ObjectBindings;
+}
+
 void UMovieSceneSequencePlayer::BeginDestroy()
 {
 	Stop();
diff --git a/Engine/Source/Runtime/MovieScene/Public/MovieSceneSequencePlayer.h b/Engine/Source/Runtime/MovieScene/Public/MovieSceneSequencePlayer.h
--- a/Engine/Source/Runtime/MovieScene/Public/MovieSceneSequencePlayer.h
+++ b/Engine/Source/Runtime/MovieScene/Public/MovieSceneSequencePlayer.h
@@ -402,6 +402,10 @@ public:
 	UFUNCTION(BlueprintCallable, Category="Game|Cinematic")
 	TArray<UObject*> GetBoundObjects(FMovieSceneObjectBindingID ObjectBinding);
 
+	/** Get the object bindings for the requested object */
+	UFUNCTION(BlueprintCallable, Category="Game|Cinematic")
+	TArray<FMovieSceneObjectBindingID> GetObjectBindings(UObject* InObject);
+
 public:
 
 	/** Update the sequence for the current time, if playing */
