Fixed shadowed variable warning in MovieScene.cpp

#trivial

diff --git a/Engine/Source/Runtime/MovieScene/Private/MovieScene.cpp b/Engine/Source/Runtime/MovieScene/Private/MovieScene.cpp
--- a/Engine/Source/Runtime/MovieScene/Private/MovieScene.cpp
+++ b/Engine/Source/Runtime/MovieScene/Private/MovieScene.cpp
@@ -953,7 +953,7 @@ void UMovieScene::ReplaceBinding(const FGuid& OldGuid, const FGuid& NewGuid, con
 
 void UMovieScene::ReplaceBinding(const FGuid& BindingToReplaceGuid, const FMovieSceneBinding& NewBinding)
 {
-	FMovieSceneBinding* Binding = ObjectBindings.FindByPredicate([BindingToReplaceGuid](const FMovieSceneBinding& Binding) { return Binding.GetObjectGuid() == BindingToReplaceGuid; });
+	FMovieSceneBinding* Binding = ObjectBindings.FindByPredicate([BindingToReplaceGuid](const FMovieSceneBinding& CheckedBinding) { return CheckedBinding.GetObjectGuid() == BindingToReplaceGuid; });
 	if (Binding)
 	{
 		*Binding = NewBinding;
