Sequencer: Resolve merge conflict with 4.20

diff --git a/Engine/Source/Runtime/MovieScene/Public/MovieScene.h b/Engine/Source/Runtime/MovieScene/Public/MovieScene.h
--- a/Engine/Source/Runtime/MovieScene/Public/MovieScene.h
+++ b/Engine/Source/Runtime/MovieScene/Public/MovieScene.h
@@ -451,16 +451,6 @@ public:
 	 */
 	void MoveBindingContents(const FGuid& SourceBindingId, const FGuid& DestinationBindingId);
 
-	/**
-	 * Tries to find an FMovieSceneBinding for the specified Guid.
-	 * 
-	 * @param ForGuid	The binding's Guid to look for.
-	 * @return			Pointer to the binding, otherwise nullptr.
-	 */
-	FMovieSceneBinding* FindBinding(const FGuid& ForGuid)
-	{
-		return ObjectBindings.FindByPredicate([ForGuid](const FMovieSceneBinding& Binding) { return Binding.GetObjectGuid() == ForGuid; });
-	}
 public:
 
 	// @todo sequencer: the following methods really shouldn't be here
