Sequencer: Fix GoToEndAndStop() and other jump state so that when calling stop, sequencer gets torn down properly.

This fixes a bug where if you just call GoToEndAndStop(), the sequence won't actually stop and tear down.

diff --git a/Engine/Source/Runtime/MovieScene/Private/MovieSceneSequencePlayer.cpp b/Engine/Source/Runtime/MovieScene/Private/MovieSceneSequencePlayer.cpp
--- a/Engine/Source/Runtime/MovieScene/Private/MovieSceneSequencePlayer.cpp
+++ b/Engine/Source/Runtime/MovieScene/Private/MovieSceneSequencePlayer.cpp
@@ -203,7 +203,7 @@ void UMovieSceneSequencePlayer::Scrub()
 
 void UMovieSceneSequencePlayer::Stop()
 {
-	if (IsPlaying() || IsPaused())
+	if (IsPlaying() || IsPaused() || RootTemplateInstance.IsValid())
 	{
 		if (bIsEvaluating)
 		{
