Sequencer: Add default sections for all tracks (not just property tracks)

diff --git a/Engine/Plugins/MovieScene/LevelSequenceEditor/Source/LevelSequenceEditor/Private/LevelSequenceEditorToolkit.cpp b/Engine/Plugins/MovieScene/LevelSequenceEditor/Source/LevelSequenceEditor/Private/LevelSequenceEditorToolkit.cpp
--- a/Engine/Plugins/MovieScene/LevelSequenceEditor/Source/LevelSequenceEditor/Private/LevelSequenceEditorToolkit.cpp
+++ b/Engine/Plugins/MovieScene/LevelSequenceEditor/Source/LevelSequenceEditor/Private/LevelSequenceEditorToolkit.cpp
@@ -422,8 +422,12 @@ void FLevelSequenceEditorToolkit::AddDefaultTracksForActor(AActor& Actor, const
 					NewTrack = MovieScene->AddTrack(TrackClass, Binding);
 				}
 
-				// Create a section for any property tracks
-				if (Cast<UMovieScenePropertyTrack>(NewTrack))
+				bool bCreateDefaultSection = false;
+#if WITH_EDITORONLY_DATA
+				bCreateDefaultSection = NewTrack->SupportsDefaultSections();
+#endif
+
+				if (bCreateDefaultSection)
 				{
 					UMovieSceneSection* NewSection;
 					if (NewTrack->GetAllSections().Num() > 0)
diff --git a/Engine/Source/Runtime/MovieScene/Private/MovieSceneTrack.cpp b/Engine/Source/Runtime/MovieScene/Private/MovieSceneTrack.cpp
--- a/Engine/Source/Runtime/MovieScene/Private/MovieSceneTrack.cpp
+++ b/Engine/Source/Runtime/MovieScene/Private/MovieSceneTrack.cpp
@@ -21,6 +21,7 @@ UMovieSceneTrack::UMovieSceneTrack(const FObjectInitializer& InInitializer)
 #if WITH_EDITORONLY_DATA
 	TrackTint = FColor(127, 127, 127, 0);
 	SortingOrder = -1;
+	bSupportsDefaultSections = true;
 #endif
 }
 
diff --git a/Engine/Source/Runtime/MovieScene/Public/MovieSceneTrack.h b/Engine/Source/Runtime/MovieScene/Public/MovieSceneTrack.h
--- a/Engine/Source/Runtime/MovieScene/Public/MovieSceneTrack.h
+++ b/Engine/Source/Runtime/MovieScene/Public/MovieSceneTrack.h
@@ -238,8 +238,8 @@ public:
 	virtual void RemoveAllAnimationData() { }
 
 	/**
-	 * @return Whether or not this track supports multiple row indices.
-	 */
+	* @return Whether or not this track supports multiple row indices.
+	*/
 	virtual bool SupportsMultipleRows() const
 	{
 		return SupportedBlendTypes.Num() != 0;
@@ -340,23 +340,31 @@ public:
 	}
 
 	/**
-	* Get this folder's desired sorting order
-	*/
+	 * Get this folder's desired sorting order
+	 */
 	int32 GetSortingOrder() const
 	{
 		return SortingOrder;
 	}
 
 	/**
-	* Set this folder's desired sorting order.
-	*
-	* @param InSortingOrder The higher the value the further down the list the folder will be.
-	*/
+	 * Set this folder's desired sorting order.
+	 *
+	 * @param InSortingOrder The higher the value the further down the list the folder will be.
+	 */
 	void SetSortingOrder(const int32 InSortingOrder)
 	{
 		SortingOrder = InSortingOrder;
 	}
 
+	/**
+	 * @return Whether or not this track supports the creation of default sections when the track is created.
+	 */
+	virtual bool SupportsDefaultSections() const
+	{
+		return bSupportsDefaultSections;
+	}
+
 protected:
 
 	/** This track's tint color */
@@ -366,6 +374,11 @@ protected:
 	/** This folder's desired sorting order */
 	UPROPERTY()
 	int32 SortingOrder;
+
+	/** Does this track support the creation of a default section when created? */
+	UPROPERTY()
+	bool bSupportsDefaultSections;
+
 public:
 #endif
 
diff --git a/Engine/Source/Runtime/MovieSceneTracks/Private/Tracks/MovieSceneSkeletalAnimationTrack.cpp b/Engine/Source/Runtime/MovieSceneTracks/Private/Tracks/MovieSceneSkeletalAnimationTrack.cpp
--- a/Engine/Source/Runtime/MovieSceneTracks/Private/Tracks/MovieSceneSkeletalAnimationTrack.cpp
+++ b/Engine/Source/Runtime/MovieSceneTracks/Private/Tracks/MovieSceneSkeletalAnimationTrack.cpp
@@ -21,6 +21,7 @@ UMovieSceneSkeletalAnimationTrack::UMovieSceneSkeletalAnimationTrack(const FObje
 {
 #if WITH_EDITORONLY_DATA
 	TrackTint = FColor(124, 15, 124, 65);
+	bSupportsDefaultSections = false;
 #endif
 
 	SupportedBlendTypes.Add(EMovieSceneBlendType::Absolute);
