Sequencer: Fix crash splitting a sub section without a sequence.

diff --git a/Engine/Source/Runtime/MovieScene/Private/Sections/MovieSceneSubSection.cpp b/Engine/Source/Runtime/MovieScene/Private/Sections/MovieSceneSubSection.cpp
--- a/Engine/Source/Runtime/MovieScene/Private/Sections/MovieSceneSubSection.cpp
+++ b/Engine/Source/Runtime/MovieScene/Private/Sections/MovieSceneSubSection.cpp
@@ -220,7 +220,20 @@ UMovieSceneSection* UMovieSceneSubSection::SplitSection( FQualifiedFrameTime Spl
 		{
 			// Sections need their offsets calculated in their local resolution. Different sequences can have different tick resolutions 
 			// so we need to transform from the parent resolution to the local one before splitting them.
-			FFrameRate LocalTickResolution = GetSequence()->GetMovieScene()->GetTickResolution();
+			FFrameRate LocalTickResolution;
+			if (GetSequence())
+			{
+				LocalTickResolution = GetSequence()->GetMovieScene()->GetTickResolution();
+			}
+			else
+			{
+				UMovieScene* OuterScene = GetTypedOuter<UMovieScene>();
+				if (OuterScene)
+				{
+					LocalTickResolution = OuterScene->GetTickResolution();
+				}
+			}
+
 			FFrameNumber LocalResolutionStartOffset = FFrameRate::TransformTime(SplitTime.Time.GetFrame() - MovieScene::DiscreteInclusiveLower(InitialRange), SplitTime.Rate, LocalTickResolution).FrameNumber;
 
 			FFrameNumber NewStartOffset = LocalResolutionStartOffset / Parameters.TimeScale;
