Sequencer: When re-keying, don't set interpolation on existing keys

diff --git a/Engine/Source/Runtime/MovieScene/Private/Channels/MovieSceneFloatChannel.cpp b/Engine/Source/Runtime/MovieScene/Private/Channels/MovieSceneFloatChannel.cpp
--- a/Engine/Source/Runtime/MovieScene/Private/Channels/MovieSceneFloatChannel.cpp
+++ b/Engine/Source/Runtime/MovieScene/Private/Channels/MovieSceneFloatChannel.cpp
@@ -812,33 +812,6 @@ FKeyHandle AddKeyToChannel(FMovieSceneFloatChannel* Channel, FFrameNumber InFram
 	{
 		FMovieSceneFloatValue& Value = ChannelData.GetValues()[ExistingIndex]; //-V758
 		Value.Value = InValue;
-		switch (Interpolation)
-		{
-		case EMovieSceneKeyInterpolation::User:
-			Value.InterpMode = RCIM_Cubic;
-			Value.TangentMode = RCTM_User;
-			break;
-
-		case EMovieSceneKeyInterpolation::Break:
-			Value.InterpMode = RCIM_Cubic;
-			Value.TangentMode = RCTM_Break;
-			break;
-
-		case EMovieSceneKeyInterpolation::Linear:
-			Value.InterpMode = RCIM_Linear;
-			Value.TangentMode = RCTM_Auto;
-			break;
-
-		case EMovieSceneKeyInterpolation::Constant:
-			Value.InterpMode = RCIM_Constant;
-			Value.TangentMode = RCTM_Auto;
-			break;
-
-		default:
-			Value.InterpMode = RCIM_Cubic;
-			Value.TangentMode = RCTM_Auto;
-			break;
-		}
 	}
 	else switch (Interpolation)
 	{
