Sequencer: FMovieSceneAccumulatedBlendState no longer consolidates blendable tokens if there are none to blend

diff --git a/Engine/Source/Runtime/MovieScene/Private/Evaluation/Blending/MovieSceneAccumulatedBlendState.cpp b/Engine/Source/Runtime/MovieScene/Private/Evaluation/Blending/MovieSceneAccumulatedBlendState.cpp
--- a/Engine/Source/Runtime/MovieScene/Private/Evaluation/Blending/MovieSceneAccumulatedBlendState.cpp
+++ b/Engine/Source/Runtime/MovieScene/Private/Evaluation/Blending/MovieSceneAccumulatedBlendState.cpp
@@ -5,6 +5,11 @@
 
 void FMovieSceneAccumulatedBlendState::Consolidate(TMap<UObject*, TMap<FMovieSceneBlendingActuatorID, FActuatorTokenStackPtr>>& InOutBlendState, FMovieSceneEvaluationOperand InOperand, IMovieScenePlayer& Player)
 {
+	if (TokensToBlend.Num() == 0)
+	{
+		return;
+	}
+
 	if (InOperand.ObjectBindingID.IsValid())
 	{
 		for (TWeakObjectPtr<> WeakObj : Player.FindBoundObjects(InOperand))
