Sequencer: If playback is at the end, calling play again should reset to the beginning.

diff --git a/Engine/Source/Runtime/MovieScene/Private/MovieSceneSequencePlayer.cpp b/Engine/Source/Runtime/MovieScene/Private/MovieSceneSequencePlayer.cpp
--- a/Engine/Source/Runtime/MovieScene/Private/MovieSceneSequencePlayer.cpp
+++ b/Engine/Source/Runtime/MovieScene/Private/MovieSceneSequencePlayer.cpp
@@ -88,8 +88,23 @@ void UMovieSceneSequencePlayer::PlayInternal()
 {
 	if (!IsPlaying() && Sequence && CanPlay())
 	{
-		// Start playing
+		// If at the end, rewind to the beginning to restart playback
+		if (bReversePlayback)
+		{
+			if (GetCurrentTime().Time == FFrameTime(StartTime))
+			{
+				JumpToFrame(GetLastValidTime());
+			}
+		}
+		else
+		{
+			if (GetCurrentTime().Time == GetLastValidTime())
+			{
+				JumpToFrame(FFrameTime(StartTime));
+			}
+		}
 
+		// Start playing
 		// @todo Sequencer playback: Should we recreate the instance every time?
 		// We must not recreate the instance since it holds stateful information (such as which objects it has spawned). Recreating the instance would break any 
 		// @todo: Is this still the case now that eval state is stored (correctly) in the player?
