Sequencer: Removed persistent FMovieSceneEvaluationTemplateInstance and replaced with a completely transient ptr cache

  - With the introduction of DirectorInstances (4.21), there became a need for a persistent store of data so directors could be kept alive for each sub-sequence
  - Making the instance map persistent seemed like the best way to achieve this, however, that resulted in pointers to data structures being held when they could be reallocated
  - Rather than revert back to 4.19's idea of a TransientInstances member that is repopulated every evaluation, the concept of an instance-per-sequence-id has been completely removed and changed for a stack allocated cache of the ptrs required for evaluation. Director instances are now stored in a separate map
  - This makes it much clearer that the ptrs are not to be held persistently, and reduces the probability of this confusion happening again

diff --git a/Engine/Source/Runtime/MovieScene/Private/Evaluation/MovieSceneEvaluationTemplateInstance.cpp b/Engine/Source/Runtime/MovieScene/Private/Evaluation/MovieSceneEvaluationTemplateInstance.cpp
--- a/Engine/Source/Runtime/MovieScene/Private/Evaluation/MovieSceneEvaluationTemplateInstance.cpp
+++ b/Engine/Source/Runtime/MovieScene/Private/Evaluation/MovieSceneEvaluationTemplateInstance.cpp
@@ -1,6 +1,7 @@
 // Copyright 1998-2018 Epic Games, Inc. All Rights Reserved.
 
 #include "Evaluation/MovieSceneEvaluationTemplateInstance.h"
+#include "Containers/SortedMap.h"
 #include "IMovieScenePlayer.h"
 #include "MovieSceneSequence.h"
 #include "MovieSceneSequence.h"
@@ -16,6 +17,65 @@ DECLARE_CYCLE_STAT(TEXT("Gather Entries For Frame"), MovieSceneEval_GatherEntrie
 DECLARE_CYCLE_STAT(TEXT("Call Setup() and TearDown()"), MovieSceneEval_CallSetupTearDown, STATGROUP_MovieSceneEval);
 DECLARE_CYCLE_STAT(TEXT("Evaluate Group"), MovieSceneEval_EvaluateGroup, STATGROUP_MovieSceneEval);
 
+
+/**
+ * Structure of ptrs that are cached only for the duration of an evaluation frame for a particular sequence.
+ * Any of this information may be re-allocated outside of an evaluation so cannot be persistently cached.
+ * 24 bytes
+ */
+struct FMovieSceneEvaluationPtrs
+{
+	/** The cached sequence ptr - always non-nullptr */
+	UMovieSceneSequence*              Sequence;
+	/** The cached template ptr - always non-nullptr */
+	FMovieSceneEvaluationTemplate*    Template;
+	/** The cached sub data ptr from the hierarchy. Only valid for sub sequences. */
+	const FMovieSceneSubSequenceData* SubData;
+};
+
+/**
+ * A cache of ptrs that need to be efficiently referred to during evaluation,
+ * but must not persist outside that evaluation.
+ */
+struct FMovieSceneEvaluationPtrCache
+{
+	typedef TSortedMap<FMovieSceneSequenceID, FMovieSceneEvaluationPtrs, TInlineAllocator<8>>::TConstIterator FConstIterator;
+
+	/** Construct the cache from a root sequence, template store and a set of sequence IDs that are to be included this frame */
+	FMovieSceneEvaluationPtrCache(const FMovieSceneRootOverridePath& RootOverridePath, UMovieSceneSequence* InRootSequence, IMovieSceneSequenceTemplateStore* InTemplateStore, TArrayView<FMovieSceneSequenceID const> InSubSequences);
+
+	/**
+	 * Attempt to locate the cached pointers for the specified sequence ID, failing gracefully if they were not found
+	 */
+	const FMovieSceneEvaluationPtrs* Find(FMovieSceneSequenceID SequenceID) const
+	{
+		return CachedPtrs.Find(SequenceID);
+	}
+
+	/**
+	 * Locate the cached pointers for the specified sequence ID assuming they exist, asseting if not
+	 */
+	const FMovieSceneEvaluationPtrs& GetChecked(FMovieSceneSequenceID SequenceID) const
+	{
+		return CachedPtrs.FindChecked(SequenceID);
+	}
+
+	FConstIterator CreateConstIterator() const
+	{
+		return CachedPtrs.CreateConstIterator();
+	}
+
+private:
+
+	/**
+	 * Generally we are dealing with very small numbers of sub sequences (or, just a single master sequence)
+	 * For this reason we allocate using a sorted map and an inline allocator to ensure that lookups are as
+	 * fast as possible in the common case
+	 */
+	TSortedMap<FMovieSceneSequenceID, FMovieSceneEvaluationPtrs, TInlineAllocator<8>> CachedPtrs;
+};
+
+
 /** Scoped helper class that facilitates the delayed restoration of preanimated state for specific evaluation keys */
 struct FDelayedPreAnimatedStateRestore
 {
@@ -49,27 +109,51 @@ private:
 	TArray<FMovieSceneEvaluationKey> KeysToRestore;
 };
 
-FMovieSceneEvaluationTemplateInstance::FMovieSceneEvaluationTemplateInstance()
-	: Sequence(nullptr), Template(nullptr), SubData(nullptr), DirectorInstance(nullptr)
-{}
+FMovieSceneEvaluationPtrCache::FMovieSceneEvaluationPtrCache(const FMovieSceneRootOverridePath& RootOverridePath, UMovieSceneSequence* InRootSequence, IMovieSceneSequenceTemplateStore* InTemplateStore, TArrayView<FMovieSceneSequenceID const> InSubSequences)
+{
+	// No root sequence == empty container
+	if (!InRootSequence)
+	{
+		return;
+	}
 
-FMovieSceneEvaluationTemplateInstance::FMovieSceneEvaluationTemplateInstance(UMovieSceneSequence& InSequence, FMovieSceneEvaluationTemplate* InTemplate)
-	: Sequence(&InSequence), Template(InTemplate), SubData(nullptr), DirectorInstance(nullptr)
-{}
+	check(InTemplateStore);
 
-FMovieSceneEvaluationTemplateInstance::FMovieSceneEvaluationTemplateInstance(const FMovieSceneSubSequenceData* InSubData, IMovieSceneSequenceTemplateStore& TemplateStore)
-	: Sequence(nullptr), Template(nullptr), SubData(InSubData), DirectorInstance(nullptr)
-{
-	Sequence = InSubData ? InSubData->GetSequence() : nullptr;
-	if (Sequence)
+	// Find the root template from the template store
+	FMovieSceneEvaluationTemplate* RootTemplate = &InTemplateStore->AccessTemplate(*InRootSequence);
+
+	// We always remap sequence IDs to their root space to ensure that spawnables presist properly when jumping into/out of shots with Eval in Isolation turned on
+	FMovieSceneSequenceID RemappedSequenceID = RootOverridePath.Remap(MovieSceneSequenceID::Root);
+
+	// Cache all the ptrs for the root sequence
+	CachedPtrs.Add(RemappedSequenceID, FMovieSceneEvaluationPtrs{ InRootSequence, RootTemplate, nullptr });
+
+	// Cache all sub-sequence ptrs
+	const FMovieSceneSequenceHierarchy& RootHierarchy = RootTemplate->Hierarchy;
+	for (FMovieSceneSequenceID SubSequenceID : InSubSequences)
 	{
-		Template = &TemplateStore.AccessTemplate(*Sequence);
+		if (SubSequenceID == MovieSceneSequenceID::Root)
+		{
+			continue;
+		}
+
+		const FMovieSceneSubSequenceData* SubData     = RootHierarchy.FindSubData(SubSequenceID);
+		UMovieSceneSequence*              SubSequence = SubData ? SubData->GetSequence() : nullptr;
+
+		// We gracefully handle nullptr here because in some rare cases a previous frame's meta-data may
+		// be referencing stale data that no longer exists
+		if (SubSequence)
+		{
+			RemappedSequenceID = RootOverridePath.Remap(SubSequenceID);
+			CachedPtrs.Add(RemappedSequenceID, FMovieSceneEvaluationPtrs{ SubSequence, &InTemplateStore->AccessTemplate(*SubSequence), SubData });
+		}
 	}
 }
 
 FMovieSceneRootEvaluationTemplateInstance::FMovieSceneRootEvaluationTemplateInstance()
 	: RootSequence(nullptr)
 	, TemplateStore(MakeShared<FMovieSceneSequencePrecompiledTemplateStore>())
+	, RootID(MovieSceneSequenceID::Root)
 {
 }
 
@@ -93,6 +177,8 @@ void FMovieSceneRootEvaluationTemplateInstance::Initialize(UMovieSceneSequence&
 {
 	if (RootSequence.Get() != &InRootSequence)
 	{
+		Finish(Player);
+
 		// Always ensure that there is no persistent data when initializing a new sequence
 		// to ensure we don't collide with the previous sequence's entity keys
 		Player.State.PersistentEntityData.Reset();
@@ -101,17 +187,12 @@ void FMovieSceneRootEvaluationTemplateInstance::Initialize(UMovieSceneSequence&
 		LastFrameMetaData.Reset();
 		ThisFrameMetaData.Reset();
 		ExecutionTokens = FMovieSceneExecutionTokens();
-
-		Instances.Reset();
 	}
 
 	RootSequence = &InRootSequence;
 	RootTemplate = &TemplateStore->AccessTemplate(InRootSequence);
 
-	Instances.RootID = MovieSceneSequenceID::Root;
-	Instances.RootInstance = FMovieSceneEvaluationTemplateInstance(InRootSequence, RootTemplate);
-
-	Player.State.AssignSequence(MovieSceneSequenceID::Root, InRootSequence, Player);
+	RootID = MovieSceneSequenceID::Root;
 }
 
 void FMovieSceneRootEvaluationTemplateInstance::Finish(IMovieScenePlayer& Player)
@@ -121,8 +202,7 @@ void FMovieSceneRootEvaluationTemplateInstance::Finish(IMovieScenePlayer& Player
 
 	CallSetupTearDown(Player);
 
-	Instances.RootInstance.DirectorInstance = nullptr;
-	Instances.ResetSubInstances();
+	ResetDirectorInstances();
 }
 
 void FMovieSceneRootEvaluationTemplateInstance::Evaluate(FMovieSceneContext Context, IMovieScenePlayer& Player, FMovieSceneSequenceID InOverrideRootID)
@@ -132,97 +212,114 @@ void FMovieSceneRootEvaluationTemplateInstance::Evaluate(FMovieSceneContext Cont
 	Swap(ThisFrameMetaData, LastFrameMetaData);
 	ThisFrameMetaData.Reset();
 
-	if (Instances.RootID != InOverrideRootID)
+	if (RootID != InOverrideRootID)
 	{
 		// Tear everything down if we're evaluating a different root sequence
 		CallSetupTearDown(Player);
 		LastFrameMetaData.Reset();
 	}
 
-	const FMovieSceneEvaluationGroup* GroupToEvaluate = SetupFrame(Player, InOverrideRootID, Context);
+	UMovieSceneSequence* OverrideRootSequence = GetSequence(InOverrideRootID);
+	if (!OverrideRootSequence)
+	{
+		CallSetupTearDown(Player);
+		return;
+	}
+
+	const FMovieSceneEvaluationGroup* GroupToEvaluate = SetupFrame(OverrideRootSequence, InOverrideRootID, Context);
 	if (!GroupToEvaluate)
 	{
 		CallSetupTearDown(Player);
 		return;
 	}
 
+	// Cache all the pointers needed for this frame
+	FMovieSceneEvaluationPtrCache EvaluationPtrCache = ConstructEvaluationPtrCacheForFrame(OverrideRootSequence);
+
+	// Ensure the correct sequences are assigned for each sequence ID
+	for (auto Iter = EvaluationPtrCache.CreateConstIterator(); Iter; ++Iter)
+	{
+		Player.State.AssignSequence(Iter.Key(), *Iter.Value().Sequence, Player);
+	}
+
 	// Cause stale tracks to not restore until after evaluation. This fixes issues when tracks that are set to 'Restore State' are regenerated, causing the state to be restored then re-animated by the new track
 	FDelayedPreAnimatedStateRestore DelayedRestore(Player);
 
 	// Run the post root evaluate steps which invoke tear downs for anything no longer evaluated.
 	// Do this now to ensure they don't undo any of the current frame's execution tokens 
-	CallSetupTearDown(Player, &DelayedRestore);
+	CallSetupTearDown(EvaluationPtrCache, Player, &DelayedRestore);
 
 	// Ensure any null objects are not cached
 	Player.State.InvalidateExpiredObjects();
 
 	// Accumulate execution tokens into this structure
-	EvaluateGroup(*GroupToEvaluate, Context, Player);
+	EvaluateGroup(EvaluationPtrCache, *GroupToEvaluate, Context, Player);
 
 	// Process execution tokens
 	ExecutionTokens.Apply(Context, Player);
 }
 
-const FMovieSceneEvaluationGroup* FMovieSceneRootEvaluationTemplateInstance::SetupFrame(IMovieScenePlayer& Player, FMovieSceneSequenceID InOverrideRootID, FMovieSceneContext Context)
+FMovieSceneEvaluationPtrCache FMovieSceneRootEvaluationTemplateInstance::ConstructEvaluationPtrCacheForFrame(UMovieSceneSequence* OverrideRootSequence)
 {
-	Instances.RootID = InOverrideRootID;
+	// We recreate all necessary sequence data for the current and previous frames by diffing the sequences active last frame, with this frame
+	TArray<FMovieSceneSequenceID> PreviousAndCurrentFrameSequenceIDs = LastFrameMetaData.ActiveSequences;
+	ThisFrameMetaData.DiffSequences(LastFrameMetaData, &PreviousAndCurrentFrameSequenceIDs, nullptr);
+
+	return FMovieSceneEvaluationPtrCache(RootOverridePath, OverrideRootSequence, TemplateStore.Get(), PreviousAndCurrentFrameSequenceIDs);
+}
+
+const FMovieSceneEvaluationGroup* FMovieSceneRootEvaluationTemplateInstance::SetupFrame(UMovieSceneSequence* OverrideRootSequence, FMovieSceneSequenceID InOverrideRootID, FMovieSceneContext Context)
+{
+	RootID = InOverrideRootID;
 	RootOverridePath.Set(InOverrideRootID, GetHierarchy());
 
-	UMovieSceneSequence* RootSequencePtr = RootSequence.Get();
-	if (!RootSequencePtr)
+	FMovieSceneEvaluationTemplate* OverrideRootTemplate = nullptr;
+	FMovieSceneSequenceTransform   RootOverrideTransform;
+
+	if (InOverrideRootID == MovieSceneSequenceID::Root)
 	{
-		return nullptr;
+		OverrideRootTemplate = RootTemplate;
 	}
-
-	// Create the instance for the current root override
-	FMovieSceneEvaluationTemplateInstance* RootOverrideInstance = Instances.Find(InOverrideRootID);
-
-	if (InOverrideRootID != MovieSceneSequenceID::Root && !RootOverrideInstance)
+	else
 	{
+		// Evaluate Sub Sequences in Isolation is turned on
+		OverrideRootTemplate = &TemplateStore->AccessTemplate(*OverrideRootSequence);
 		if (const FMovieSceneSubSequenceData* OverrideSubData = GetHierarchy().FindSubData(InOverrideRootID))
 		{
-			Instances.Add(InOverrideRootID, FMovieSceneEvaluationTemplateInstance(OverrideSubData, *TemplateStore));
-			RootOverrideInstance = Instances.Find(InOverrideRootID);
+			RootOverrideTransform = OverrideSubData->RootToSequenceTransform;
 		}
 	}
 
-
-	if (!ensureMsgf(RootOverrideInstance && RootOverrideInstance->IsValid(), TEXT("Could not find valid sequence or template for supplied sequence ID.")))
+	if (!ensureMsgf(OverrideRootSequence && OverrideRootTemplate, TEXT("Could not find valid sequence or template for supplied sequence ID.")))
 	{
 		return nullptr;
 	}
 
 	// Ensure the root is up to date
-	if (RootOverrideInstance->Template->SequenceSignature != RootOverrideInstance->Sequence->GetSignature())
-	{
-		FMovieSceneEvaluationTemplateGenerator(*RootOverrideInstance->Sequence, *RootOverrideInstance->Template).Generate();
-	}
-
-	FMovieSceneSequenceTransform MasterToRootOverrideTransform;
-	if (RootOverrideInstance->SubData)
+	if (OverrideRootTemplate->SequenceSignature != OverrideRootSequence->GetSignature())
 	{
-		MasterToRootOverrideTransform = RootOverrideInstance->SubData->RootToSequenceTransform;
+		FMovieSceneEvaluationTemplateGenerator(*OverrideRootSequence, *OverrideRootTemplate).Generate();
 	}
 
-	FFrameNumber RootOverrideTime = (Context.GetTime() * MasterToRootOverrideTransform).FloorToFrame();
+	FFrameNumber RootOverrideTime = (Context.GetTime() * RootOverrideTransform).FloorToFrame();
 
 	// First off, attempt to find the evaluation group in the existing evaluation field data from the template
-	int32 TemplateFieldIndex = RootOverrideInstance->Template->EvaluationField.GetSegmentFromTime(RootOverrideTime);
+	int32 TemplateFieldIndex = OverrideRootTemplate->EvaluationField.GetSegmentFromTime(RootOverrideTime);
 
 	if (TemplateFieldIndex != INDEX_NONE)
 	{
-		const FMovieSceneEvaluationMetaData& FieldMetaData = RootOverrideInstance->Template->EvaluationField.GetMetaData(TemplateFieldIndex);
+		const FMovieSceneEvaluationMetaData& FieldMetaData = OverrideRootTemplate->EvaluationField.GetMetaData(TemplateFieldIndex);
 
 		// Verify that this field entry is still valid (all its cached signatures are still the same)
 		TRange<FFrameNumber> InvalidatedSubSequenceRange = TRange<FFrameNumber>::Empty();
-		if (FieldMetaData.IsDirty(RootOverrideInstance->Template->Hierarchy, *TemplateStore, &InvalidatedSubSequenceRange))
+		if (FieldMetaData.IsDirty(OverrideRootTemplate->Hierarchy, *TemplateStore, &InvalidatedSubSequenceRange))
 		{
 			TemplateFieldIndex = INDEX_NONE;
 
 			if (!InvalidatedSubSequenceRange.IsEmpty())
 			{
 				// Invalidate the evaluation field for the root template (it may not exist until we compile below)
-				RootOverrideInstance->Template->EvaluationField.Invalidate(InvalidatedSubSequenceRange);
+				OverrideRootTemplate->EvaluationField.Invalidate(InvalidatedSubSequenceRange);
 			}
 		}
 	}
@@ -233,17 +330,17 @@ const FMovieSceneEvaluationGroup* FMovieSceneRootEvaluationTemplateInstance::Set
 		static bool bFullCompile = false;
 		if (bFullCompile)
 		{
-			FMovieSceneCompiler::Compile(*RootOverrideInstance->Sequence, *TemplateStore);
-			TemplateFieldIndex = RootOverrideInstance->Template->EvaluationField.GetSegmentFromTime(RootOverrideTime);
+			FMovieSceneCompiler::Compile(*OverrideRootSequence, *TemplateStore);
+			TemplateFieldIndex = OverrideRootTemplate->EvaluationField.GetSegmentFromTime(RootOverrideTime);
 		}
 		else
 		{
-			TOptional<FCompiledGroupResult> CompileResult = FMovieSceneCompiler::CompileTime(RootOverrideTime, *RootOverrideInstance->Sequence, *TemplateStore);
+			TOptional<FCompiledGroupResult> CompileResult = FMovieSceneCompiler::CompileTime(RootOverrideTime, *OverrideRootSequence, *TemplateStore);
 
 			if (CompileResult.IsSet())
 			{
 				TRange<FFrameNumber> FieldRange = CompileResult->Range;
-				TemplateFieldIndex = RootOverrideInstance->Template->EvaluationField.Insert(
+				TemplateFieldIndex = OverrideRootTemplate->EvaluationField.Insert(
 					RootOverrideTime,
 					FieldRange,
 					MoveTemp(CompileResult->Group),
@@ -256,17 +353,14 @@ const FMovieSceneEvaluationGroup* FMovieSceneRootEvaluationTemplateInstance::Set
 	if (TemplateFieldIndex != INDEX_NONE)
 	{
 		// Set meta-data
-		ThisFrameMetaData = RootOverrideInstance->Template->EvaluationField.GetMetaData(TemplateFieldIndex);
-
-		RecreateInstances(*RootOverrideInstance, InOverrideRootID, Player);
-
-		return &RootOverrideInstance->Template->EvaluationField.GetGroup(TemplateFieldIndex);
+		ThisFrameMetaData = OverrideRootTemplate->EvaluationField.GetMetaData(TemplateFieldIndex);
+		return &OverrideRootTemplate->EvaluationField.GetGroup(TemplateFieldIndex);
 	}
 
 	return nullptr;
 }
 
-void FMovieSceneRootEvaluationTemplateInstance::EvaluateGroup(const FMovieSceneEvaluationGroup& Group, const FMovieSceneContext& RootContext, IMovieScenePlayer& Player)
+void FMovieSceneRootEvaluationTemplateInstance::EvaluateGroup(const FMovieSceneEvaluationPtrCache& EvaluationPtrCache, const FMovieSceneEvaluationGroup& Group, const FMovieSceneContext& RootContext, IMovieScenePlayer& Player)
 {
 	MOVIESCENE_DETAILED_SCOPE_CYCLE_COUNTER(MovieSceneEval_EvaluateGroup);
 
@@ -289,8 +383,8 @@ void FMovieSceneRootEvaluationTemplateInstance::EvaluateGroup(const FMovieSceneE
 			// Ensure we're able to find the sequence instance in our root if we've overridden
 			SegmentPtr.SequenceID = RootOverridePath.Remap(SegmentPtr.SequenceID);
 
-			const FMovieSceneEvaluationTemplateInstance& Instance = Instances.GetChecked(SegmentPtr.SequenceID);
-			const FMovieSceneEvaluationTrack* Track = Instance.Template->FindTrack(SegmentPtr.TrackIdentifier);
+			const FMovieSceneEvaluationPtrs&  EvalPtrs = EvaluationPtrCache.GetChecked(SegmentPtr.SequenceID);
+			const FMovieSceneEvaluationTrack* Track    = EvalPtrs.Template->FindTrack(SegmentPtr.TrackIdentifier);
 
 			if (Track)
 			{
@@ -303,13 +397,13 @@ void FMovieSceneRootEvaluationTemplateInstance::EvaluateGroup(const FMovieSceneE
 				Player.PreAnimatedState.SetCaptureEntity(TrackKey, EMovieSceneCompletionMode::KeepState);
 
 				SubContext = Context;
-				if (Instance.SubData)
+				if (EvalPtrs.SubData)
 				{
-					SubContext = Context.Transform(Instance.SubData->RootToSequenceTransform, Instance.SubData->TickResolution);
+					SubContext = Context.Transform(EvalPtrs.SubData->RootToSequenceTransform, EvalPtrs.SubData->TickResolution);
 
 					// Hittest against the sequence's pre and postroll ranges
-					SubContext.ReportOuterSectionRanges(Instance.SubData->PreRollRange.Value, Instance.SubData->PostRollRange.Value);
-					SubContext.SetHierarchicalBias(Instance.SubData->HierarchicalBias);
+					SubContext.ReportOuterSectionRanges(EvalPtrs.SubData->PreRollRange.Value, EvalPtrs.SubData->PostRollRange.Value);
+					SubContext.SetHierarchicalBias(EvalPtrs.SubData->HierarchicalBias);
 				}
 
 				Track->Initialize(SegmentPtr.SegmentID, Operand, SubContext, PersistentDataProxy, Player);
@@ -329,8 +423,8 @@ void FMovieSceneRootEvaluationTemplateInstance::EvaluateGroup(const FMovieSceneE
 			// Ensure we're able to find the sequence instance in our root if we've overridden
 			SegmentPtr.SequenceID = RootOverridePath.Remap(SegmentPtr.SequenceID);
 
-			const FMovieSceneEvaluationTemplateInstance& Instance = Instances.GetChecked(SegmentPtr.SequenceID);
-			const FMovieSceneEvaluationTrack* Track = Instance.Template->FindTrack(SegmentPtr.TrackIdentifier);
+			const FMovieSceneEvaluationPtrs&  EvalPtrs = EvaluationPtrCache.GetChecked(SegmentPtr.SequenceID);
+			const FMovieSceneEvaluationTrack* Track    = EvalPtrs.Template->FindTrack(SegmentPtr.TrackIdentifier);
 
 			if (Track)
 			{
@@ -345,13 +439,13 @@ void FMovieSceneRootEvaluationTemplateInstance::EvaluateGroup(const FMovieSceneE
 				ExecutionTokens.SetCurrentScope(FMovieSceneEvaluationScope(TrackKey, EMovieSceneCompletionMode::KeepState));
 
 				SubContext = Context;
-				if (Instance.SubData)
+				if (EvalPtrs.SubData)
 				{
-					SubContext = Context.Transform(Instance.SubData->RootToSequenceTransform, Instance.SubData->TickResolution);
+					SubContext = Context.Transform(EvalPtrs.SubData->RootToSequenceTransform, EvalPtrs.SubData->TickResolution);
 
 					// Hittest against the sequence's pre and postroll ranges
-					SubContext.ReportOuterSectionRanges(Instance.SubData->PreRollRange.Value, Instance.SubData->PostRollRange.Value);
-					SubContext.SetHierarchicalBias(Instance.SubData->HierarchicalBias);
+					SubContext.ReportOuterSectionRanges(EvalPtrs.SubData->PreRollRange.Value, EvalPtrs.SubData->PostRollRange.Value);
+					SubContext.SetHierarchicalBias(EvalPtrs.SubData->HierarchicalBias);
 				}
 
 				Track->Evaluate(
@@ -367,7 +461,41 @@ void FMovieSceneRootEvaluationTemplateInstance::EvaluateGroup(const FMovieSceneE
 	}
 }
 
-void FMovieSceneRootEvaluationTemplateInstance::CallSetupTearDown(IMovieScenePlayer& Player, FDelayedPreAnimatedStateRestore* DelayedRestore)
+void FMovieSceneRootEvaluationTemplateInstance::CallSetupTearDown(IMovieScenePlayer& Player)
+{
+	UMovieSceneSequence* RootSequencePtr      = RootSequence.Get();
+	UMovieSceneSequence* OverrideRootSequence = nullptr;
+
+	if (RootID == MovieSceneSequenceID::Root)
+	{
+		OverrideRootSequence = RootSequencePtr;
+	}
+	else if (RootSequencePtr && RootTemplate)
+	{
+		// Evaluate Sub Sequences in Isolation is turned on
+		const FMovieSceneSubSequenceData* OverrideSubData = GetHierarchy().FindSubData(RootID);
+		if (OverrideSubData)
+		{
+			OverrideRootSequence = OverrideSubData->GetSequence();
+		}
+	}
+
+	if (OverrideRootSequence)
+	{
+		// Cache all the pointers needed for the teardown
+		FMovieSceneEvaluationPtrCache PtrCache = ConstructEvaluationPtrCacheForFrame(OverrideRootSequence);
+
+		// Ensure the correct sequences are assigned for each sequence ID
+		for (auto Iter = PtrCache.CreateConstIterator(); Iter; ++Iter)
+		{
+			Player.State.AssignSequence(Iter.Key(), *Iter.Value().Sequence, Player);
+		}
+
+		CallSetupTearDown(PtrCache, Player);
+	}
+}
+
+void FMovieSceneRootEvaluationTemplateInstance::CallSetupTearDown(const FMovieSceneEvaluationPtrCache& EvaluationPtrCache, IMovieScenePlayer& Player, FDelayedPreAnimatedStateRestore* DelayedRestore)
 {
 	MOVIESCENE_DETAILED_SCOPE_CYCLE_COUNTER(MovieSceneEval_CallSetupTearDown);
 
@@ -384,11 +512,11 @@ void FMovieSceneRootEvaluationTemplateInstance::CallSetupTearDown(IMovieScenePla
 		// Ensure we're able to find the sequence instance in our root if we've overridden
 		Key.SequenceID = RootOverridePath.Remap(Key.SequenceID);
 
-		const FMovieSceneEvaluationTemplateInstance* Instance = Instances.Find(Key.SequenceID);
-		if (Instance && Instance->IsValid())
+		const FMovieSceneEvaluationPtrs* EvalPtrs = EvaluationPtrCache.Find(Key.SequenceID);
+		if (EvalPtrs)
 		{
-			const FMovieSceneEvaluationTrack* Track = Instance->Template->FindTrack(Key.TrackIdentifier);
-			const bool bStaleTrack = Instance->Template->IsTrackStale(Key.TrackIdentifier);
+			const FMovieSceneEvaluationTrack* Track = EvalPtrs->Template->FindTrack(Key.TrackIdentifier);
+			const bool bStaleTrack = EvalPtrs->Template->IsTrackStale(Key.TrackIdentifier);
 
 			// Track data key may be required by both tracks and sections
 			PersistentDataProxy.SetTrackKey(Key.AsTrack());
@@ -431,10 +559,10 @@ void FMovieSceneRootEvaluationTemplateInstance::CallSetupTearDown(IMovieScenePla
 		// Ensure we're able to find the sequence instance in our root if we've overridden
 		Key.SequenceID = RootOverridePath.Remap(Key.SequenceID);
 
-		const FMovieSceneEvaluationTemplateInstance& Instance = Instances.GetChecked(Key.SequenceID);
-		check(Instance.IsValid());
+		const FMovieSceneEvaluationPtrs&  EvalPtrs = EvaluationPtrCache.GetChecked(Key.SequenceID);
+		const FMovieSceneEvaluationTrack* Track    = EvalPtrs.Template->FindTrack(Key.TrackIdentifier);
 
-		if (const FMovieSceneEvaluationTrack* Track = Instance.Template->FindTrack(Key.TrackIdentifier))
+		if (Track)
 		{
 			PersistentDataProxy.SetTrackKey(Key.AsTrack());
 
@@ -463,50 +591,60 @@ void FMovieSceneRootEvaluationTemplateInstance::CallSetupTearDown(IMovieScenePla
 
 bool FMovieSceneRootEvaluationTemplateInstance::IsDirty(TSet<UMovieSceneSequence*>* OutDirtySequences) const
 {
-	bool bIsDirty = false;
+	UMovieSceneSequence* RootSequencePtr = RootSequence.Get();
 
 	// Dirty if our master sequence is no longer valid
-	if (!Instances.RootInstance.IsValid())
+	if (!RootSequencePtr || !RootTemplate)
 	{
-		bIsDirty = true;
+		return true;
 	}
+
+	bool bIsDirty = false;
+
 	// Dirty if our master sequence signature doesn't match the template
-	else if (Instances.RootInstance.Template->SequenceSignature != Instances.RootInstance.Sequence->GetSignature())
+	if (RootTemplate->SequenceSignature != RootSequencePtr->GetSignature())
 	{
 		bIsDirty = true;
 
 		if (OutDirtySequences)
 		{
-			OutDirtySequences->Add(Instances.RootInstance.Sequence);
+			OutDirtySequences->Add(RootSequencePtr);
 		}
 	}
 
-	// Find the sequence we're actually evaluating (only != MovieSceneSequenceID::Root when "Evaluate Sequences in Isolation" is on)
-	const FMovieSceneEvaluationTemplateInstance* RootOverrideInstance = Instances.Find(Instances.RootID);
+	UMovieSceneSequence*           OverrideRootSequence = RootSequencePtr;
+	FMovieSceneEvaluationTemplate* OverrideRootTemplate = RootTemplate;
 
-	// Dirty if the root sequence is not valid
-	if (!RootOverrideInstance || !RootOverrideInstance->IsValid())
-	{
-		bIsDirty = true;
-	}
-	else
+	// Find the sequence we're actually evaluating (only != MovieSceneSequenceID::Root when "Evaluate Sequences in Isolation" is on)
+	if (RootID != MovieSceneSequenceID::Root)
 	{
-		// Dirty if our root override template signature doesn't match the sequence
-		if (RootOverrideInstance->Template->SequenceSignature != RootOverrideInstance->Sequence->GetSignature())
+		OverrideRootSequence = GetSequence(RootID);
+		OverrideRootTemplate = OverrideRootSequence ? &TemplateStore->AccessTemplate(*OverrideRootSequence) : nullptr;
+
+		// Dirty if the root sequence is not valid
+		if (!OverrideRootSequence || !OverrideRootTemplate)
 		{
 			bIsDirty = true;
-
-			if (OutDirtySequences)
+		}
+		else
+		{
+			// Dirty if our root override template signature doesn't match the sequence
+			if (OverrideRootTemplate->SequenceSignature != OverrideRootSequence->GetSignature())
 			{
-				OutDirtySequences->Add(RootOverrideInstance->Sequence);
+				bIsDirty = true;
+
+				if (OutDirtySequences)
+				{
+					OutDirtySequences->Add(OverrideRootSequence);
+				}
 			}
 		}
+	}
 
-		// Dirty if anything we evaluated last frame is dirty
-		if (LastFrameMetaData.IsDirty(RootOverrideInstance->Template->Hierarchy, *TemplateStore, nullptr, OutDirtySequences))
-		{
-			bIsDirty = true;
-		}
+	// Dirty if anything we evaluated last frame is dirty
+	if (OverrideRootTemplate && LastFrameMetaData.IsDirty(OverrideRootTemplate->Hierarchy, *TemplateStore, nullptr, OutDirtySequences))
+	{
+		bIsDirty = true;
 	}
 
 	return bIsDirty;
@@ -519,9 +657,14 @@ void FMovieSceneRootEvaluationTemplateInstance::CopyActuators(FMovieSceneBlendin
 
 UMovieSceneSequence* FMovieSceneRootEvaluationTemplateInstance::GetSequence(FMovieSceneSequenceIDRef SequenceID) const
 {
-	if (SequenceID == MovieSceneSequenceID::Root)
+	UMovieSceneSequence* RootSequencePtr = RootSequence.Get();
+	if (!RootSequencePtr)
+	{
+		return nullptr;
+	}
+	else if (SequenceID == MovieSceneSequenceID::Root)
 	{
-		return RootSequence.Get();
+		return RootSequencePtr;
 	}
 	
 	const FMovieSceneSubSequenceData* SubData = GetHierarchy().FindSubData(SequenceID);
@@ -543,64 +686,24 @@ FMovieSceneEvaluationTemplate* FMovieSceneRootEvaluationTemplateInstance::FindTe
 
 UObject* FMovieSceneRootEvaluationTemplateInstance::GetOrCreateDirectorInstance(FMovieSceneSequenceIDRef SequenceID, IMovieScenePlayer& Player)
 {
-	FMovieSceneEvaluationTemplateInstance* Instance = Instances.Find(SequenceID);
-
-	if (Instance && Instance->Sequence)
+	UObject* ExistingDirectorInstance = DirectorInstances.FindRef(SequenceID);
+	if (ExistingDirectorInstance)
 	{
-		if (!Instance->DirectorInstance)
-		{
-			Instance->DirectorInstance = Instance->Sequence->CreateDirectorInstance(Player);
-		}
-
-		return Instance->DirectorInstance;
+		return ExistingDirectorInstance;
 	}
 
-	return nullptr;
-}
-
-void FMovieSceneRootEvaluationTemplateInstance::ResetDirectorInstances()
-{
-	Instances.RootInstance.DirectorInstance = nullptr;
+	UMovieSceneSequence* Sequence            = GetSequence(SequenceID);
+	UObject*             NewDirectorInstance = Sequence->CreateDirectorInstance(Player);
 
-	for (TTuple<FMovieSceneSequenceID, FMovieSceneEvaluationTemplateInstance>& Pair : Instances.SubInstances)
+	if (NewDirectorInstance)
 	{
-		Pair.Value.DirectorInstance = nullptr;
+		DirectorInstances.Add(SequenceID, NewDirectorInstance);
 	}
+
+	return NewDirectorInstance;
 }
 
-void FMovieSceneRootEvaluationTemplateInstance::RecreateInstances(const FMovieSceneEvaluationTemplateInstance& RootOverrideInstance, FMovieSceneSequenceID InOverrideRootID, IMovieScenePlayer& Player)
+void FMovieSceneRootEvaluationTemplateInstance::ResetDirectorInstances()
 {
-	check(RootOverrideInstance.IsValid());
-
-	// We recreate all necessary sequence instances for the current and previous frames by diffing the sequences active last frame, with this frame
-	TArray<FMovieSceneSequenceID> PreviousAndCurrentFrameSequenceIDs = LastFrameMetaData.ActiveSequences;
-	ThisFrameMetaData.DiffSequences(LastFrameMetaData, &PreviousAndCurrentFrameSequenceIDs, nullptr);
-
-	for (FMovieSceneSequenceID SequenceID : PreviousAndCurrentFrameSequenceIDs)
-	{
-		if (Instances.Find(SequenceID))
-		{
-			continue;
-		}
-
-		const FMovieSceneSubSequenceData* SubData = RootOverrideInstance.Template->Hierarchy.FindSubData(SequenceID);
-		if (SubData)
-		{
-			FMovieSceneEvaluationTemplateInstance NewInstance(SubData, *TemplateStore);
-
-			if (!NewInstance.DirectorInstance && NewInstance.Sequence)
-			{
-				NewInstance.DirectorInstance = NewInstance.Sequence->CreateDirectorInstance(Player);
-			}
-
-			// Always add sequence IDs as the root
-			SequenceID = RootOverridePath.Remap(SequenceID);
-			Instances.Add(SequenceID, NewInstance);
-
-			if (NewInstance.Sequence)
-			{
-				Player.State.AssignSequence(SequenceID, *NewInstance.Sequence, Player);
-			}
-		}
-	}
-}
\ No newline at end of file
+	DirectorInstances.Reset();
+}
diff --git a/Engine/Source/Runtime/MovieScene/Private/Evaluation/MovieSceneSequenceTemplateStore.cpp b/Engine/Source/Runtime/MovieScene/Private/Evaluation/MovieSceneSequenceTemplateStore.cpp
--- a/Engine/Source/Runtime/MovieScene/Private/Evaluation/MovieSceneSequenceTemplateStore.cpp
+++ b/Engine/Source/Runtime/MovieScene/Private/Evaluation/MovieSceneSequenceTemplateStore.cpp
@@ -7,9 +7,4 @@
 FMovieSceneEvaluationTemplate& FMovieSceneSequencePrecompiledTemplateStore::AccessTemplate(UMovieSceneSequence& Sequence)
 {
 	return Sequence.PrecompiledEvaluationTemplate;
-}
-
-FMovieSceneEvaluationTemplate& FMovieSceneSequenceTemporaryTemplateStore::AccessTemplate(UMovieSceneSequence& Sequence)
-{
-	return Templates.FindOrAdd(&Sequence);
 }
\ No newline at end of file
diff --git a/Engine/Source/Runtime/MovieScene/Public/Evaluation/MovieSceneEvaluationTemplateInstance.h b/Engine/Source/Runtime/MovieScene/Public/Evaluation/MovieSceneEvaluationTemplateInstance.h
--- a/Engine/Source/Runtime/MovieScene/Public/Evaluation/MovieSceneEvaluationTemplateInstance.h
+++ b/Engine/Source/Runtime/MovieScene/Public/Evaluation/MovieSceneEvaluationTemplateInstance.h
@@ -2,7 +2,10 @@
 
 #pragma once
 
-#include "CoreMinimal.h"
+#include "CoreTypes.h"
+#include "Containers/Map.h"
+#include "UObject/WeakObjectPtr.h"
+#include "UObject/WeakObjectPtrTemplates.h"
 #include "MovieSceneSequenceID.h"
 #include "Evaluation/MovieSceneSequenceTransform.h"
 #include "Evaluation/MovieScenePlayback.h"
@@ -14,93 +17,8 @@
 
 class UMovieSceneSequence;
 struct FDelayedPreAnimatedStateRestore;
-struct FCompileOnTheFlyData;
+struct FMovieSceneEvaluationPtrCache;
 
-/**
- * An instance of an evaluation template. Fast to initialize and evaluate.
- */
-USTRUCT()
-struct FMovieSceneEvaluationTemplateInstance
-{
-	GENERATED_BODY()
-
-	FMovieSceneEvaluationTemplateInstance();
-
-	FMovieSceneEvaluationTemplateInstance(UMovieSceneSequence& InSequence, FMovieSceneEvaluationTemplate* InTemplate);
-
-	FMovieSceneEvaluationTemplateInstance(const FMovieSceneSubSequenceData* InSubData, IMovieSceneSequenceTemplateStore& TemplateStore);
-
-	bool IsValid() const
-	{
-		return Sequence && Template;
-	}
-
-	/** The sequence this applies to */
-	UPROPERTY()
-	UMovieSceneSequence* Sequence;
-
-	/** Pointer to the evaluation template we're evaluating */
-	FMovieSceneEvaluationTemplate* Template;
-
-	/** Sub Data or nullptr */
-	const FMovieSceneSubSequenceData* SubData;
-
-	/** The sequence this applies to */
-	UPROPERTY()
-	UObject* DirectorInstance;
-};
-
-USTRUCT()
-struct FMovieSceneEvaluationTemplateInstanceContainer
-{
-	GENERATED_BODY()
-
-	FMovieSceneEvaluationTemplateInstanceContainer()
-		: RootID(MovieSceneSequenceID::Invalid)
-	{}
-
-	void Reset()
-	{
-		RootID = MovieSceneSequenceID::Invalid;
-		RootInstance = FMovieSceneEvaluationTemplateInstance();
-		ResetSubInstances();
-	}
-
-	void ResetSubInstances()
-	{
-		SubInstances.Reset();
-	}
-
-	void Add(FMovieSceneSequenceID InID, const FMovieSceneEvaluationTemplateInstance& InInstance)
-	{
-		check(InID != RootID);
-		SubInstances.Add(InID, InInstance);
-	}
-
-	const FMovieSceneEvaluationTemplateInstance& GetChecked(FMovieSceneSequenceID InID) const
-	{
-		return InID == MovieSceneSequenceID::Root ? RootInstance : SubInstances.FindChecked(InID);
-	}
-
-	FMovieSceneEvaluationTemplateInstance* Find(FMovieSceneSequenceID InID)
-	{
-		return InID == MovieSceneSequenceID::Root ? &RootInstance : SubInstances.Find(InID);
-	}
-
-	const FMovieSceneEvaluationTemplateInstance* Find(FMovieSceneSequenceID InID) const
-	{
-		return InID == MovieSceneSequenceID::Root ? &RootInstance : SubInstances.Find(InID);
-	}
-
-	/** Sequence ID that is used to evaluate from for the current frame */
-	FMovieSceneSequenceID RootID;
-	
-	UPROPERTY()
-	FMovieSceneEvaluationTemplateInstance RootInstance;
-
-	UPROPERTY()
-	TMap<FMovieSceneSequenceID, FMovieSceneEvaluationTemplateInstance> SubInstances;
-};
 
 /**
  * Root evaluation template instance used to play back any sequence
@@ -232,36 +150,37 @@ public:
 	 */
 	MOVIESCENE_API void CopyActuators(FMovieSceneBlendingAccumulator& Accumulator) const;
 
-	/**
-	 * Const access to all the instances contained within this template instance
-	 */
-	const FMovieSceneEvaluationTemplateInstanceContainer& GetInstances() const
-	{
-		return Instances;
-	}
-
 private:
 
 	/**
 	 * Setup the current frame by finding or generating the necessary evaluation group and meta-data
 	 *
-	 * @param Context				Evaluation context containing the time (or range) to evaluate
+	 * @param OverrideRootSequence	Pointer to the sequence that is considered the root for this evaluation (that maps to InOverrideRootID)
+	 * @param OverrideRootID		The sequence ID of the currently considered root (normally MovieSceneSequenceID::Root unless Evaluate Sub-Sequences in Isolation is active)
+	 * @param Context				The evaluation context for this frame
+	 * @return The evaluation group within the root template's evaluation field to evaluate, or nullptr if none could be found or compiled
 	 */
-	const FMovieSceneEvaluationGroup* SetupFrame(IMovieScenePlayer& Player, FMovieSceneSequenceID InOverrideRootID, FMovieSceneContext Context);
+	const FMovieSceneEvaluationGroup* SetupFrame(UMovieSceneSequence* OverrideRootSequence, FMovieSceneSequenceID InOverrideRootID, FMovieSceneContext Context);
 
 	/**
 	 * Process entities that are newly evaluated, and those that are no longer being evaluated
 	 */
-	void CallSetupTearDown(IMovieScenePlayer& Player, FDelayedPreAnimatedStateRestore* DelayedRestore = nullptr);
+	void CallSetupTearDown(IMovieScenePlayer& Player);
 
 	/**
-	 * Evaluate a particular group of a segment
+	 * Process entities that are newly evaluated, and those that are no longer being evaluated
 	 */
-	void EvaluateGroup(const FMovieSceneEvaluationGroup& Group, const FMovieSceneContext& Context, IMovieScenePlayer& Player);
+	void CallSetupTearDown(const FMovieSceneEvaluationPtrCache& EvaluationCache, IMovieScenePlayer& Player, FDelayedPreAnimatedStateRestore* DelayedRestore = nullptr);
 
-private:
+	/**
+	 * Evaluate a particular group of a segment
+	 */
+	void EvaluateGroup(const FMovieSceneEvaluationPtrCache& EvaluationCache, const FMovieSceneEvaluationGroup& Group, const FMovieSceneContext& Context, IMovieScenePlayer& Player);
 
-	void RecreateInstances(const FMovieSceneEvaluationTemplateInstance& RootOverrideInstance, FMovieSceneSequenceID InOverrideRootID, IMovieScenePlayer& Player);
+	/**
+	 * Construct all the template and sub-data ptrs required for this frame by combining all those needed last frame, with those needed this frame
+	 */
+	FMovieSceneEvaluationPtrCache ConstructEvaluationPtrCacheForFrame(UMovieSceneSequence* OverrideRootSequence);
 
 private:
 
@@ -269,8 +188,12 @@ private:
 
 	FMovieSceneEvaluationTemplate* RootTemplate;
 
+	/** Sequence ID that was last used to evaluate from */
+	FMovieSceneSequenceID RootID;
+
+	/** Map of director instances by sequence ID. Kept alive by this map assuming this struct is reference collected */
 	UPROPERTY()
-	FMovieSceneEvaluationTemplateInstanceContainer Instances;
+	TMap<FMovieSceneSequenceID, UObject*> DirectorInstances;
 
 	/** Cache of everything that was evaluated last frame */
 	FMovieSceneEvaluationMetaData LastFrameMetaData;
diff --git a/Engine/Source/Runtime/MovieScene/Public/Evaluation/MovieSceneSequenceTemplateStore.h b/Engine/Source/Runtime/MovieScene/Public/Evaluation/MovieSceneSequenceTemplateStore.h
--- a/Engine/Source/Runtime/MovieScene/Public/Evaluation/MovieSceneSequenceTemplateStore.h
+++ b/Engine/Source/Runtime/MovieScene/Public/Evaluation/MovieSceneSequenceTemplateStore.h
@@ -7,21 +7,25 @@
 
 class UMovieSceneSequence;
 
+/**
+ * Interface that is used to retrieve an evaluation template for a given sequence
+ */
 struct IMovieSceneSequenceTemplateStore
 {
 	virtual ~IMovieSceneSequenceTemplateStore() {}
 
-	MOVIESCENE_API virtual FMovieSceneEvaluationTemplate& AccessTemplate(UMovieSceneSequence& Sequence) = 0;
+	/**
+	 * Access the template for the specified sequence.
+	 * Returned templates must be valid, and must never be re-allocated.
+	 */
+	virtual FMovieSceneEvaluationTemplate& AccessTemplate(UMovieSceneSequence& Sequence) = 0;
 };
 
-struct FMovieSceneSequencePrecompiledTemplateStore : IMovieSceneSequenceTemplateStore
-{
-	MOVIESCENE_API virtual FMovieSceneEvaluationTemplate& AccessTemplate(UMovieSceneSequence& Sequence);
-};
 
-struct FMovieSceneSequenceTemporaryTemplateStore : IMovieSceneSequenceTemplateStore
+/**
+ * Implementation of a template store that just returns UMovieSceneSequence::PrecompiledEvaluationTemplate
+ */
+struct FMovieSceneSequencePrecompiledTemplateStore : IMovieSceneSequenceTemplateStore
 {
 	MOVIESCENE_API virtual FMovieSceneEvaluationTemplate& AccessTemplate(UMovieSceneSequence& Sequence);
-private:
-	TMap<FObjectKey, FMovieSceneEvaluationTemplate> Templates;
-};
+};
\ No newline at end of file
