Sequencer: Added new (disabled) perf test that can be used for performance profiling of compilation and evaluation

diff --git a/Engine/Source/Runtime/MovieScene/Private/Tests/MovieSceneCompilerTests.cpp b/Engine/Source/Runtime/MovieScene/Private/Tests/MovieSceneCompilerTests.cpp
new file mode 100644
--- /dev/null
+++ b/Engine/Source/Runtime/MovieScene/Private/Tests/MovieSceneCompilerTests.cpp
@@ -0,0 +1,80 @@
+// Copyright 1998-2018 Epic Games, Inc. All Rights Reserved.
+
+#include "MovieSceneTestObjects.h"
+#include "Misc/AutomationTest.h"
+#include "Compilation/MovieSceneCompiler.h"
+#include "Compilation/MovieSceneSegmentCompiler.h"
+#include "Compilation/MovieSceneCompilerRules.h"
+#include "MovieSceneTestsCommon.h"
+#include "Evaluation/MovieSceneEvaluationTrack.h"
+#include "Evaluation/MovieSceneEvaluationField.h"
+#include "Evaluation/MovieSceneSequenceTemplateStore.h"
+#include "UObject/Package.h"
+
+#if WITH_DEV_AUTOMATION_TESTS
+
+IMPLEMENT_SIMPLE_AUTOMATION_TEST(FMovieSceneCompilerPerfTest, "System.Engine.Sequencer.Compiler.Perf", EAutomationTestFlags::EditorContext | EAutomationTestFlags::EngineFilter | EAutomationTestFlags::Disabled)
+bool FMovieSceneCompilerPerfTest::RunTest(const FString& Parameters)
+{
+	static bool  bFullCompile              = true;
+	static bool  bInvalidateEveryIteration = false;
+	static int32 NumIterations             = 1000000;
+
+	FFrameRate TickResolution(1000, 1);
+
+	UTestMovieSceneSequence* Sequence = NewObject<UTestMovieSceneSequence>(GetTransientPackage());
+	Sequence->MovieScene->SetTickResolutionDirectly(TickResolution);
+
+	for (int32 i = 0; i < 100; ++i)
+	{
+		UTestMovieSceneTrack* Track = NewObject<UTestMovieSceneTrack>(GetTransientPackage());
+
+		int32 NumSections = FMath::Rand() % 10;
+		for (int32 SectionIndex = 0; SectionIndex < NumSections; ++SectionIndex)
+		{
+			UTestMovieSceneSection* Section = NewObject<UTestMovieSceneSection>(Track);
+
+			double StartSeconds    = FMath::FRand() * 60.f;
+			double DurationSeconds = FMath::FRand() * 60.f;
+			Section->SetRange(TRange<FFrameNumber>::Inclusive( (StartSeconds * TickResolution).RoundToFrame(), ((StartSeconds + DurationSeconds) * TickResolution).RoundToFrame() ));
+			Track->SectionArray.Add(Section);
+		}
+	}
+
+	struct FTestMovieScenePlayer : IMovieScenePlayer
+	{
+		FMovieSceneRootEvaluationTemplateInstance RootInstance;
+		virtual FMovieSceneRootEvaluationTemplateInstance& GetEvaluationTemplate() override { return RootInstance; }
+		virtual void UpdateCameraCut(UObject* CameraObject, UObject* UnlockIfCameraObject = nullptr, bool bJumpCut = false) override {}
+		virtual void SetViewportSettings(const TMap<FViewportClient*, EMovieSceneViewportParams>& ViewportParamsMap) override {}
+		virtual void GetViewportSettings(TMap<FViewportClient*, EMovieSceneViewportParams>& ViewportParamsMap) const override {}
+		virtual EMovieScenePlayerStatus::Type GetPlaybackStatus() const override { return EMovieScenePlayerStatus::Playing; }
+		virtual void SetPlaybackStatus(EMovieScenePlayerStatus::Type InPlaybackStatus) override {}
+	} TestPlayer;
+
+	TestPlayer.RootInstance.Initialize(*Sequence, TestPlayer);
+
+	if (bFullCompile)
+	{
+		FMovieSceneSequencePrecompiledTemplateStore Store;
+		FMovieSceneCompiler::Compile(*Sequence, Store);
+	}
+
+	for (int32 i = 0; i < NumIterations; ++i)
+	{
+		if (bInvalidateEveryIteration)
+		{
+			Sequence->PrecompiledEvaluationTemplate.EvaluationField = FMovieSceneEvaluationField();
+		}
+
+		double StartSeconds    = FMath::FRand() * 60.f;
+		double DurationSeconds = FMath::FRand() * 1.f;
+
+		FMovieSceneEvaluationRange EvaluatedRange(TRange<FFrameTime>(StartSeconds * TickResolution, (StartSeconds + DurationSeconds) * TickResolution), TickResolution, EPlayDirection::Forwards);
+		TestPlayer.RootInstance.Evaluate(EvaluatedRange, TestPlayer);
+	}
+
+	return true;
+}
+
+#endif // WITH_DEV_AUTOMATION_TESTS
\ No newline at end of file
diff --git a/Engine/Source/Runtime/MovieScene/Private/Tests/MovieSceneSegmentCompilerTests.cpp b/Engine/Source/Runtime/MovieScene/Private/Tests/MovieSceneSegmentCompilerTests.cpp
--- a/Engine/Source/Runtime/MovieScene/Private/Tests/MovieSceneSegmentCompilerTests.cpp
+++ b/Engine/Source/Runtime/MovieScene/Private/Tests/MovieSceneSegmentCompilerTests.cpp
@@ -1,6 +1,6 @@
 // Copyright 1998-2018 Epic Games, Inc. All Rights Reserved.
 
-#include "MovieSceneSegmentCompilerTests.h"
+#include "MovieSceneTestObjects.h"
 #include "CoreMinimal.h"
 #include "Misc/AutomationTest.h"
 #include "Compilation/MovieSceneCompiler.h"
@@ -285,17 +285,17 @@ bool FMovieSceneTrackCompilerTest::RunTest(const FString& Parameters)
 	//	High-pass Filter Expected			[ 					|						0					|		 1			|							]
 
 	{
-		UMovieSceneSegmentCompilerTestTrack* Track = NewObject<UMovieSceneSegmentCompilerTestTrack>(GetTransientPackage());
+		UTestMovieSceneTrack* Track = NewObject<UTestMovieSceneTrack>(GetTransientPackage());
 		Track->EvalOptions.bCanEvaluateNearestSection = true;
 
-		UMovieSceneSegmentCompilerTestSection* Section0 = NewObject<UMovieSceneSegmentCompilerTestSection>(Track);
+		UTestMovieSceneSection* Section0 = NewObject<UTestMovieSceneSection>(Track);
 		Section0->SetRange(TRange<FFrameNumber>(
 			TRangeBound<FFrameNumber>::Inclusive(10),
 			TRangeBound<FFrameNumber>::Exclusive(25)
 		));
 		Section0->SetRowIndex(0);
 
-		UMovieSceneSegmentCompilerTestSection* Section1 = NewObject<UMovieSceneSegmentCompilerTestSection>(Track);
+		UTestMovieSceneSection* Section1 = NewObject<UTestMovieSceneSection>(Track);
 		Section1->SetRange(TRange<FFrameNumber>(
 			TRangeBound<FFrameNumber>::Inclusive(20),
 			TRangeBound<FFrameNumber>::Inclusive(30)
@@ -380,21 +380,21 @@ bool FMovieSceneTrackCompilerTest::RunTest(const FString& Parameters)
 	//	No Nearest Section Expected			[ 		1			|	(0,1)	|			(3,0,1)			|		(2,0,1)		|			1				]
 	//	High-Pass Filter Expected			[ 		1			|	0		|			3				|		2			|			1				]
 	{
-		UMovieSceneSegmentCompilerTestTrack* Track = NewObject<UMovieSceneSegmentCompilerTestTrack>(GetTransientPackage());
+		UTestMovieSceneTrack* Track = NewObject<UTestMovieSceneTrack>(GetTransientPackage());
 
-		UMovieSceneSegmentCompilerTestSection* Section0 = NewObject<UMovieSceneSegmentCompilerTestSection>(Track);
+		UTestMovieSceneSection* Section0 = NewObject<UTestMovieSceneSection>(Track);
 		Section0->SetRange(TRange<FFrameNumber>(10, 30));
 		Section0->SetRowIndex(1);
 
-		UMovieSceneSegmentCompilerTestSection* Section1 = NewObject<UMovieSceneSegmentCompilerTestSection>(Track);
+		UTestMovieSceneSection* Section1 = NewObject<UTestMovieSceneSection>(Track);
 		Section1->SetRange(TRange<FFrameNumber>::All());
 		Section1->SetRowIndex(2);
 		
-		UMovieSceneSegmentCompilerTestSection* Section2 = NewObject<UMovieSceneSegmentCompilerTestSection>(Track);
+		UTestMovieSceneSection* Section2 = NewObject<UTestMovieSceneSection>(Track);
 		Section2->SetRange(TRange<FFrameNumber>(20, 30));
 		Section2->SetRowIndex(0);
 
-		UMovieSceneSegmentCompilerTestSection* Section3 = NewObject<UMovieSceneSegmentCompilerTestSection>(Track);
+		UTestMovieSceneSection* Section3 = NewObject<UTestMovieSceneSection>(Track);
 		Section3->SetRange(TRange<FFrameNumber>(15, 25));
 		Section3->SetRowIndex(0);
 		Section3->SetOverlapPriority(100.f);
@@ -602,12 +602,12 @@ bool FMovieSceneCompilerTreeIteratorBoundsTest::RunTest(const FString& Parameter
 
 
 
-FMovieSceneEvalTemplatePtr UMovieSceneSegmentCompilerTestTrack::CreateTemplateForSection(const UMovieSceneSection& InSection) const
+FMovieSceneEvalTemplatePtr UTestMovieSceneTrack::CreateTemplateForSection(const UMovieSceneSection& InSection) const
 {
 	return FTestMovieSceneEvalTemplate();
 }
 
-FMovieSceneTrackSegmentBlenderPtr UMovieSceneSegmentCompilerTestTrack::GetTrackSegmentBlender() const
+FMovieSceneTrackSegmentBlenderPtr UTestMovieSceneTrack::GetTrackSegmentBlender() const
 {
 	struct FSegmentBlender : FMovieSceneTrackSegmentBlender
 	{
diff --git a/Engine/Source/Runtime/MovieScene/Private/Tests/MovieSceneSegmentCompilerTests.h b/Engine/Source/Runtime/MovieScene/Private/Tests/MovieSceneSegmentCompilerTests.h
deleted file mode 100644
diff --git a/Engine/Source/Runtime/MovieScene/Private/Tests/MovieSceneTestObjects.h b/Engine/Source/Runtime/MovieScene/Private/Tests/MovieSceneTestObjects.h
new file mode 100644
--- /dev/null
+++ b/Engine/Source/Runtime/MovieScene/Private/Tests/MovieSceneTestObjects.h
@@ -0,0 +1,64 @@
+// Copyright 1998-2018 Epic Games, Inc. All Rights Reserved.
+
+#pragma once
+
+#include "MovieScene.h"
+#include "MovieSceneTrack.h"
+#include "MovieSceneSection.h"
+#include "MovieSceneSequence.h"
+#include "Evaluation/MovieSceneEvalTemplate.h"
+#include "MovieSceneTestObjects.generated.h"
+
+USTRUCT()
+struct FTestMovieSceneEvalTemplate : public FMovieSceneEvalTemplate
+{
+	GENERATED_BODY()
+
+	virtual UScriptStruct& GetScriptStructImpl() const { return *StaticStruct(); }
+};
+
+UCLASS(MinimalAPI)
+class UTestMovieSceneTrack : public UMovieSceneTrack
+{
+public:
+
+	GENERATED_BODY()
+
+	virtual const TArray<UMovieSceneSection*>& GetAllSections() const override { return SectionArray; }
+	virtual FMovieSceneEvalTemplatePtr CreateTemplateForSection(const UMovieSceneSection& InSection) const override;
+	virtual FMovieSceneTrackSegmentBlenderPtr GetTrackSegmentBlender() const override;
+
+	UPROPERTY()
+	bool bHighPassFilter;
+
+	UPROPERTY()
+	TArray<UMovieSceneSection*> SectionArray;
+};
+
+UCLASS(MinimalAPI)
+class UTestMovieSceneSection : public UMovieSceneSection
+{
+public:
+	GENERATED_BODY()
+};
+
+UCLASS(MinimalAPI)
+class UTestMovieSceneSequence : public UMovieSceneSequence
+{
+public:
+	GENERATED_BODY()
+
+	UTestMovieSceneSequence(const FObjectInitializer& ObjInit)
+		: Super(ObjInit)
+	{
+		if (!HasAnyFlags(RF_ClassDefaultObject))
+		{
+			MovieScene = ObjInit.CreateDefaultSubobject<UMovieScene>(this, "MovieScene");
+		}
+	}
+
+	virtual UMovieScene* GetMovieScene() const override { return MovieScene; }
+
+	UPROPERTY()
+	UMovieScene* MovieScene;
+};
