Appease static analysis

diff --git a/Engine/Source/Runtime/MovieScene/Private/Evaluation/MovieSceneEvaluationTemplateInstance.cpp b/Engine/Source/Runtime/MovieScene/Private/Evaluation/MovieSceneEvaluationTemplateInstance.cpp
--- a/Engine/Source/Runtime/MovieScene/Private/Evaluation/MovieSceneEvaluationTemplateInstance.cpp
+++ b/Engine/Source/Runtime/MovieScene/Private/Evaluation/MovieSceneEvaluationTemplateInstance.cpp
@@ -270,6 +270,8 @@ FMovieSceneEvaluationPtrCache FMovieSceneRootEvaluationTemplateInstance::Constru
 
 const FMovieSceneEvaluationGroup* FMovieSceneRootEvaluationTemplateInstance::SetupFrame(UMovieSceneSequence* OverrideRootSequence, FMovieSceneSequenceID InOverrideRootID, FMovieSceneContext Context)
 {
+	check(OverrideRootSequence);
+
 	RootID = InOverrideRootID;
 	RootOverridePath.Set(InOverrideRootID, GetHierarchy());
 
@@ -290,7 +292,7 @@ const FMovieSceneEvaluationGroup* FMovieSceneRootEvaluationTemplateInstance::Set
 		}
 	}
 
-	if (!ensureMsgf(OverrideRootSequence && OverrideRootTemplate, TEXT("Could not find valid sequence or template for supplied sequence ID.")))
+	if (!ensureMsgf(OverrideRootTemplate, TEXT("Could not find valid template for supplied sequence ID.")))
 	{
 		return nullptr;
 	}
