Sequencer: Fix movie scene player with start offset, which was getting nullified because stop() was getting called last.

diff --git a/Engine/Source/Runtime/MovieScene/Private/MovieSceneSequencePlayer.cpp b/Engine/Source/Runtime/MovieScene/Private/MovieSceneSequencePlayer.cpp
--- a/Engine/Source/Runtime/MovieScene/Private/MovieSceneSequencePlayer.cpp
+++ b/Engine/Source/Runtime/MovieScene/Private/MovieSceneSequencePlayer.cpp
@@ -442,6 +442,8 @@ void UMovieSceneSequencePlayer::Initialize(UMovieSceneSequence* InSequence, cons
 	Sequence = InSequence;
 	PlaybackSettings = InSettings;
 
+	FFrameTime StartTimeWithOffset = StartTime;
+
 	EUpdateClockSource ClockToUse = EUpdateClockSource::Tick;
 
 	UMovieScene* MovieScene = Sequence->GetMovieScene();
@@ -478,7 +480,7 @@ void UMovieSceneSequencePlayer::Initialize(UMovieSceneSequence* InSequence, cons
 			? FFrameTime(FMath::Rand() % GetFrameDuration())
 			: FMath::Clamp<FFrameTime>(SpecifiedStartOffset, 0, GetFrameDuration()-1);
 
-		PlayPosition.Reset(StartTime + StartingTimeOffset);
+		StartTimeWithOffset = StartTime + StartingTimeOffset;
 
 		ClockToUse = MovieScene->GetClockSource();
 	}
@@ -493,11 +495,14 @@ void UMovieSceneSequencePlayer::Initialize(UMovieSceneSequence* InSequence, cons
 		}
 	}
 
-	PlaybackSettings.TimeController->Reset(GetCurrentTime());
 	RootTemplateInstance.Initialize(*Sequence, *this);
 
 	// Ensure everything is set up, ready for playback
 	Stop();
+
+	// Set up playback position (with offset) after Stop(), which will reset the starting time to StartTime
+	PlayPosition.Reset(StartTimeWithOffset);
+	PlaybackSettings.TimeController->Reset(GetCurrentTime());
 }
 
 void UMovieSceneSequencePlayer::Update(const float DeltaSeconds)
