Sequencer: When evaluating a range of a sequence, all overlapping segments of that range are now guaranteed to be up-to-date and compiled.

diff --git a/Engine/Source/Runtime/MovieScene/Private/Compilation/MovieSceneCompiler.cpp b/Engine/Source/Runtime/MovieScene/Private/Compilation/MovieSceneCompiler.cpp
--- a/Engine/Source/Runtime/MovieScene/Private/Compilation/MovieSceneCompiler.cpp
+++ b/Engine/Source/Runtime/MovieScene/Private/Compilation/MovieSceneCompiler.cpp
@@ -11,6 +11,10 @@
 #include "MovieSceneSequence.h"
 #include "Sections/MovieSceneSubSection.h"
 #include "Tracks/MovieSceneSubTrack.h"
+#include "MovieSceneCommonHelpers.h"
+
+DECLARE_CYCLE_STAT(TEXT("Full Compile"),  MovieSceneEval_CompileFull,  STATGROUP_MovieSceneEval);
+DECLARE_CYCLE_STAT(TEXT("Compile Range"), MovieSceneEval_CompileRange, STATGROUP_MovieSceneEval);
 
 /** Parameter structure used for gathering entities for a given time or range */
 struct FGatherParameters
@@ -100,13 +104,45 @@ struct FCompileOnTheFlyData
 struct FMovieSceneGatheredCompilerData
 {
 	/** Intersection of any empty space that overlaps the currently evaluating time range */
-	TRange<FFrameNumber> EmptyOverlappingRange;
+	FMovieSceneEvaluationTree EmptySpace;
 	/** Tree of tracks to evaluate */
 	TMovieSceneEvaluationTree<FCompileOnTheFlyData> Tracks;
 	/** Tree of active sequences */
 	TMovieSceneEvaluationTree<FMovieSceneSequenceID> Sequences;
 };
 
+/**
+ * Populate the specified tree with all the ranges from the specified array that fully encompass the specified range
+ * This is specifically used when compiling a specific range of an evaluation field in FMovieSceneCompiler::CompileRange()
+ * The desire is to have the first range-entry that exists before TestRange, the last entry-range that exists
+ * after TestRange, and all those in between. With this information we can quickly iterate the relevant gaps in
+ * the field along with the compiled data.
+ */
+void PopulateIterableTreeWithEncompassingRanges(const TRange<FFrameNumber>& TestRange, TArrayView<const FMovieSceneFrameRange> Ranges, TMovieSceneEvaluationTree<int32>& OutFieldTree)
+{
+	// Add the first range that's before the input range
+	int32 FirstIndex = Algo::LowerBoundBy(Ranges, TestRange.GetLowerBound(), &FMovieSceneFrameRange::GetLowerBound, MovieSceneHelpers::SortLowerBounds);
+	if (FirstIndex - 1 >= 0)
+	{
+		--FirstIndex;
+	}
+
+	TRangeBound<FFrameNumber> StopAfterBound = TRangeBound<FFrameNumber>::FlipInclusion(TestRange.GetUpperBound());
+
+	// Add all ranges that overlap the input range, and the first subsequent range
+	for (int32 Index = FirstIndex; Index < Ranges.Num(); ++Index)
+	{
+		OutFieldTree.Add(Ranges[Index].Value, Index);
+
+		// If this range's lower bound is >= the end of TestRange, we have enough information now to perform the compile
+		TRangeBound<FFrameNumber> ThisLowerBound = Ranges[Index].Value.GetLowerBound();
+		if (StopAfterBound.IsClosed() && ThisLowerBound.IsClosed() && TRangeBound<FFrameNumber>::MaxLower(ThisLowerBound, StopAfterBound) == ThisLowerBound)
+		{
+			break;
+		}
+	}
+}
+
 IMovieSceneModule& GetMovieSceneModule()
 {
 	static TWeakPtr<IMovieSceneModule> WeakMovieSceneModule;
@@ -124,6 +160,8 @@ IMovieSceneModule& GetMovieSceneModule()
 
 void FMovieSceneCompiler::Compile(UMovieSceneSequence& InCompileSequence, IMovieSceneSequenceTemplateStore& InTemplateStore)
 {
+	SCOPE_CYCLE_COUNTER(MovieSceneEval_CompileFull)
+
 	FMovieSceneEvaluationTemplate& CompileTemplate = InTemplateStore.AccessTemplate(InCompileSequence);
 
 	// Pass down a mutable path to the gather functions
@@ -168,36 +206,83 @@ void FMovieSceneCompiler::Compile(UMovieSceneSequence& InCompileSequence, IMovie
 	}
 }
 
-TOptional<FCompiledGroupResult> FMovieSceneCompiler::CompileTime(FFrameNumber InGlobalTime, UMovieSceneSequence& InCompileSequence, IMovieSceneSequenceTemplateStore& InTemplateStore)
+void FMovieSceneCompiler::CompileRange(TRange<FFrameNumber> InGlobalRange, UMovieSceneSequence& InCompileSequence, IMovieSceneSequenceTemplateStore& InTemplateStore)
 {
-	FMovieSceneEvaluationTemplate& CompileTemplate = InTemplateStore.AccessTemplate(InCompileSequence);
+	SCOPE_CYCLE_COUNTER(MovieSceneEval_CompileRange)
 
-	TRange<FFrameNumber> CompileRange = TRange<FFrameNumber>::Inclusive(InGlobalTime, InGlobalTime);
+	FMovieSceneEvaluationTemplate& CompileTemplate = InTemplateStore.AccessTemplate(InCompileSequence);
 
 	// Pass down a mutable path to the gather functions
 	FMovieSceneRootOverridePath RootPath;
 
-	// Gather everything that happens at this time, recursively
+	// Gather everything that happens over this range, recursively throughout the entire sequence
 	FMovieSceneGatheredCompilerData GatherData;
-	FGatherParameters GatherParams(RootPath, CompileTemplate.Hierarchy, InTemplateStore, CompileRange);
+	FGatherParameters GatherParams(RootPath, CompileTemplate.Hierarchy, InTemplateStore, InGlobalRange);
 	GatherCompileOnTheFlyData(InCompileSequence, GatherParams, GatherData);
 
-	FMovieSceneEvaluationTreeRangeIterator NodeAtTime = GatherData.Tracks.IterateFromTime(InGlobalTime);
-	if (ensure(NodeAtTime))
+	// --------------------------------------------------------------------------------------------------------------------
+	// When compiling a range we want to compile *at least* the range specified by InGlobalRange
+	// We may compile outside of this range if a gap in the evaluation field overlaps either bound, and the actual
+	// unique sequence state defines sections outside of the range.
+	// The general idea here is to iterate over any empty gaps in the evaluation field, populating it with the compiled
+	// result for each lower bound. Note that there will be one or more new field entries added for each gap, depending on
+	// whether any tracks or sections begin or end during the range of the gap.
+	// --------------------------------------------------------------------------------------------------------------------
+
+	// Populate an iterable tree with the ranges that at least encompass the range we want to comile,
+	// plus one either side of InGlobalRange if they exist. This allows us to fully understand which gaps we want to fill in
+	TMovieSceneEvaluationTree<int32> EvaluationFieldAsTree;
+	PopulateIterableTreeWithEncompassingRanges(InGlobalRange, CompileTemplate.EvaluationField.GetRanges(), EvaluationFieldAsTree);
+
+	// Start adding new field entries from the lower bound of the desired global range.
+	// IterFromBound should be <= InGlobalRange.GetLowerBound at this point
+	TRangeBound<FFrameNumber> IterFromBound = InGlobalRange.GetLowerBound();
+	FMovieSceneEvaluationTreeRangeIterator ExistingEvaluationFieldIter = EvaluationFieldAsTree.IterateFromLowerBound(IterFromBound);
+
+	// Now keep iterating the empty spaces in the field until we have nothing left to do.
+	// We only increment ExistingEvaluationFieldIter when it is at an already populated range,
+	// or if we've just compiled a range that has the same upper bound as the current gap (empty space)
+	TArray<FCompileOnTheFlyData> SortedCompileData;
+	for ( ; ExistingEvaluationFieldIter && !IterFromBound.IsOpen(); )
 	{
+		// If EvaluationFieldAsTree has any data at the current iterator position for it,
+		// the evaluation field is already populated for that node
+		if (EvaluationFieldAsTree.GetAllData(ExistingEvaluationFieldIter.Node()))
+		{
+			IterFromBound = TRangeBound<FFrameNumber>::FlipInclusion(ExistingEvaluationFieldIter.Range().GetUpperBound());
+			++ExistingEvaluationFieldIter;
+			continue;
+		}
+
+		TRange<FFrameNumber> EmptySpaceRange = ExistingEvaluationFieldIter.Range();
+
+		// Find the intersection of all the current ranges (the gap in the evaluation field, the track field, sub sequence field, and empty space)
+		FMovieSceneEvaluationTreeRangeIterator TrackIteratorFromHere       = GatherData.Tracks.IterateFromLowerBound(IterFromBound);
+		FMovieSceneEvaluationTreeRangeIterator SubSequenceIteratorFromHere = GatherData.Sequences.IterateFromLowerBound(IterFromBound);
+		FMovieSceneEvaluationTreeRangeIterator EmptySpaceIteratorFromHere  = GatherData.EmptySpace.IterateFromLowerBound(IterFromBound);
+
+		// Find the intersection of all the compiled data
 		TRange<FFrameNumber> CompiledRange = TRange<FFrameNumber>::Intersection(
+			EmptySpaceRange,
 			TRange<FFrameNumber>::Intersection(
-				GatherData.EmptyOverlappingRange,
-				GatherData.Sequences.IterateFromLowerBound(CompileRange.GetLowerBound()).Range()
-			),
-			NodeAtTime.Range());
-
-		// Always include the global time in the compiled range - this is necessary when floating point rounding
-		// introduced during sub sequence transformations shift the range to just outside the current time
-		CompiledRange = TRange<FFrameNumber>::Hull(CompiledRange, CompileRange);
+				TrackIteratorFromHere.Range(),
+				TRange<FFrameNumber>::Intersection(
+					EmptySpaceIteratorFromHere.Range(),
+					SubSequenceIteratorFromHere.Range()
+				)
+			)
+		);
+
+		// If the range we just compiled no longer overlaps the range we were asked to compile,
+		// Break out of the loop as all of our work is done. This will happen if there is a gap
+		// in the evaluation field that overlaps with the upper bound of InGlobalRange.
+		if (!CompiledRange.Overlaps(InGlobalRange))
+		{
+			break;
+		}
 
-		TArray<FCompileOnTheFlyData> SortedCompileData;
-		for (const FCompileOnTheFlyData& TrackData : GatherData.Tracks.GetAllData(NodeAtTime.Node()))
+		SortedCompileData.Reset();
+		for (const FCompileOnTheFlyData& TrackData : GatherData.Tracks.GetAllData(TrackIteratorFromHere.Node()))
 		{
 			SortedCompileData.Add(TrackData);
 		}
@@ -215,13 +300,23 @@ TOptional<FCompiledGroupResult> FMovieSceneCompiler::CompileTime(FFrameNumber In
 		PopulateEvaluationGroup(Result, SortedCompileData);
 
 		// Compute meta data for this segment
-		TMovieSceneEvaluationTreeDataIterator<FMovieSceneSequenceID> SubSequences = GatherData.Sequences.GetAllData(GatherData.Sequences.IterateFromLowerBound(CompileRange.GetLowerBound()).Node());
+		TMovieSceneEvaluationTreeDataIterator<FMovieSceneSequenceID> SubSequences = GatherData.Sequences.GetAllData(SubSequenceIteratorFromHere.Node());
 		PopulateMetaData(Result, CompileTemplate.Hierarchy, InTemplateStore, SortedCompileData, SubSequences);
 
-		return Result;
-	}
+		// Add the results to the evaluation field and continue iterating starting from the end of the compiled range
+		CompileTemplate.EvaluationField.Insert(Result.Range, MoveTemp(Result.Group), MoveTemp(Result.MetaData));
+
+		// We may still have some to compile
+		IterFromBound = TRangeBound<FFrameNumber>::FlipInclusion(CompiledRange.GetUpperBound());
 
-	return TOptional<FCompiledGroupResult>();
+		// If the range that we just compiled goes right up to the end of the gap, increment onto the
+		// next entry in the evaluation field iterator (which should be a populated range)
+		if (CompiledRange.GetUpperBound() == EmptySpaceRange.GetUpperBound())
+		{
+			IterFromBound = TRangeBound<FFrameNumber>::FlipInclusion(CompiledRange.GetUpperBound());
+			++ExistingEvaluationFieldIter;
+		}
+	}
 }
 
 void FMovieSceneCompiler::CompileHierarchy(const UMovieSceneSequence& InRootSequence, FMovieSceneSequenceHierarchy& OutHierarchy, FMovieSceneSequenceID RootSequenceID, int32 MaxDepth)
@@ -323,9 +418,6 @@ void FMovieSceneCompiler::GatherCompileOnTheFlyData(UMovieSceneSequence& InSeque
 	// Start iterating the field from the lower bound of the compile range
 	FMovieSceneEvaluationTreeRangeIterator SubSectionIt(SubSectionField.IterateFromLowerBound(CompileClampIntersection.GetLowerBound()));
 
-	// Intersect the unique range in the tree with the current overlapping empty range to constrict the resulting compile range in the case where this is a gap between sub sections
-	OutData.EmptyOverlappingRange = TRange<FFrameNumber>::Intersection(OutData.EmptyOverlappingRange, Params.ClampRoot(SubSectionIt.Range() * Params.RootToSequenceTransform.Inverse()));
-
 	for ( ; SubSectionIt && SubSectionIt.Range().Overlaps(CompileClampIntersection); ++SubSectionIt)
 	{
 		TRange<FFrameNumber> ThisSegmentRangeRoot = Params.ClampRoot(SubSectionIt.Range() * Params.RootToSequenceTransform.Inverse());
@@ -336,6 +428,8 @@ void FMovieSceneCompiler::GatherCompileOnTheFlyData(UMovieSceneSequence& InSeque
 
 		SubSectionGatherParams.SetClampRange(ThisSegmentRangeRoot);
 
+		bool bAnySubSections = false;
+
 		// Iterate all sub sections in the current range
 		for (const FMovieSceneSubSectionData& SubSectionData : SubSectionField.GetAllData(SubSectionIt.Node()))
 		{
@@ -353,11 +447,19 @@ void FMovieSceneCompiler::GatherCompileOnTheFlyData(UMovieSceneSequence& InSeque
 
 			if (bTrackMatchesFlags && SubSection)
 			{
+				bAnySubSections = true;
+
 				SubSectionGatherParams.Flags = SubSectionData.Flags;
 
 				GatherCompileDataForSubSection(*SubSection, SubSectionData.ObjectBindingId, SubSectionGatherParams, OutData);
 			}
 		}
+
+		if (!bAnySubSections)
+		{
+			// Intersect the unique range in the tree with the current overlapping empty range to constrict the resulting compile range in the case where this is a gap between sub sections
+			OutData.EmptySpace.AddTimeRange(Params.ClampRoot(SubSectionIt.Range() * Params.RootToSequenceTransform.Inverse()));
+		}
 	}
 }
 
@@ -452,7 +554,7 @@ void FMovieSceneCompiler::GatherCompileDataForTrack(FMovieSceneEvaluationTrack&
 		TRange<FFrameNumber> EmptyTrackSpace = Track.GetUniqueRangeFromLowerBound(Params.LocalCompileRange.GetLowerBound());
 		TRange<FFrameNumber> ClampedEmptyTrackSpaceRoot = Params.ClampRoot(EmptyTrackSpace * SequenceToRootTransform);
 
-		OutData.EmptyOverlappingRange = TRange<FFrameNumber>::Intersection(ClampedEmptyTrackSpaceRoot, OutData.EmptyOverlappingRange);
+		OutData.EmptySpace.AddTimeRange(ClampedEmptyTrackSpaceRoot);
 	}
 	else for (FMovieSceneSegmentIdentifier SegmentID : SegmentIDs)
 	{
diff --git a/Engine/Source/Runtime/MovieScene/Private/Evaluation/MovieSceneEvaluationField.cpp b/Engine/Source/Runtime/MovieScene/Private/Evaluation/MovieSceneEvaluationField.cpp
--- a/Engine/Source/Runtime/MovieScene/Private/Evaluation/MovieSceneEvaluationField.cpp
+++ b/Engine/Source/Runtime/MovieScene/Private/Evaluation/MovieSceneEvaluationField.cpp
@@ -3,11 +3,94 @@
 #include "Evaluation/MovieSceneEvaluationField.h"
 #include "Evaluation/MovieSceneEvaluationTemplateInstance.h"
 #include "Evaluation/MovieSceneSequenceTemplateStore.h"
+#include "Evaluation/MovieSceneEvaluationTree.h"
+#include "Compilation/MovieSceneCompiler.h"
 #include "MovieSceneCommonHelpers.h"
 #include "Algo/Sort.h"
 
 #include "MovieSceneSequence.h"
 
+TRange<int32> FMovieSceneEvaluationField::ConditionallyCompileRange(const TRange<FFrameNumber>& InRange, UMovieSceneSequence* InSequence, IMovieSceneSequenceTemplateStore& TemplateStore)
+{
+	check(InSequence);
+
+	// First off, attempt to find the evaluation group in the existing evaluation field data from the template
+	TRange<int32> OverlappingFieldEntries = OverlapRange(InRange);
+	int32 EvalFieldStartIndex = OverlappingFieldEntries.GetLowerBoundValue();
+	int32 EvalFieldEndIndex   = OverlappingFieldEntries.GetUpperBoundValue();
+
+	bool bIsDirty = OverlappingFieldEntries.IsEmpty();
+
+	const FMovieSceneSequenceHierarchy& RootHierarchy = TemplateStore.AccessTemplate(*InSequence).Hierarchy;
+
+	TArray<TRange<FFrameNumber>, TInlineAllocator<8>> RangesToInvalidate;
+	for (int32 Index = EvalFieldStartIndex; Index < EvalFieldEndIndex; ++Index)
+	{
+		const TRange<FFrameNumber>& ThisRange = Ranges[Index].Value;
+
+		// Check for gaps between the entries.
+		if (Index == EvalFieldStartIndex)
+		{
+			// If the first overlapping range starts after InRange's lower bound, there must be a gap before it
+			if (TRangeBound<FFrameNumber>::MinLower(ThisRange.GetLowerBound(), InRange.GetLowerBound()) != ThisRange.GetLowerBound())
+			{
+				bIsDirty = true;
+			}
+		}
+		if (Index == EvalFieldEndIndex - 1)
+		{
+			// If the last overlapping range ends before InRange's upper bound, there must be a gap after it
+			if (TRangeBound<FFrameNumber>::MaxUpper(ThisRange.GetUpperBound(), InRange.GetUpperBound()) != ThisRange.GetUpperBound())
+			{
+				bIsDirty = true;
+			}
+		}
+
+		// If adjacent ranges are not contiguous, we have a gap
+		if (Index > EvalFieldStartIndex && Ranges.IsValidIndex(Index-1) && !Ranges[Index-1].Value.Adjoins(ThisRange))
+		{
+			bIsDirty = true;
+		}
+
+		// Verify that this field entry is still valid (all its cached signatures are still the same)
+		TRange<FFrameNumber> InvalidatedSubSequenceRange = TRange<FFrameNumber>::Empty();
+		if (MetaData[Index].IsDirty(RootHierarchy, TemplateStore, &InvalidatedSubSequenceRange))
+		{
+			bIsDirty = true;
+
+			if (!InvalidatedSubSequenceRange.IsEmpty())
+			{
+				// Invalidate this evaluation field
+				RangesToInvalidate.Add(InvalidatedSubSequenceRange);
+			}
+		}
+	}
+
+	// Invalidate any areas in the evaluation field that are now out of date
+	for (const TRange<FFrameNumber>& Range : RangesToInvalidate)
+	{
+		Invalidate(Range);
+	}
+
+	if (bIsDirty)
+	{
+		// We need to compile an entry in the evaluation field
+		static bool bFullCompile = false;
+ 		if (bFullCompile)
+		{
+			FMovieSceneCompiler::Compile(*InSequence, TemplateStore);
+		}
+		else
+		{
+			FMovieSceneCompiler::CompileRange(InRange, *InSequence, TemplateStore);
+		}
+
+		return OverlapRange(InRange);
+	}
+
+	return OverlappingFieldEntries;
+}
+
 int32 FMovieSceneEvaluationField::GetSegmentFromTime(FFrameNumber Time) const
 {
 	// Linear search
@@ -23,29 +106,29 @@ int32 FMovieSceneEvaluationField::GetSegmentFromTime(FFrameNumber Time) const
 	return INDEX_NONE;
 }
 
-TRange<int32> FMovieSceneEvaluationField::OverlapRange(TRange<FFrameNumber> Range) const
+TRange<int32> FMovieSceneEvaluationField::OverlapRange(const TRange<FFrameNumber>& Range) const
 {
-	int32 StartIndex = 0, Num = 0;
-	for (int32 Index = 0; Index < Ranges.Num(); ++Index)
+	if (Ranges.Num() == 0)
 	{
-		if (Ranges[Index].Value.Overlaps(Range))
-		{
-			if (Num == 0)
-			{
-				StartIndex = Index;
-			}
-			++Num;
-		}
-		else if (Num != 0)
-		{
-			break;
-		}
+		return TRange<int32>::Empty();
 	}
 
-	return Num != 0 ? TRange<int32>(StartIndex, StartIndex + Num) : TRange<int32>::Empty();
+	// Binary search the first lower bound that's greater than the input range's lower bound
+	int32 StartIndex = Algo::UpperBoundBy(Ranges, Range.GetLowerBound(), &FMovieSceneFrameRange::GetLowerBound, MovieSceneHelpers::SortLowerBounds);
+
+	// StartIndex is always <= Ranges.Num(). If the previous range overlaps the input range, include that
+	if (StartIndex > 0 && Ranges[StartIndex-1].Value.Overlaps(Range))
+	{
+		StartIndex = StartIndex - 1;
+	}
+
+	// Binary search the first lower bound that is greater than or equal to the input range's upper bound
+	int32 EndIndex = Range.GetUpperBound().IsOpen() ? Ranges.Num() : Algo::LowerBoundBy(Ranges, TRangeBound<FFrameNumber>::FlipInclusion(Range.GetUpperBound()), &FMovieSceneFrameRange::GetLowerBound, MovieSceneHelpers::SortLowerBounds);
+
+	return EndIndex - StartIndex > 0 ? TRange<int32>(StartIndex, EndIndex) : TRange<int32>::Empty();
 }
 
-void FMovieSceneEvaluationField::Invalidate(TRange<FFrameNumber> Range)
+void FMovieSceneEvaluationField::Invalidate(const TRange<FFrameNumber>& Range)
 {
 	TRange<int32> OverlappingRange = OverlapRange(Range);
 	if (!OverlappingRange.IsEmpty())
@@ -60,32 +143,9 @@ void FMovieSceneEvaluationField::Invalidate(TRange<FFrameNumber> Range)
 	}
 }
 
-int32 FMovieSceneEvaluationField::Insert(FFrameNumber InsertTime, TRange<FFrameNumber> InRange, FMovieSceneEvaluationGroup&& InGroup, FMovieSceneEvaluationMetaData&& InMetaData)
+int32 FMovieSceneEvaluationField::Insert(const TRange<FFrameNumber>& InRange, FMovieSceneEvaluationGroup&& InGroup, FMovieSceneEvaluationMetaData&& InMetaData)
 {
-	const int32 InsertIndex = Algo::UpperBoundBy(Ranges, TRangeBound<FFrameNumber>(InsertTime), &FMovieSceneFrameRange::GetLowerBound, MovieSceneHelpers::SortLowerBounds);
-
-	// Intersect the supplied range with the allowable space between adjacent existing ranges
-	TRange<FFrameNumber> InsertSpace(
-		Ranges.IsValidIndex(InsertIndex-1) ? TRangeBound<FFrameNumber>::FlipInclusion(Ranges[InsertIndex-1].GetUpperBound()) : TRangeBound<FFrameNumber>::Open(),
-		Ranges.IsValidIndex(InsertIndex  ) ? TRangeBound<FFrameNumber>::FlipInclusion(Ranges[InsertIndex  ].GetLowerBound()) : TRangeBound<FFrameNumber>::Open()
-		);
-
-	InRange = TRange<FFrameNumber>::Intersection(InRange, InsertSpace);
-
-	// @todo: Remove this code and enforce the check below outright when we have proper time representation
-	if (Ranges.IsValidIndex(InsertIndex  ) && Ranges[InsertIndex  ].Value.Overlaps(InRange))
-	{
-		InRange = TRange<FFrameNumber>(InRange.GetLowerBound(), TRangeBound<FFrameNumber>::FlipInclusion(Ranges[InsertIndex].Value.GetLowerBound()));
-	}
-	if (Ranges.IsValidIndex(InsertIndex-1) && Ranges[InsertIndex-1].Value.Overlaps(InRange))
-	{
-		InRange = TRange<FFrameNumber>(TRangeBound<FFrameNumber>::FlipInclusion(Ranges[InsertIndex-1].Value.GetUpperBound()), InRange.GetUpperBound());
-	}
-
-	if (!ensure(!InRange.IsEmpty()))
-	{
-		return INDEX_NONE;
-	}
+	const int32 InsertIndex = Algo::UpperBoundBy(Ranges, InRange.GetLowerBound(), &FMovieSceneFrameRange::GetLowerBound, MovieSceneHelpers::SortLowerBounds);
 
 	const bool bOverlapping = 
 		(Ranges.IsValidIndex(InsertIndex  ) && Ranges[InsertIndex  ].Value.Overlaps(InRange)) ||
@@ -107,7 +167,7 @@ int32 FMovieSceneEvaluationField::Insert(FFrameNumber InsertTime, TRange<FFrameN
 	return InsertIndex;
 }
 
-void FMovieSceneEvaluationField::Add(TRange<FFrameNumber> InRange, FMovieSceneEvaluationGroup&& InGroup, FMovieSceneEvaluationMetaData&& InMetaData)
+void FMovieSceneEvaluationField::Add(const TRange<FFrameNumber>& InRange, FMovieSceneEvaluationGroup&& InGroup, FMovieSceneEvaluationMetaData&& InMetaData)
 {
 	if (ensureAlwaysMsgf(!Ranges.Num() || !Ranges.Last().Value.Overlaps(InRange), TEXT("Attempting to add overlapping ranges to sequence evaluation field.")))
 	{
diff --git a/Engine/Source/Runtime/MovieScene/Private/Evaluation/MovieSceneEvaluationTemplateInstance.cpp b/Engine/Source/Runtime/MovieScene/Private/Evaluation/MovieSceneEvaluationTemplateInstance.cpp
--- a/Engine/Source/Runtime/MovieScene/Private/Evaluation/MovieSceneEvaluationTemplateInstance.cpp
+++ b/Engine/Source/Runtime/MovieScene/Private/Evaluation/MovieSceneEvaluationTemplateInstance.cpp
@@ -305,58 +305,28 @@ const FMovieSceneEvaluationGroup* FMovieSceneRootEvaluationTemplateInstance::Set
 
 	FFrameNumber RootOverrideTime = (Context.GetTime() * RootOverrideTransform).FloorToFrame();
 
-	// First off, attempt to find the evaluation group in the existing evaluation field data from the template
-	int32 TemplateFieldIndex = OverrideRootTemplate->EvaluationField.GetSegmentFromTime(RootOverrideTime);
 
-	if (TemplateFieldIndex != INDEX_NONE)
-	{
-		const FMovieSceneEvaluationMetaData& FieldMetaData = OverrideRootTemplate->EvaluationField.GetMetaData(TemplateFieldIndex);
 
-		// Verify that this field entry is still valid (all its cached signatures are still the same)
-		TRange<FFrameNumber> InvalidatedSubSequenceRange = TRange<FFrameNumber>::Empty();
-		if (FieldMetaData.IsDirty(OverrideRootTemplate->Hierarchy, *TemplateStore, &InvalidatedSubSequenceRange))
-		{
-			TemplateFieldIndex = INDEX_NONE;
+	FMovieSceneEvaluationField& EvaluationField = OverrideRootTemplate->EvaluationField;
 
-			if (!InvalidatedSubSequenceRange.IsEmpty())
-			{
-				// Invalidate the evaluation field for the root template (it may not exist until we compile below)
-				OverrideRootTemplate->EvaluationField.Invalidate(InvalidatedSubSequenceRange);
-			}
-		}
-	}
+	// Get the range that we are evaluating in the root's space
+	TRange<FFrameNumber> ContextRange = Context.GetTraversedFrameNumberRange() * RootOverrideTransform;
 
-	if (TemplateFieldIndex == INDEX_NONE)
+	// Verify and update the evaluation field for this range, returning the bounds of the overlapping field entries
+	TRange<int32> FieldRange = EvaluationField.ConditionallyCompileRange(ContextRange, OverrideRootSequence, *TemplateStore);
+	if (FieldRange.IsEmpty())
 	{
-		// We need to compile an entry in the evaluation field
-		static bool bFullCompile = false;
-		if (bFullCompile)
-		{
-			FMovieSceneCompiler::Compile(*OverrideRootSequence, *TemplateStore);
-			TemplateFieldIndex = OverrideRootTemplate->EvaluationField.GetSegmentFromTime(RootOverrideTime);
-		}
-		else
-		{
-			TOptional<FCompiledGroupResult> CompileResult = FMovieSceneCompiler::CompileTime(RootOverrideTime, *OverrideRootSequence, *TemplateStore);
-
-			if (CompileResult.IsSet())
-			{
-				TRange<FFrameNumber> FieldRange = CompileResult->Range;
-				TemplateFieldIndex = OverrideRootTemplate->EvaluationField.Insert(
-					RootOverrideTime,
-					FieldRange,
-					MoveTemp(CompileResult->Group),
-					MoveTemp(CompileResult->MetaData)
-					);
-			}
-		}
+		return nullptr;
 	}
 
+	// The one that we want to evaluate is either the first or last index in the range.
+	// FieldRange is always of the form [First, Last+1)
+	int32 TemplateFieldIndex = Context.GetDirection() == EPlayDirection::Forwards ? FieldRange.GetUpperBoundValue() - 1 : FieldRange.GetLowerBoundValue();
 	if (TemplateFieldIndex != INDEX_NONE)
 	{
 		// Set meta-data
-		ThisFrameMetaData = OverrideRootTemplate->EvaluationField.GetMetaData(TemplateFieldIndex);
-		return &OverrideRootTemplate->EvaluationField.GetGroup(TemplateFieldIndex);
+		ThisFrameMetaData = EvaluationField.GetMetaData(TemplateFieldIndex);
+		return &EvaluationField.GetGroup(TemplateFieldIndex);
 	}
 
 	return nullptr;
diff --git a/Engine/Source/Runtime/MovieScene/Private/Evaluation/MovieScenePlayback.cpp b/Engine/Source/Runtime/MovieScene/Private/Evaluation/MovieScenePlayback.cpp
--- a/Engine/Source/Runtime/MovieScene/Private/Evaluation/MovieScenePlayback.cpp
+++ b/Engine/Source/Runtime/MovieScene/Private/Evaluation/MovieScenePlayback.cpp
@@ -55,6 +55,25 @@ FMovieSceneEvaluationRange::FMovieSceneEvaluationRange(FFrameTime InCurrentTime,
 {
 }
 
+TRange<FFrameNumber> FMovieSceneEvaluationRange::GetTraversedFrameNumberRange() const
+{
+	TRange<FFrameNumber> FrameNumberRange;
+
+	if (!EvaluationRange.GetLowerBound().IsOpen())
+	{
+		FFrameNumber StartFrame = EvaluationRange.GetLowerBoundValue().FloorToFrame();
+		FrameNumberRange.SetLowerBound(TRangeBound<FFrameNumber>::Inclusive(StartFrame));
+	}
+
+	if (!EvaluationRange.GetUpperBound().IsOpen())
+	{
+		FFrameNumber EndFrame = EvaluationRange.GetUpperBoundValue().FloorToFrame() + 1;
+		FrameNumberRange.SetUpperBound(TRangeBound<FFrameNumber>::Exclusive(EndFrame));
+	}
+
+	return FrameNumberRange;
+}
+
 TRange<FFrameNumber> FMovieSceneEvaluationRange::TimeRangeToNumberRange(const TRange<FFrameTime>& InFrameTimeRange)
 {
 	TRange<FFrameNumber> FrameNumberRange;
diff --git a/Engine/Source/Runtime/MovieScene/Public/Compilation/MovieSceneCompiler.h b/Engine/Source/Runtime/MovieScene/Public/Compilation/MovieSceneCompiler.h
--- a/Engine/Source/Runtime/MovieScene/Public/Compilation/MovieSceneCompiler.h
+++ b/Engine/Source/Runtime/MovieScene/Public/Compilation/MovieSceneCompiler.h
@@ -51,14 +51,13 @@ public:
 	MOVIESCENE_API static void Compile(UMovieSceneSequence& InCompileSequence, IMovieSceneSequenceTemplateStore& InTemplateStore);
 
 	/**
-	 * Compile the evaluation data that is relevant for the specified time
+	 * Compile the evaluation data that is relevant for the specified range
 	 *
-	 * @param InGlobalTime 				The time at which to compile the evaluation data, in the time-space of InCompileSequence
+	 * @param InGlobalRange 			The range for which to compile the evaluation data, in the time-space of InCompileSequence
 	 * @param InCompileSequence 		The Sequence to compile
 	 * @param InTemplateStore 			A template store that contains (or knows how to fetch) evaluation templates for a given sequence
-	 * @return The compiled data, ready for insertion into an FMovieSceneEvaluationField
 	 */
-	MOVIESCENE_API static TOptional<FCompiledGroupResult> CompileTime(FFrameNumber InGlobalTime, UMovieSceneSequence& InCompileSequence, IMovieSceneSequenceTemplateStore& InTemplateStore);
+	MOVIESCENE_API static void CompileRange(TRange<FFrameNumber> InGlobalRange, UMovieSceneSequence& InCompileSequence, IMovieSceneSequenceTemplateStore& InTemplateStore);
 
 	/**
 	 * Compile the hierarchy data for the specified sequence into a specific hierarchy
diff --git a/Engine/Source/Runtime/MovieScene/Public/Evaluation/MovieSceneEvaluationField.h b/Engine/Source/Runtime/MovieScene/Public/Evaluation/MovieSceneEvaluationField.h
--- a/Engine/Source/Runtime/MovieScene/Public/Evaluation/MovieSceneEvaluationField.h
+++ b/Engine/Source/Runtime/MovieScene/Public/Evaluation/MovieSceneEvaluationField.h
@@ -219,6 +219,16 @@ struct FMovieSceneEvaluationField
 {
 	GENERATED_BODY()
 
+	/**
+	 * Ensure that the evaluation field is up-to-date for the range encompassing at least the range specified, and return its entries
+	 *
+	 * @param OverlapRange	A range specifying the smallest set of times to ensure are compiled
+	 * @param InSequence	The sequence that this evaluation field relates to
+	 * @param TemplateStore	A template store that defines how to retrieve templates for the specified sequence
+	 * @return A range of indices for which GetRange() overlaps the specified OverlapRange, of the form [First, First+Num)
+	 */
+	MOVIESCENE_API TRange<int32> ConditionallyCompileRange(const TRange<FFrameNumber>& OverlapRange, UMovieSceneSequence* InSequence, IMovieSceneSequenceTemplateStore& TemplateStore);
+
 	/**
 	 * Efficiently find the entry that exists at the specified time, if any
 	 *
@@ -231,9 +241,9 @@ struct FMovieSceneEvaluationField
 	 * Deduce the indices into Ranges and Groups that overlap with the specified time range
 	 *
 	 * @param Range			The range to overlap with our field
-	 * @return A range of indices into Ranges and Groups that overlap with the requested range
+	 * @return A range of indices for which GetRange() overlaps the specified Range, of the form [First, First+Num)
 	 */
-	MOVIESCENE_API TRange<int32> OverlapRange(TRange<FFrameNumber> Range) const;
+	MOVIESCENE_API TRange<int32> OverlapRange(const TRange<FFrameNumber>& Range) const;
 
 	/**
 	 * Invalidate a range in this field
@@ -241,18 +251,17 @@ struct FMovieSceneEvaluationField
 	 * @param Range			The range to overlap with our field
 	 * @return A range of indices into Ranges and Groups that overlap with the requested range
 	 */
-	MOVIESCENE_API void Invalidate(TRange<FFrameNumber> Range);
+	MOVIESCENE_API void Invalidate(const TRange<FFrameNumber>& Range);
 
 	/**
 	 * Insert a new range into this field
 	 *
-	 * @param InsertTime	The time at which to insert the range
 	 * @param InRange		The total range to insert to the field. Will potentially be intersected with preexisting adjacent ranges
 	 * @param InGroup		The group defining what should happen at this time
 	 * @param InMetaData	The meta-data defining efficient access to what happens in this frame
 	 * @return The index the entries were inserted at
 	 */
-	MOVIESCENE_API int32 Insert(FFrameNumber InsertTime, TRange<FFrameNumber> InRange, FMovieSceneEvaluationGroup&& InGroup, FMovieSceneEvaluationMetaData&& InMetaData);
+	MOVIESCENE_API int32 Insert(const TRange<FFrameNumber>& InRange, FMovieSceneEvaluationGroup&& InGroup, FMovieSceneEvaluationMetaData&& InMetaData);
 
 	/**
 	 * Add the specified data to this field, assuming the specified range lies after any other entries
@@ -261,7 +270,7 @@ struct FMovieSceneEvaluationField
 	 * @param InGroup		The group defining what should happen at this time
 	 * @param InMetaData	The meta-data defining efficient access to what happens in this frame
 	 */
-	MOVIESCENE_API void Add(TRange<FFrameNumber> InRange, FMovieSceneEvaluationGroup&& InGroup, FMovieSceneEvaluationMetaData&& InMetaData);
+	MOVIESCENE_API void Add(const TRange<FFrameNumber>& InRange, FMovieSceneEvaluationGroup&& InGroup, FMovieSceneEvaluationMetaData&& InMetaData);
 
 	/**
 	 * Access this field's signature
@@ -281,6 +290,14 @@ struct FMovieSceneEvaluationField
 		return Ranges.Num();
 	}
 
+	/**
+	 * Access this entire field's set of ranges
+	 */
+	TArrayView<const FMovieSceneFrameRange> GetRanges() const
+	{
+		return Ranges;
+	}
+
 	/**
 	 * Lookup a valid range by index
 	 * @param Index 	The valid index within the ranges to lookup
diff --git a/Engine/Source/Runtime/MovieScene/Public/Evaluation/MovieScenePlayback.h b/Engine/Source/Runtime/MovieScene/Public/Evaluation/MovieScenePlayback.h
--- a/Engine/Source/Runtime/MovieScene/Public/Evaluation/MovieScenePlayback.h
+++ b/Engine/Source/Runtime/MovieScene/Public/Evaluation/MovieScenePlayback.h
@@ -59,6 +59,12 @@ struct MOVIESCENE_API FMovieSceneEvaluationRange
 		return TimeRangeToNumberRange(EvaluationRange);
 	}
 
+	/**
+	 * Get the range of frame numbers traversed over this evaluation range by flooring the lower bound, and ceiling the upper bound.
+	 * For example: a time range of [1.5, 5.6] will yield the equivalent of [1, 6). A time range of (2.0, 2.9) will yield the equivalent of [2,3).
+	 */
+	TRange<FFrameNumber> GetTraversedFrameNumberRange() const;
+
 	/**
 	 * Get the direction to evaluate our range
 	 */
