Sequencer: Increase curve color contrast.

diff --git a/Engine/Source/Runtime/MovieScene/Private/Channels/MovieSceneChannelEditorData.cpp b/Engine/Source/Runtime/MovieScene/Private/Channels/MovieSceneChannelEditorData.cpp
--- a/Engine/Source/Runtime/MovieScene/Private/Channels/MovieSceneChannelEditorData.cpp
+++ b/Engine/Source/Runtime/MovieScene/Private/Channels/MovieSceneChannelEditorData.cpp
@@ -14,9 +14,9 @@ const FText FCommonChannelData::ChannelG = NSLOCTEXT("MovieSceneChannels", "Chan
 const FText FCommonChannelData::ChannelB = NSLOCTEXT("MovieSceneChannels", "ChannelB", "B");
 const FText FCommonChannelData::ChannelA = NSLOCTEXT("MovieSceneChannels", "ChannelA", "A");
 
-const FLinearColor FCommonChannelData::RedChannelColor(0.7f, 0.0f, 0.0f, 0.5f);
-const FLinearColor FCommonChannelData::GreenChannelColor(0.0f, 0.7f, 0.0f, 0.5f);
-const FLinearColor FCommonChannelData::BlueChannelColor(0.0f, 0.0f, 0.7f, 0.5f);
+const FLinearColor FCommonChannelData::RedChannelColor(1.0f, 0.0f, 0.0f, 0.8f);
+const FLinearColor FCommonChannelData::GreenChannelColor(0.0f, 1.0f, 0.0f, 0.8f);
+const FLinearColor FCommonChannelData::BlueChannelColor(0.05f, 0.05f, 1.0f, 0.8f);
 
 
 FMovieSceneChannelMetaData::FMovieSceneChannelMetaData()
