Sequencer: Rewind to beginning if playing forwards, as defined as positive playrate rather than just the bReversePlayback flag (and vice versa)

diff --git a/Engine/Source/Runtime/MovieScene/Private/MovieSceneSequencePlayer.cpp b/Engine/Source/Runtime/MovieScene/Private/MovieSceneSequencePlayer.cpp
--- a/Engine/Source/Runtime/MovieScene/Private/MovieSceneSequencePlayer.cpp
+++ b/Engine/Source/Runtime/MovieScene/Private/MovieSceneSequencePlayer.cpp
@@ -88,19 +88,21 @@ void UMovieSceneSequencePlayer::PlayInternal()
 {
 	if (!IsPlaying() && Sequence && CanPlay())
 	{
-		// If at the end, rewind to the beginning to restart playback
-		if (bReversePlayback)
+		float PlayRate = bReversePlayback ? -PlaybackSettings.PlayRate : PlaybackSettings.PlayRate;
+
+		// If at the end and playing forwards, rewind to beginning
+		if (GetCurrentTime().Time == GetLastValidTime())
 		{
-			if (GetCurrentTime().Time == FFrameTime(StartTime))
+			if (PlayRate > 0.f)
 			{
-				JumpToFrame(GetLastValidTime());
+				JumpToFrame(FFrameTime(StartTime));
 			}
 		}
-		else
+		else if (GetCurrentTime().Time == FFrameTime(StartTime))
 		{
-			if (GetCurrentTime().Time == GetLastValidTime())
+			if (PlayRate < 0.f)
 			{
-				JumpToFrame(FFrameTime(StartTime));
+				JumpToFrame(GetLastValidTime());
 			}
 		}
 
