Sequencer: Fixed a crash caused by setter functions that have return values
  - Any setter function that has either a return value or more than a single property is now provided the correct sized parameter pack to ensure that there is no possibility of stack-stomping.
  - The new value is assumed to be the first parameter in the function (and early return if this is not the case)
  - Subsequent parameters are simply default initialized - no provision of these values is (or shall be) supported
  - Return values are ignored

diff --git a/Engine/Source/Runtime/MovieScene/Private/MovieSceneCommonHelpers.cpp b/Engine/Source/Runtime/MovieScene/Private/MovieSceneCommonHelpers.cpp
--- a/Engine/Source/Runtime/MovieScene/Private/MovieSceneCommonHelpers.cpp
+++ b/Engine/Source/Runtime/MovieScene/Private/MovieSceneCommonHelpers.cpp
@@ -436,10 +436,9 @@ FTrackInstancePropertyBindings::FPropertyAddress FTrackInstancePropertyBindings:
 void FTrackInstancePropertyBindings::CallFunctionForEnum( UObject& InRuntimeObject, int64 PropertyValue )
 {
 	FPropertyAndFunction PropAndFunction = FindOrAdd(InRuntimeObject);
-	if (UFunction* Setter = PropAndFunction.SetterFunction.Get())
+	if (UFunction* SetterFunction = PropAndFunction.SetterFunction.Get())
 	{
-		// ProcessEvent should really be taking const void*
-		InRuntimeObject.ProcessEvent(Setter, (void*)&PropertyValue);
+		InvokeSetterFunction(&InRuntimeObject, SetterFunction, PropertyValue);
 	}
 	else if (UProperty* Property = PropAndFunction.PropertyAddress.GetProperty())
 	{
@@ -468,11 +467,18 @@ void FTrackInstancePropertyBindings::CacheBinding(const UObject& Object)
 {
 	FPropertyAndFunction PropAndFunction;
 	{
-		PropAndFunction.SetterFunction = Object.FindFunction(FunctionName);
 		PropAndFunction.PropertyAddress = FindProperty(Object, PropertyPath);
-		if (NotifyFunctionName != NAME_None)
+
+		UFunction* SetterFunction = Object.FindFunction(FunctionName);
+		if (SetterFunction && SetterFunction->NumParms >= 1)
+		{
+			PropAndFunction.SetterFunction = SetterFunction;
+		}
+		
+		UFunction* NotifyFunction = NotifyFunctionName != NAME_None ? Object.FindFunction(NotifyFunctionName) : nullptr;
+		if (NotifyFunction && NotifyFunction->NumParms == 0 && NotifyFunction->ReturnValueOffset == MAX_uint16)
 		{
-			PropAndFunction.NotifyFunction = Object.FindFunction(NotifyFunctionName);
+			PropAndFunction.NotifyFunction = NotifyFunction;
 		}
 	}
 
@@ -520,8 +526,7 @@ template<> void FTrackInstancePropertyBindings::CallFunction<bool>(UObject& InRu
 	FPropertyAndFunction PropAndFunction = FindOrAdd(InRuntimeObject);
 	if (UFunction* SetterFunction = PropAndFunction.SetterFunction.Get())
 	{
-		// ProcessEvent should really be taking const void*
-		InRuntimeObject.ProcessEvent(SetterFunction, (void*)&PropertyValue);
+		InvokeSetterFunction(&InRuntimeObject, SetterFunction, PropertyValue);
 	}
 	else if (UProperty* Property = PropAndFunction.PropertyAddress.GetProperty())
 	{
@@ -591,8 +596,7 @@ template<> void FTrackInstancePropertyBindings::CallFunction<UObject*>(UObject&
 	FPropertyAndFunction PropAndFunction = FindOrAdd(InRuntimeObject);
 	if (UFunction* SetterFunction = PropAndFunction.SetterFunction.Get())
 	{
-		// ProcessEvent should really be taking const void*
-		InRuntimeObject.ProcessEvent(SetterFunction, (void*)&PropertyValue);
+		InvokeSetterFunction(&InRuntimeObject, SetterFunction, PropertyValue);
 	}
 	else if (UObjectPropertyBase* ObjectProperty = Cast<UObjectPropertyBase>(PropAndFunction.PropertyAddress.GetProperty()))
 	{
diff --git a/Engine/Source/Runtime/MovieScene/Public/MovieSceneCommonHelpers.h b/Engine/Source/Runtime/MovieScene/Public/MovieSceneCommonHelpers.h
--- a/Engine/Source/Runtime/MovieScene/Public/MovieSceneCommonHelpers.h
+++ b/Engine/Source/Runtime/MovieScene/Public/MovieSceneCommonHelpers.h
@@ -149,8 +149,7 @@ public:
 		FPropertyAndFunction PropAndFunction = FindOrAdd(InRuntimeObject);
 		if (UFunction* SetterFunction = PropAndFunction.SetterFunction.Get())
 		{
-			// ProcessEvent should really be taking const void*
-			InRuntimeObject.ProcessEvent(SetterFunction, (void*)&PropertyValue);
+			InvokeSetterFunction(&InRuntimeObject, SetterFunction, PropertyValue);
 		}
 		else if (ValueType* Val = PropAndFunction.GetPropertyAddress<ValueType>())
 		{
@@ -260,6 +259,14 @@ public:
 
 private:
 
+	/**
+	 * Wrapper for UObject::ProcessEvent that attempts to pass the new property value directly to the function as a parameter,
+	 * but handles cases where multiple parameters or a return value exists. The setter parameter must be the first in the list,
+	 * any other parameters will be default constructed.
+	 */
+	template<typename T>
+	static void InvokeSetterFunction(UObject* InRuntimeObject, UFunction* Setter, T&& InPropertyValue);
+
 	struct FPropertyAddress
 	{
 		TWeakObjectPtr<UProperty> Property;
@@ -345,3 +352,53 @@ template<> MOVIESCENE_API void FTrackInstancePropertyBindings::SetCurrentValue<b
 template<> MOVIESCENE_API void FTrackInstancePropertyBindings::CallFunction<UObject*>(UObject& InRuntimeObject, UObject* PropertyValue);
 template<> MOVIESCENE_API UObject* FTrackInstancePropertyBindings::GetCurrentValue<UObject*>(const UObject& InRuntimeObject);
 template<> MOVIESCENE_API void FTrackInstancePropertyBindings::SetCurrentValue<UObject*>(UObject& InRuntimeObject, UObject* InValue);
+
+
+template<typename T>
+void FTrackInstancePropertyBindings::InvokeSetterFunction(UObject* InRuntimeObject, UFunction* Setter, T&& InPropertyValue)
+{
+	// CacheBinding already guarantees that the function has >= 1 parameters
+	const int32 ParmsSize = Setter->ParmsSize;
+
+	// This should all be const really, but ProcessEvent only takes a non-const void*
+	void* InputParameter = const_cast<typename TDecay<T>::Type*>(&InPropertyValue);
+
+	// By default we try and use the existing stack value
+	uint8* Params = reinterpret_cast<uint8*>(InputParameter);
+
+	check(InRuntimeObject && Setter);
+	if (Setter->ReturnValueOffset != MAX_uint16 || Setter->NumParms > 1)
+	{
+		// Function has a return value or multiple parameters, we need to initialize memory for the entire parameter pack
+		// We use alloca here (as in UObject::ProcessEvent) to avoid a heap allocation. Alloca memory survives the current function's stack frame.
+		Params = reinterpret_cast<uint8*>(FMemory_Alloca(ParmsSize));
+
+		bool bFirstProperty = true;
+		for (UProperty* Property = Setter->PropertyLink; Property; Property = Property->PropertyLinkNext)
+		{
+			// Initialize the parameter pack with any param properties that reside in the container
+			if (Property->IsInContainer(ParmsSize))
+			{
+				Property->InitializeValue_InContainer(Params);
+
+				// The first encountered property is assumed to be the input value so initialize this with the user-specified value from InPropertyValue
+				if (Property->HasAnyPropertyFlags(CPF_Parm) && !Property->HasAnyPropertyFlags(CPF_ReturnParm) && bFirstProperty)
+				{
+					const bool bIsValid = ensureMsgf(sizeof(T) == Property->ElementSize, TEXT("Property type does not match for Sequencer setter function %s::%s (%ibytes != %ibytes"), *InRuntimeObject->GetName(), *Setter->GetName(), sizeof(T), Property->ElementSize);
+					if (bIsValid)
+					{
+						Property->CopyCompleteValue(Property->ContainerPtrToValuePtr<void>(Params), &InPropertyValue);
+					}
+					else
+					{
+						return;
+					}
+				}
+				bFirstProperty = false;
+			}
+		}
+	}
+
+	// Now we have the parameters set up correctly, call the function
+	InRuntimeObject->ProcessEvent(Setter, Params);
+}
\ No newline at end of file
