Sequencer: Null section ptr check

diff --git a/Engine/Source/Runtime/MovieScene/Private/MovieSceneTrack.cpp b/Engine/Source/Runtime/MovieScene/Private/MovieSceneTrack.cpp
--- a/Engine/Source/Runtime/MovieScene/Private/MovieSceneTrack.cpp
+++ b/Engine/Source/Runtime/MovieScene/Private/MovieSceneTrack.cpp
@@ -248,7 +248,10 @@ int32 UMovieSceneTrack::GetMaxRowIndex() const
 	int32 MaxRowIndex = 0;
 	for (UMovieSceneSection* Section : GetAllSections())
 	{
-		MaxRowIndex = FMath::Max(MaxRowIndex, Section->GetRowIndex());
+		if (Section != nullptr)
+		{
+			MaxRowIndex = FMath::Max(MaxRowIndex, Section->GetRowIndex());
+		}
 	}
 
 	return MaxRowIndex;
