Sequencer: Remove null folders

diff --git a/Engine/Source/Runtime/MovieScene/Private/MovieScene.cpp b/Engine/Source/Runtime/MovieScene/Private/MovieScene.cpp
--- a/Engine/Source/Runtime/MovieScene/Private/MovieScene.cpp
+++ b/Engine/Source/Runtime/MovieScene/Private/MovieScene.cpp
@@ -877,6 +877,18 @@ void UMovieScene::PostLoad()
 		}
 	}
 
+	for (int32 RootFolderIndex = 0; RootFolderIndex < RootFolders.Num();)
+	{
+		if (RootFolders[RootFolderIndex] == nullptr)
+		{
+			RootFolders.RemoveAt(RootFolderIndex);
+		}
+		else 
+		{
+			++RootFolderIndex;
+		}
+	}
+
 	UpgradeTimeRanges();
 
 	for (FMovieSceneSpawnable& Spawnable : Spawnables)
diff --git a/Engine/Source/Runtime/MovieScene/Private/MovieSceneFolder.cpp b/Engine/Source/Runtime/MovieScene/Private/MovieSceneFolder.cpp
--- a/Engine/Source/Runtime/MovieScene/Private/MovieSceneFolder.cpp
+++ b/Engine/Source/Runtime/MovieScene/Private/MovieSceneFolder.cpp
@@ -76,6 +76,23 @@ void UMovieSceneFolder::RemoveChildObjectBinding( const FGuid& InObjectBinding )
 	ChildObjectBindings.Remove( InObjectBinding );
 }
 
+void UMovieSceneFolder::PostLoad()
+{
+	// Remove any null folders
+	for (int32 ChildFolderIndex = 0; ChildFolderIndex < ChildFolders.Num(); )
+	{
+		if (ChildFolders[ChildFolderIndex] == nullptr)
+		{
+			ChildFolders.RemoveAt(ChildFolderIndex);
+		}
+		else
+		{
+			++ChildFolderIndex;
+		}
+	}
+
+	Super::PostLoad();
+}
 
 void UMovieSceneFolder::Serialize( FArchive& Archive )
 {
diff --git a/Engine/Source/Runtime/MovieScene/Public/MovieSceneFolder.h b/Engine/Source/Runtime/MovieScene/Public/MovieSceneFolder.h
--- a/Engine/Source/Runtime/MovieScene/Public/MovieSceneFolder.h
+++ b/Engine/Source/Runtime/MovieScene/Public/MovieSceneFolder.h
@@ -50,6 +50,8 @@ class MOVIESCENE_API UMovieSceneFolder : public UObject
 	/** Removes a guid for an object binding from this folder. */
 	void RemoveChildObjectBinding( const FGuid& InObjectBinding );
 
+	/** Called after this object has been deserialized */
+	virtual void PostLoad() override;
 
 	virtual void Serialize( FArchive& Archive );
 
