Sequencer: Fix editor only data

diff --git a/Engine/Source/Runtime/MovieScene/Private/MovieScene.cpp b/Engine/Source/Runtime/MovieScene/Private/MovieScene.cpp
--- a/Engine/Source/Runtime/MovieScene/Private/MovieScene.cpp
+++ b/Engine/Source/Runtime/MovieScene/Private/MovieScene.cpp
@@ -877,6 +877,7 @@ void UMovieScene::PostLoad()
 		}
 	}
 
+#if WITH_EDITORONLY_DATA
 	for (int32 RootFolderIndex = 0; RootFolderIndex < RootFolders.Num();)
 	{
 		if (RootFolders[RootFolderIndex] == nullptr)
@@ -888,6 +889,7 @@ void UMovieScene::PostLoad()
 			++RootFolderIndex;
 		}
 	}
+#endif
 
 	UpgradeTimeRanges();
 
