Agressively maintain a single-ownership rule for child tracks/folders/bindings in Movie Scene Folders. This now programatically enforces that no other folder in the same scene can have the same child track/folder/binding, and it agressively cleans up old folders on post load to deal with legacy content which through some series of events ends up with certain objects being left in folders when moved/deleted/etc.

This is an editor-only change as Folders are only stored in a sequence as editor-only data.

These non-valid objects were causing issues with the Curve Editor Tree which now enforces that things exist as expected.

diff --git a/Engine/Source/Runtime/MovieScene/Private/MovieSceneFolder.cpp b/Engine/Source/Runtime/MovieScene/Private/MovieSceneFolder.cpp
--- a/Engine/Source/Runtime/MovieScene/Private/MovieSceneFolder.cpp
+++ b/Engine/Source/Runtime/MovieScene/Private/MovieSceneFolder.cpp
@@ -1,6 +1,21 @@
 // Copyright 1998-2019 Epic Games, Inc. All Rights Reserved.
 
 #include "MovieSceneFolder.h"
+#include "MovieScene.h"
+#include "MovieSceneTrack.h"
+#include "Algo/Count.h"
+
+void GetMovieSceneFoldersRecursive(const TArray<UMovieSceneFolder*>& InFoldersToRecurse, TArray<UMovieSceneFolder*>& OutFolders)
+{
+	for (UMovieSceneFolder* Folder : InFoldersToRecurse)
+	{
+		if (Folder)
+		{
+			OutFolders.Add(Folder);
+			GetMovieSceneFoldersRecursive(Folder->GetChildFolders(), OutFolders);
+		}
+	}
+}
 
 UMovieSceneFolder::UMovieSceneFolder( const FObjectInitializer& ObjectInitializer )
 	: Super( ObjectInitializer )
@@ -31,6 +46,20 @@ const TArray<UMovieSceneFolder*>& UMovieSceneFolder::GetChildFolders() const
 
 void UMovieSceneFolder::AddChildFolder( UMovieSceneFolder* InChildFolder )
 {
+	// Ensure the added folder does not belong to any other folder in the same scene.
+	UMovieScene* OwningScene = GetTypedOuter<UMovieScene>();
+	if (OwningScene)
+	{
+		TArray<UMovieSceneFolder*> AllFolders;
+		GetMovieSceneFoldersRecursive(OwningScene->GetRootFolders(), AllFolders);
+
+		for (UMovieSceneFolder* MovieSceneFolder : AllFolders)
+		{
+			MovieSceneFolder->RemoveChildFolder(InChildFolder);
+		}
+	}
+
+	// Now add it as a child of ourself
 	ChildFolders.Add( InChildFolder );
 }
 
@@ -49,6 +78,19 @@ const TArray<UMovieSceneTrack*>& UMovieSceneFolder::GetChildMasterTracks() const
 
 void UMovieSceneFolder::AddChildMasterTrack( UMovieSceneTrack* InMasterTrack )
 {
+	// Ensure the added track does not belong to any other folder in the same scene.
+	UMovieScene* OwningScene = GetTypedOuter<UMovieScene>();
+	if (OwningScene)
+	{
+		TArray<UMovieSceneFolder*> AllFolders;
+		GetMovieSceneFoldersRecursive(OwningScene->GetRootFolders(), AllFolders);
+
+		for (UMovieSceneFolder* MovieSceneFolder : AllFolders)
+		{
+			MovieSceneFolder->RemoveChildMasterTrack(InMasterTrack);
+		}
+	}
+
 	ChildMasterTracks.Add( InMasterTrack );
 }
 
@@ -67,6 +109,19 @@ const TArray<FGuid>& UMovieSceneFolder::GetChildObjectBindings() const
 
 void UMovieSceneFolder::AddChildObjectBinding(const FGuid& InObjectBinding )
 {
+	// Ensure the added object  does not belong to any other folder in the same scene.
+	UMovieScene* OwningScene = GetTypedOuter<UMovieScene>();
+	if (OwningScene)
+	{
+		TArray<UMovieSceneFolder*> AllFolders;
+		GetMovieSceneFoldersRecursive(OwningScene->GetRootFolders(), AllFolders);
+
+		for (UMovieSceneFolder* MovieSceneFolder : AllFolders)
+		{
+			MovieSceneFolder->RemoveChildObjectBinding(InObjectBinding);
+		}
+	}
+
 	ChildObjectBindings.Add( InObjectBinding );
 }
 
@@ -91,6 +146,61 @@ void UMovieSceneFolder::PostLoad()
 		}
 	}
 
+	// Historically we've not been very strict about ensuring a folder, track, or object binding existed
+	// only in one folder. This is now enforced (via automatically removing the item from other folders
+	// when they are added to this folder), and checked (the tree view trips an ensure on invalid children)
+	// but all legacy content can still have the invalid children which continuously trips the ensure.
+	// Since we now enforce child-only-exists-in-one-folder, we can safely remove any invalid children on
+	// load, and be confident that we shouldn't run into situations in the future where an invalid child is
+	// left in a folder.
+	UMovieScene* OwningScene = GetTypedOuter<UMovieScene>();
+	if (OwningScene)
+	{
+		// Validate child Master Tracks
+		for(int32 ChildMasterTrackIndex = 0; ChildMasterTrackIndex < ChildMasterTracks.Num(); ChildMasterTrackIndex++)
+		{
+			const UMovieSceneTrack* ChildTrack = ChildMasterTracks[ChildMasterTrackIndex];
+			if (!OwningScene->GetMasterTracks().Contains(ChildTrack))
+			{
+				ChildMasterTracks.RemoveAt(ChildMasterTrackIndex);
+				ChildMasterTrackIndex--;
+
+				UE_LOG(LogMovieScene, Warning, TEXT("Folder (%s) contained a reference to a Master Track (%s) that no longer exists in the sequence, removing."), *GetFolderName().ToString(), *GetNameSafe(ChildTrack));
+			}
+		}
+
+		// Validate child Object Bindings
+		for (int32 ChildObjectBindingIndex = 0; ChildObjectBindingIndex < ChildObjectBindings.Num(); ChildObjectBindingIndex++)
+		{
+			const FGuid& ChildBinding = ChildObjectBindings[ChildObjectBindingIndex];
+			if (!OwningScene->FindBinding(ChildBinding))
+			{
+				ChildObjectBindings.RemoveAt(ChildObjectBindingIndex);
+				ChildObjectBindingIndex--;
+
+				UE_LOG(LogMovieScene, Warning, TEXT("Folder (%s) contained a reference to an Object Binding (%s) that no longer exists in the sequence, removing."), *GetFolderName().ToString(), *ChildBinding.ToString());
+			}
+		}
+
+		// A folder should exist in only one place in the tree, as a child of ourself. If they exist in more
+		// than one place, two folders point to the same actual UObject, so we'll remove it from ourself. When
+		// the that folder is PostLoaded it will search the whole tree again and only find the one reference as
+		// our reference will no longer exist.
+		TArray<UMovieSceneFolder*> AllFolders;
+		GetMovieSceneFoldersRecursive(OwningScene->GetRootFolders(), AllFolders);
+		for (int32 ChildFolderIndex = 0; ChildFolderIndex < ChildFolders.Num(); ChildFolderIndex++)
+		{
+			int32 NumFolderInstances = Algo::Count(AllFolders, ChildFolders[ChildFolderIndex]);
+			if (NumFolderInstances > 1)
+			{
+				ChildFolders.RemoveAt(ChildFolderIndex);
+				ChildFolderIndex--;
+
+				UE_LOG(LogMovieScene, Warning, TEXT("Folder (%s) contained a reference to an Folder (%s) that exists in multiple places in the sequence, removing."), *GetFolderName().ToString(), *ChildFolders[ChildFolderIndex]->GetFolderName().ToString());
+			}
+		}
+	}
+
 	Super::PostLoad();
 }
 
