CIS fixes for movie scene folders in non-editor compiles.

diff --git a/Engine/Source/Runtime/MovieScene/Private/MovieSceneFolder.cpp b/Engine/Source/Runtime/MovieScene/Private/MovieSceneFolder.cpp
--- a/Engine/Source/Runtime/MovieScene/Private/MovieSceneFolder.cpp
+++ b/Engine/Source/Runtime/MovieScene/Private/MovieSceneFolder.cpp
@@ -46,6 +46,8 @@ const TArray<UMovieSceneFolder*>& UMovieSceneFolder::GetChildFolders() const
 
 void UMovieSceneFolder::AddChildFolder( UMovieSceneFolder* InChildFolder )
 {
+
+#if WITH_EDITORONLY_DATA
 	// Ensure the added folder does not belong to any other folder in the same scene.
 	UMovieScene* OwningScene = GetTypedOuter<UMovieScene>();
 	if (OwningScene)
@@ -58,6 +60,7 @@ void UMovieSceneFolder::AddChildFolder( UMovieSceneFolder* InChildFolder )
 			MovieSceneFolder->RemoveChildFolder(InChildFolder);
 		}
 	}
+#endif
 
 	// Now add it as a child of ourself
 	ChildFolders.Add( InChildFolder );
@@ -78,6 +81,7 @@ const TArray<UMovieSceneTrack*>& UMovieSceneFolder::GetChildMasterTracks() const
 
 void UMovieSceneFolder::AddChildMasterTrack( UMovieSceneTrack* InMasterTrack )
 {
+#if WITH_EDITORONLY_DATA
 	// Ensure the added track does not belong to any other folder in the same scene.
 	UMovieScene* OwningScene = GetTypedOuter<UMovieScene>();
 	if (OwningScene)
@@ -90,6 +94,7 @@ void UMovieSceneFolder::AddChildMasterTrack( UMovieSceneTrack* InMasterTrack )
 			MovieSceneFolder->RemoveChildMasterTrack(InMasterTrack);
 		}
 	}
+#endif
 
 	ChildMasterTracks.Add( InMasterTrack );
 }
@@ -109,6 +114,7 @@ const TArray<FGuid>& UMovieSceneFolder::GetChildObjectBindings() const
 
 void UMovieSceneFolder::AddChildObjectBinding(const FGuid& InObjectBinding )
 {
+#if WITH_EDITORONLY_DATA
 	// Ensure the added object  does not belong to any other folder in the same scene.
 	UMovieScene* OwningScene = GetTypedOuter<UMovieScene>();
 	if (OwningScene)
@@ -121,6 +127,7 @@ void UMovieSceneFolder::AddChildObjectBinding(const FGuid& InObjectBinding )
 			MovieSceneFolder->RemoveChildObjectBinding(InObjectBinding);
 		}
 	}
+#endif
 
 	ChildObjectBindings.Add( InObjectBinding );
 }
@@ -146,6 +153,7 @@ void UMovieSceneFolder::PostLoad()
 		}
 	}
 
+#if WITH_EDITORONLY_DATA
 	// Historically we've not been very strict about ensuring a folder, track, or object binding existed
 	// only in one folder. This is now enforced (via automatically removing the item from other folders
 	// when they are added to this folder), and checked (the tree view trips an ensure on invalid children)
@@ -200,6 +208,7 @@ void UMovieSceneFolder::PostLoad()
 			}
 		}
 	}
+#endif
 
 	Super::PostLoad();
 }
