Sequencer: Fix crash on movie rendering from dev-core. IsDestructionThreadSafe returns false.

diff --git a/Engine/Source/Runtime/MovieScene/Public/MovieSceneSequencePlayer.h b/Engine/Source/Runtime/MovieScene/Public/MovieSceneSequencePlayer.h
--- a/Engine/Source/Runtime/MovieScene/Public/MovieSceneSequencePlayer.h
+++ b/Engine/Source/Runtime/MovieScene/Public/MovieSceneSequencePlayer.h
@@ -551,6 +551,7 @@ protected:
 	/*~ Begin UObject interface */
 	virtual void BeginDestroy() override;
 	virtual bool IsSupportedForNetworking() const { return true; }
+	virtual bool IsDestructionThreadSafe() const override { return false; }
 	virtual int32 GetFunctionCallspace(UFunction* Function, FFrame* Stack) override;
 	virtual bool CallRemoteFunction(UFunction* Function, void* Parameters, FOutParmRec* OutParms, FFrame* Stack) override;
 	virtual void PostNetReceive() override;
