Sequencer: Set old max tick rate only if unset during update so that it's not overwritten during playback.

diff --git a/Engine/Source/Runtime/MovieScene/Private/MovieSceneSequencePlayer.cpp b/Engine/Source/Runtime/MovieScene/Private/MovieSceneSequencePlayer.cpp
--- a/Engine/Source/Runtime/MovieScene/Private/MovieSceneSequencePlayer.cpp
+++ b/Engine/Source/Runtime/MovieScene/Private/MovieSceneSequencePlayer.cpp
@@ -171,7 +171,11 @@ void UMovieSceneSequencePlayer::PlayInternal()
 
 		if (PlayPosition.GetEvaluationType() == EMovieSceneEvaluationType::FrameLocked)
 		{
-			OldMaxTickRate = GEngine->GetMaxFPS();
+			if (!OldMaxTickRate.IsSet())
+			{
+				OldMaxTickRate = GEngine->GetMaxFPS();
+			}
+
 			GEngine->SetMaxFPS(1.f / PlayPosition.GetInputRate().AsInterval());
 		}
 
@@ -308,6 +312,7 @@ void UMovieSceneSequencePlayer::StopInternal(FFrameTime TimeToResetTo)
 		if (OldMaxTickRate.IsSet())
 		{
 			GEngine->SetMaxFPS(OldMaxTickRate.GetValue());
+			OldMaxTickRate.Reset();
 		}
 
 		if (HasAuthority())
