Changed the Movie Scene Folder api to automatically call Modify() on itself. This solves a lot of issues where Modify was not getting called before using the API which caused inconsistencies when using Undo/Redo.

diff --git a/Engine/Plugins/MovieScene/SequencerScripting/Source/SequencerScripting/Private/ExtensionLibraries/MovieSceneFolderExtensions.cpp b/Engine/Plugins/MovieScene/SequencerScripting/Source/SequencerScripting/Private/ExtensionLibraries/MovieSceneFolderExtensions.cpp
--- a/Engine/Plugins/MovieScene/SequencerScripting/Source/SequencerScripting/Private/ExtensionLibraries/MovieSceneFolderExtensions.cpp
+++ b/Engine/Plugins/MovieScene/SequencerScripting/Source/SequencerScripting/Private/ExtensionLibraries/MovieSceneFolderExtensions.cpp
@@ -42,6 +42,7 @@ bool UMovieSceneFolderExtensions::SetFolderColor(UMovieSceneFolder* Folder, FCol
 #if WITH_EDITORONLY_DATA
 	if (Folder)
 	{
+		Folder->Modify();
 		Folder->SetFolderColor(InFolderColor);
 		return true;
 	}
diff --git a/Engine/Source/Editor/Sequencer/Private/DisplayNodes/SequencerFolderNode.cpp b/Engine/Source/Editor/Sequencer/Private/DisplayNodes/SequencerFolderNode.cpp
--- a/Engine/Source/Editor/Sequencer/Private/DisplayNodes/SequencerFolderNode.cpp
+++ b/Engine/Source/Editor/Sequencer/Private/DisplayNodes/SequencerFolderNode.cpp
@@ -96,7 +96,6 @@ void FSequencerFolderNode::SetDisplayName( const FText& NewDisplayName )
 		FName UniqueName = FSequencerUtilities::GetUniqueName( FName( *NewDisplayName.ToString() ), SiblingNames );
 
 		const FScopedTransaction Transaction( NSLOCTEXT( "SequencerFolderNode", "RenameFolder", "Rename folder." ) );
-		MovieSceneFolder.Modify();
 		MovieSceneFolder.SetFolderName( UniqueName );
 	}
 }
@@ -242,7 +241,6 @@ void FSequencerFolderNode::MoveDisplayNodeToFolder(TSharedRef<FSequencerDisplayN
 			{
 				checkf(ParentSeqNode->GetType() == ESequencerNode::Folder, TEXT("Can not remove from unsupported parent node."));
 				TSharedPtr<FSequencerFolderNode> ParentFolder = StaticCastSharedPtr<FSequencerFolderNode>(ParentSeqNode);
-				ParentFolder->GetFolder().Modify();
 				ParentFolder->GetFolder().RemoveChildFolder(&DraggedFolderNode->GetFolder());
 			}
 			else
@@ -263,7 +261,6 @@ void FSequencerFolderNode::MoveDisplayNodeToFolder(TSharedRef<FSequencerDisplayN
 			{
 				checkf(ParentSeqNode->GetType() == ESequencerNode::Folder, TEXT("Can not remove from unsupported parent node."));
 				TSharedPtr<FSequencerFolderNode> ParentFolder = StaticCastSharedPtr<FSequencerFolderNode>(ParentSeqNode);
-				ParentFolder->GetFolder().Modify();
 				ParentFolder->GetFolder().RemoveChildMasterTrack(DraggedTrackNode->GetTrack());
 			}
 			
@@ -278,7 +275,6 @@ void FSequencerFolderNode::MoveDisplayNodeToFolder(TSharedRef<FSequencerDisplayN
 			{
 				checkf(ParentSeqNode->GetType() == ESequencerNode::Folder, TEXT("Can not remove from unsupported parent node."));
 				TSharedPtr<FSequencerFolderNode> ParentFolder = StaticCastSharedPtr<FSequencerFolderNode>(ParentSeqNode);
-				ParentFolder->GetFolder().Modify();
 				ParentFolder->GetFolder().RemoveChildObjectBinding(DraggedObjectBindingNode->GetObjectBinding());
 			}
 
diff --git a/Engine/Source/Editor/Sequencer/Private/Sequencer.cpp b/Engine/Source/Editor/Sequencer/Private/Sequencer.cpp
--- a/Engine/Source/Editor/Sequencer/Private/Sequencer.cpp
+++ b/Engine/Source/Editor/Sequencer/Private/Sequencer.cpp
@@ -2882,6 +2882,7 @@ UMovieSceneFolder* FSequencer::CreateFoldersRecursively(const TArray<FString>& F
 		else
 		{
 			// If we have no parent folder then we must be at the root so we add it to the root of the movie scene
+			OwningMovieScene->Modify();
 			OwningMovieScene->GetRootFolders().Add(FolderToUse);
 		}
 	}
@@ -4527,7 +4528,6 @@ bool FSequencer::OnRequestNodeDeleted( TSharedRef<const FSequencerDisplayNode> N
 		if ( NodeToBeDeleted->GetParent().IsValid() )
 		{
 			TSharedPtr<FSequencerFolderNode> ParentFolder = StaticCastSharedPtr<FSequencerFolderNode>( NodeToBeDeleted->GetParent() );
-			ParentFolder->GetFolder().Modify();
 			ParentFolder->GetFolder().RemoveChildFolder( &FolderToBeDeleted->GetFolder() );
 		}
 		else
@@ -4556,7 +4556,6 @@ bool FSequencer::OnRequestNodeDeleted( TSharedRef<const FSequencerDisplayNode> N
 		if ( NodeToBeDeleted->GetParent().IsValid() && NodeToBeDeleted->GetParent()->GetType() == ESequencerNode::Folder )
 		{
 			TSharedPtr<FSequencerFolderNode> ParentFolder = StaticCastSharedPtr<FSequencerFolderNode>( NodeToBeDeleted->GetParent() );
-			ParentFolder->GetFolder().Modify();
 			ParentFolder->GetFolder().RemoveChildObjectBinding( BindingToRemove );
 		}
 		
@@ -4583,7 +4582,6 @@ bool FSequencer::OnRequestNodeDeleted( TSharedRef<const FSequencerDisplayNode> N
 		if ( NodeToBeDeleted->GetParent().IsValid() && NodeToBeDeleted->GetParent()->GetType() == ESequencerNode::Folder )
 		{
 			TSharedPtr<FSequencerFolderNode> ParentFolder = StaticCastSharedPtr<FSequencerFolderNode>( NodeToBeDeleted->GetParent() );
-			ParentFolder->GetFolder().Modify();
 			ParentFolder->GetFolder().RemoveChildMasterTrack( Track );
 		}
 
@@ -5883,7 +5881,6 @@ FGuid FSequencer::DoAssignActor(AActor*const* InActors, int32 NumActors, FGuid I
 		FoldersToCheck.RemoveAt(0);
 		if ( Folder->GetChildObjectBindings().Contains( InObjectBinding ) )
 		{
-			Folder->Modify();
 			Folder->RemoveChildObjectBinding( InObjectBinding );
 			Folder->AddChildObjectBinding( NewGuid );
 			bFolderFound = true;
@@ -8081,7 +8078,6 @@ void FSequencer::OnAddFolder()
 
 	if ( SelectedParentFolders.Num() == 1 )
 	{
-		SelectedParentFolders[0]->Modify();
 		SelectedParentFolders[0]->AddChildFolder( NewFolder );
 	}
 	else
diff --git a/Engine/Source/Editor/Sequencer/Private/SequencerNodeTree.cpp b/Engine/Source/Editor/Sequencer/Private/SequencerNodeTree.cpp
--- a/Engine/Source/Editor/Sequencer/Private/SequencerNodeTree.cpp
+++ b/Engine/Source/Editor/Sequencer/Private/SequencerNodeTree.cpp
@@ -406,7 +406,6 @@ void FSequencerNodeTree::MoveDisplayNodeToRoot(TSharedRef<FSequencerDisplayNode>
 			{
 				checkf(ParentSeqNode->GetType() == ESequencerNode::Folder, TEXT("Can not remove from unsupported parent node."));
 				TSharedPtr<FSequencerFolderNode> ParentFolder = StaticCastSharedPtr<FSequencerFolderNode>(ParentSeqNode);
-				ParentFolder->GetFolder().Modify();
 				ParentFolder->GetFolder().RemoveChildFolder(&FolderNode->GetFolder());
 			}
 			else
@@ -426,7 +425,6 @@ void FSequencerNodeTree::MoveDisplayNodeToRoot(TSharedRef<FSequencerDisplayNode>
 			{
 				checkf(ParentSeqNode->GetType() == ESequencerNode::Folder, TEXT("Can not remove from unsupported parent node."));
 				TSharedPtr<FSequencerFolderNode> ParentFolder = StaticCastSharedPtr<FSequencerFolderNode>(ParentSeqNode);
-				ParentFolder->GetFolder().Modify();
 				ParentFolder->GetFolder().RemoveChildMasterTrack(DraggedTrackNode->GetTrack());
 			}
 			break;
@@ -440,7 +438,6 @@ void FSequencerNodeTree::MoveDisplayNodeToRoot(TSharedRef<FSequencerDisplayNode>
 			{
 				checkf(ParentSeqNode->GetType() == ESequencerNode::Folder, TEXT("Can not remove from unsupported parent node."));
 				TSharedPtr<FSequencerFolderNode> ParentFolder = StaticCastSharedPtr<FSequencerFolderNode>(ParentSeqNode);
-				ParentFolder->GetFolder().Modify();
 				ParentFolder->GetFolder().RemoveChildObjectBinding(DraggedObjectBindingNode->GetObjectBinding());
 			}
 			break;
diff --git a/Engine/Source/Runtime/MovieScene/Private/MovieSceneFolder.cpp b/Engine/Source/Runtime/MovieScene/Private/MovieSceneFolder.cpp
--- a/Engine/Source/Runtime/MovieScene/Private/MovieSceneFolder.cpp
+++ b/Engine/Source/Runtime/MovieScene/Private/MovieSceneFolder.cpp
@@ -34,6 +34,8 @@ FName UMovieSceneFolder::GetFolderName() const
 
 void UMovieSceneFolder::SetFolderName( FName InFolderName )
 {
+	Modify();
+
 	FolderName = InFolderName;
 }
 
@@ -46,6 +48,7 @@ const TArray<UMovieSceneFolder*>& UMovieSceneFolder::GetChildFolders() const
 
 void UMovieSceneFolder::AddChildFolder( UMovieSceneFolder* InChildFolder )
 {
+	Modify();
 
 #if WITH_EDITORONLY_DATA
 	// Ensure the added folder does not belong to any other folder in the same scene.
@@ -69,6 +72,8 @@ void UMovieSceneFolder::AddChildFolder( UMovieSceneFolder* InChildFolder )
 
 void UMovieSceneFolder::RemoveChildFolder( UMovieSceneFolder* InChildFolder )
 {
+	Modify();
+
 	ChildFolders.Remove(InChildFolder);
 }
 
@@ -81,6 +86,8 @@ const TArray<UMovieSceneTrack*>& UMovieSceneFolder::GetChildMasterTracks() const
 
 void UMovieSceneFolder::AddChildMasterTrack( UMovieSceneTrack* InMasterTrack )
 {
+	Modify();
+
 #if WITH_EDITORONLY_DATA
 	// Ensure the added track does not belong to any other folder in the same scene.
 	UMovieScene* OwningScene = GetTypedOuter<UMovieScene>();
@@ -102,6 +109,8 @@ void UMovieSceneFolder::AddChildMasterTrack( UMovieSceneTrack* InMasterTrack )
 
 void UMovieSceneFolder::RemoveChildMasterTrack( UMovieSceneTrack* InMasterTrack )
 {
+	Modify();
+
 	ChildMasterTracks.Remove( InMasterTrack );
 }
 
@@ -114,6 +123,8 @@ const TArray<FGuid>& UMovieSceneFolder::GetChildObjectBindings() const
 
 void UMovieSceneFolder::AddChildObjectBinding(const FGuid& InObjectBinding )
 {
+	Modify();
+
 #if WITH_EDITORONLY_DATA
 	// Ensure the added object  does not belong to any other folder in the same scene.
 	UMovieScene* OwningScene = GetTypedOuter<UMovieScene>();
@@ -135,6 +146,8 @@ void UMovieSceneFolder::AddChildObjectBinding(const FGuid& InObjectBinding )
 
 void UMovieSceneFolder::RemoveChildObjectBinding( const FGuid& InObjectBinding )
 {
+	Modify();
+
 	ChildObjectBindings.Remove( InObjectBinding );
 }
 
diff --git a/Engine/Source/Runtime/MovieScene/Public/MovieSceneFolder.h b/Engine/Source/Runtime/MovieScene/Public/MovieSceneFolder.h
--- a/Engine/Source/Runtime/MovieScene/Public/MovieSceneFolder.h
+++ b/Engine/Source/Runtime/MovieScene/Public/MovieSceneFolder.h
@@ -20,34 +20,34 @@ class MOVIESCENE_API UMovieSceneFolder : public UObject
 	/** Gets the name of this folder. */
 	FName GetFolderName() const;
 
-	/** Sets the name of this folder. */
+	/** Sets the name of this folder. Automatically calls Modify() on the folder object. */
 	void SetFolderName( FName InFolderName );
 
 	/** Gets the folders contained by this folder. */
 	const TArray<UMovieSceneFolder*>& GetChildFolders() const;
 
-	/** Adds a child folder to this folder. */
+	/** Adds a child folder to this folder. Automatically calls Modify() on the folder object. */
 	void AddChildFolder( UMovieSceneFolder* InChildFolder );
 
-	/** Removes a child folder from this folder. */
+	/** Removes a child folder from this folder. Automatically calls Modify() on the folder object. */
 	void RemoveChildFolder( UMovieSceneFolder* InChildFolder );
 
 	/** Gets the master tracks contained by this folder. */
 	const TArray<UMovieSceneTrack*>& GetChildMasterTracks() const;
 
-	/** Adds a master track to this folder. */
+	/** Adds a master track to this folder. Automatically calls Modify() on the folder object. */
 	void AddChildMasterTrack( UMovieSceneTrack* InMasterTrack );
 
-	/** Removes a master track from this folder. */
+	/** Removes a master track from this folder. Automatically calls Modify() on the folder object. */
 	void RemoveChildMasterTrack( UMovieSceneTrack* InMasterTrack );
 
 	/** Gets the guids for the object bindings contained by this folder. */
 	const TArray<FGuid>& GetChildObjectBindings() const;
 
-	/** Adds a guid for an object binding to this folder. */
+	/** Adds a guid for an object binding to this folder. Automatically calls Modify() on the folder object. */
 	void AddChildObjectBinding(const FGuid& InObjectBinding );
 
-	/** Removes a guid for an object binding from this folder. */
+	/** Removes a guid for an object binding from this folder. Automatically calls Modify() on the folder object. */
 	void RemoveChildObjectBinding( const FGuid& InObjectBinding );
 
 	/** Called after this object has been deserialized */
@@ -71,7 +71,7 @@ class MOVIESCENE_API UMovieSceneFolder : public UObject
 	}
 
 	/**
-	 * Set this folder's color.
+	 * Set this folder's color. Does not call Modify() on the folder object for legacy reasons.
 	 *
 	 * @param InFolderColor The folder color to set.
 	 */
@@ -89,7 +89,7 @@ class MOVIESCENE_API UMovieSceneFolder : public UObject
 	}
 
 	/**
-	 * Set this folder's desired sorting order.
+	 * Set this folder's desired sorting order. Does not call Modify() internally for legacy reasons.
 	 *
 	 * @param InSortingOrder The higher the value the further down the list the folder will be.
 	 */
