Sequencer: Optimization pass at hot CPU paths

  - Changed a map of maps to use a single map with a combined key and therefore reduce heap allocation (and allow for pre-reservation of sensible starting allocation)
  - Skipped unnecessarily checking for binding overrides in the spawn template if there were already bound objects
  - Skipped redundant string manipulation in actor spawner to do with level names where the common case is no specified level name

diff --git a/Engine/Source/Runtime/LevelSequence/Private/LevelSequenceActorSpawner.cpp b/Engine/Source/Runtime/LevelSequence/Private/LevelSequenceActorSpawner.cpp
--- a/Engine/Source/Runtime/LevelSequence/Private/LevelSequenceActorSpawner.cpp
+++ b/Engine/Source/Runtime/LevelSequence/Private/LevelSequenceActorSpawner.cpp
@@ -22,17 +22,23 @@ UClass* FLevelSequenceActorSpawner::GetSupportedTemplateType() const
 	return AActor::StaticClass();
 }
 
-ULevelStreaming* GetLevelStreaming(FString SafeLevelName, UWorld& World)
+ULevelStreaming* GetLevelStreaming(const FName& DesiredLevelName, const TArray<ULevelStreaming*>& StreamingLevels)
 {
-	if (FPackageName::IsShortPackageName(SafeLevelName))
+	if (DesiredLevelName == NAME_None)
+	{
+		return nullptr;
+	}
+
+	FString SafeLevelNameString = DesiredLevelName.ToString();
+	if (FPackageName::IsShortPackageName(SafeLevelNameString))
 	{
 		// Make sure MyMap1 and Map1 names do not resolve to a same streaming level
-		SafeLevelName = TEXT("/") + SafeLevelName;
+		SafeLevelNameString.InsertAt('/', 0);
 	}
 
-	for (ULevelStreaming* LevelStreaming : World.GetStreamingLevels())
+	for (ULevelStreaming* LevelStreaming : StreamingLevels)
 	{
-		if (LevelStreaming && LevelStreaming->GetWorldAssetPackageName().EndsWith(SafeLevelName, ESearchCase::IgnoreCase))
+		if (LevelStreaming && LevelStreaming->GetWorldAssetPackageName().EndsWith(SafeLevelNameString, ESearchCase::IgnoreCase))
 		{
 			return LevelStreaming;
 		}
@@ -61,14 +67,18 @@ UObject* FLevelSequenceActorSpawner::SpawnObject(FMovieSceneSpawnable& Spawnable
 		WorldContext = GWorld;
 	}
 
-	ULevelStreaming* LevelStreaming = GetLevelStreaming(Spawnable.GetLevelName().ToString(), *WorldContext);
-	if (LevelStreaming && LevelStreaming->GetWorldAsset().IsValid())
-	{
-		WorldContext = LevelStreaming->GetWorldAsset().Get();
-	}
-	else if (Spawnable.GetLevelName() != NAME_None)
+	FName DesiredLevelName = Spawnable.GetLevelName();
+	if (DesiredLevelName != NAME_None)
 	{
-		UE_LOG(LogMovieScene, Warning, TEXT("Can't find sublevel '%s' to spawn '%s' into"), *Spawnable.GetLevelName().ToString(), *Spawnable.GetName());
+		ULevelStreaming* LevelStreaming = GetLevelStreaming(DesiredLevelName, WorldContext->GetStreamingLevels());
+		if (LevelStreaming && LevelStreaming->GetWorldAsset().IsValid())
+		{
+			WorldContext = LevelStreaming->GetWorldAsset().Get();
+		}
+		else
+		{
+			UE_LOG(LogMovieScene, Warning, TEXT("Can't find sublevel '%s' to spawn '%s' into"), *DesiredLevelName.ToString(), *Spawnable.GetName());
+		}
 	}
 
 	// Construct the object with the same name that we will set later on the actor to avoid renaming it inside SetActorLabel
diff --git a/Engine/Source/Runtime/MovieScene/Private/Evaluation/Blending/MovieSceneAccumulatedBlendState.cpp b/Engine/Source/Runtime/MovieScene/Private/Evaluation/Blending/MovieSceneAccumulatedBlendState.cpp
--- a/Engine/Source/Runtime/MovieScene/Private/Evaluation/Blending/MovieSceneAccumulatedBlendState.cpp
+++ b/Engine/Source/Runtime/MovieScene/Private/Evaluation/Blending/MovieSceneAccumulatedBlendState.cpp
@@ -3,7 +3,7 @@
 #include "Evaluation/Blending/MovieSceneAccumulatedBlendState.h"
 #include "IMovieScenePlayer.h"
 
-void FMovieSceneAccumulatedBlendState::Consolidate(TMap<UObject*, TMap<FMovieSceneBlendingActuatorID, FActuatorTokenStackPtr>>& InOutBlendState, FMovieSceneEvaluationOperand InOperand, IMovieScenePlayer& Player)
+void FMovieSceneAccumulatedBlendState::Consolidate(TMap<FMovieSceneBlendingKey, FActuatorTokenStackPtr>& InOutBlendState, FMovieSceneEvaluationOperand InOperand, IMovieScenePlayer& Player)
 {
 	if (TokensToBlend.Num() == 0)
 	{
@@ -16,21 +16,28 @@ void FMovieSceneAccumulatedBlendState::Consolidate(TMap<UObject*, TMap<FMovieSce
 		{
 			if (UObject* Obj = WeakObj.Get())
 			{
-				Consolidate(InOutBlendState.FindOrAdd(Obj));
+				for (TInlineValue<FTokenEntry>& Token : TokensToBlend)
+				{
+					FMovieSceneBlendingKey Key = { Obj, Token->GetActuatorID() };
+
+					Token->Consolidate(InOutBlendState.FindOrAdd(Key));
+				}
 			}
 		}
 	}
 	else
 	{
 		// Explicit nullptr means master tracks
-		Consolidate(InOutBlendState.FindOrAdd(nullptr));
+		Consolidate(InOutBlendState);
 	}
 }
 
-void FMovieSceneAccumulatedBlendState::Consolidate(TMap<FMovieSceneBlendingActuatorID, FActuatorTokenStackPtr>& InOutBlendState)
+void FMovieSceneAccumulatedBlendState::Consolidate(TMap<FMovieSceneBlendingKey, FActuatorTokenStackPtr>& InOutBlendState)
 {
-	for (auto& Token : TokensToBlend)
+	for (TInlineValue<FTokenEntry>& Token : TokensToBlend)
 	{
-		Token->Consolidate(InOutBlendState);
+		FMovieSceneBlendingKey Key = { nullptr, Token->GetActuatorID() };
+
+		Token->Consolidate(InOutBlendState.FindOrAdd(Key));
 	}
 }
diff --git a/Engine/Source/Runtime/MovieScene/Private/Evaluation/Blending/MovieSceneBlendingAccumulator.cpp b/Engine/Source/Runtime/MovieScene/Private/Evaluation/Blending/MovieSceneBlendingAccumulator.cpp
--- a/Engine/Source/Runtime/MovieScene/Private/Evaluation/Blending/MovieSceneBlendingAccumulator.cpp
+++ b/Engine/Source/Runtime/MovieScene/Private/Evaluation/Blending/MovieSceneBlendingAccumulator.cpp
@@ -9,29 +9,30 @@ FMovieSceneAnimTypeID GetInitialValueTypeID()
 	return ID;
 }
 
+
 void FMovieSceneBlendingAccumulator::Apply(const FMovieSceneContext& Context, FPersistentEvaluationData& PersistentData, IMovieScenePlayer& Player)
 {
-	TMap<UObject*, TMap<FMovieSceneBlendingActuatorID, FActuatorTokenStackPtr>> BlendState;
+	TMap<FMovieSceneBlendingKey, FActuatorTokenStackPtr> BlendState;
+	BlendState.Reserve(128);
 
 	UnboundBlendState.Consolidate(BlendState, FMovieSceneEvaluationOperand(), Player);
-	for (auto& Pair : OperandToBlendState)
+	for (TTuple<FMovieSceneEvaluationOperand, FMovieSceneAccumulatedBlendState>& Pair : OperandToBlendState)
 	{
 		Pair.Value.Consolidate(BlendState, Pair.Key, Player);
 	}
 
-	// Evaluate the token stacks
-	for (auto& Pair : BlendState)
+	for (TTuple<FMovieSceneBlendingKey, FActuatorTokenStackPtr>& Pair : BlendState)
 	{
-		UObject* Object = Pair.Key;
-		
-		for (auto& InnerPair : Pair.Value)
-		{
-			InnerPair.Value->ComputeAndActuate(Object, *this, InnerPair.Key, Context, PersistentData, Player);
-		}
+		IBlendableTokenStack* TokenStack = Pair.Value.GetPtr();
+
+		UObject* Object = Pair.Key.ObjectPtr;
+		FMovieSceneBlendingActuatorID ActuatorType = Pair.Key.ActuatorType;
+
+		TokenStack->ComputeAndActuate(Object, *this, ActuatorType, Context, PersistentData, Player);
 	}
 
 	UnboundBlendState.Reset();
-	for (auto& Pair : OperandToBlendState)
+	for (TTuple<FMovieSceneEvaluationOperand, FMovieSceneAccumulatedBlendState>& Pair : OperandToBlendState)
 	{
 		Pair.Value.Reset();
 	}
@@ -40,18 +41,21 @@ void FMovieSceneBlendingAccumulator::Apply(const FMovieSceneContext& Context, FP
 void FMovieSceneBlendingAccumulator::Interrogate(const FMovieSceneContext& Context, FMovieSceneInterrogationData& InterrogationData, UObject* AnimatedObject)
 {
 	//@todo: interrogation currently does not work on entire sequences, so we just accumulate into a single map.
-	TMap<FMovieSceneBlendingActuatorID, FActuatorTokenStackPtr> BlendState;
+	TMap<FMovieSceneBlendingKey, FActuatorTokenStackPtr> BlendState;
 
 	UnboundBlendState.Consolidate(BlendState);
-	for (auto& Pair : OperandToBlendState)
+	for (TTuple<FMovieSceneEvaluationOperand, FMovieSceneAccumulatedBlendState>& Pair : OperandToBlendState)
 	{
 		Pair.Value.Consolidate(BlendState);
 	}
 
 	// Evaluate the token stacks
-	for (auto& Pair : BlendState)
+	for (TTuple<FMovieSceneBlendingKey, FActuatorTokenStackPtr>& Pair : BlendState)
 	{
-		Pair.Value->Interrogate(AnimatedObject, InterrogationData, *this, Pair.Key, Context);
+		IBlendableTokenStack* TokenStack = Pair.Value.GetPtr();
+		FMovieSceneBlendingActuatorID ActuatorType = Pair.Key.ActuatorType;
+
+		TokenStack->Interrogate(AnimatedObject, InterrogationData, *this, ActuatorType, Context);
 	}
 
 	UnboundBlendState.Reset();
diff --git a/Engine/Source/Runtime/MovieScene/Public/Evaluation/Blending/MovieSceneAccumulatedBlendState.h b/Engine/Source/Runtime/MovieScene/Public/Evaluation/Blending/MovieSceneAccumulatedBlendState.h
--- a/Engine/Source/Runtime/MovieScene/Public/Evaluation/Blending/MovieSceneAccumulatedBlendState.h
+++ b/Engine/Source/Runtime/MovieScene/Public/Evaluation/Blending/MovieSceneAccumulatedBlendState.h
@@ -15,6 +15,25 @@
  */
 typedef TInlineValue<IBlendableTokenStack, 32> FActuatorTokenStackPtr;
 
+struct FMovieSceneBlendingKey
+{
+	/** The object that the animation applies to */
+	UObject* ObjectPtr;
+
+	/** The type of the animation being applied to the object above */
+	FMovieSceneBlendingActuatorID ActuatorType;
+
+	friend uint32 GetTypeHash(const FMovieSceneBlendingKey& InKey)
+	{
+		return GetTypeHash(InKey.ObjectPtr) ^ GetTypeHash(InKey.ActuatorType);
+	}
+
+	friend bool operator==(const FMovieSceneBlendingKey& A, const FMovieSceneBlendingKey& B)
+	{
+		return A.ObjectPtr == B.ObjectPtr && A.ActuatorType == B.ActuatorType;
+	}
+};
+
 /**
  * Blendable token state that's accumulated per-operand
  */
@@ -49,14 +68,14 @@ public:
 	 * @param InOperand			The operand that is being animated
 	 * @param Player			The movie scene player that is currently being animated
 	 */
-	MOVIESCENE_API void Consolidate(TMap<UObject*, TMap<FMovieSceneBlendingActuatorID, FActuatorTokenStackPtr>>& InOutBlendState, FMovieSceneEvaluationOperand InOperand, IMovieScenePlayer& Player);
+	MOVIESCENE_API void Consolidate(TMap<FMovieSceneBlendingKey, FActuatorTokenStackPtr>& InOutBlendState, FMovieSceneEvaluationOperand InOperand, IMovieScenePlayer& Player);
 
 	/**
 	 * Consolidate all tokens currently accumulated into the specified container. Used for offline interrogation.
 	 * 
 	 * @param InOutBlendState 	Map of token stacks to populate. One FActuatorTokenStackPtr per actuator type ID.
 	 */
-	MOVIESCENE_API void Consolidate(TMap<FMovieSceneBlendingActuatorID, FActuatorTokenStackPtr>& InOutBlendState);
+	MOVIESCENE_API void Consolidate(TMap<FMovieSceneBlendingKey, FActuatorTokenStackPtr>& InOutBlendState);
 
 	/**
 	 * Reset this container
@@ -77,7 +96,12 @@ private:
 		 * Consolidate this token into the specified array of stacks, one per actuator type
 		 * @param Stacks 		Array of token stacks to populate, one per actuator type
 		 */
-		virtual void Consolidate(TMap<FMovieSceneBlendingActuatorID, FActuatorTokenStackPtr>& Stacks) = 0;
+		virtual void Consolidate(FActuatorTokenStackPtr& OutStack) = 0;
+
+		FMovieSceneBlendingActuatorID GetActuatorID() const
+		{
+			return ActuatorTypeID;
+		}
 
 		/**
 		 * Access the result of GetBlendingDataType<DataType>() for the data that this token contains
@@ -115,15 +139,12 @@ private:
 		 * Consolidate this token into the specified array of stacks, one per actuator type
 		 * @param Stacks 		Array of token stacks to populate, one per actuator type
 		 */
-		virtual void Consolidate(TMap<FMovieSceneBlendingActuatorID, FActuatorTokenStackPtr>& Stacks) override final
+		virtual void Consolidate(FActuatorTokenStackPtr& OutStack) override final
 		{
 			// Attempt to find an existing stack for this actuator
-			FActuatorTokenStackPtr* ExistingStack = Stacks.Find(ActuatorTypeID);
-			if (!ExistingStack)
+			if (!OutStack.IsValid())
 			{
-				// Add a new one if there is not one already
-				Stacks.Add(ActuatorTypeID, TBlendableTokenStack<DataType>());
-				ExistingStack = &Stacks.FindChecked(ActuatorTypeID);
+				OutStack = TBlendableTokenStack<DataType>();
 			}
 
 			// Ensure that the stack is operating on the same data type.
@@ -141,10 +162,10 @@ private:
 			//		// Assert - data type mismatch - cannot associate a TBlendableToken<int32> with an actuator ID that operates on a float
 			//		ExecutionTokens.BlendToken(FMyActuator::GetActuatorTypeID(), TBlendableToken<int32>(1.f, ...));
 
-			if (ensureMsgf((*ExistingStack)->DataTypeID == DataTypeID, TEXT("Data type mismatch between actuators of the same ID")))
+			if (ensureMsgf(OutStack->DataTypeID == DataTypeID, TEXT("Data type mismatch between actuators of the same ID")))
 			{
 				// static_cast should be safe now that we've verified the data type matches
-				TBlendableTokenStack<DataType>& TypedStack = static_cast<TBlendableTokenStack<DataType>&>(**ExistingStack);
+				TBlendableTokenStack<DataType>& TypedStack = static_cast<TBlendableTokenStack<DataType>&>(OutStack.GetValue());
 
 				// Add a pointer to this token to the stack to blend
 				TypedStack.AddToken(&Token);
diff --git a/Engine/Source/Runtime/MovieSceneTracks/Private/Evaluation/MovieSceneSpawnTemplate.cpp b/Engine/Source/Runtime/MovieSceneTracks/Private/Evaluation/MovieSceneSpawnTemplate.cpp
--- a/Engine/Source/Runtime/MovieSceneTracks/Private/Evaluation/MovieSceneSpawnTemplate.cpp
+++ b/Engine/Source/Runtime/MovieSceneTracks/Private/Evaluation/MovieSceneSpawnTemplate.cpp
@@ -51,7 +51,7 @@ struct FSpawnObjectToken : IMovieSceneExecutionToken
 		
 		// Check binding overrides to see if this spawnable has been overridden, and whether it allows the default spawnable to exist
 		const IMovieScenePlaybackClient* PlaybackClient = Player.GetPlaybackClient();
-		if (PlaybackClient)
+		if (!bHasSpawnedObject && PlaybackClient)
 		{
 			TArray<UObject*, TInlineAllocator<1>> FoundObjects;
 			bool bUseDefaultBinding = PlaybackClient->RetrieveBindingOverrides(Operand.ObjectBindingID, Operand.SequenceID, FoundObjects);
