Sequencer: RestorePreAnimatedState only for the recompiled blueprint. This fixes an issue where the level visibility would flush level streaming while loading.

diff --git a/Engine/Source/Editor/Sequencer/Private/Sequencer.cpp b/Engine/Source/Editor/Sequencer/Private/Sequencer.cpp
--- a/Engine/Source/Editor/Sequencer/Private/Sequencer.cpp
+++ b/Engine/Source/Editor/Sequencer/Private/Sequencer.cpp
@@ -338,7 +338,10 @@ void FSequencer::InitSequencer(const FSequencerInitParams& InitParams, const TSh
 		FDelegateHandle OnBlueprintPreCompileHandle = GEditor->OnBlueprintPreCompile().AddLambda([&](UBlueprint* InBlueprint)
 		{
 			// Restore pre animate state since objects will be reinstanced and current cached state will no longer be valid.
-			RestorePreAnimatedState();
+			if (InBlueprint)
+			{
+				RestorePreAnimatedState(InBlueprint->GeneratedClass.Get());
+			}
 		});
 		AcquiredResources.Add([=] { GEditor->OnBlueprintPreCompile().Remove(OnBlueprintPreCompileHandle); });
 
diff --git a/Engine/Source/Runtime/MovieScene/Private/Evaluation/MovieScenePreAnimatedState.cpp b/Engine/Source/Runtime/MovieScene/Private/Evaluation/MovieScenePreAnimatedState.cpp
--- a/Engine/Source/Runtime/MovieScene/Private/Evaluation/MovieScenePreAnimatedState.cpp
+++ b/Engine/Source/Runtime/MovieScene/Private/Evaluation/MovieScenePreAnimatedState.cpp
@@ -304,6 +304,27 @@ void FMovieScenePreAnimatedState::RestorePreAnimatedState(IMovieScenePlayer& Pla
 	}
 }
 
+void FMovieScenePreAnimatedState::RestorePreAnimatedState(IMovieScenePlayer& Player, UClass* GeneratedClass)
+{
+	for (auto& ObjectTokenPair : ObjectTokens)
+	{
+		UObject* Object = ObjectTokenPair.Key.ResolveObjectPtr();
+		if (Object)
+		{
+			if (Object->IsA(GeneratedClass) || Object->GetOuter()->IsA(GeneratedClass))
+			{
+				ObjectTokenPair.Value.Restore(Player);
+
+				for (auto& Pair : EntityToAnimatedObjects)
+				{
+					Pair.Value.Remove(ObjectTokenPair.Key);
+				}
+			}
+		}
+	}
+}
+
+
 void FMovieScenePreAnimatedState::RestorePreAnimatedState(IMovieScenePlayer& Player, UObject& Object, TFunctionRef<bool(FMovieSceneAnimTypeID)> InFilter)
 {
 	auto* FoundObjectTokens = ObjectTokens.Find(&Object);
diff --git a/Engine/Source/Runtime/MovieScene/Public/Evaluation/MovieScenePreAnimatedState.h b/Engine/Source/Runtime/MovieScene/Public/Evaluation/MovieScenePreAnimatedState.h
--- a/Engine/Source/Runtime/MovieScene/Public/Evaluation/MovieScenePreAnimatedState.h
+++ b/Engine/Source/Runtime/MovieScene/Public/Evaluation/MovieScenePreAnimatedState.h
@@ -267,6 +267,8 @@ public:
 
 	MOVIESCENE_API void RestorePreAnimatedState(IMovieScenePlayer& Player, UObject& Object);
 
+	MOVIESCENE_API void RestorePreAnimatedState(IMovieScenePlayer& Player, UClass* GeneratedClass);
+
 	MOVIESCENE_API void RestorePreAnimatedState(IMovieScenePlayer& Player, UObject& Object, TFunctionRef<bool(FMovieSceneAnimTypeID)> InFilter);
 
 	/**
diff --git a/Engine/Source/Runtime/MovieScene/Public/IMovieScenePlayer.h b/Engine/Source/Runtime/MovieScene/Public/IMovieScenePlayer.h
--- a/Engine/Source/Runtime/MovieScene/Public/IMovieScenePlayer.h
+++ b/Engine/Source/Runtime/MovieScene/Public/IMovieScenePlayer.h
@@ -296,6 +296,16 @@ public:
 		PreAnimatedState.RestorePreAnimatedState(*this, Object);
 	}
 
+	/**
+	 * Restore any pre-animated state that has been cached for the specified class
+	 *
+	 * @param GeneratedClass			The class of the object to restore
+	 */
+	void RestorePreAnimatedState(UClass* GeneratedClass)
+	{
+		PreAnimatedState.RestorePreAnimatedState(*this, GeneratedClass);
+	}
+
 	/**
 	 * Restore any pre-animated state that has been cached for the specified object
 	 *
