Sequencer: Check for property type in evaluation. This fixes an issue where the property type is changed after the track has been created.

diff --git a/Engine/Source/Runtime/MovieScene/Private/MovieSceneCommonHelpers.cpp b/Engine/Source/Runtime/MovieScene/Private/MovieSceneCommonHelpers.cpp
--- a/Engine/Source/Runtime/MovieScene/Private/MovieSceneCommonHelpers.cpp
+++ b/Engine/Source/Runtime/MovieScene/Private/MovieSceneCommonHelpers.cpp
@@ -296,22 +296,29 @@ FTrackInstancePropertyBindings::FPropertyAddress FTrackInstancePropertyBindings:
 
 	if (PropertyAndIndex.ArrayIndex != INDEX_NONE)
 	{
-		UArrayProperty* ArrayProp = CastChecked<UArrayProperty>(PropertyAndIndex.Property);
-
-		FScriptArrayHelper ArrayHelper(ArrayProp, ArrayProp->ContainerPtrToValuePtr<void>(BasePointer));
-		if (ArrayHelper.IsValidIndex(PropertyAndIndex.ArrayIndex))
+		if (PropertyAndIndex.Property->IsA(UArrayProperty::StaticClass()))
 		{
-			UStructProperty* InnerStructProp = Cast<UStructProperty>(ArrayProp->Inner);
-			if (InnerStructProp && InPropertyNames.IsValidIndex(Index+1))
-			{
-				return FindPropertyRecursive(ArrayHelper.GetRawPtr(PropertyAndIndex.ArrayIndex), InnerStructProp->Struct, InPropertyNames, Index+1);
-			}
-			else
+			UArrayProperty* ArrayProp = CastChecked<UArrayProperty>(PropertyAndIndex.Property);
+
+			FScriptArrayHelper ArrayHelper(ArrayProp, ArrayProp->ContainerPtrToValuePtr<void>(BasePointer));
+			if (ArrayHelper.IsValidIndex(PropertyAndIndex.ArrayIndex))
 			{
-				NewAddress.Property = ArrayProp->Inner;
-				NewAddress.Address = ArrayHelper.GetRawPtr(PropertyAndIndex.ArrayIndex);
+				UStructProperty* InnerStructProp = Cast<UStructProperty>(ArrayProp->Inner);
+				if (InnerStructProp && InPropertyNames.IsValidIndex(Index + 1))
+				{
+					return FindPropertyRecursive(ArrayHelper.GetRawPtr(PropertyAndIndex.ArrayIndex), InnerStructProp->Struct, InPropertyNames, Index + 1);
+				}
+				else
+				{
+					NewAddress.Property = ArrayProp->Inner;
+					NewAddress.Address = ArrayHelper.GetRawPtr(PropertyAndIndex.ArrayIndex);
+				}
 			}
 		}
+		else
+		{
+			UE_LOG(LogMovieScene, Error, TEXT("Mismatch in property evaluation. %s is not of type: %s"), *PropertyAndIndex.Property->GetName(), *UArrayProperty::StaticClass()->GetName());
+		}
 	}
 	else if (UStructProperty* StructProp = Cast<UStructProperty>(PropertyAndIndex.Property))
 	{
@@ -365,11 +372,18 @@ void FTrackInstancePropertyBindings::CallFunctionForEnum( UObject& InRuntimeObje
 	}
 	else if (UProperty* Property = PropAndFunction.PropertyAddress.GetProperty())
 	{
-		if (UEnumProperty* EnumProperty = CastChecked<UEnumProperty>(Property))
+		if (Property->IsA(UEnumProperty::StaticClass()))
 		{
-			UNumericProperty* UnderlyingProperty = EnumProperty->GetUnderlyingProperty();
-			void* ValueAddr = EnumProperty->ContainerPtrToValuePtr<void>(PropAndFunction.PropertyAddress.Address);
-			UnderlyingProperty->SetIntPropertyValue(ValueAddr, PropertyValue);
+			if (UEnumProperty* EnumProperty = CastChecked<UEnumProperty>(Property))
+			{
+				UNumericProperty* UnderlyingProperty = EnumProperty->GetUnderlyingProperty();
+				void* ValueAddr = EnumProperty->ContainerPtrToValuePtr<void>(PropAndFunction.PropertyAddress.Address);
+				UnderlyingProperty->SetIntPropertyValue(ValueAddr, PropertyValue);
+			}
+		}
+		else
+		{
+			UE_LOG(LogMovieScene, Error, TEXT("Mismatch in property evaluation. %s is not of type: %s"), *Property->GetName(), *UEnumProperty::StaticClass()->GetName());
 		}
 	}
 
@@ -411,12 +425,19 @@ int64 FTrackInstancePropertyBindings::GetCurrentValueForEnum(const UObject& Obje
 
 	if (UProperty* Property = PropAndFunction.PropertyAddress.GetProperty())
 	{
-		if(UEnumProperty* EnumProperty = CastChecked<UEnumProperty>(Property))
+		if (Property->IsA(UEnumProperty::StaticClass()))
 		{
-			UNumericProperty* UnderlyingProperty = EnumProperty->GetUnderlyingProperty();
-			void* ValueAddr = EnumProperty->ContainerPtrToValuePtr<void>(PropAndFunction.PropertyAddress.Address);
-			int64 Result = UnderlyingProperty->GetSignedIntPropertyValue(ValueAddr);
-			return Result;
+			if (UEnumProperty* EnumProperty = CastChecked<UEnumProperty>(Property))
+			{
+				UNumericProperty* UnderlyingProperty = EnumProperty->GetUnderlyingProperty();
+				void* ValueAddr = EnumProperty->ContainerPtrToValuePtr<void>(PropAndFunction.PropertyAddress.Address);
+				int64 Result = UnderlyingProperty->GetSignedIntPropertyValue(ValueAddr);
+				return Result;
+			}
+		}
+		else
+		{
+			UE_LOG(LogMovieScene, Error, TEXT("Mismatch in property evaluation. %s is not of type: %s"), *Property->GetName(), *UEnumProperty::StaticClass()->GetName());
 		}
 	}
 
@@ -433,10 +454,17 @@ template<> void FTrackInstancePropertyBindings::CallFunction<bool>(UObject& InRu
 	}
 	else if (UProperty* Property = PropAndFunction.PropertyAddress.GetProperty())
 	{
-		if (UBoolProperty* BoolProperty = CastChecked<UBoolProperty>(Property))
+		if (Property->IsA(UBoolProperty::StaticClass()))
 		{
-			uint8* ValuePtr = BoolProperty->ContainerPtrToValuePtr<uint8>(PropAndFunction.PropertyAddress.Address);
-			BoolProperty->SetPropertyValue(ValuePtr, PropertyValue);
+			if (UBoolProperty* BoolProperty = CastChecked<UBoolProperty>(Property))
+			{
+				uint8* ValuePtr = BoolProperty->ContainerPtrToValuePtr<uint8>(PropAndFunction.PropertyAddress.Address);
+				BoolProperty->SetPropertyValue(ValuePtr, PropertyValue);
+			}
+		}
+		else
+		{
+			UE_LOG(LogMovieScene, Error, TEXT("Mismatch in property evaluation. %s is not of type: %s"), *Property->GetName(), *UBoolProperty::StaticClass()->GetName());
 		}
 	}
 
@@ -451,10 +479,17 @@ template<> bool FTrackInstancePropertyBindings::GetCurrentValue<bool>(const UObj
 	FPropertyAndFunction PropAndFunction = FindOrAdd(Object);
 	if (UProperty* Property = PropAndFunction.PropertyAddress.GetProperty())
 	{
-		if (UBoolProperty* BoolProperty = CastChecked<UBoolProperty>(Property))
+		if (Property->IsA(UBoolProperty::StaticClass()))
+		{
+			if (UBoolProperty* BoolProperty = CastChecked<UBoolProperty>(Property))
+			{
+				const uint8* ValuePtr = BoolProperty->ContainerPtrToValuePtr<uint8>(PropAndFunction.PropertyAddress.Address);
+				return BoolProperty->GetPropertyValue(ValuePtr);
+			}
+		}
+		else
 		{
-			const uint8* ValuePtr = BoolProperty->ContainerPtrToValuePtr<uint8>(PropAndFunction.PropertyAddress.Address);
-			return BoolProperty->GetPropertyValue(ValuePtr);
+			UE_LOG(LogMovieScene, Error, TEXT("Mismatch in property evaluation. %s is not of type: %s"), *Property->GetName(), *UBoolProperty::StaticClass()->GetName());
 		}
 	}
 
