Sequencer: Fix key reduction not working on imported fbx channel.

Implement channel Optimize() so that it follows the curve editor's SimplifySelection by evaluating whether a key is needed (rather than removing the key and checking whether the value has changed).

The main fix is that the tangents need to be calculated as non-normalized (ie. with the DisplayRate)

diff --git a/Engine/Source/Editor/MovieSceneTools/Private/MovieSceneToolHelpers.cpp b/Engine/Source/Editor/MovieSceneTools/Private/MovieSceneToolHelpers.cpp
--- a/Engine/Source/Editor/MovieSceneTools/Private/MovieSceneToolHelpers.cpp
+++ b/Engine/Source/Editor/MovieSceneTools/Private/MovieSceneToolHelpers.cpp
@@ -946,14 +946,15 @@ bool ImportFBXProperty(FString NodeName, FString AnimatedPropertyName, FGuid Obj
 
 				}
 
+				Channel->AutoSetTangents();
+
 				if (ImportFBXSettings->bReduceKeys)
 				{
 					FKeyDataOptimizationParams Params;
 					Params.Tolerance = ImportFBXSettings->ReduceKeysTolerance;
-					using namespace MovieScene;
-					Optimize(Channel, Params);
+					Params.DisplayRate = FrameRate;
+					Channel->Optimize(Params);
 				}
-				Channel->AutoSetTangents();
 
 				return true;
 			}
@@ -1000,16 +1001,16 @@ void ImportTransformChannel(const FRichCurve& Source, FMovieSceneFloatChannel* D
 
 	}
 
+	Dest->AutoSetTangents();
 
 	const UMovieSceneUserImportFBXSettings* ImportFBXSettings = GetDefault<UMovieSceneUserImportFBXSettings>();
 	if (ImportFBXSettings->bReduceKeys)
 	{
 		FKeyDataOptimizationParams Params;
 		Params.Tolerance = ImportFBXSettings->ReduceKeysTolerance;
-		using namespace MovieScene;
-		Optimize(Dest, Params);
+		Params.DisplayRate = DestFrameRate;
+		Dest->Optimize(Params);
 	}
-	Dest->AutoSetTangents();
 }
 
 bool ImportFBXTransform(FString NodeName, FGuid ObjectBinding, UnFbx::FFbxCurvesAPI& CurveAPI, UMovieScene* InMovieScene)
diff --git a/Engine/Source/Runtime/MovieScene/Private/Channels/MovieSceneFloatChannel.cpp b/Engine/Source/Runtime/MovieScene/Private/Channels/MovieSceneFloatChannel.cpp
--- a/Engine/Source/Runtime/MovieScene/Private/Channels/MovieSceneFloatChannel.cpp
+++ b/Engine/Source/Runtime/MovieScene/Private/Channels/MovieSceneFloatChannel.cpp
@@ -116,6 +116,45 @@ float BezierInterp(float P0, float P1, float P2, float P3, float Alpha)
 	return P0123;
 }
 
+
+static float EvalForTwoKeys(const FMovieSceneFloatValue& Key1, FFrameNumber Key1Time,
+							const FMovieSceneFloatValue& Key2, FFrameNumber Key2Time,
+							FFrameNumber InTime,
+							FFrameRate DisplayRate)
+{
+	double DecimalRate = DisplayRate.AsDecimal();
+
+	float Diff = (float)(Key2Time - Key1Time).Value;
+	Diff /= DecimalRate;
+	
+	if (Diff > 0 && Key1.InterpMode != RCIM_Constant)
+	{
+		const float Alpha = ((float)(InTime - Key1Time).Value / DecimalRate) / Diff;
+		const float P0 = Key1.Value;
+		const float P3 = Key2.Value;
+
+		if (Key1.InterpMode == RCIM_Linear)
+		{
+			return FMath::Lerp(P0, P3, Alpha);
+		}
+		else
+		{
+			float LeaveTangent = Key1.Tangent.LeaveTangent * DecimalRate;
+			float ArriveTangent = Key2.Tangent.ArriveTangent * DecimalRate;
+
+			const float OneThird = 1.0f / 3.0f;
+			const float P1 = P0 + (LeaveTangent * Diff*OneThird);
+			const float P2 = P3 - (ArriveTangent * Diff*OneThird);
+
+			return BezierInterp(P0, P1, P2, P3, Alpha);
+		}
+	}
+	else
+	{
+		return Key1.Value;
+	}
+}
+
 struct FCycleParams
 {
 	FFrameTime Time;
@@ -763,35 +802,40 @@ void FMovieSceneFloatChannel::Offset(FFrameNumber DeltaPosition)
 void FMovieSceneFloatChannel::Optimize(const FKeyDataOptimizationParams& Params)
 {
 	TMovieSceneChannelData<FMovieSceneFloatValue> ChannelData = GetData();
-	if (Times.Num() > 1)
-	{
-		int32 StartIndex = 0;
-		int32 EndIndex = 0;
-
-		{
-			StartIndex = Params.Range.GetLowerBound().IsClosed() ? Algo::LowerBound(Times, Params.Range.GetLowerBoundValue()) : 0;
-			EndIndex   = Params.Range.GetUpperBound().IsClosed() ? Algo::UpperBound(Times, Params.Range.GetUpperBoundValue()) : Times.Num();
-		}
+	TArray<FFrameNumber> OutKeyTimes;
+	TArray<FKeyHandle> OutKeyHandles;
 
-		for (int32 Index = StartIndex; Index < EndIndex; ++Index)
-		{
-			FFrameNumber Time = Times[Index];
-			FMovieSceneFloatValue OriginalValue = Values[Index];
+	GetKeys(Params.Range, &OutKeyTimes, &OutKeyHandles);
 
-			// If the channel evaluates the same with this key removed, we can leave it out
-			ChannelData.RemoveKey(Index);
+	if (OutKeyHandles.Num() > 2)
+	{
+		int32 MostRecentKeepKeyIndex = 0;
+		TArray<FKeyHandle> KeysToRemove;
 
-			float NewValue = 0.f;
-			if (Evaluate(Time, NewValue) && FMath::IsNearlyEqual(NewValue, OriginalValue.Value, Params.Tolerance))
+		for (int32 TestIndex = 1; TestIndex < OutKeyHandles.Num() - 1; ++TestIndex)
+		{
+			int32 Index = ChannelData.GetIndex(OutKeyHandles[TestIndex]);
+			int32 NextIndex = ChannelData.GetIndex(OutKeyHandles[TestIndex+1]);
+
+			const float KeyValue = ChannelData.GetValues()[Index].Value;
+			const float ValueWithoutKey = EvalForTwoKeys(
+				ChannelData.GetValues()[MostRecentKeepKeyIndex], ChannelData.GetTimes()[MostRecentKeepKeyIndex].Value,
+				ChannelData.GetValues()[NextIndex], ChannelData.GetTimes()[NextIndex].Value,
+				ChannelData.GetTimes()[Index].Value,
+				Params.DisplayRate);
+				
+			if (FMath::Abs(ValueWithoutKey - KeyValue) > Params.Tolerance) // Is this key needed
 			{
-				Index--;
+				MostRecentKeepKeyIndex = Index;
 			}
 			else
 			{
-				ChannelData.AddKey(Time, OriginalValue);
+				KeysToRemove.Add(OutKeyHandles[TestIndex]);
 			}
 		}
 
+		ChannelData.DeleteKeys(KeysToRemove);
+
 		if (Params.bAutoSetInterpolation)
 		{
 			AutoSetTangents();
diff --git a/Engine/Source/Runtime/MovieScene/Public/Channels/MovieSceneChannelTraits.h b/Engine/Source/Runtime/MovieScene/Public/Channels/MovieSceneChannelTraits.h
--- a/Engine/Source/Runtime/MovieScene/Public/Channels/MovieSceneChannelTraits.h
+++ b/Engine/Source/Runtime/MovieScene/Public/Channels/MovieSceneChannelTraits.h
@@ -24,6 +24,9 @@ struct FKeyDataOptimizationParams
 
 	/** A range inside which to optimize keys */
 	TRange<FFrameNumber> Range = TRange<FFrameNumber>::All();
+
+	/** The display rate to use for calculating tangents in non-normalized space */
+	FFrameRate DisplayRate;
 };
 
 /**
