Sequencer: Use cached object bindings when drawing motion trails.

FindObjectId can be slow because it clears the cached bindings. Added FindCachedObjectId which retrieves the binding without clearing the cache.

diff --git a/Engine/Source/Editor/Sequencer/Private/SequencerEdMode.cpp b/Engine/Source/Editor/Sequencer/Private/SequencerEdMode.cpp
--- a/Engine/Source/Editor/Sequencer/Private/SequencerEdMode.cpp
+++ b/Engine/Source/Editor/Sequencer/Private/SequencerEdMode.cpp
@@ -286,7 +286,7 @@ void FSequencerEdMode::GetParents(TArray<const UObject *>& Parents, const UObjec
 /** This is not that scalable moving forward with stuff like the control rig , need a better caching solution there */
 bool FSequencerEdMode::GetParentTM(FTransform& CurrentRefTM, const TSharedPtr<FSequencer>& Sequencer, UObject* ParentObject, FFrameTime KeyTime)
 {
-	FGuid ObjectBinding = Sequencer->FindObjectId(*ParentObject, Sequencer->GetFocusedTemplateID());
+	FGuid ObjectBinding = Sequencer->FindCachedObjectId(*ParentObject, Sequencer->GetFocusedTemplateID());
 
 	if (ObjectBinding.IsValid())
 	{
@@ -427,7 +427,7 @@ FTransform FSequencerEdMode::GetRefFrame(const TSharedPtr<FSequencer>& Sequencer
 		// Check if our parent is animated in this Sequencer
 
 		UObject* ParentObject = SceneComponent->GetAttachParent() == SceneComponent->GetOwner()->GetRootComponent() ? static_cast<UObject*>(SceneComponent->GetOwner()) : SceneComponent->GetAttachParent();
-		FGuid ObjectBinding = Sequencer->FindObjectId(*ParentObject, Sequencer->GetFocusedTemplateID());
+		FGuid ObjectBinding = Sequencer->FindCachedObjectId(*ParentObject, Sequencer->GetFocusedTemplateID());
 
 		if (ObjectBinding.IsValid())
 		{
diff --git a/Engine/Source/Runtime/MovieScene/Private/Evaluation/MovieSceneEvaluationState.cpp b/Engine/Source/Runtime/MovieScene/Private/Evaluation/MovieSceneEvaluationState.cpp
--- a/Engine/Source/Runtime/MovieScene/Private/Evaluation/MovieSceneEvaluationState.cpp
+++ b/Engine/Source/Runtime/MovieScene/Private/Evaluation/MovieSceneEvaluationState.cpp
@@ -75,6 +75,18 @@ FGuid FMovieSceneObjectCache::FindObjectId(UObject& InObject, IMovieScenePlayer&
 	// complete lookup from scratch. This is required for UMG as it interchanges content slots without notifying sequencer.
 	Clear(Player);
 
+	return FindCachedObjectId(InObject, Player);
+}
+
+FGuid FMovieSceneObjectCache::FindCachedObjectId(UObject& InObject, IMovieScenePlayer& Player)
+{
+	UMovieSceneSequence* Sequence = WeakSequence.Get();
+	UMovieScene* MovieScene = Sequence ? Sequence->GetMovieScene() : nullptr;
+	if (!MovieScene)
+	{
+		return FGuid();
+	}
+
 	TWeakObjectPtr<> ObjectToFind(&InObject);
 
 	// Search all possessables
@@ -314,3 +326,9 @@ FGuid FMovieSceneEvaluationState::FindObjectId(UObject& Object, FMovieSceneSeque
 	FMovieSceneObjectCache* Cache = ObjectCaches.Find(InSequenceID);
 	return Cache ? Cache->FindObjectId(Object, Player) : FGuid();
 }
+
+FGuid FMovieSceneEvaluationState::FindCachedObjectId(UObject& Object, FMovieSceneSequenceIDRef InSequenceID, IMovieScenePlayer& Player)
+{
+	FMovieSceneObjectCache* Cache = ObjectCaches.Find(InSequenceID);
+	return Cache ? Cache->FindCachedObjectId(Object, Player) : FGuid();
+}
diff --git a/Engine/Source/Runtime/MovieScene/Public/Evaluation/MovieSceneEvaluationState.h b/Engine/Source/Runtime/MovieScene/Public/Evaluation/MovieSceneEvaluationState.h
--- a/Engine/Source/Runtime/MovieScene/Public/Evaluation/MovieSceneEvaluationState.h
+++ b/Engine/Source/Runtime/MovieScene/Public/Evaluation/MovieSceneEvaluationState.h
@@ -55,6 +55,16 @@ struct FMovieSceneObjectCache
 	 */
 	MOVIESCENE_API FGuid FindObjectId(UObject& InObject, IMovieScenePlayer& Player);
 
+	/**
+	* Attempt deduce the posessable or spawnable that relates to the specified object
+	* @note Does not clear the existing cache
+	*
+	* @param InObject			The object whose binding ID is to be find
+	* @param Player			The movie scene player that is playing back the sequence
+	* @return The object's spawnable or possessable GUID, or a zero GUID if it was not found
+	*/
+	MOVIESCENE_API FGuid FindCachedObjectId(UObject& InObject, IMovieScenePlayer& Player);
+
 	/**
 	 * Invalidate any object bindings for objects that have been destroyed
 	 */
@@ -158,6 +168,16 @@ struct FMovieSceneEvaluationState
 	 */
 	MOVIESCENE_API FGuid FindObjectId(UObject& Object, FMovieSceneSequenceIDRef InSequenceID, IMovieScenePlayer& Player);
 
+	/**
+	* Attempt deduce the posessable or spawnable that relates to the specified object
+	* @note Does not clear the existing cache
+	*
+	* @param InObject			The object whose binding ID is to be find
+	* @param Player			The movie scene player that is playing back the sequence
+	* @return The object's spawnable or possessable GUID, or a zero GUID if it was not found
+	*/
+	MOVIESCENE_API FGuid FindCachedObjectId(UObject& Object, FMovieSceneSequenceIDRef InSequenceID, IMovieScenePlayer& Player);
+
 	/**
 	 * Find an object cache pertaining to the specified sequence
 	 *
diff --git a/Engine/Source/Runtime/MovieScene/Public/IMovieScenePlayer.h b/Engine/Source/Runtime/MovieScene/Public/IMovieScenePlayer.h
--- a/Engine/Source/Runtime/MovieScene/Public/IMovieScenePlayer.h
+++ b/Engine/Source/Runtime/MovieScene/Public/IMovieScenePlayer.h
@@ -209,6 +209,20 @@ public:
 		return State.FindObjectId(InObject, SequenceID, *this);
 	}
 
+	/**
+	* Attempt to find the object binding ID for the specified object, in the specified sequence
+	* @note: Does not clear the existing cache
+	*
+	* @param InObject 			The object to find a GUID for
+	* @param SequenceID 		The sequence ID to search within
+	*
+	* @return The guid of the object's binding, or zero guid if it was not found
+	*/
+	FGuid FindCachedObjectId(UObject& InObject, FMovieSceneSequenceIDRef SequenceID)
+	{
+		return State.FindCachedObjectId(InObject, SequenceID, *this);
+	}
+
 	/**
 	 * Attempt to save specific state for the specified token state before it animates an object.
 	 * @note: Will only call IMovieSceneExecutionToken::CacheExistingState if no state has been previously cached for the specified token type
