Sequencer: Fixed edge case where equvalent range bounds were not considered equal

diff --git a/Engine/Source/Runtime/MovieScene/Private/Compilation/MovieSceneCompiler.cpp b/Engine/Source/Runtime/MovieScene/Private/Compilation/MovieSceneCompiler.cpp
--- a/Engine/Source/Runtime/MovieScene/Private/Compilation/MovieSceneCompiler.cpp
+++ b/Engine/Source/Runtime/MovieScene/Private/Compilation/MovieSceneCompiler.cpp
@@ -314,7 +314,6 @@ void FMovieSceneCompiler::CompileRange(TRange<FFrameNumber> InGlobalRange, UMovi
 		// next entry in the evaluation field iterator (which should be a populated range)
 		if (CompiledRange.GetUpperBound() == EmptySpaceRange.GetUpperBound())
 		{
-			IterFromBound = TRangeBound<FFrameNumber>::FlipInclusion(CompiledRange.GetUpperBound());
 			++ExistingEvaluationFieldIter;
 		}
 	}
diff --git a/Engine/Source/Runtime/MovieScene/Private/Evaluation/MovieSceneEvaluationTree.cpp b/Engine/Source/Runtime/MovieScene/Private/Evaluation/MovieSceneEvaluationTree.cpp
--- a/Engine/Source/Runtime/MovieScene/Private/Evaluation/MovieSceneEvaluationTree.cpp
+++ b/Engine/Source/Runtime/MovieScene/Private/Evaluation/MovieSceneEvaluationTree.cpp
@@ -65,7 +65,7 @@ FMovieSceneEvaluationTreeRangeIterator::FMovieSceneEvaluationTreeRangeIterator(c
 
 		// Binary search children's lower bounds for the first that's >= the starting bound. That results in either ChildIndex or ChildIndex-1 being the overlapping range (if any).
 		int32 ChildIndex = Algo::LowerBoundBy(Children, StartingBound, GetLowerBound, MovieSceneHelpers::SortLowerBounds);
-		if (Children.IsValidIndex(ChildIndex) && Children[ChildIndex].Range.GetLowerBound() == StartingBound)
+		if (Children.IsValidIndex(ChildIndex) && MovieScene::DiscreteLowerBoundsAreEquivalent(Children[ChildIndex].Range.GetLowerBound(), StartingBound))
 		{
 			CurrentNodeHandle = FMovieSceneEvaluationTreeNodeHandle(ThisNode.ChildrenID, ChildIndex);
 		}
diff --git a/Engine/Source/Runtime/MovieScene/Public/MovieSceneTimeHelpers.h b/Engine/Source/Runtime/MovieScene/Public/MovieSceneTimeHelpers.h
--- a/Engine/Source/Runtime/MovieScene/Public/MovieSceneTimeHelpers.h
+++ b/Engine/Source/Runtime/MovieScene/Public/MovieSceneTimeHelpers.h
@@ -119,6 +119,32 @@ inline bool DiscreteRangesOverlap(const TRange<FFrameNumber>& A, const TRange<FF
 	}
 }
 
+/**
+ * Check whether the two specified lower bounds are equivalent in a discrete domain
+ */
+inline bool DiscreteLowerBoundsAreEquivalent(const TRangeBound<FFrameNumber>& A, const TRangeBound<FFrameNumber>& B)
+{
+	if (A.IsOpen() || B.IsOpen())
+	{
+		// Either bound is open
+		return A.IsOpen() && B.IsOpen();
+	}
+	else if (A.IsExclusive() == B.IsExclusive())
+	{
+		// Both inclusive or both exclusive
+		return A.GetValue() == B.GetValue();
+	}
+	else if (A.IsExclusive())
+	{
+		// A exclusive, B inclusive
+		return A.GetValue() == B.GetValue()-1;
+	}
+	else
+	{
+		// A inclusive, B exclusive
+		return B.GetValue() == A.GetValue()-1;
+	}
+}
 
 /**
  * Dilate the specified range by adding a specific size to the lower and upper bounds (if closed)
