Sequencer: Replaced instances of TRangeBound operator== and TRange::IsEmpty in movie scene compiler code with functions that understand equivalence in a discrete domain

diff --git a/Engine/Source/Runtime/MovieScene/Private/Compilation/MovieSceneCompiler.cpp b/Engine/Source/Runtime/MovieScene/Private/Compilation/MovieSceneCompiler.cpp
--- a/Engine/Source/Runtime/MovieScene/Private/Compilation/MovieSceneCompiler.cpp
+++ b/Engine/Source/Runtime/MovieScene/Private/Compilation/MovieSceneCompiler.cpp
@@ -312,7 +312,7 @@ void FMovieSceneCompiler::CompileRange(TRange<FFrameNumber> InGlobalRange, UMovi
 
 		// If the range that we just compiled goes right up to the end of the gap, increment onto the
 		// next entry in the evaluation field iterator (which should be a populated range)
-		if (CompiledRange.GetUpperBound() == EmptySpaceRange.GetUpperBound())
+		if (MovieScene::DiscreteUpperBoundsAreEquivalent(CompiledRange.GetUpperBound(), EmptySpaceRange.GetUpperBound()))
 		{
 			++ExistingEvaluationFieldIter;
 		}
@@ -421,7 +421,7 @@ void FMovieSceneCompiler::GatherCompileOnTheFlyData(UMovieSceneSequence& InSeque
 	for ( ; SubSectionIt && SubSectionIt.Range().Overlaps(CompileClampIntersection); ++SubSectionIt)
 	{
 		TRange<FFrameNumber> ThisSegmentRangeRoot = Params.ClampRoot(SubSectionIt.Range() * Params.RootToSequenceTransform.Inverse());
-		if (ThisSegmentRangeRoot.IsEmpty())
+		if (MovieScene::DiscreteRangeIsEmpty(ThisSegmentRangeRoot))
 		{
 			continue;
 		}
@@ -569,7 +569,7 @@ void FMovieSceneCompiler::GatherCompileDataForTrack(FMovieSceneEvaluationTrack&
 		Data.bRequiresInit = ThisSegment.Impls.ContainsByPredicate(RequiresInit);
 
 		TRange<FFrameNumber> IntersectionRange = Params.ClampRoot(ThisSegment.Range * SequenceToRootTransform);
-		if (!IntersectionRange.IsEmpty())
+		if (!MovieScene::DiscreteRangeIsEmpty(IntersectionRange))
 		{
 			OutData.Tracks.Add(IntersectionRange, Data);
 		}
diff --git a/Engine/Source/Runtime/MovieScene/Private/Compilation/MovieSceneEvaluationTemplateGenerator.cpp b/Engine/Source/Runtime/MovieScene/Private/Compilation/MovieSceneEvaluationTemplateGenerator.cpp
--- a/Engine/Source/Runtime/MovieScene/Private/Compilation/MovieSceneEvaluationTemplateGenerator.cpp
+++ b/Engine/Source/Runtime/MovieScene/Private/Compilation/MovieSceneEvaluationTemplateGenerator.cpp
@@ -8,6 +8,7 @@
 #include "IMovieSceneModule.h"
 #include "Tracks/MovieSceneSubTrack.h"
 #include "Sections/MovieSceneSubSection.h"
+#include "MovieSceneTimeHelpers.h"
 #include "Compilation/MovieSceneSegmentCompiler.h"
 
 FMovieSceneEvaluationTemplateGenerator::FMovieSceneEvaluationTemplateGenerator(UMovieSceneSequence& InSequence, FMovieSceneEvaluationTemplate& OutTemplate)
@@ -103,7 +104,7 @@ void FMovieSceneEvaluationTemplateGenerator::ProcessSubTrack(const UMovieSceneSu
 	for (int32 SectionIndex = 0; SectionIndex < AllSections.Num(); ++SectionIndex)
 	{
 		UMovieSceneSubSection* SubSection = CastChecked<UMovieSceneSubSection>(AllSections[SectionIndex]);
-		if (!SubSection || !SubSection->IsActive() || SubSection->GetRange().IsEmpty())
+		if (!SubSection || !SubSection->IsActive() || MovieScene::DiscreteRangeIsEmpty(SubSection->GetRange()))
 		{
 			continue;
 		}
diff --git a/Engine/Source/Runtime/MovieScene/Private/Evaluation/MovieSceneEvaluationField.cpp b/Engine/Source/Runtime/MovieScene/Private/Evaluation/MovieSceneEvaluationField.cpp
--- a/Engine/Source/Runtime/MovieScene/Private/Evaluation/MovieSceneEvaluationField.cpp
+++ b/Engine/Source/Runtime/MovieScene/Private/Evaluation/MovieSceneEvaluationField.cpp
@@ -59,7 +59,7 @@ TRange<int32> FMovieSceneEvaluationField::ConditionallyCompileRange(const TRange
 		{
 			bIsDirty = true;
 
-			if (!InvalidatedSubSequenceRange.IsEmpty())
+			if (!MovieScene::DiscreteRangeIsEmpty(InvalidatedSubSequenceRange))
 			{
 				// Invalidate this evaluation field
 				RangesToInvalidate.Add(InvalidatedSubSequenceRange);
diff --git a/Engine/Source/Runtime/MovieScene/Private/Evaluation/MovieSceneEvaluationTrack.cpp b/Engine/Source/Runtime/MovieScene/Private/Evaluation/MovieSceneEvaluationTrack.cpp
--- a/Engine/Source/Runtime/MovieScene/Private/Evaluation/MovieSceneEvaluationTrack.cpp
+++ b/Engine/Source/Runtime/MovieScene/Private/Evaluation/MovieSceneEvaluationTrack.cpp
@@ -286,7 +286,7 @@ TArray<FMovieSceneSegmentIdentifier> FMovieSceneEvaluationTrack::GetSegmentsInRa
 {
 	TArray<FMovieSceneSegmentIdentifier> SegmentsInRange;
 
-	while (!InLocalRange.IsEmpty())
+	while (!MovieScene::DiscreteRangeIsEmpty(InLocalRange))
 	{
 		FMovieSceneSegmentIdentifier SegmentID = FindFirstSegment(InLocalRange);
 
@@ -470,7 +470,7 @@ namespace
 	bool IntersectSegmentRanges(const FMovieSceneSegment& Segment, TRange<FFrameNumber> TraversedRange, TMap<int32, TRange<FFrameNumber>>& ImplToAccumulatedRange)
 	{
 		TRange<FFrameNumber> Intersection = TRange<FFrameNumber>::Intersection(Segment.Range, TraversedRange);
-		if (Intersection.IsEmpty())
+		if (MovieScene::DiscreteRangeIsEmpty(Intersection))
 		{
 			return false;
 		}
diff --git a/Engine/Source/Runtime/MovieScene/Private/Evaluation/MovieSceneEvaluationTree.cpp b/Engine/Source/Runtime/MovieScene/Private/Evaluation/MovieSceneEvaluationTree.cpp
--- a/Engine/Source/Runtime/MovieScene/Private/Evaluation/MovieSceneEvaluationTree.cpp
+++ b/Engine/Source/Runtime/MovieScene/Private/Evaluation/MovieSceneEvaluationTree.cpp
@@ -104,7 +104,7 @@ void FMovieSceneEvaluationTreeRangeIterator::Iter(bool bForwards)
 		return;
 	}
 
-	TRangeBound<FFrameNumber> NewLeadingBound = CurrentRange.IsEmpty() ? TRangeBound<FFrameNumber>::Open() : TRangeBound<FFrameNumber>::FlipInclusion(GetTrailingBound(bForwards, CurrentRange));
+	TRangeBound<FFrameNumber> NewLeadingBound = MovieScene::DiscreteRangeIsEmpty(CurrentRange) ? TRangeBound<FFrameNumber>::Open() : TRangeBound<FFrameNumber>::FlipInclusion(GetTrailingBound(bForwards, CurrentRange));
 
 	// Iterate into children when possible (where the leading bound matches the current new leading bound)
 	FMovieSceneEvaluationTreeNodeHandle NextChildHandle = FindNextChild(CurrentNodeHandle, NewLeadingBound, bForwards);
@@ -253,7 +253,7 @@ void FMovieSceneEvaluationTree::AddTimeRange(TRange<FFrameNumber> InTimeRange, c
 				TRange<FFrameNumber> PrecedingSpace(LastBound, TRangeBound<FFrameNumber>::FlipInclusion(ChildNode.Range.GetLowerBound()));
 				TRange<FFrameNumber> PrecedingIntersection = TRange<FFrameNumber>::Intersection(PrecedingSpace, InTimeRange);
 
-				if (!PrecedingIntersection.IsEmpty())
+				if (!MovieScene::DiscreteRangeIsEmpty(PrecedingIntersection))
 				{
 					// Insert this child at the current child index
 					InsertNewChild(PrecedingIntersection, InOperator, ChildIndex, InParent);
@@ -278,7 +278,7 @@ void FMovieSceneEvaluationTree::AddTimeRange(TRange<FFrameNumber> InTimeRange, c
 			TRange<FFrameNumber> ProceedingSpace(LastBound, ThisNode.Range.GetUpperBound());
 			TRange<FFrameNumber> ProceedingIntersection = TRange<FFrameNumber>::Intersection(ProceedingSpace, InTimeRange);
 
-			if (!ProceedingIntersection.IsEmpty())
+			if (!MovieScene::DiscreteRangeIsEmpty(ProceedingIntersection))
 			{
 				// Insert this child at the end
 				InsertNewChild(ProceedingIntersection, InOperator, ChildNodes.Get(ThisNode.ChildrenID).Num(), InParent);
diff --git a/Engine/Source/Runtime/MovieScene/Private/MovieSceneTrack.cpp b/Engine/Source/Runtime/MovieScene/Private/MovieSceneTrack.cpp
--- a/Engine/Source/Runtime/MovieScene/Private/MovieSceneTrack.cpp
+++ b/Engine/Source/Runtime/MovieScene/Private/MovieSceneTrack.cpp
@@ -226,7 +226,7 @@ EMovieSceneCompileResult UMovieSceneTrack::Compile(FMovieSceneEvaluationTrack& O
 		for (const UMovieSceneSection* Section : GetAllSections())
 		{
 			const TRange<FFrameNumber> SectionRange = Section->GetRange();
-			if (!Section->IsActive() || SectionRange.IsEmpty())
+			if (!Section->IsActive() || MovieScene::DiscreteRangeIsEmpty(SectionRange))
 			{
 				continue;
 			}
diff --git a/Engine/Source/Runtime/MovieScene/Public/MovieSceneTimeHelpers.h b/Engine/Source/Runtime/MovieScene/Public/MovieSceneTimeHelpers.h
--- a/Engine/Source/Runtime/MovieScene/Public/MovieSceneTimeHelpers.h
+++ b/Engine/Source/Runtime/MovieScene/Public/MovieSceneTimeHelpers.h
@@ -146,6 +146,63 @@ inline bool DiscreteLowerBoundsAreEquivalent(const TRangeBound<FFrameNumber>& A,
 	}
 }
 
+/**
+ * Check whether the two specified upper bounds are equivalent in a discrete domain
+ */
+inline bool DiscreteUpperBoundsAreEquivalent(const TRangeBound<FFrameNumber>& A, const TRangeBound<FFrameNumber>& B)
+{
+	if (A.IsOpen() || B.IsOpen())
+	{
+		// Either bound is open
+		return A.IsOpen() && B.IsOpen();
+	}
+	else if (A.IsExclusive() == B.IsExclusive())
+	{
+		// Both inclusive or both exclusive
+		return A.GetValue() == B.GetValue();
+	}
+	else if (A.IsExclusive())
+	{
+		// A exclusive, B inclusive
+		return A.GetValue() == B.GetValue()+1;
+	}
+	else
+	{
+		// A inclusive, B exclusive
+		return B.GetValue() == A.GetValue()+1;
+	}
+}
+
+/**
+ * Check whether the specified range contains any integer frame numbers or not
+ */
+inline bool DiscreteRangeIsEmpty(const TRange<FFrameNumber>& InRange)
+{
+	if (InRange.GetLowerBound().IsOpen() || InRange.GetUpperBound().IsOpen())
+	{
+		return false;
+	}
+
+	// From here on we're handling ranges of the form [x,y], [x,y), (x,y] and (x,y)
+	const bool bLowerInclusive = InRange.GetLowerBound().IsInclusive();
+	const bool bUpperInclusive = InRange.GetUpperBound().IsInclusive();
+
+	if (bLowerInclusive)
+	{
+		// Lower is inclusive
+		return bUpperInclusive
+			? InRange.GetLowerBoundValue() >  InRange.GetUpperBoundValue()		// [x, y] - empty if x >  y
+			: InRange.GetLowerBoundValue() >= InRange.GetUpperBoundValue();		// [x, y) - empty if x >= y
+	}
+	else
+	{
+		// Lower is exclusive
+		return bUpperInclusive
+			? InRange.GetLowerBoundValue() >= InRange.GetUpperBoundValue()		// (x, y] - empty if x >= y
+			: InRange.GetLowerBoundValue() >= InRange.GetUpperBoundValue()-1;	// (x, y) - empty if x >= y-1
+	}
+}
+
 /**
  * Dilate the specified range by adding a specific size to the lower and upper bounds (if closed)
  */
