Sequencer: Added automation test to verify issue-free compilation of range boundary conditions

diff --git a/Engine/Source/Runtime/MovieScene/Private/Tests/MovieSceneCompilerTests.cpp b/Engine/Source/Runtime/MovieScene/Private/Tests/MovieSceneCompilerTests.cpp
--- a/Engine/Source/Runtime/MovieScene/Private/Tests/MovieSceneCompilerTests.cpp
+++ b/Engine/Source/Runtime/MovieScene/Private/Tests/MovieSceneCompilerTests.cpp
@@ -27,7 +27,7 @@ bool FMovieSceneCompilerPerfTest::RunTest(const FString& Parameters)
 
 	for (int32 i = 0; i < 100; ++i)
 	{
-		UTestMovieSceneTrack* Track = NewObject<UTestMovieSceneTrack>(GetTransientPackage());
+		UTestMovieSceneTrack* Track = Sequence->MovieScene->AddMasterTrack<UTestMovieSceneTrack>();
 
 		int32 NumSections = FMath::Rand() % 10;
 		for (int32 SectionIndex = 0; SectionIndex < NumSections; ++SectionIndex)
@@ -77,4 +77,106 @@ bool FMovieSceneCompilerPerfTest::RunTest(const FString& Parameters)
 	return true;
 }
 
+TRange<FFrameNumber> MakeRange(FFrameNumber LowerBound, FFrameNumber UpperBound, ERangeBoundTypes::Type LowerType, ERangeBoundTypes::Type UpperType)
+{
+	return LowerType == ERangeBoundTypes::Inclusive
+		? UpperType == ERangeBoundTypes::Inclusive
+			? TRange<FFrameNumber>(TRangeBound<FFrameNumber>::Inclusive(LowerBound), TRangeBound<FFrameNumber>::Inclusive(UpperBound))
+			: TRange<FFrameNumber>(TRangeBound<FFrameNumber>::Inclusive(LowerBound), TRangeBound<FFrameNumber>::Exclusive(UpperBound))
+		: UpperType == ERangeBoundTypes::Inclusive
+			? TRange<FFrameNumber>(TRangeBound<FFrameNumber>::Exclusive(LowerBound), TRangeBound<FFrameNumber>::Inclusive(UpperBound))
+			: TRange<FFrameNumber>(TRangeBound<FFrameNumber>::Exclusive(LowerBound), TRangeBound<FFrameNumber>::Exclusive(UpperBound));
+}
+
+IMPLEMENT_SIMPLE_AUTOMATION_TEST(FMovieSceneCompilerRangeTest, "System.Engine.Sequencer.Compiler.Ranges", EAutomationTestFlags::EditorContext | EAutomationTestFlags::EngineFilter)
+bool FMovieSceneCompilerRangeTest::RunTest(const FString& Parameters)
+{
+	FFrameNumber CompileAtTimes[] = {
+		-3, -2, -1, 0, 1, 2, 3, 9, 10, 11, 12, 13, 14, 15, 16, 17, 18, 19, 20
+	};
+	// Test each combination of inc/excl boundary conditions for adjacent and adjoining ranges
+	TRange<FFrameNumber> Ranges[] = {
+		MakeRange(-2, -1, ERangeBoundTypes::Inclusive, ERangeBoundTypes::Inclusive),
+		MakeRange(-2, -1, ERangeBoundTypes::Inclusive, ERangeBoundTypes::Exclusive),
+		MakeRange(-2, -1, ERangeBoundTypes::Exclusive, ERangeBoundTypes::Inclusive),
+		MakeRange(-2, -1, ERangeBoundTypes::Exclusive, ERangeBoundTypes::Exclusive),
+
+		MakeRange(-1, -1, ERangeBoundTypes::Inclusive, ERangeBoundTypes::Inclusive),
+		MakeRange(-1, -1, ERangeBoundTypes::Inclusive, ERangeBoundTypes::Exclusive),
+		MakeRange(-1, -1, ERangeBoundTypes::Exclusive, ERangeBoundTypes::Inclusive),
+		MakeRange(-1, -1, ERangeBoundTypes::Exclusive, ERangeBoundTypes::Exclusive),
+
+		MakeRange(-1,  0, ERangeBoundTypes::Inclusive, ERangeBoundTypes::Inclusive),
+		MakeRange(-1,  0, ERangeBoundTypes::Inclusive, ERangeBoundTypes::Exclusive),
+		MakeRange(-1,  0, ERangeBoundTypes::Exclusive, ERangeBoundTypes::Inclusive),
+		MakeRange(-1,  0, ERangeBoundTypes::Exclusive, ERangeBoundTypes::Exclusive),
+
+		MakeRange( 0,  0, ERangeBoundTypes::Inclusive, ERangeBoundTypes::Inclusive),
+		MakeRange( 0,  0, ERangeBoundTypes::Inclusive, ERangeBoundTypes::Exclusive),
+		MakeRange( 0,  0, ERangeBoundTypes::Exclusive, ERangeBoundTypes::Inclusive),
+		MakeRange( 0,  0, ERangeBoundTypes::Exclusive, ERangeBoundTypes::Exclusive),
+
+		MakeRange( 0,  1, ERangeBoundTypes::Inclusive, ERangeBoundTypes::Inclusive),
+		MakeRange( 0,  1, ERangeBoundTypes::Inclusive, ERangeBoundTypes::Exclusive),
+		MakeRange( 0,  1, ERangeBoundTypes::Exclusive, ERangeBoundTypes::Inclusive),
+		MakeRange( 0,  1, ERangeBoundTypes::Exclusive, ERangeBoundTypes::Exclusive),
+
+		MakeRange( 1,  1, ERangeBoundTypes::Inclusive, ERangeBoundTypes::Inclusive),
+		MakeRange( 1,  1, ERangeBoundTypes::Inclusive, ERangeBoundTypes::Exclusive),
+		MakeRange( 1,  1, ERangeBoundTypes::Exclusive, ERangeBoundTypes::Inclusive),
+		MakeRange( 1,  1, ERangeBoundTypes::Exclusive, ERangeBoundTypes::Exclusive),
+
+		MakeRange( 0,  2, ERangeBoundTypes::Inclusive, ERangeBoundTypes::Inclusive),
+		MakeRange( 0,  2, ERangeBoundTypes::Inclusive, ERangeBoundTypes::Exclusive),
+		MakeRange( 0,  2, ERangeBoundTypes::Exclusive, ERangeBoundTypes::Inclusive),
+		MakeRange( 0,  2, ERangeBoundTypes::Exclusive, ERangeBoundTypes::Exclusive),
+
+		MakeRange( 10, 15, ERangeBoundTypes::Inclusive, ERangeBoundTypes::Inclusive),
+		MakeRange( 9,  15, ERangeBoundTypes::Exclusive, ERangeBoundTypes::Inclusive),
+		MakeRange( 10, 15, ERangeBoundTypes::Exclusive, ERangeBoundTypes::Inclusive),
+		MakeRange( 11, 15, ERangeBoundTypes::Exclusive, ERangeBoundTypes::Inclusive),
+		
+		MakeRange( 13, 17, ERangeBoundTypes::Inclusive, ERangeBoundTypes::Inclusive),
+		MakeRange( 13, 18, ERangeBoundTypes::Inclusive, ERangeBoundTypes::Exclusive),
+		MakeRange( 13, 19, ERangeBoundTypes::Inclusive, ERangeBoundTypes::Inclusive),
+		MakeRange( 13, 18, ERangeBoundTypes::Inclusive, ERangeBoundTypes::Inclusive),
+	};
+
+	UTestMovieSceneSequence* Sequence = NewObject<UTestMovieSceneSequence>(GetTransientPackage());
+	for (TRange<FFrameNumber> Range : Ranges)
+	{
+		UTestMovieSceneTrack*   Track   = Sequence->MovieScene->AddMasterTrack<UTestMovieSceneTrack>();
+		UTestMovieSceneSection* Section = NewObject<UTestMovieSceneSection>(Track);
+
+		Section->SetRange(Range);
+		Track->SectionArray.Add(Section);
+	}
+
+	struct FTemplateStore : IMovieSceneSequenceTemplateStore
+	{
+		FMovieSceneEvaluationTemplate& AccessTemplate(UMovieSceneSequence& Sequence) override
+		{
+			return Template;
+		}
+
+		FMovieSceneEvaluationTemplate Template;
+	} Store;
+
+	// Compile individual times
+	for (FFrameNumber Time : CompileAtTimes)
+	{
+		FMovieSceneCompiler::CompileRange(TRange<FFrameNumber>::Inclusive(Time, Time), *Sequence, Store);
+	}
+
+	// Compile a whole range
+	Store.Template = FMovieSceneEvaluationTemplate();
+	FMovieSceneCompiler::CompileRange(TRange<FFrameNumber>::All(), *Sequence, Store);
+
+	// Compile the whole sequence
+	Store.Template = FMovieSceneEvaluationTemplate();
+	FMovieSceneCompiler::Compile(*Sequence, Store);
+
+	return true;
+}
+
 #endif // WITH_DEV_AUTOMATION_TESTS
\ No newline at end of file
