Fixed erroneous shadowed variable warning

diff --git a/Engine/Source/Runtime/MovieScene/Private/Tests/MovieSceneCompilerTests.cpp b/Engine/Source/Runtime/MovieScene/Private/Tests/MovieSceneCompilerTests.cpp
--- a/Engine/Source/Runtime/MovieScene/Private/Tests/MovieSceneCompilerTests.cpp
+++ b/Engine/Source/Runtime/MovieScene/Private/Tests/MovieSceneCompilerTests.cpp
@@ -154,7 +154,7 @@ bool FMovieSceneCompilerRangeTest::RunTest(const FString& Parameters)
 
 	struct FTemplateStore : IMovieSceneSequenceTemplateStore
 	{
-		FMovieSceneEvaluationTemplate& AccessTemplate(UMovieSceneSequence& Sequence) override
+		FMovieSceneEvaluationTemplate& AccessTemplate(UMovieSceneSequence&) override
 		{
 			return Template;
 		}
