Sequencer: Fixed discrepancy between UI ranges and compiled/evaluation ranges

  - Ranges with inclusive upper bounds in Sequencer UI are represented as exclusive upper bounds of the subsequent frame, but this special handling was not being matched by the compiler or evaluation logic. Rather than introduce this complexity into the compiler, the compiler now uses UMovieSceneSection::GetTrueRange so that the complete TRange API can be used at runtime.

diff --git a/Engine/Source/Runtime/MovieScene/Private/Compilation/MovieSceneCompiler.cpp b/Engine/Source/Runtime/MovieScene/Private/Compilation/MovieSceneCompiler.cpp
--- a/Engine/Source/Runtime/MovieScene/Private/Compilation/MovieSceneCompiler.cpp
+++ b/Engine/Source/Runtime/MovieScene/Private/Compilation/MovieSceneCompiler.cpp
@@ -277,7 +277,7 @@ void FMovieSceneCompiler::CompileRange(TRange<FFrameNumber> InGlobalRange, UMovi
 		// If the range we just compiled no longer overlaps the range we were asked to compile,
 		// Break out of the loop as all of our work is done. This will happen if there is a gap
 		// in the evaluation field that overlaps with the upper bound of InGlobalRange.
-		if (!MovieScene::DiscreteRangesOverlap(CompiledRange, InGlobalRange))
+		if (!CompiledRange.Overlaps(InGlobalRange))
 		{
 			break;
 		}
@@ -312,7 +312,7 @@ void FMovieSceneCompiler::CompileRange(TRange<FFrameNumber> InGlobalRange, UMovi
 
 		// If the range that we just compiled goes right up to the end of the gap, increment onto the
 		// next entry in the evaluation field iterator (which should be a populated range)
-		if (MovieScene::DiscreteUpperBoundsAreEquivalent(CompiledRange.GetUpperBound(), EmptySpaceRange.GetUpperBound()))
+		if (CompiledRange.GetUpperBound() == EmptySpaceRange.GetUpperBound())
 		{
 			++ExistingEvaluationFieldIter;
 		}
@@ -421,7 +421,7 @@ void FMovieSceneCompiler::GatherCompileOnTheFlyData(UMovieSceneSequence& InSeque
 	for ( ; SubSectionIt && SubSectionIt.Range().Overlaps(CompileClampIntersection); ++SubSectionIt)
 	{
 		TRange<FFrameNumber> ThisSegmentRangeRoot = Params.ClampRoot(SubSectionIt.Range() * Params.RootToSequenceTransform.Inverse());
-		if (MovieScene::DiscreteRangeIsEmpty(ThisSegmentRangeRoot))
+		if (ThisSegmentRangeRoot.IsEmpty())
 		{
 			continue;
 		}
@@ -569,7 +569,7 @@ void FMovieSceneCompiler::GatherCompileDataForTrack(FMovieSceneEvaluationTrack&
 		Data.bRequiresInit = ThisSegment.Impls.ContainsByPredicate(RequiresInit);
 
 		TRange<FFrameNumber> IntersectionRange = Params.ClampRoot(ThisSegment.Range * SequenceToRootTransform);
-		if (!MovieScene::DiscreteRangeIsEmpty(IntersectionRange))
+		if (!IntersectionRange.IsEmpty())
 		{
 			OutData.Tracks.Add(IntersectionRange, Data);
 		}
diff --git a/Engine/Source/Runtime/MovieScene/Private/Compilation/MovieSceneEvaluationTemplateGenerator.cpp b/Engine/Source/Runtime/MovieScene/Private/Compilation/MovieSceneEvaluationTemplateGenerator.cpp
--- a/Engine/Source/Runtime/MovieScene/Private/Compilation/MovieSceneEvaluationTemplateGenerator.cpp
+++ b/Engine/Source/Runtime/MovieScene/Private/Compilation/MovieSceneEvaluationTemplateGenerator.cpp
@@ -104,55 +104,56 @@ void FMovieSceneEvaluationTemplateGenerator::ProcessSubTrack(const UMovieSceneSu
 	for (int32 SectionIndex = 0; SectionIndex < AllSections.Num(); ++SectionIndex)
 	{
 		UMovieSceneSubSection* SubSection = CastChecked<UMovieSceneSubSection>(AllSections[SectionIndex]);
-		if (!SubSection || !SubSection->IsActive() || MovieScene::DiscreteRangeIsEmpty(SubSection->GetRange()))
+		if (!SubSection || !SubSection->IsActive())
+		{
+			continue;
+		}
+
+		const TRange<FFrameNumber> SectionRange = SubSection->GetTrueRange();
+		if (SectionRange.IsEmpty())
 		{
 			continue;
 		}
 
 		CompiledSignatures.Add(SubSection->GetSignature());
 
+		// Process the actual section range
 		if (bRequiresSegmentBlending)
 		{
-			const TRange<FFrameNumber> SectionRange = SubSection->GetRange();
+			SubSectionBlendTree.Add(SectionRange, FSectionEvaluationData(SectionIndex, ESectionEvaluationFlags::None));
+		}
+		else
+		{
+			Template.AddSubSectionRange(*SubSection, ObjectBindingId, SectionRange, ESectionEvaluationFlags::None);
+		}
+
+		// Process the section preroll range
+		if (!SectionRange.GetLowerBound().IsOpen() && SubSection->GetPreRollFrames() > 0)
+		{
+			TRange<FFrameNumber> PreRollRange(SectionRange.GetLowerBoundValue() - SubSection->GetPreRollFrames(), TRangeBound<FFrameNumber>::FlipInclusion(SectionRange.GetLowerBoundValue()));
 
-			// Process the actual section range
 			if (bRequiresSegmentBlending)
 			{
-				SubSectionBlendTree.Add(SectionRange, FSectionEvaluationData(SectionIndex, ESectionEvaluationFlags::None));
+				SubSectionBlendTree.Add(PreRollRange, FSectionEvaluationData(SectionIndex, ESectionEvaluationFlags::PreRoll));
 			}
 			else
 			{
-				Template.AddSubSectionRange(*SubSection, ObjectBindingId, SectionRange, ESectionEvaluationFlags::None);
+				Template.AddSubSectionRange(*SubSection, ObjectBindingId, PreRollRange, ESectionEvaluationFlags::PreRoll);
 			}
+		}
+
+		// Process the section postroll range
+		if (!SectionRange.GetUpperBound().IsOpen() && SubSection->GetPostRollFrames() > 0)
+		{
+			TRange<FFrameNumber> PostRollRange(TRangeBound<FFrameNumber>::FlipInclusion(SectionRange.GetUpperBoundValue()), SectionRange.GetUpperBoundValue() + SubSection->GetPostRollFrames());
 
-			// Process the section preroll range
-			if (!SectionRange.GetLowerBound().IsOpen() && SubSection->GetPreRollFrames() > 0)
+			if (bRequiresSegmentBlending)
 			{
-				TRange<FFrameNumber> PreRollRange(SectionRange.GetLowerBoundValue() - SubSection->GetPreRollFrames(), TRangeBound<FFrameNumber>::FlipInclusion(SectionRange.GetLowerBoundValue()));
-
-				if (bRequiresSegmentBlending)
-				{
-					SubSectionBlendTree.Add(PreRollRange, FSectionEvaluationData(SectionIndex, ESectionEvaluationFlags::PreRoll));
-				}
-				else
-				{
-					Template.AddSubSectionRange(*SubSection, ObjectBindingId, PreRollRange, ESectionEvaluationFlags::PreRoll);
-				}
+				SubSectionBlendTree.Add(PostRollRange, FSectionEvaluationData(SectionIndex, ESectionEvaluationFlags::PostRoll));
 			}
-
-			// Process the section postroll range
-			if (!SectionRange.GetUpperBound().IsOpen() && SubSection->GetPostRollFrames() > 0)
+			else
 			{
-				TRange<FFrameNumber> PostRollRange(TRangeBound<FFrameNumber>::FlipInclusion(SectionRange.GetUpperBoundValue()), SectionRange.GetUpperBoundValue() + SubSection->GetPostRollFrames());
-
-				if (bRequiresSegmentBlending)
-				{
-					SubSectionBlendTree.Add(PostRollRange, FSectionEvaluationData(SectionIndex, ESectionEvaluationFlags::PostRoll));
-				}
-				else
-				{
-					Template.AddSubSectionRange(*SubSection, ObjectBindingId, PostRollRange, ESectionEvaluationFlags::PostRoll);
-				}
+				Template.AddSubSectionRange(*SubSection, ObjectBindingId, PostRollRange, ESectionEvaluationFlags::PostRoll);
 			}
 		}
 	}
diff --git a/Engine/Source/Runtime/MovieScene/Private/Evaluation/MovieSceneEvaluationField.cpp b/Engine/Source/Runtime/MovieScene/Private/Evaluation/MovieSceneEvaluationField.cpp
--- a/Engine/Source/Runtime/MovieScene/Private/Evaluation/MovieSceneEvaluationField.cpp
+++ b/Engine/Source/Runtime/MovieScene/Private/Evaluation/MovieSceneEvaluationField.cpp
@@ -59,7 +59,7 @@ TRange<int32> FMovieSceneEvaluationField::ConditionallyCompileRange(const TRange
 		{
 			bIsDirty = true;
 
-			if (!MovieScene::DiscreteRangeIsEmpty(InvalidatedSubSequenceRange))
+			if (!InvalidatedSubSequenceRange.IsEmpty())
 			{
 				// Invalidate this evaluation field
 				RangesToInvalidate.Add(InvalidatedSubSequenceRange);
@@ -120,7 +120,7 @@ TRange<int32> FMovieSceneEvaluationField::OverlapRange(const TRange<FFrameNumber
 	int32 StartIndex = Algo::UpperBoundBy(RangesToSearch, Range.GetLowerBound(), &FMovieSceneFrameRange::GetLowerBound, MovieSceneHelpers::SortLowerBounds);
 
 	// StartIndex is always <= RangesToSearch.Num(). If the previous range overlaps the input range, include that
-	if (StartIndex > 0 && MovieScene::DiscreteRangesOverlap(RangesToSearch[StartIndex-1].Value, Range))
+	if (StartIndex > 0 && RangesToSearch[StartIndex-1].Value.Overlaps(Range))
 	{
 		StartIndex = StartIndex - 1;
 	}
@@ -138,7 +138,7 @@ TRange<int32> FMovieSceneEvaluationField::OverlapRange(const TRange<FFrameNumber
 	int32 Length = Range.GetUpperBound().IsOpen() ? RangesToSearch.Num() : Algo::UpperBoundBy(RangesToSearch, Range.GetUpperBound(), &FMovieSceneFrameRange::GetUpperBound, MovieSceneHelpers::SortUpperBounds);
 
 	// Length is always <= RangesToSearch.Num(). If the previous range overlaps the input range, include that
-	if (Length < RangesToSearch.Num() && MovieScene::DiscreteRangesOverlap(RangesToSearch[Length].Value, Range))
+	if (Length < RangesToSearch.Num() && RangesToSearch[Length].Value.Overlaps(Range))
 	{
 		Length = Length + 1;
 	}
@@ -166,8 +166,8 @@ int32 FMovieSceneEvaluationField::Insert(const TRange<FFrameNumber>& InRange, FM
 	const int32 InsertIndex = Algo::UpperBoundBy(Ranges, InRange.GetLowerBound(), &FMovieSceneFrameRange::GetLowerBound, MovieSceneHelpers::SortLowerBounds);
 
 	const bool bOverlapping = 
-		(Ranges.IsValidIndex(InsertIndex  ) && MovieScene::DiscreteRangesOverlap(Ranges[InsertIndex  ].Value, InRange)) ||
-		(Ranges.IsValidIndex(InsertIndex-1) && MovieScene::DiscreteRangesOverlap(Ranges[InsertIndex-1].Value, InRange));
+		(Ranges.IsValidIndex(InsertIndex  ) && Ranges[InsertIndex  ].Value.Overlaps(InRange)) ||
+		(Ranges.IsValidIndex(InsertIndex-1) && Ranges[InsertIndex-1].Value.Overlaps(InRange));
 
 	if (!ensureAlwaysMsgf(!bOverlapping, TEXT("Attempting to insert an overlapping range into the evaluation field.")))
 	{
@@ -187,7 +187,7 @@ int32 FMovieSceneEvaluationField::Insert(const TRange<FFrameNumber>& InRange, FM
 
 void FMovieSceneEvaluationField::Add(const TRange<FFrameNumber>& InRange, FMovieSceneEvaluationGroup&& InGroup, FMovieSceneEvaluationMetaData&& InMetaData)
 {
-	if (ensureAlwaysMsgf(!Ranges.Num() || !MovieScene::DiscreteRangesOverlap(Ranges.Last().Value, InRange), TEXT("Attempting to add overlapping ranges to sequence evaluation field.")))
+	if (ensureAlwaysMsgf(!Ranges.Num() || !Ranges.Last().Value.Overlaps(InRange), TEXT("Attempting to add overlapping ranges to sequence evaluation field.")))
 	{
 		Ranges.Add(InRange);
 		MetaData.Add(MoveTemp(InMetaData));
diff --git a/Engine/Source/Runtime/MovieScene/Private/Evaluation/MovieSceneEvaluationTrack.cpp b/Engine/Source/Runtime/MovieScene/Private/Evaluation/MovieSceneEvaluationTrack.cpp
--- a/Engine/Source/Runtime/MovieScene/Private/Evaluation/MovieSceneEvaluationTrack.cpp
+++ b/Engine/Source/Runtime/MovieScene/Private/Evaluation/MovieSceneEvaluationTrack.cpp
@@ -75,8 +75,8 @@ FMovieSceneSegmentIdentifier FMovieSceneEvaluationTrackSegments::Add(FMovieScene
 		InsertIndex = FMath::Clamp(InsertIndex, 0, SortedSegments.Num());
 	}
 
-	if (!ensureMsgf(   ( !SortedSegments.IsValidIndex(InsertIndex  ) || !MovieScene::DiscreteRangesOverlap(SortedSegments[InsertIndex  ].Range, Segment.Range) ) && 
-			           ( !SortedSegments.IsValidIndex(InsertIndex+1) || !MovieScene::DiscreteRangesOverlap(SortedSegments[InsertIndex+1].Range, Segment.Range) ), TEXT("Attempting to add overlapping segment to segment array. This is invalid.")))
+	if (!ensureMsgf(   ( !SortedSegments.IsValidIndex(InsertIndex  ) || !SortedSegments[InsertIndex  ].Range.Overlaps(Segment.Range) ) && 
+			           ( !SortedSegments.IsValidIndex(InsertIndex+1) || !SortedSegments[InsertIndex+1].Range.Overlaps(Segment.Range) ), TEXT("Attempting to add overlapping segment to segment array. This is invalid.")))
 	{
 		return FMovieSceneSegmentIdentifier();
 	}
@@ -274,7 +274,7 @@ FMovieSceneSegmentIdentifier FMovieSceneEvaluationTrack::FindFirstSegment(TRange
 	{
 		return SortedSegments[SegmentIndex].ID;
 	}
-	else if (SortedSegments.IsValidIndex(SegmentIndex-1) && MovieScene::DiscreteRangesOverlap(SortedSegments[SegmentIndex-1].Range, InLocalRange))
+	else if (SortedSegments.IsValidIndex(SegmentIndex-1) && SortedSegments[SegmentIndex-1].Range.Overlaps(InLocalRange))
 	{
 		return SortedSegments[SegmentIndex-1].ID;
 	}
@@ -286,7 +286,7 @@ TArray<FMovieSceneSegmentIdentifier> FMovieSceneEvaluationTrack::GetSegmentsInRa
 {
 	TArray<FMovieSceneSegmentIdentifier> SegmentsInRange;
 
-	while (!MovieScene::DiscreteRangeIsEmpty(InLocalRange))
+	while (!InLocalRange.IsEmpty())
 	{
 		FMovieSceneSegmentIdentifier SegmentID = FindFirstSegment(InLocalRange);
 
@@ -470,7 +470,7 @@ namespace
 	bool IntersectSegmentRanges(const FMovieSceneSegment& Segment, TRange<FFrameNumber> TraversedRange, TMap<int32, TRange<FFrameNumber>>& ImplToAccumulatedRange)
 	{
 		TRange<FFrameNumber> Intersection = TRange<FFrameNumber>::Intersection(Segment.Range, TraversedRange);
-		if (MovieScene::DiscreteRangeIsEmpty(Intersection))
+		if (Intersection.IsEmpty())
 		{
 			return false;
 		}
diff --git a/Engine/Source/Runtime/MovieScene/Private/Evaluation/MovieSceneEvaluationTree.cpp b/Engine/Source/Runtime/MovieScene/Private/Evaluation/MovieSceneEvaluationTree.cpp
--- a/Engine/Source/Runtime/MovieScene/Private/Evaluation/MovieSceneEvaluationTree.cpp
+++ b/Engine/Source/Runtime/MovieScene/Private/Evaluation/MovieSceneEvaluationTree.cpp
@@ -65,11 +65,11 @@ FMovieSceneEvaluationTreeRangeIterator::FMovieSceneEvaluationTreeRangeIterator(c
 
 		// Binary search children's lower bounds for the first that's >= the starting bound. That results in either ChildIndex or ChildIndex-1 being the overlapping range (if any).
 		int32 ChildIndex = Algo::LowerBoundBy(Children, StartingBound, GetLowerBound, MovieSceneHelpers::SortLowerBounds);
-		if (Children.IsValidIndex(ChildIndex) && MovieScene::DiscreteLowerBoundsAreEquivalent(Children[ChildIndex].Range.GetLowerBound(), StartingBound))
+		if (Children.IsValidIndex(ChildIndex) && Children[ChildIndex].Range.GetLowerBound() == StartingBound)
 		{
 			CurrentNodeHandle = FMovieSceneEvaluationTreeNodeHandle(ThisNode.ChildrenID, ChildIndex);
 		}
-		else if (Children.IsValidIndex(ChildIndex-1) && MovieScene::DiscreteRangesOverlap(Children[ChildIndex-1].Range, CompareRange))
+		else if (Children.IsValidIndex(ChildIndex-1) && Children[ChildIndex-1].Range.Overlaps(CompareRange))
 		{
 			CurrentNodeHandle = FMovieSceneEvaluationTreeNodeHandle(ThisNode.ChildrenID, ChildIndex - 1);
 		}
@@ -85,7 +85,7 @@ FMovieSceneEvaluationTreeRangeIterator::FMovieSceneEvaluationTreeRangeIterator(c
 		}
 	}
 
-	check(MovieScene::DiscreteRangesOverlap(CurrentRange, CompareRange));
+	check(CurrentRange.Overlaps(CompareRange));
 }
 
 void FMovieSceneEvaluationTreeRangeIterator::Iter(bool bForwards)
@@ -104,7 +104,7 @@ void FMovieSceneEvaluationTreeRangeIterator::Iter(bool bForwards)
 		return;
 	}
 
-	TRangeBound<FFrameNumber> NewLeadingBound = MovieScene::DiscreteRangeIsEmpty(CurrentRange) ? TRangeBound<FFrameNumber>::Open() : TRangeBound<FFrameNumber>::FlipInclusion(GetTrailingBound(bForwards, CurrentRange));
+	TRangeBound<FFrameNumber> NewLeadingBound = CurrentRange.IsEmpty() ? TRangeBound<FFrameNumber>::Open() : TRangeBound<FFrameNumber>::FlipInclusion(GetTrailingBound(bForwards, CurrentRange));
 
 	// Iterate into children when possible (where the leading bound matches the current new leading bound)
 	FMovieSceneEvaluationTreeNodeHandle NextChildHandle = FindNextChild(CurrentNodeHandle, NewLeadingBound, bForwards);
@@ -220,7 +220,7 @@ void FMovieSceneEvaluationTree::AddTimeRange(TRange<FFrameNumber> InTimeRange, c
 {
 	// Take a temporary copy of the node as the container may be reallocated in this function
 	FMovieSceneEvaluationTreeNode ThisNode = GetNode(InParent);
-	if (!ensure(MovieScene::DiscreteRangesOverlap(ThisNode.Range, InTimeRange)))
+	if (!ensure(ThisNode.Range.Overlaps(InTimeRange)))
 	{
 		return;
 	}
@@ -253,7 +253,7 @@ void FMovieSceneEvaluationTree::AddTimeRange(TRange<FFrameNumber> InTimeRange, c
 				TRange<FFrameNumber> PrecedingSpace(LastBound, TRangeBound<FFrameNumber>::FlipInclusion(ChildNode.Range.GetLowerBound()));
 				TRange<FFrameNumber> PrecedingIntersection = TRange<FFrameNumber>::Intersection(PrecedingSpace, InTimeRange);
 
-				if (!MovieScene::DiscreteRangeIsEmpty(PrecedingIntersection))
+				if (!PrecedingIntersection.IsEmpty())
 				{
 					// Insert this child at the current child index
 					InsertNewChild(PrecedingIntersection, InOperator, ChildIndex, InParent);
@@ -263,7 +263,7 @@ void FMovieSceneEvaluationTree::AddTimeRange(TRange<FFrameNumber> InTimeRange, c
 				}
 			}
 
-			if (MovieScene::DiscreteRangesOverlap(ChildNode.Range, InTimeRange))
+			if (ChildNode.Range.Overlaps(InTimeRange))
 			{
 				// Find the node again since InsertNewChild may have re-allocated the nodes
 				AddTimeRange(InTimeRange, InOperator, FMovieSceneEvaluationTreeNodeHandle(ThisNode.ChildrenID, ChildIndex));
@@ -278,7 +278,7 @@ void FMovieSceneEvaluationTree::AddTimeRange(TRange<FFrameNumber> InTimeRange, c
 			TRange<FFrameNumber> ProceedingSpace(LastBound, ThisNode.Range.GetUpperBound());
 			TRange<FFrameNumber> ProceedingIntersection = TRange<FFrameNumber>::Intersection(ProceedingSpace, InTimeRange);
 
-			if (!MovieScene::DiscreteRangeIsEmpty(ProceedingIntersection))
+			if (!ProceedingIntersection.IsEmpty())
 			{
 				// Insert this child at the end
 				InsertNewChild(ProceedingIntersection, InOperator, ChildNodes.Get(ThisNode.ChildrenID).Num(), InParent);
diff --git a/Engine/Source/Runtime/MovieScene/Private/Evaluation/MovieSceneSequenceHierarchy.cpp b/Engine/Source/Runtime/MovieScene/Private/Evaluation/MovieSceneSequenceHierarchy.cpp
--- a/Engine/Source/Runtime/MovieScene/Private/Evaluation/MovieSceneSequenceHierarchy.cpp
+++ b/Engine/Source/Runtime/MovieScene/Private/Evaluation/MovieSceneSequenceHierarchy.cpp
@@ -30,12 +30,9 @@ FMovieSceneSubSequenceData::FMovieSceneSubSequenceData(const UMovieSceneSubSecti
 
 	TickResolution = MovieScenePtr->GetTickResolution();
 
-	TRange<FFrameNumber> SubRange = InSubSection.GetRange();
+	TRange<FFrameNumber> SubRange = InSubSection.GetTrueRange();
 	checkf(SubRange.GetLowerBound().IsClosed() && SubRange.GetUpperBound().IsClosed(), TEXT("Use of open (infinite) bounds with sub sections is not supported."));
 
-	const FFrameNumber OuterStartTime = MovieScene::DiscreteInclusiveLower(SubRange);
-	const FFrameNumber OuterEndTime   = MovieScene::DiscreteExclusiveUpper(SubRange);
-
 	OuterToInnerTransform = RootToSequenceTransform = InSubSection.OuterToInnerTransform();
 
 	PlayRange.Value = SubRange * RootToSequenceTransform;
diff --git a/Engine/Source/Runtime/MovieScene/Private/MovieSceneTrack.cpp b/Engine/Source/Runtime/MovieScene/Private/MovieSceneTrack.cpp
--- a/Engine/Source/Runtime/MovieScene/Private/MovieSceneTrack.cpp
+++ b/Engine/Source/Runtime/MovieScene/Private/MovieSceneTrack.cpp
@@ -225,8 +225,8 @@ EMovieSceneCompileResult UMovieSceneTrack::Compile(FMovieSceneEvaluationTrack& O
 	{
 		for (const UMovieSceneSection* Section : GetAllSections())
 		{
-			const TRange<FFrameNumber> SectionRange = Section->GetRange();
-			if (!Section->IsActive() || MovieScene::DiscreteRangeIsEmpty(SectionRange))
+			const TRange<FFrameNumber> SectionRange = Section->GetTrueRange();
+			if (!Section->IsActive() || SectionRange.IsEmpty())
 			{
 				continue;
 			}
diff --git a/Engine/Source/Runtime/MovieScene/Private/Tests/MovieSceneCompilerTests.cpp b/Engine/Source/Runtime/MovieScene/Private/Tests/MovieSceneCompilerTests.cpp
--- a/Engine/Source/Runtime/MovieScene/Private/Tests/MovieSceneCompilerTests.cpp
+++ b/Engine/Source/Runtime/MovieScene/Private/Tests/MovieSceneCompilerTests.cpp
@@ -92,7 +92,7 @@ IMPLEMENT_SIMPLE_AUTOMATION_TEST(FMovieSceneCompilerRangeTest, "System.Engine.Se
 bool FMovieSceneCompilerRangeTest::RunTest(const FString& Parameters)
 {
 	FFrameNumber CompileAtTimes[] = {
-		-3, -2, -1, 0, 1, 2, 3, 9, 10, 11, 12, 13, 14, 15, 16, 17, 18, 19, 20
+		-3, -2, -1, 0, 1, 2, 3, 9, 10, 11, 12, 13, 14, 15, 16, 17, 18, 19, 20, 21, 22, 23, 24,
 	};
 	// Test each combination of inc/excl boundary conditions for adjacent and adjoining ranges
 	TRange<FFrameNumber> Ranges[] = {
@@ -140,6 +140,10 @@ bool FMovieSceneCompilerRangeTest::RunTest(const FString& Parameters)
 		MakeRange( 13, 18, ERangeBoundTypes::Inclusive, ERangeBoundTypes::Exclusive),
 		MakeRange( 13, 19, ERangeBoundTypes::Inclusive, ERangeBoundTypes::Inclusive),
 		MakeRange( 13, 18, ERangeBoundTypes::Inclusive, ERangeBoundTypes::Inclusive),
+
+		// Explicitly test two adjacent ranges that would produce effectively empty space in between them when iterating
+		MakeRange( 21, 22, ERangeBoundTypes::Inclusive, ERangeBoundTypes::Inclusive),
+		MakeRange( 23, 24, ERangeBoundTypes::Inclusive, ERangeBoundTypes::Inclusive),
 	};
 
 	UTestMovieSceneSequence* Sequence = NewObject<UTestMovieSceneSequence>(GetTransientPackage());
diff --git a/Engine/Source/Runtime/MovieScene/Private/Tests/MovieSceneSegmentCompilerTests.cpp b/Engine/Source/Runtime/MovieScene/Private/Tests/MovieSceneSegmentCompilerTests.cpp
--- a/Engine/Source/Runtime/MovieScene/Private/Tests/MovieSceneSegmentCompilerTests.cpp
+++ b/Engine/Source/Runtime/MovieScene/Private/Tests/MovieSceneSegmentCompilerTests.cpp
@@ -298,7 +298,7 @@ bool FMovieSceneTrackCompilerTest::RunTest(const FString& Parameters)
 		UTestMovieSceneSection* Section1 = NewObject<UTestMovieSceneSection>(Track);
 		Section1->SetRange(TRange<FFrameNumber>(
 			TRangeBound<FFrameNumber>::Inclusive(20),
-			TRangeBound<FFrameNumber>::Inclusive(30)
+			TRangeBound<FFrameNumber>::Exclusive(30)
 		));
 		Section1->SetRowIndex(1);
 		
@@ -315,7 +315,7 @@ bool FMovieSceneTrackCompilerTest::RunTest(const FString& Parameters)
 			FMovieSceneSegment Expected[] = {
 				FMovieSceneSegment(TRange<FFrameNumber>(TRangeBound<FFrameNumber>::Inclusive(10),	TRangeBound<FFrameNumber>::Exclusive(20)), { FSectionEvaluationData(0) }),
 				FMovieSceneSegment(TRange<FFrameNumber>(TRangeBound<FFrameNumber>::Inclusive(20),	TRangeBound<FFrameNumber>::Exclusive(25)), { FSectionEvaluationData(0), FSectionEvaluationData(1) }),
-				FMovieSceneSegment(TRange<FFrameNumber>(TRangeBound<FFrameNumber>::Inclusive(25),	TRangeBound<FFrameNumber>::Inclusive(30)), { FSectionEvaluationData(1) }),
+				FMovieSceneSegment(TRange<FFrameNumber>(TRangeBound<FFrameNumber>::Inclusive(25),	TRangeBound<FFrameNumber>::Exclusive(30)), { FSectionEvaluationData(1) }),
 			};
 			AssertSegmentValues(this, Expected, EvalTrack.GetSortedSegments());
 		}
@@ -331,8 +331,8 @@ bool FMovieSceneTrackCompilerTest::RunTest(const FString& Parameters)
 				FMovieSceneSegment(TRange<FFrameNumber>(Inf,										TRangeBound<FFrameNumber>::Exclusive(10)),	{ FSectionEvaluationData(0, 10) }),
 				FMovieSceneSegment(TRange<FFrameNumber>(TRangeBound<FFrameNumber>::Inclusive(10),	TRangeBound<FFrameNumber>::Exclusive(20)),	{ FSectionEvaluationData(0) }),
 				FMovieSceneSegment(TRange<FFrameNumber>(TRangeBound<FFrameNumber>::Inclusive(20),	TRangeBound<FFrameNumber>::Exclusive(25)),	{ FSectionEvaluationData(0), FSectionEvaluationData(1) }),
-				FMovieSceneSegment(TRange<FFrameNumber>(TRangeBound<FFrameNumber>::Inclusive(25),	TRangeBound<FFrameNumber>::Inclusive(30)),	{ FSectionEvaluationData(1) }),
-				FMovieSceneSegment(TRange<FFrameNumber>(TRangeBound<FFrameNumber>::Exclusive(30),	Inf), 										{ FSectionEvaluationData(1, 30) }),
+				FMovieSceneSegment(TRange<FFrameNumber>(TRangeBound<FFrameNumber>::Inclusive(25),	TRangeBound<FFrameNumber>::Exclusive(30)),	{ FSectionEvaluationData(1) }),
+				FMovieSceneSegment(TRange<FFrameNumber>(TRangeBound<FFrameNumber>::Inclusive(30),	Inf), 										{ FSectionEvaluationData(1, 30) }),
 			};
 			AssertSegmentValues(this, Expected, EvalTrack.GetSortedSegments());
 		}
@@ -347,7 +347,7 @@ bool FMovieSceneTrackCompilerTest::RunTest(const FString& Parameters)
 			FMovieSceneSegment Expected[] = {
 				FMovieSceneSegment(TRange<FFrameNumber>(TRangeBound<FFrameNumber>::Inclusive(10),	TRangeBound<FFrameNumber>::Exclusive(20)), { FSectionEvaluationData(0) }),
 				FMovieSceneSegment(TRange<FFrameNumber>(TRangeBound<FFrameNumber>::Inclusive(20),	TRangeBound<FFrameNumber>::Exclusive(25)), { FSectionEvaluationData(0), FSectionEvaluationData(1) }),
-				FMovieSceneSegment(TRange<FFrameNumber>(TRangeBound<FFrameNumber>::Inclusive(25),	TRangeBound<FFrameNumber>::Inclusive(30)), { FSectionEvaluationData(1) }),
+				FMovieSceneSegment(TRange<FFrameNumber>(TRangeBound<FFrameNumber>::Inclusive(25),	TRangeBound<FFrameNumber>::Exclusive(30)), { FSectionEvaluationData(1) }),
 			};
 			AssertSegmentValues(this, Expected, EvalTrack.GetSortedSegments());
 		}
@@ -362,7 +362,7 @@ bool FMovieSceneTrackCompilerTest::RunTest(const FString& Parameters)
 
 			FMovieSceneSegment Expected[] = {
 				FMovieSceneSegment(TRange<FFrameNumber>(TRangeBound<FFrameNumber>::Inclusive(10),	TRangeBound<FFrameNumber>::Exclusive(25)), { FSectionEvaluationData(0) }),
-				FMovieSceneSegment(TRange<FFrameNumber>(TRangeBound<FFrameNumber>::Inclusive(25),	TRangeBound<FFrameNumber>::Inclusive(30)), { FSectionEvaluationData(1) }),
+				FMovieSceneSegment(TRange<FFrameNumber>(TRangeBound<FFrameNumber>::Inclusive(25),	TRangeBound<FFrameNumber>::Exclusive(30)), { FSectionEvaluationData(1) }),
 			};
 			AssertSegmentValues(this, Expected, EvalTrack.GetSortedSegments());
 		}
diff --git a/Engine/Source/Runtime/MovieScene/Public/MovieSceneSection.h b/Engine/Source/Runtime/MovieScene/Public/MovieSceneSection.h
--- a/Engine/Source/Runtime/MovieScene/Public/MovieSceneSection.h
+++ b/Engine/Source/Runtime/MovieScene/Public/MovieSceneSection.h
@@ -167,6 +167,28 @@ public:
 		return SectionRange.Value;
 	}
 
+	/**
+	 * A true representation of this section's range with an inclusive start frame and an exclusive end frame.
+	 * The resulting range defines that the section lies between { lower <= time < upper }
+	 */
+	TRange<FFrameNumber> GetTrueRange() const
+	{
+		TRangeBound<FFrameNumber> SectionLower = SectionRange.Value.GetLowerBound();
+		TRangeBound<FFrameNumber> SectionUpper = SectionRange.Value.GetUpperBound();
+
+		// Make exclusive lower bounds inclusive on the next frame
+		if (SectionLower.IsExclusive())
+		{
+			SectionLower = TRangeBound<FFrameNumber>::Inclusive(SectionLower.GetValue() + 1);
+		}
+		// Make inclusive upper bounds exclusive on the next frame
+		if (SectionUpper.IsInclusive())
+		{
+			SectionUpper = TRangeBound<FFrameNumber>::Exclusive(SectionUpper.GetValue() + 1);
+		}
+		return TRange<FFrameNumber>(SectionLower, SectionUpper);
+	}
+
 	/**
 	 * Expands this section's range to include the specified time
 	 */
diff --git a/Engine/Source/Runtime/MovieScene/Public/MovieSceneTimeHelpers.h b/Engine/Source/Runtime/MovieScene/Public/MovieSceneTimeHelpers.h
--- a/Engine/Source/Runtime/MovieScene/Public/MovieSceneTimeHelpers.h
+++ b/Engine/Source/Runtime/MovieScene/Public/MovieSceneTimeHelpers.h
@@ -101,78 +101,6 @@ inline int32 DiscreteSize(const TRange<FFrameNumber>& InRange)
 	return (int64)DiscreteExclusiveUpper(InRange).Value - (int64)DiscreteInclusiveLower(InRange).Value;
 }
 
-/**
- * Check whether the two specified discrete ranges share any common discrete values
- * 
- * @return The size of the range (considering inclusive and exclusive boundaries)
- */
-inline bool DiscreteRangesOverlap(const TRange<FFrameNumber>& A, const TRange<FFrameNumber>& B)
-{
-	TRange<FFrameNumber> Intersection = TRange<FFrameNumber>::Intersection(A, B);
-	if (Intersection.GetLowerBound().IsOpen() || Intersection.GetUpperBound().IsOpen())
-	{
-		return true;
-	}
-	else
-	{
-		return DiscreteSize(Intersection) > 0;
-	}
-}
-
-/**
- * Check whether the two specified lower bounds are equivalent in a discrete domain
- */
-inline bool DiscreteLowerBoundsAreEquivalent(const TRangeBound<FFrameNumber>& A, const TRangeBound<FFrameNumber>& B)
-{
-	if (A.IsOpen() || B.IsOpen())
-	{
-		// Either bound is open
-		return A.IsOpen() && B.IsOpen();
-	}
-	else if (A.IsExclusive() == B.IsExclusive())
-	{
-		// Both inclusive or both exclusive
-		return A.GetValue() == B.GetValue();
-	}
-	else if (A.IsExclusive())
-	{
-		// A exclusive, B inclusive
-		return A.GetValue() == B.GetValue()-1;
-	}
-	else
-	{
-		// A inclusive, B exclusive
-		return B.GetValue() == A.GetValue()-1;
-	}
-}
-
-/**
- * Check whether the two specified upper bounds are equivalent in a discrete domain
- */
-inline bool DiscreteUpperBoundsAreEquivalent(const TRangeBound<FFrameNumber>& A, const TRangeBound<FFrameNumber>& B)
-{
-	if (A.IsOpen() || B.IsOpen())
-	{
-		// Either bound is open
-		return A.IsOpen() && B.IsOpen();
-	}
-	else if (A.IsExclusive() == B.IsExclusive())
-	{
-		// Both inclusive or both exclusive
-		return A.GetValue() == B.GetValue();
-	}
-	else if (A.IsExclusive())
-	{
-		// A exclusive, B inclusive
-		return A.GetValue() == B.GetValue()+1;
-	}
-	else
-	{
-		// A inclusive, B exclusive
-		return B.GetValue() == A.GetValue()+1;
-	}
-}
-
 /**
  * Check whether the specified range contains any integer frame numbers or not
  */
