Sequencer: Fixed compilation of tracks with empty space where a compile range overlapped both empty space and an existing compiled segment.

The issue is that previously we only reported empty space when compiling a specific _time_, which meant that the inference that empty space must exist where no overlapping track segments existed was correct. However, with the shift to compiling the entire evaluated range, this inference is incorrect and a more robust approach is required. We now iterate the unique segments in the track for the whole compilation range in exactly the same way that we iterate sub sections, adding empty space only where no compiled segments exist (or could be compiled) in the track.

Also added a new test (Empty Space On The Fly) to verify this failure case, which failed under the old code.

diff --git a/Engine/Source/Runtime/MovieScene/Private/Compilation/MovieSceneCompiler.cpp b/Engine/Source/Runtime/MovieScene/Private/Compilation/MovieSceneCompiler.cpp
--- a/Engine/Source/Runtime/MovieScene/Private/Compilation/MovieSceneCompiler.cpp
+++ b/Engine/Source/Runtime/MovieScene/Private/Compilation/MovieSceneCompiler.cpp
@@ -544,34 +544,37 @@ void FMovieSceneCompiler::GatherCompileDataForTrack(FMovieSceneEvaluationTrack&
 		return Track.HasChildTemplate(EvalData.ImplIndex) && Track.GetChildTemplate(EvalData.ImplIndex).RequiresInitialization();
 	};
 
-	FMovieSceneSequenceTransform SequenceToRootTransform = Params.RootToSequenceTransform.Inverse();
-	FMovieSceneSequenceID CurrentSequenceID              = Params.RootPath.Remap(MovieSceneSequenceID::Root);
+	FMovieSceneSequenceTransform SequenceToRootTransform  = Params.RootToSequenceTransform.Inverse();
+	FMovieSceneSequenceID        CurrentSequenceID        = Params.RootPath.Remap(MovieSceneSequenceID::Root);
+	TRange<FFrameNumber>         CompileClampIntersection = TRange<FFrameNumber>::Intersection(Params.LocalCompileRange, Params.LocalClampRange);
 
-	TArray<FMovieSceneSegmentIdentifier> SegmentIDs      = Track.GetSegmentsInRange(Params.LocalCompileRange);
-	if (SegmentIDs.Num() == 0)
+	FMovieSceneEvaluationTreeRangeIterator TrackIter = Track.IterateFrom(CompileClampIntersection.GetLowerBound());
+	for ( ; TrackIter && TrackIter.Range().Overlaps(CompileClampIntersection); ++TrackIter)
 	{
-		// No segment at this time, so just report the time range of the empty space.
-		TRange<FFrameNumber> EmptyTrackSpace = Track.GetUniqueRangeFromLowerBound(Params.LocalCompileRange.GetLowerBound());
-		TRange<FFrameNumber> ClampedEmptyTrackSpaceRoot = Params.ClampRoot(EmptyTrackSpace * SequenceToRootTransform);
-
-		OutData.EmptySpace.AddTimeRange(ClampedEmptyTrackSpaceRoot);
-	}
-	else for (FMovieSceneSegmentIdentifier SegmentID : SegmentIDs)
-	{
-		const FMovieSceneSegment& ThisSegment = Track.GetSegment(SegmentID);
-
-		FCompileOnTheFlyData Data;
-		Data.Segment = FMovieSceneEvaluationFieldSegmentPtr(CurrentSequenceID, TrackID, SegmentID);
-		Data.GroupEvaluationPriority = GetMovieSceneModule().GetEvaluationGroupParameters(Track.GetEvaluationGroup()).EvaluationPriority;
-		Data.HierarchicalBias = Params.HierarchicalBias;
-		Data.EvaluationPriority = Track.GetEvaluationPriority();
-		Data.Track = &Track;
-		Data.bRequiresInit = ThisSegment.Impls.ContainsByPredicate(RequiresInit);
-
-		TRange<FFrameNumber> IntersectionRange = Params.ClampRoot(ThisSegment.Range * SequenceToRootTransform);
-		if (!IntersectionRange.IsEmpty())
+		FMovieSceneSegmentIdentifier SegmentID = Track.GetSegmentFromIterator(TrackIter);
+		if (!SegmentID.IsValid())
 		{
-			OutData.Tracks.Add(IntersectionRange, Data);
+			// No segment at this time, so just report the time range of the empty space.
+			TRange<FFrameNumber> ClampedEmptyTrackSpaceRoot = Params.ClampRoot(TrackIter.Range() * SequenceToRootTransform);
+			OutData.EmptySpace.AddTimeRange(ClampedEmptyTrackSpaceRoot);
+		}
+		else
+		{
+			const FMovieSceneSegment& ThisSegment = Track.GetSegment(SegmentID);
+
+			FCompileOnTheFlyData Data;
+			Data.Segment = FMovieSceneEvaluationFieldSegmentPtr(CurrentSequenceID, TrackID, SegmentID);
+			Data.GroupEvaluationPriority = GetMovieSceneModule().GetEvaluationGroupParameters(Track.GetEvaluationGroup()).EvaluationPriority;
+			Data.HierarchicalBias = Params.HierarchicalBias;
+			Data.EvaluationPriority = Track.GetEvaluationPriority();
+			Data.Track = &Track;
+			Data.bRequiresInit = ThisSegment.Impls.ContainsByPredicate(RequiresInit);
+
+			TRange<FFrameNumber> IntersectionRange = Params.ClampRoot(ThisSegment.Range * SequenceToRootTransform);
+			if (!IntersectionRange.IsEmpty())
+			{
+				OutData.Tracks.Add(IntersectionRange, Data);
+			}
 		}
 	}
 }
diff --git a/Engine/Source/Runtime/MovieScene/Private/Evaluation/MovieSceneEvaluationTrack.cpp b/Engine/Source/Runtime/MovieScene/Private/Evaluation/MovieSceneEvaluationTrack.cpp
--- a/Engine/Source/Runtime/MovieScene/Private/Evaluation/MovieSceneEvaluationTrack.cpp
+++ b/Engine/Source/Runtime/MovieScene/Private/Evaluation/MovieSceneEvaluationTrack.cpp
@@ -347,11 +347,6 @@ FMovieSceneSegmentIdentifier FMovieSceneEvaluationTrack::GetSegmentFromIterator(
 	return CompileSegment(Iterator);
 }
 
-TRange<FFrameNumber> FMovieSceneEvaluationTrack::GetUniqueRangeFromLowerBound(TRangeBound<FFrameNumber> InLowerBound) const
-{
-	return EvaluationTree.Tree.IterateFromLowerBound(InLowerBound).Range();
-}
-
 void FMovieSceneEvaluationTrack::ValidateSegments()
 {
 	// We don't remove segments as this may break ptrs that have been set up in the evaluation field. Instead we just remove invalid track indices.
diff --git a/Engine/Source/Runtime/MovieScene/Private/Tests/MovieSceneCompilerTests.cpp b/Engine/Source/Runtime/MovieScene/Private/Tests/MovieSceneCompilerTests.cpp
--- a/Engine/Source/Runtime/MovieScene/Private/Tests/MovieSceneCompilerTests.cpp
+++ b/Engine/Source/Runtime/MovieScene/Private/Tests/MovieSceneCompilerTests.cpp
@@ -9,6 +9,7 @@
 #include "Evaluation/MovieSceneEvaluationTrack.h"
 #include "Evaluation/MovieSceneEvaluationField.h"
 #include "Evaluation/MovieSceneSequenceTemplateStore.h"
+#include "Algo/Find.h"
 #include "UObject/Package.h"
 
 #if WITH_DEV_AUTOMATION_TESTS
@@ -183,4 +184,151 @@ bool FMovieSceneCompilerRangeTest::RunTest(const FString& Parameters)
 	return true;
 }
 
+
+IMPLEMENT_SIMPLE_AUTOMATION_TEST(FMovieSceneCompilerEmptySpaceOnTheFlyTest, "System.Engine.Sequencer.Compiler.Empty Space On The Fly", EAutomationTestFlags::EditorContext | EAutomationTestFlags::EngineFilter)
+bool FMovieSceneCompilerEmptySpaceOnTheFlyTest::RunTest(const FString& Parameters)
+{
+	// Tests that compiling ranges that contain empty space works correctly by verifying that the result evaluation field entries are either populated or empty as expected
+
+	struct FResult
+	{
+		TRange<FFrameNumber> FieldRange;
+		bool                 bExpectEmpty;
+	};
+
+	struct FTest
+	{
+		TArray<TRange<FFrameNumber>> CompileRanges;
+		TArray<FResult>              ExpectedResults;
+	};
+
+	TRange<FFrameNumber> SectionRanges[] = {
+		TRange<FFrameNumber>(0,  10),
+		TRange<FFrameNumber>(20, 30),
+		TRange<FFrameNumber>(40, 50),
+		TRange<FFrameNumber>(60, 70),
+	};
+
+	FResult ExpectedResults[] = {
+		{ TRange<FFrameNumber>(0,  10), false },
+		{ TRange<FFrameNumber>(10, 20), true  },
+		{ TRange<FFrameNumber>(20, 30), false },
+		{ TRange<FFrameNumber>(30, 40), true  },
+		{ TRange<FFrameNumber>(40, 50), false },
+		{ TRange<FFrameNumber>(50, 60), true  },
+		{ TRange<FFrameNumber>(60, 70), false },
+	};
+
+	TArray<FTest> Tests;
+
+	// Test 0: Test that compiling a range that only overlaps a section results in only that section's time being compiled
+	{
+		Tests.Emplace();
+		Tests.Last().CompileRanges   = { TRange<FFrameNumber>(5, 6) };
+		Tests.Last().ExpectedResults = { ExpectedResults[0] };
+	}
+
+	// Test 1: Test that compiling a range that overlaps both a section and empty space results in an entry for the section and the empty space
+	{
+		Tests.Emplace();
+		Tests.Last().CompileRanges   = { TRange<FFrameNumber>(6, 15) };
+		Tests.Last().ExpectedResults = { ExpectedResults[0], ExpectedResults[1] };
+	}
+
+	// Test 2: Test that compiling a range that only overlaps empty space works as expected
+	{
+		Tests.Emplace();
+		Tests.Last().CompileRanges   = { TRange<FFrameNumber>(14, 15) };
+		Tests.Last().ExpectedResults = { ExpectedResults[1] };
+	}
+
+	// Test 3: Test that compiling a section range followed by a range that overlaps both that section, and subsequent empty space compiles the empty space correctly
+	{
+		Tests.Emplace();
+		Tests.Last().CompileRanges   = { TRange<FFrameNumber>(5, 6), TRange<FFrameNumber>(6, 15) };
+		Tests.Last().ExpectedResults = { ExpectedResults[0], ExpectedResults[1] };
+	}
+
+	// Test 4: Test that compiling section range followed by a range that overlaps empty space preceeding that section and the section itself, compiles correctly (reverse of Test3)
+	{
+		Tests.Emplace();
+		Tests.Last().CompileRanges   = { TRange<FFrameNumber>(24, 25), TRange<FFrameNumber>(15, 24), TRange<FFrameNumber>(5, 6) };
+		Tests.Last().ExpectedResults = { ExpectedResults[0], ExpectedResults[1], ExpectedResults[2] };
+	}
+
+	// Test 5: Test that compiling a range encomassing the entire track results in the correct field ranges
+	{
+		Tests.Emplace();
+		Tests.Last().CompileRanges   = { TRange<FFrameNumber>(0, 70) };
+		Tests.Last().ExpectedResults = TArray<FResult>(&ExpectedResults[0], ARRAY_COUNT(ExpectedResults));
+	}
+
+	UTestMovieSceneSequence* Sequence = NewObject<UTestMovieSceneSequence>(GetTransientPackage());
+	UTestMovieSceneTrack*    Track    = Sequence->MovieScene->AddMasterTrack<UTestMovieSceneTrack>();
+
+	for (TRange<FFrameNumber> Range : SectionRanges)
+	{
+		UTestMovieSceneSection* Section = NewObject<UTestMovieSceneSection>(Track);
+		Section->SetRange(Range);
+		Track->SectionArray.Add(Section);
+	}
+
+	struct FTemplateStore : IMovieSceneSequenceTemplateStore
+	{
+		FMovieSceneEvaluationTemplate& AccessTemplate(UMovieSceneSequence&) override
+		{
+			return Template;
+		}
+
+		FMovieSceneEvaluationTemplate Template;
+	} Store;
+
+
+	for (int32 Index = 0; Index < Tests.Num(); ++Index)
+	{
+		const FTest& Test = Tests[Index];
+
+		// Wipe the evaluation template before each test
+		Store.Template = FMovieSceneEvaluationTemplate();
+
+		// Compile all the ranges that the test demands
+		for (TRange<FFrameNumber> CompileRange : Test.CompileRanges)
+		{
+			FMovieSceneCompiler::CompileRange(CompileRange, *Sequence, Store);
+		}
+
+		// Verify that the resulting evaluation field is what we expect
+		TArrayView<const FMovieSceneFrameRange> FieldRanges = Store.Template.EvaluationField.GetRanges();
+		if (!FieldRanges.Num())
+		{
+			AddError(FString::Printf(TEXT("Test index %02d: No evaluation field entries were compiled."), Index));
+			continue;
+		}
+
+		for (const FResult& Result : Test.ExpectedResults)
+		{
+			// Find the field entry that exactly matches our expected result
+			const FMovieSceneFrameRange* FieldRange = Algo::FindBy(FieldRanges, Result.FieldRange, &FMovieSceneFrameRange::Value);
+			if (!FieldRange)
+			{
+				AddError(FString::Printf(TEXT("Test index %02d: Expected to find an evaluation field entry for range %s but did not."), Index, *LexToString(Result.FieldRange)));
+			}
+			else
+			{
+				// Verify that the field entry is either empty or populated as the test expects
+				const int32 FieldIndex       = FieldRange - FieldRanges.GetData();
+				const bool  FieldIsEmptyHere = (Store.Template.EvaluationField.GetGroup(FieldIndex).SegmentPtrLUT.Num() == 0);
+
+				if (Result.bExpectEmpty != FieldIsEmptyHere)
+				{
+					const TCHAR* ExpectedString = Result.bExpectEmpty ? TEXT("empty")     : TEXT("populated");
+					const TCHAR* ActualString   = FieldIsEmptyHere    ? TEXT("populated") : TEXT("empty");
+					AddError(FString::Printf(TEXT("Test index %02d: Expected evaluation field entry range %s to be %s but it was %s."), Index, ExpectedString, *LexToString(Result.FieldRange), ActualString));
+				}
+			}
+		}
+	}
+
+	return true;
+}
 #endif // WITH_DEV_AUTOMATION_TESTS
\ No newline at end of file
diff --git a/Engine/Source/Runtime/MovieScene/Public/Evaluation/MovieSceneEvaluationTrack.h b/Engine/Source/Runtime/MovieScene/Public/Evaluation/MovieSceneEvaluationTrack.h
--- a/Engine/Source/Runtime/MovieScene/Public/Evaluation/MovieSceneEvaluationTrack.h
+++ b/Engine/Source/Runtime/MovieScene/Public/Evaluation/MovieSceneEvaluationTrack.h
@@ -394,14 +394,6 @@ public:
 	 */
 	TArray<FMovieSceneSegmentIdentifier> GetSegmentsInRange(TRange<FFrameNumber> InLocalRange);
 
-	/**
-	 * Get the smallest range of unique FSectionEvaluationData combinations that overlaps the specified lower bound
-	 *
-	 * @param InLowerBound 			The lower bound from which to start looking for a time range
-	 * @return The smallest time range of unique evaluation data entries that encompasses the specified lower bound
-	 */
-	TRange<FFrameNumber> GetUniqueRangeFromLowerBound(TRangeBound<FFrameNumber> InLowerBound) const;
-
 	/**
 	 * Set the source track from which this track originates
 	 */
@@ -449,7 +441,15 @@ public:
 	}
 
 	/**
-	 * Iterate all this track's unique time ranges
+	 * Iterate all this track's unique time ranges starting from the specified lower bound
+	 */
+	FMovieSceneEvaluationTreeRangeIterator IterateFrom(TRangeBound<FFrameNumber> LowerBound) const
+	{
+		return FMovieSceneEvaluationTreeRangeIterator(EvaluationTree.Tree, LowerBound);
+	}
+
+	/**
+	 * Retrieve the section evaluation data for the specified node in this track's evaluation tree
 	 */
 	TMovieSceneEvaluationTreeDataIterator<FSectionEvaluationData> GetData(FMovieSceneEvaluationTreeNodeHandle TreeNodeHandle) const
 	{
