Sequencer: Clamp segment ranges to iterated ranges when compiling tracks

  - This fixes an issue where segments that were combined together could end up in the evaluation field twice

diff --git a/Engine/Source/Runtime/MovieScene/Private/Compilation/MovieSceneCompiler.cpp b/Engine/Source/Runtime/MovieScene/Private/Compilation/MovieSceneCompiler.cpp
--- a/Engine/Source/Runtime/MovieScene/Private/Compilation/MovieSceneCompiler.cpp
+++ b/Engine/Source/Runtime/MovieScene/Private/Compilation/MovieSceneCompiler.cpp
@@ -570,7 +570,8 @@ void FMovieSceneCompiler::GatherCompileDataForTrack(FMovieSceneEvaluationTrack&
 			Data.Track = &Track;
 			Data.bRequiresInit = ThisSegment.Impls.ContainsByPredicate(RequiresInit);
 
-			TRange<FFrameNumber> IntersectionRange = Params.ClampRoot(ThisSegment.Range * SequenceToRootTransform);
+			TRange<FFrameNumber> SegmentTrackIntersection = TRange<FFrameNumber>::Intersection(ThisSegment.Range, TrackIter.Range());
+			TRange<FFrameNumber> IntersectionRange        = Params.ClampRoot(SegmentTrackIntersection * SequenceToRootTransform);
 			if (!IntersectionRange.IsEmpty())
 			{
 				OutData.Tracks.Add(IntersectionRange, Data);
