Sequencer: More descriptive comments for PlayToFrame, JumpToFrame and ScrubToFrame.

diff --git a/Engine/Source/Runtime/MovieScene/Public/MovieSceneSequencePlayer.h b/Engine/Source/Runtime/MovieScene/Public/MovieSceneSequencePlayer.h
--- a/Engine/Source/Runtime/MovieScene/Public/MovieSceneSequencePlayer.h
+++ b/Engine/Source/Runtime/MovieScene/Public/MovieSceneSequencePlayer.h
@@ -295,7 +295,8 @@ public:
 public:
 
 	/**
-	 * Play the sequence from the current time, to the specified frame position
+	 * Low-level call to set the current time of the player by evaluating from the current time to the specified time, as if the sequence is playing. 
+	 * Triggers events that lie within the evaluated range. Does not alter the persistent playback status of the player (IsPlaying).
 	 *
 	 * @param NewPosition     The new frame time to play to
 	 */
@@ -303,7 +304,8 @@ public:
 	void PlayToFrame(FFrameTime NewPosition);
 
 	/**
-	 * Scrub the sequence from the current time, to the specified frame position
+	 * Low-level call to set the current time of the player by evaluating only the specified time. Will not trigger any events. 
+	 * Does not alter the persistent playback status of the player (IsPlaying).
 	 *
 	 * @param NewPosition     The new frame time to scrub to
 	 */
@@ -311,7 +313,8 @@ public:
 	void ScrubToFrame(FFrameTime NewPosition);
 
 	/**
-	 * Jump to the specified frame position, without evaluating the sequence in between the current and desired time (as if in a paused state)
+	 * Low-level call to set the current time of the player by evaluating only the specified time, as if scrubbing the timeline. Will trigger only events that exist at the specified time. 
+	 * Does not alter the persistent playback status of the player (IsPlaying).
 	 *
 	 * @param NewPosition     The new frame time to jump to
 	 */
