Sequencer: Deprecation removal pass

diff --git a/Engine/Source/Runtime/MovieScene/Private/Compilation/MovieSceneSegmentCompiler.cpp b/Engine/Source/Runtime/MovieScene/Private/Compilation/MovieSceneSegmentCompiler.cpp
--- a/Engine/Source/Runtime/MovieScene/Private/Compilation/MovieSceneSegmentCompiler.cpp
+++ b/Engine/Source/Runtime/MovieScene/Private/Compilation/MovieSceneSegmentCompiler.cpp
@@ -4,27 +4,6 @@
 #include "Compilation/MovieSceneCompilerRules.h"
 #include "Evaluation/MovieSceneEvaluationTree.h"
 
-PRAGMA_DISABLE_DEPRECATION_WARNINGS
-void FMovieSceneSegmentCompilerRules::Blend(FSegmentBlendData& BlendData) const
-{
-	FMovieSceneSegment TmpSegment;
-
-	for (int32 Index = 0; Index < BlendData.Num(); ++Index)
-	{
-		TmpSegment.Impls.Add(FSectionEvaluationData(Index, BlendData[Index].Flags));
-	}
-
-	BlendSegment(TmpSegment, BlendData);
-
-	FSegmentBlendData NewBlendData;
-	for (FSectionEvaluationData EvalData : TmpSegment.Impls)
-	{
-		NewBlendData.Add(BlendData[EvalData.ImplIndex]);
-		NewBlendData.Last().Flags = EvalData.Flags;
-	}
-}
-PRAGMA_ENABLE_DEPRECATION_WARNINGS
-
 bool FMovieSceneAdditiveCameraTrackBlender::SortByStartTime(const FMovieSceneSectionData& A, const FMovieSceneSectionData& B)
 {
 	return TRangeBound<FFrameNumber>::MinLower(A.Section->GetRange().GetLowerBound(), B.Section->GetRange().GetLowerBound()) == A.Section->GetRange().GetLowerBound();
diff --git a/Engine/Source/Runtime/MovieScene/Private/MovieSceneSequencePlayer.cpp b/Engine/Source/Runtime/MovieScene/Private/MovieSceneSequencePlayer.cpp
--- a/Engine/Source/Runtime/MovieScene/Private/MovieSceneSequencePlayer.cpp
+++ b/Engine/Source/Runtime/MovieScene/Private/MovieSceneSequencePlayer.cpp
@@ -394,11 +394,6 @@ bool UMovieSceneSequencePlayer::IsReversed() const
 	return bReversePlayback;
 }
 
-float UMovieSceneSequencePlayer::GetLength() const
-{
-	return GetDuration().AsSeconds();
-}
-
 float UMovieSceneSequencePlayer::GetPlayRate() const
 {
 	return PlaybackSettings.PlayRate;
@@ -409,11 +404,6 @@ void UMovieSceneSequencePlayer::SetPlayRate(float PlayRate)
 	PlaybackSettings.PlayRate = PlayRate;
 }
 
-void UMovieSceneSequencePlayer::SetPlaybackRange( float NewStartTime, float NewEndTime )
-{
-	SetTimeRange(NewStartTime, NewEndTime - NewStartTime);
-}
-
 FFrameTime UMovieSceneSequencePlayer::GetLastValidTime() const
 {
 	return DurationFrames > 0 ? FFrameTime(StartTime + DurationFrames-1, 0.99999994f) : FFrameTime(StartTime);
diff --git a/Engine/Source/Runtime/MovieScene/Private/MovieSceneTrack.cpp b/Engine/Source/Runtime/MovieScene/Private/MovieSceneTrack.cpp
--- a/Engine/Source/Runtime/MovieScene/Private/MovieSceneTrack.cpp
+++ b/Engine/Source/Runtime/MovieScene/Private/MovieSceneTrack.cpp
@@ -138,40 +138,14 @@ void UMovieSceneTrack::UpdateEasing()
 	}
 }
 
-PRAGMA_DISABLE_DEPRECATION_WARNINGS
-TInlineValue<FMovieSceneSegmentCompilerRules> UMovieSceneTrack::GetRowCompilerRules() const
-{
-	return TInlineValue<FMovieSceneSegmentCompilerRules>();
-}
-
-TInlineValue<FMovieSceneSegmentCompilerRules> UMovieSceneTrack::GetTrackCompilerRules() const
-{
-	return TInlineValue<FMovieSceneSegmentCompilerRules>();
-}
-
 FMovieSceneTrackRowSegmentBlenderPtr UMovieSceneTrack::GetRowSegmentBlender() const
 {
-	// Handle legacy row compiler rules
-	TInlineValue<FMovieSceneSegmentCompilerRules> LegacyRules = GetRowCompilerRules();
-	if (LegacyRules.IsValid())
-	{
-		return TLegacyTrackRowSegmentBlender<FMovieSceneTrackRowSegmentBlender>(MoveTemp(LegacyRules));
-	}
-	else
-	{
-		return FDefaultTrackRowSegmentBlender();
-	}
+	return FDefaultTrackRowSegmentBlender();
 }
 
 FMovieSceneTrackSegmentBlenderPtr UMovieSceneTrack::GetTrackSegmentBlender() const
 {
-	// Handle legacy track compiler rules
-	TInlineValue<FMovieSceneSegmentCompilerRules> LegacyRules = GetTrackCompilerRules();
-	if (LegacyRules.IsValid())
-	{
-		return TLegacyTrackRowSegmentBlender<FMovieSceneTrackSegmentBlender>(MoveTemp(LegacyRules));
-	}
-	else if (EvalOptions.bCanEvaluateNearestSection && EvalOptions.bEvalNearestSection)
+	if (EvalOptions.bCanEvaluateNearestSection && EvalOptions.bEvalNearestSection)
 	{
 		return FEvaluateNearestSegmentBlender();
 	}
@@ -180,7 +154,6 @@ FMovieSceneTrackSegmentBlenderPtr UMovieSceneTrack::GetTrackSegmentBlender() con
 		return FMovieSceneTrackSegmentBlenderPtr();
 	}
 }
-PRAGMA_ENABLE_DEPRECATION_WARNINGS
 
 void UMovieSceneTrack::GenerateTemplate(const FMovieSceneTrackCompilerArgs& Args) const
 {
diff --git a/Engine/Source/Runtime/MovieScene/Public/Compilation/IMovieSceneTemplateGenerator.h b/Engine/Source/Runtime/MovieScene/Public/Compilation/IMovieSceneTemplateGenerator.h
--- a/Engine/Source/Runtime/MovieScene/Public/Compilation/IMovieSceneTemplateGenerator.h
+++ b/Engine/Source/Runtime/MovieScene/Public/Compilation/IMovieSceneTemplateGenerator.h
@@ -27,24 +27,6 @@ struct IMovieSceneTemplateGenerator
 	 */
 	virtual void AddOwnedTrack(FMovieSceneEvaluationTrack&& InTrackTemplate, const UMovieSceneTrack& SourceTrack) = 0;
 
-
-public:
-
-	DEPRECATED(4.19, "Shared tracks are no longer supported, please emit shared Execution tokens instead (FMovieSceneExecutionTokens::AddShared")
-	void AddSharedTrack(FMovieSceneEvaluationTrack&& InTrackTemplate, const FMovieSceneSharedDataId& SharedId, const UMovieSceneTrack& SourceTrack){}
-
-	DEPRECATED(4.19, "Support for legacy tracks has now been removed")
-	void AddLegacyTrack(FMovieSceneEvaluationTrack&& InTrackTemplate, const UMovieSceneTrack& SourceTrack){}
-
-	DEPRECATED(4.19, "External segments are now handled internally")
-	void AddExternalSegments(TRange<float> RootRange, TArrayView<const FMovieSceneEvaluationFieldSegmentPtr> SegmentPtrs, ESectionEvaluationFlags Flags){}
-
-	DEPRECATED(4.19, "Sub sequence templates are no longer generated at compile time within template generators")
-	FMovieSceneSequenceTransform GetSequenceTransform(FMovieSceneSequenceIDRef InSequenceID) const { return FMovieSceneSequenceTransform(); }
-
-	DEPRECATED(4.19, "Sub sequences are now handled internally")
-	void AddSubSequence(const FMovieSceneSubSequenceData& SequenceData, FMovieSceneSequenceIDRef ParentID, FMovieSceneSequenceID SequenceID) {}
-
 protected:
 	virtual ~IMovieSceneTemplateGenerator() { }
 };
diff --git a/Engine/Source/Runtime/MovieScene/Public/Compilation/MovieSceneCompilerRules.h b/Engine/Source/Runtime/MovieScene/Public/Compilation/MovieSceneCompilerRules.h
--- a/Engine/Source/Runtime/MovieScene/Public/Compilation/MovieSceneCompilerRules.h
+++ b/Engine/Source/Runtime/MovieScene/Public/Compilation/MovieSceneCompilerRules.h
@@ -145,30 +145,3 @@ private:
 
 	MOVIESCENE_API static bool SortByStartTime(const FMovieSceneSectionData& A, const FMovieSceneSectionData& B);
 };
-
-
-PRAGMA_DISABLE_DEPRECATION_WARNINGS
-
-/** Track row segment blender that acts as a proxy to an instance of the legacy FMovieSceneSegmentCompilerRules */
-template<typename BaseType>
-struct TLegacyTrackRowSegmentBlender : BaseType
-{
-	TLegacyTrackRowSegmentBlender(TInlineValue<FMovieSceneSegmentCompilerRules>&& InRules)
-		: LegacyRules(MoveTemp(InRules))
-	{}
-
-	~TLegacyTrackRowSegmentBlender()
-	{}
-
-	TLegacyTrackRowSegmentBlender(TLegacyTrackRowSegmentBlender&&) = default;
-
-	virtual void Blend(FSegmentBlendData& BlendData) const override
-	{
-		return LegacyRules->Blend(BlendData);
-	}
-
-private:
-	TInlineValue<FMovieSceneSegmentCompilerRules> LegacyRules;
-};
-
-PRAGMA_ENABLE_DEPRECATION_WARNINGS
\ No newline at end of file
diff --git a/Engine/Source/Runtime/MovieScene/Public/Compilation/MovieSceneSegmentCompiler.h b/Engine/Source/Runtime/MovieScene/Public/Compilation/MovieSceneSegmentCompiler.h
--- a/Engine/Source/Runtime/MovieScene/Public/Compilation/MovieSceneSegmentCompiler.h
+++ b/Engine/Source/Runtime/MovieScene/Public/Compilation/MovieSceneSegmentCompiler.h
@@ -140,13 +140,3 @@ protected:
 
 typedef TInlineValue<FMovieSceneTrackRowSegmentBlender, 16> FMovieSceneTrackRowSegmentBlenderPtr;
 typedef TInlineValue<FMovieSceneTrackSegmentBlender, 16> FMovieSceneTrackSegmentBlenderPtr;
-
-
-/** Deprecated type support */
-struct FMovieSceneSegmentCompilerRules : FMovieSceneTrackSegmentBlender
-{
-	virtual void Blend(FSegmentBlendData& BlendData) const override;
-
-protected:
-	virtual void BlendSegment(FMovieSceneSegment& Segment, const TArrayView<const FMovieSceneSectionData>& SourceData) const = 0;
-};
diff --git a/Engine/Source/Runtime/MovieScene/Public/Evaluation/MovieSceneEvalTemplate.h b/Engine/Source/Runtime/MovieScene/Public/Evaluation/MovieSceneEvalTemplate.h
--- a/Engine/Source/Runtime/MovieScene/Public/Evaluation/MovieSceneEvalTemplate.h
+++ b/Engine/Source/Runtime/MovieScene/Public/Evaluation/MovieSceneEvalTemplate.h
@@ -113,14 +113,6 @@ public:
 	 * @param ExecutionTokens	Stack of execution tokens that will be used to apply animated state to the environment at a later time.
 	*/
 	virtual void EvaluateSwept(const FMovieSceneEvaluationOperand& Operand, const FMovieSceneContext& Context, const TRange<FFrameNumber>& SweptRange, const FPersistentEvaluationData& PersistentData, FMovieSceneExecutionTokens& ExecutionTokens) const
-	{
-		// Call deprecated one
-		PRAGMA_DISABLE_DEPRECATION_WARNINGS
-		EvaluateSwept(Operand, Context, PersistentData, ExecutionTokens);
-		PRAGMA_ENABLE_DEPRECATION_WARNINGS
-	}
-	DEPRECATED(4.21, "Please override void EvaluateSwept(const FMovieSceneEvaluationOperand&, const FMovieSceneContext&, const TRange<FFrameNumber>&, const FPersistentEvaluationData&, FMovieSceneExecutionTokens&) instead.")
-	virtual void EvaluateSwept(const FMovieSceneEvaluationOperand& Operand, const FMovieSceneContext& Context, const FPersistentEvaluationData& PersistentData, FMovieSceneExecutionTokens& ExecutionTokens) const
 	{
 		ensureMsgf(false, TEXT("FMovieSceneEvalTemplate::EvaluateSwept has not been implemented. Verify that this template's evaluation track has correct evaluation method (usually set in UMovieSceneTrack::PostCompile), or implement this function."));
 	}
diff --git a/Engine/Source/Runtime/MovieScene/Public/Evaluation/MovieScenePropertyTemplate.h b/Engine/Source/Runtime/MovieScene/Public/Evaluation/MovieScenePropertyTemplate.h
--- a/Engine/Source/Runtime/MovieScene/Public/Evaluation/MovieScenePropertyTemplate.h
+++ b/Engine/Source/Runtime/MovieScene/Public/Evaluation/MovieScenePropertyTemplate.h
@@ -40,55 +40,6 @@ namespace PropertyTemplate
 		FMovieSceneAnimTypeID PropertyID;
 	};
 
-	/** The value of the object as it existed before this frame's evaluation */
-	template<typename PropertyValueType> struct
-	DEPRECATED(4.17, "Precaching of property values should no longer be necessary as it was only used to pass default values to curves on evaluation. Curves should now be checked for emptyness before attempting to animate an object.")
-	TCachedValue
-	{
-		TWeakObjectPtr<UObject> WeakObject;
-		PropertyValueType Value;
-	};
-
-	PRAGMA_DISABLE_DEPRECATION_WARNINGS
-	template<typename T> struct TCachedValue_DEPRECATED : TCachedValue<T> {};
-	PRAGMA_ENABLE_DEPRECATION_WARNINGS
-
-	template<typename PropertyValueType> struct 
-	DEPRECATED(4.17, "Precaching of property values should no longer be necessary as it was only used to pass default values to curves on evaluation. Curves should now be checked for emptyness before attempting to animate an object.")
-	TCachedSectionData : FSectionData
-	{
-		/** Setup the data for the current frame */
-		void SetupFrame(const FMovieSceneEvaluationOperand& Operand, IMovieScenePlayer& Player)
-		{
-			ObjectsAndValues.Reset();
-			
-			for (TWeakObjectPtr<> Object : Player.FindBoundObjects(Operand))
-			{
-				if (UObject* ObjectPtr = Object.Get())
-				{
-					PropertyBindings->CacheBinding(*ObjectPtr);
-					if (UProperty* Property = PropertyBindings->GetProperty(*ObjectPtr))
-					{
-						if (Property->GetSize() == sizeof(PropertyValueType))
-						{
-							ObjectsAndValues.Add(TCachedValue_DEPRECATED<PropertyValueType>{ ObjectPtr, PropertyBindings->GetCurrentValue<PropertyValueType>(*ObjectPtr) });
-						}
-#if !(UE_BUILD_SHIPPING || UE_BUILD_TEST)
-						else
-						{
-							FMessageLog("Sequencer").Warning()
-								->AddToken(FUObjectToken::Create(Player.GetEvaluationTemplate().GetSequence(MovieSceneSequenceID::Root)))
-								->AddToken(FTextToken::Create(FText::Format(NSLOCTEXT("MovieScene", "IncompatibleDataWarning", "Property size mismatch for property '{0}'. Expected '{1}', found '{2}'. Recreate the track with the new property type."), FText::FromString(PropertyBindings->GetPropertyPath()), FText::FromString(TNameOf<PropertyValueType>::GetName()), FText::FromString(Property->GetCPPType()))));
-						}
-#endif
-					}
-				}
-			}
-		}
-
-		TArray<TCachedValue_DEPRECATED<PropertyValueType>, TInlineAllocator<1>> ObjectsAndValues;
-	};
-
 	template<typename PropertyValueType, typename IntermediateType = PropertyValueType>
 	struct TTemporarySetterType { typedef PropertyValueType Type; };
 
@@ -164,32 +115,6 @@ namespace PropertyTemplate
 	};
 }
 
-
-template<typename PropertyValueType> struct
-DEPRECATED(4.17, "Precaching of property values should no longer be necessary as it was only used to pass default values to curves on evaluation. Curves should now be checked for emptyness before attempting to animate an object.")
-TCachedPropertyTrackExecutionToken : IMovieSceneExecutionToken
-{
-	virtual void Execute(const FMovieSceneContext& Context, const FMovieSceneEvaluationOperand& Operand, FPersistentEvaluationData& PersistentData, IMovieScenePlayer& Player) override
-	{
-		using namespace PropertyTemplate;
-
-		MOVIESCENE_DETAILED_SCOPE_CYCLE_COUNTER(MovieSceneEval_PropertyTrack_TokenExecute)
-		
-		TCachedSectionData<PropertyValueType>& PropertyTrackData = PersistentData.GetSectionData<TCachedSectionData<PropertyValueType>>();
-		FTrackInstancePropertyBindings* PropertyBindings = PropertyTrackData.PropertyBindings.Get();
-
-		for (TCachedValue<PropertyValueType>& ObjectAndValue : PropertyTrackData.ObjectsAndValues)
-		{
-			if (UObject* ObjectPtr = ObjectAndValue.WeakObject.Get())
-			{
-				Player.SavePreAnimatedState(*ObjectPtr, PropertyTrackData.PropertyID, FTokenProducer<PropertyValueType>(*PropertyBindings));
-				
-				PropertyBindings->CallFunction<PropertyValueType>(*ObjectPtr, ObjectAndValue.Value);
-			}
-		}
-	}
-};
-
 /** Execution token that simple stores a value, and sets it when executed */
 template<typename PropertyValueType, typename IntermediateType = PropertyValueType>
 struct TPropertyTrackExecutionToken : IMovieSceneExecutionToken
@@ -290,22 +215,6 @@ struct FMovieScenePropertySectionData
 	{
 		PersistentData.AddSectionData<T>().Initialize(PropertyName, PropertyPath, FunctionName, NotifyFunctionName);
 	}
-
-	template<typename T>
-	DEPRECATED(4.17, "Precaching of property values should no longer be necessary as it was only used to pass default values to curves on evaluation. Curves should now be checked for emptyness before attempting to animate an object.")
-	void SetupCachedTrack(FPersistentEvaluationData& PersistentData) const
-	{
-		typedef PropertyTemplate::TCachedSectionData<T> FSectionData;
-		PersistentData.AddSectionData<FSectionData>().Initialize(PropertyName, PropertyPath, FunctionName, NotifyFunctionName);
-	}
-
-	template<typename T>
-	DEPRECATED(4.17, "Precaching of property values should no longer be necessary as it was only used to pass default values to curves on evaluation. Curves should now be checked for emptyness before attempting to animate an object.")
-	void SetupCachedFrame(const FMovieSceneEvaluationOperand& Operand, FPersistentEvaluationData& PersistentData, IMovieScenePlayer& Player) const
-	{
-		typedef PropertyTemplate::TCachedSectionData<T> FSectionData;
-		PersistentData.GetSectionData<FSectionData>().SetupFrame(Operand, Player);
-	}
 };
 
 USTRUCT()
diff --git a/Engine/Source/Runtime/MovieScene/Public/IMovieSceneModule.h b/Engine/Source/Runtime/MovieScene/Public/IMovieSceneModule.h
--- a/Engine/Source/Runtime/MovieScene/Public/IMovieSceneModule.h
+++ b/Engine/Source/Runtime/MovieScene/Public/IMovieSceneModule.h
@@ -17,11 +17,6 @@ struct FMovieSceneEvaluationGroupParameters
 		: EvaluationPriority(InPriority)
 	{}
 
-	DEPRECATED(4.17, "Please remove bInRequiresImmediateFlush parameter")
-	FMovieSceneEvaluationGroupParameters(uint16 InPriority, bool bInRequiresImmediateFlush)
-		: EvaluationPriority(InPriority)
-	{}
-
 	/** Prioirty assigned to this group. Higher priorities are evaluated first */
 	uint16 EvaluationPriority;
 };
diff --git a/Engine/Source/Runtime/MovieScene/Public/MovieSceneCommonHelpers.h b/Engine/Source/Runtime/MovieScene/Public/MovieSceneCommonHelpers.h
--- a/Engine/Source/Runtime/MovieScene/Public/MovieSceneCommonHelpers.h
+++ b/Engine/Source/Runtime/MovieScene/Public/MovieSceneCommonHelpers.h
@@ -243,17 +243,6 @@ public:
 		return PropertyName;
 	}
 
-	template<typename ValueType>
-	DEPRECATED(4.15, "Please use GetCurrentValue(const UObject&)")
-	ValueType GetCurrentValue(const UObject* Object) { check(Object) return GetCurrentValue<ValueType>(*Object); }
-	template <typename ValueType>
-	DEPRECATED(4.15, "Please use CallFunction(UObject&)")
-	void CallFunction( UObject* InRuntimeObject, ValueType* PropertyValue ) { CallFunction<ValueType>(*InRuntimeObject, *PropertyValue); }
-	DEPRECATED(4.15, "UpdateBindings is no longer necessary")
-	void UpdateBindings( const TArray<TWeakObjectPtr<UObject>>& InRuntimeObjects ){}
-	DEPRECATED(4.15, "UpdateBindings is no longer necessary")
-	void UpdateBinding( const TWeakObjectPtr<UObject>& InRuntimeObject ) {}
-
 private:
 
 	struct FPropertyAddress
diff --git a/Engine/Source/Runtime/MovieScene/Public/MovieSceneSequencePlayer.h b/Engine/Source/Runtime/MovieScene/Public/MovieSceneSequencePlayer.h
--- a/Engine/Source/Runtime/MovieScene/Public/MovieSceneSequencePlayer.h
+++ b/Engine/Source/Runtime/MovieScene/Public/MovieSceneSequencePlayer.h
@@ -168,65 +168,6 @@ public:
 	UFUNCTION(BlueprintCallable, Category="Game|Cinematic", meta = (ToolTip = "Go to end of the sequence and stop. Adheres to 'When Finished' section rules."))
 	void GoToEndAndStop();
 
-public:
-
-	/**
-	 * Get the current playback position
-	 * @return The current playback position
-	 */
-	DEPRECATED(4.20, "Please use GetCurrentTime instead")
-	UFUNCTION(BlueprintCallable, Category="Game|Cinematic")
-	float GetPlaybackPosition() const { return GetCurrentTime().AsSeconds() - StartTime / PlayPosition.GetInputRate(); }
-
-	/**
-	 * Get the playback length of the sequence
-	 */
-	DEPRECATED(4.20, "Please use GetDuration instead")
-	UFUNCTION(BlueprintCallable, Category="Game|Cinematic")
-	float GetLength() const;
-
-	/**
-	 * Get the offset within the level sequence to start playing
-	 */
-	DEPRECATED(4.20, "Please use GetStartTime instead")
-	UFUNCTION(BlueprintCallable, Category="Game|Cinematic")
-	float GetPlaybackStart() const { return StartTime / PlayPosition.GetInputRate(); }
-
-	/**
-	 * Get the offset within the level sequence to finish playing
-	 */
-	DEPRECATED(4.20, "Please use GetEndTime instead")
-	UFUNCTION(BlueprintCallable, Category="Game|Cinematic")
-	float GetPlaybackEnd() const { return (StartTime + DurationFrames) / PlayPosition.GetInputRate(); }
-
-	/**
-	 * Set the current playback position
-	 * @param NewPlaybackPosition - The new playback position to set.
-	 * If the animation is currently playing, it will continue to do so from the new position
-	 */
-	DEPRECATED(4.20, "Please use PlayToFrame instead")
-	UFUNCTION(BlueprintCallable, Category="Game|Cinematic")
-	void SetPlaybackPosition(float NewPlaybackPosition) { Status == EMovieScenePlayerStatus::Playing ? PlayToSeconds(NewPlaybackPosition + StartTime / PlayPosition.GetInputRate()) : JumpToSeconds(NewPlaybackPosition + StartTime / PlayPosition.GetInputRate()); }
-
-	/**
-	 * Sets the range in time to be played back by this player, overriding the default range stored in the asset
-	 *
-	 * @param	NewStartTime	The new starting time for playback
-	 * @param	NewEndTime		The new ending time for playback.  Must be larger than the start time.
-	 */
-	DEPRECATED(4.20, "Please use SetFrameRange or SetTimeRange instead")
-	UFUNCTION(BlueprintCallable, Category="Game|Cinematic")
-	void SetPlaybackRange( const float NewStartTime, const float NewEndTime );
-
-	/**
-	 * Jump to new playback position
-	 * @param NewPlaybackPosition - The new playback position to set.
-	 * This can be used to update sequencer repeatedly, as if in a scrubbing state
-	 */
-	DEPRECATED(4.20, "Please use ScrubToTime instead")
-	UFUNCTION(BlueprintCallable, Category="Game|Cinematic")
-	void JumpToPosition(float NewPlaybackPosition) { ScrubToSeconds(NewPlaybackPosition); }
-
 public:
 
 	/**
diff --git a/Engine/Source/Runtime/MovieScene/Public/MovieSceneTrack.h b/Engine/Source/Runtime/MovieScene/Public/MovieSceneTrack.h
--- a/Engine/Source/Runtime/MovieScene/Public/MovieSceneTrack.h
+++ b/Engine/Source/Runtime/MovieScene/Public/MovieSceneTrack.h
@@ -18,10 +18,6 @@ struct FMovieSceneTrackSegmentBlender;
 struct FMovieSceneTrackRowSegmentBlender;
 struct IMovieSceneTemplateGenerator;
 
-PRAGMA_DISABLE_DEPRECATION_WARNINGS
-	typedef TInlineValue<FMovieSceneSegmentCompilerRules> FMovieSceneDeprecatedCompilerRulesPtr;
-PRAGMA_ENABLE_DEPRECATION_WARNINGS
-
 /** Flags used to perform cook-time optimization of movie scene data */
 enum class ECookOptimizationFlags
 {
@@ -49,16 +45,6 @@ struct FMovieSceneTrackCompilerArgs
 
 	/** The generator responsible for generating the template */
 	IMovieSceneTemplateGenerator& Generator;
-
-	struct FMovieSceneSequenceTemplateStore {};
-
-	DEPRECATED(4.19, "Template store is no longer supplied as part of track compilation")
-	FMovieSceneSequenceTemplateStore SubSequenceStore;
-
-	struct FMovieSceneTrackCompilationParams { 	bool bForEditorPreview, bDuringBlueprintCompile; };
-
-	DEPRECATED(4.19, "Template store is no longer supplied as part of track compilation")
-	FMovieSceneTrackCompilationParams Params;
 };
 
 /** Generic evaluation options for any track */
@@ -154,12 +140,6 @@ public:
 	 */
 	MOVIESCENE_API virtual FMovieSceneTrackSegmentBlenderPtr GetTrackSegmentBlender() const;
 
-	DEPRECATED(4.19, "Please override GetRowSegmentBlender() instead.")
-	MOVIESCENE_API virtual FMovieSceneDeprecatedCompilerRulesPtr GetRowCompilerRules() const;
-
-	DEPRECATED(4.19, "Please override GetTrackSegmentBlender() instead.")
-	MOVIESCENE_API virtual FMovieSceneDeprecatedCompilerRulesPtr GetTrackCompilerRules() const;
-
 	/**
 	 * Generate a template for this track
 	 *
@@ -277,14 +257,6 @@ public:
 	 */
 	virtual const TArray<UMovieSceneSection*>& GetAllSections() const PURE_VIRTUAL(UMovieSceneTrack::GetAllSections, static TArray<UMovieSceneSection*> Empty; return Empty;);
 
-	/**
-	 * Gets the section boundaries of this track.
-	 * 
-	 * @return The range of time boundaries.
-	 */
-	DEPRECATED(4.20, "This function is no longer used.")
-	virtual TRange<FFrameNumber> GetSectionBoundaries() const PURE_VIRTUAL(UMovieSceneTrack::GetSectionBoundaries, return TRange<FFrameNumber>::Empty(););
-
 	/**
 	 * Checks to see if the section is in this track.
 	 *
diff --git a/Engine/Source/Runtime/MovieSceneTracks/Private/Sections/MovieScene3DConstraintSection.cpp b/Engine/Source/Runtime/MovieSceneTracks/Private/Sections/MovieScene3DConstraintSection.cpp
--- a/Engine/Source/Runtime/MovieSceneTracks/Private/Sections/MovieScene3DConstraintSection.cpp
+++ b/Engine/Source/Runtime/MovieSceneTracks/Private/Sections/MovieScene3DConstraintSection.cpp
@@ -17,13 +17,6 @@ void UMovieScene3DConstraintSection::SetConstraintId(const FGuid& InConstraintId
 	}
 }
 
-
-FGuid UMovieScene3DConstraintSection::GetConstraintId() const
-{
-	return FGuid();
-}
-
-
 void UMovieScene3DConstraintSection::OnBindingsUpdated(const TMap<FGuid, FGuid>& OldGuidToNewGuidMap)
 {
 	if (OldGuidToNewGuidMap.Contains(ConstraintBindingID.GetGuid()))
diff --git a/Engine/Source/Runtime/MovieSceneTracks/Public/Sections/MovieScene3DConstraintSection.h b/Engine/Source/Runtime/MovieSceneTracks/Public/Sections/MovieScene3DConstraintSection.h
--- a/Engine/Source/Runtime/MovieSceneTracks/Public/Sections/MovieScene3DConstraintSection.h
+++ b/Engine/Source/Runtime/MovieSceneTracks/Public/Sections/MovieScene3DConstraintSection.h
@@ -25,9 +25,6 @@ public:
 	/** Sets the constraint id for this section */
 	virtual void SetConstraintId(const FGuid& InId);
 
-	DEPRECATED(4.18, "Constraint guid no longer supported, Use GetConstraintBindingID.")
-	virtual FGuid GetConstraintId() const;
-
 	/** Gets the constraint binding for this Constraint section */
 	const FMovieSceneObjectBindingID& GetConstraintBindingID() const
 	{
diff --git a/Engine/Source/Runtime/MovieSceneTracks/Public/Sections/MovieSceneCameraCutSection.h b/Engine/Source/Runtime/MovieSceneTracks/Public/Sections/MovieSceneCameraCutSection.h
--- a/Engine/Source/Runtime/MovieSceneTracks/Public/Sections/MovieSceneCameraCutSection.h
+++ b/Engine/Source/Runtime/MovieSceneTracks/Public/Sections/MovieSceneCameraCutSection.h
@@ -31,12 +31,6 @@ public:
 				EMovieSceneCompletionMode::ProjectDefault);
 	}
 
-	DEPRECATED(4.18, "Camera guid no longer supported, Use GetCameraBindingID.")
-	FGuid GetCameraGuid() const
-	{
-		return FGuid();
-	}
-
 	/** Sets the camera binding for this CameraCut section. Evaluates from the sequence binding ID */
 	void SetCameraGuid(const FGuid& InGuid)
 	{
