Undo BP removal changes from changelist 4466484 since BP warnings were not being emitted for deprecated functions yet.

diff --git a/Engine/Source/Runtime/MovieScene/Private/MovieSceneSequencePlayer.cpp b/Engine/Source/Runtime/MovieScene/Private/MovieSceneSequencePlayer.cpp
--- a/Engine/Source/Runtime/MovieScene/Private/MovieSceneSequencePlayer.cpp
+++ b/Engine/Source/Runtime/MovieScene/Private/MovieSceneSequencePlayer.cpp
@@ -394,6 +394,11 @@ bool UMovieSceneSequencePlayer::IsReversed() const
 	return bReversePlayback;
 }
 
+float UMovieSceneSequencePlayer::GetLength() const
+{
+	return GetDuration().AsSeconds();
+}
+
 float UMovieSceneSequencePlayer::GetPlayRate() const
 {
 	return PlaybackSettings.PlayRate;
@@ -404,6 +409,11 @@ void UMovieSceneSequencePlayer::SetPlayRate(float PlayRate)
 	PlaybackSettings.PlayRate = PlayRate;
 }
 
+void UMovieSceneSequencePlayer::SetPlaybackRange( float NewStartTime, float NewEndTime )
+{
+	SetTimeRange(NewStartTime, NewEndTime - NewStartTime);
+}
+
 FFrameTime UMovieSceneSequencePlayer::GetLastValidTime() const
 {
 	return DurationFrames > 0 ? FFrameTime(StartTime + DurationFrames-1, 0.99999994f) : FFrameTime(StartTime);
diff --git a/Engine/Source/Runtime/MovieScene/Public/MovieSceneSequencePlayer.h b/Engine/Source/Runtime/MovieScene/Public/MovieSceneSequencePlayer.h
--- a/Engine/Source/Runtime/MovieScene/Public/MovieSceneSequencePlayer.h
+++ b/Engine/Source/Runtime/MovieScene/Public/MovieSceneSequencePlayer.h
@@ -168,6 +168,65 @@ public:
 	UFUNCTION(BlueprintCallable, Category="Game|Cinematic", meta = (ToolTip = "Go to end of the sequence and stop. Adheres to 'When Finished' section rules."))
 	void GoToEndAndStop();
 
+public:
+
+	/**
+	 * Get the current playback position
+	 * @return The current playback position
+	 */
+	DEPRECATED(4.20, "Please use GetCurrentTime instead")
+	UFUNCTION(BlueprintCallable, Category="Game|Cinematic")
+	float GetPlaybackPosition() const { return GetCurrentTime().AsSeconds() - StartTime / PlayPosition.GetInputRate(); }
+
+	/**
+	 * Get the playback length of the sequence
+	 */
+	DEPRECATED(4.20, "Please use GetDuration instead")
+	UFUNCTION(BlueprintCallable, Category="Game|Cinematic")
+	float GetLength() const;
+
+	/**
+	 * Get the offset within the level sequence to start playing
+	 */
+	DEPRECATED(4.20, "Please use GetStartTime instead")
+	UFUNCTION(BlueprintCallable, Category="Game|Cinematic")
+	float GetPlaybackStart() const { return StartTime / PlayPosition.GetInputRate(); }
+
+	/**
+	 * Get the offset within the level sequence to finish playing
+	 */
+	DEPRECATED(4.20, "Please use GetEndTime instead")
+	UFUNCTION(BlueprintCallable, Category="Game|Cinematic")
+	float GetPlaybackEnd() const { return (StartTime + DurationFrames) / PlayPosition.GetInputRate(); }
+
+	/**
+	 * Set the current playback position
+	 * @param NewPlaybackPosition - The new playback position to set.
+	 * If the animation is currently playing, it will continue to do so from the new position
+	 */
+	DEPRECATED(4.20, "Please use PlayToFrame instead")
+	UFUNCTION(BlueprintCallable, Category="Game|Cinematic")
+	void SetPlaybackPosition(float NewPlaybackPosition) { Status == EMovieScenePlayerStatus::Playing ? PlayToSeconds(NewPlaybackPosition + StartTime / PlayPosition.GetInputRate()) : JumpToSeconds(NewPlaybackPosition + StartTime / PlayPosition.GetInputRate()); }
+
+	/**
+	 * Sets the range in time to be played back by this player, overriding the default range stored in the asset
+	 *
+	 * @param	NewStartTime	The new starting time for playback
+	 * @param	NewEndTime		The new ending time for playback.  Must be larger than the start time.
+	 */
+	DEPRECATED(4.20, "Please use SetFrameRange or SetTimeRange instead")
+	UFUNCTION(BlueprintCallable, Category="Game|Cinematic")
+	void SetPlaybackRange( const float NewStartTime, const float NewEndTime );
+
+	/**
+	 * Jump to new playback position
+	 * @param NewPlaybackPosition - The new playback position to set.
+	 * This can be used to update sequencer repeatedly, as if in a scrubbing state
+	 */
+	DEPRECATED(4.20, "Please use ScrubToTime instead")
+	UFUNCTION(BlueprintCallable, Category="Game|Cinematic")
+	void JumpToPosition(float NewPlaybackPosition) { ScrubToSeconds(NewPlaybackPosition); }
+
 public:
 
 	/**
