Sequencer: Extend UMovieSceneSequence::Stop to accept the time to reset the cursor to

  - This ensures that at the point we trigger explicit network Stop events, the time cursor will not change unless the user has explicitly done so, which resolves an issue where a passive net update with a slightly different time can come in after the explicit stop RPC event has been handled by a client.
  - It also means we are no longer resetting the time twice when stopping.

[ Copy from 4586508 ]

diff --git a/Engine/Source/Runtime/MovieScene/Private/MovieSceneSequencePlayer.cpp b/Engine/Source/Runtime/MovieScene/Private/MovieSceneSequencePlayer.cpp
--- a/Engine/Source/Runtime/MovieScene/Private/MovieSceneSequencePlayer.cpp
+++ b/Engine/Source/Runtime/MovieScene/Private/MovieSceneSequencePlayer.cpp
@@ -261,20 +261,29 @@ void UMovieSceneSequencePlayer::Scrub()
 }
 
 void UMovieSceneSequencePlayer::Stop()
+{
+	StopInternal(bReversePlayback ? GetLastValidTime() : FFrameTime(StartTime));
+}
+
+void UMovieSceneSequencePlayer::StopAtCurrentTime()
+{
+	StopInternal(PlayPosition.GetCurrentPosition());
+}
+
+void UMovieSceneSequencePlayer::StopInternal(FFrameTime TimeToResetTo)
 {
 	if (IsPlaying() || IsPaused() || RootTemplateInstance.IsValid())
 	{
 		if (bIsEvaluating)
 		{
-			LatentActions.Emplace(FLatentAction::Stop);
+			LatentActions.Emplace(FLatentAction::Stop, TimeToResetTo);
 			return;
 		}
 
 		Status = EMovieScenePlayerStatus::Stopped;
 
-		// Put the cursor back to the start
-		PlayPosition.Reset(bReversePlayback ? GetLastValidTime() : FFrameTime(StartTime));
-
+		// Put the cursor at the specified position
+		PlayPosition.Reset(TimeToResetTo);
 		TimeController->StopPlaying(GetCurrentTime());
 
 		CurrentNumLoops = 0;
@@ -294,7 +303,7 @@ void UMovieSceneSequencePlayer::Stop()
 		if (HasAuthority())
 		{
 			// Explicitly handle Stop() events through an RPC call
-			RPC_OnStopEvent(PlayPosition.GetCurrentPosition());
+			RPC_OnStopEvent(TimeToResetTo);
 		}
 		UpdateNetworkSyncProperties();
 
@@ -315,8 +324,9 @@ void UMovieSceneSequencePlayer::Stop()
 
 void UMovieSceneSequencePlayer::GoToEndAndStop()
 {
-	JumpToFrame(GetLastValidTime());
-	Stop();
+	FFrameTime Time = GetLastValidTime();
+	JumpToFrame(Time);
+	StopInternal(Time);
 }
 
 FQualifiedFrameTime UMovieSceneSequencePlayer::GetCurrentTime() const
@@ -515,7 +525,7 @@ void UMovieSceneSequencePlayer::Initialize(UMovieSceneSequence* InSequence, cons
 	// may be called during PostLoad.
 	if (Sequence)
 	{
-		Stop();
+		StopAtCurrentTime();
 	}
 
 	Sequence = InSequence;
@@ -714,10 +724,7 @@ void UMovieSceneSequencePlayer::UpdateTimeCursorPosition_Internal(FFrameTime New
 			}
 			else
 			{
-				Stop();
-
-				// When playback stops naturally, the time cursor is put at the boundary that was crossed to make ping-pong playback easy
-				PlayPosition.Reset(NewPosition);
+				StopInternal(NewPosition);
 			}
 
 			TimeController->StopPlaying(GetCurrentTime());
@@ -776,8 +783,8 @@ void UMovieSceneSequencePlayer::ApplyLatentActions()
 	{
 		switch (LatentAction.Type)
 		{
-		case FLatentAction::Stop:          Stop();                         continue;
-		case FLatentAction::Pause:         Pause();                        continue;
+		case FLatentAction::Stop:          StopInternal(LatentAction.Position); continue;
+		case FLatentAction::Pause:         Pause();                             continue;
 		}
 
 		check(LatentAction.Type == FLatentAction::Update);
@@ -937,7 +944,7 @@ void UMovieSceneSequencePlayer::RPC_OnStopEvent_Implementation(FFrameTime Stoppe
 	case EMovieScenePlayerStatus::Scrubbing: ScrubToFrame(StoppedTime);  break;
 	}
 
-	Stop();
+	StopInternal(StoppedTime);
 }
 
 void UMovieSceneSequencePlayer::PostNetReceive()
diff --git a/Engine/Source/Runtime/MovieScene/Public/MovieSceneSequencePlayer.h b/Engine/Source/Runtime/MovieScene/Public/MovieSceneSequencePlayer.h
--- a/Engine/Source/Runtime/MovieScene/Public/MovieSceneSequencePlayer.h
+++ b/Engine/Source/Runtime/MovieScene/Public/MovieSceneSequencePlayer.h
@@ -195,10 +195,14 @@ public:
 	UFUNCTION(BlueprintCallable, Category="Game|Cinematic")
 	void Scrub();
 
-	/** Stop playback. */
+	/** Stop playback and move the cursor to the end (or start, for reversed playback) of the sequence. */
 	UFUNCTION(BlueprintCallable, Category="Game|Cinematic")
 	void Stop();
 
+	/** Stop playback without moving the cursor. */
+	UFUNCTION(BlueprintCallable, Category="Game|Cinematic")
+	void StopAtCurrentTime();
+
 	/** Go to end and stop. */
 	UFUNCTION(BlueprintCallable, Category="Game|Cinematic", meta = (ToolTip = "Go to end of the sequence and stop. Adheres to 'When Finished' section rules."))
 	void GoToEndAndStop();
@@ -478,6 +482,7 @@ public:
 protected:
 
 	void PlayInternal();
+	void StopInternal(FFrameTime TimeToResetTo);
 
 	void UpdateMovieSceneInstance(FMovieSceneEvaluationRange InRange, EMovieScenePlayerStatus::Type PlayerStatus, bool bHasJumped = false);
 
@@ -584,7 +589,7 @@ protected:
 	{
 		enum EType { Stop, Pause, Update };
 
-		FLatentAction(EType InType)
+		FLatentAction(EType InType, FFrameTime DesiredTime = 0)
 			: Type(InType)
 		{}
 
