Sequencer: Do not react to network replication if the sequence player is still pending initialization

[ Copy from 4586509 ]

diff --git a/Engine/Source/Runtime/MovieScene/Private/MovieSceneSequencePlayer.cpp b/Engine/Source/Runtime/MovieScene/Private/MovieSceneSequencePlayer.cpp
--- a/Engine/Source/Runtime/MovieScene/Private/MovieSceneSequencePlayer.cpp
+++ b/Engine/Source/Runtime/MovieScene/Private/MovieSceneSequencePlayer.cpp
@@ -913,9 +913,9 @@ void UMovieSceneSequencePlayer::RPC_OnStopEvent_Implementation(FFrameTime Stoppe
 	// ~~~~~~~~~~~~~~~~~~~~~~~~~~~~~~~~~~~~~~~~~~~~~~~~
 	// Handle an explicit Stop command from the server.
 
-	if (HasAuthority())
+	if (HasAuthority() || !Sequence)
 	{
-		// Never run network sync operations on authoritative players
+		// Never run network sync operations on authoritative players or players that have not been initialized yet
 		return;
 	}
 
@@ -954,9 +954,9 @@ void UMovieSceneSequencePlayer::PostNetReceive()
 
 	Super::PostNetReceive();
 
-	if (!ensure(!HasAuthority()))
+	if (!ensure(!HasAuthority()) || !Sequence)
 	{
-		// Never run network sync operations on authoritative players
+		// Never run network sync operations on authoritative players or players that have not been initialized yet
 		return;
 	}
 
