Sequencer: Custom float channel serialize

diff --git a/Engine/Source/Runtime/Core/Public/UObject/SequencerObjectVersion.h b/Engine/Source/Runtime/Core/Public/UObject/SequencerObjectVersion.h
--- a/Engine/Source/Runtime/Core/Public/UObject/SequencerObjectVersion.h
+++ b/Engine/Source/Runtime/Core/Public/UObject/SequencerObjectVersion.h
@@ -39,6 +39,9 @@ struct CORE_API FSequencerObjectVersion
 		// Finish UMG evaluation on end
 		FinishUMGEvaluation,
 
+		// Manual serialization of float channel
+		SerializeFloatChannel,
+
 		// -----<new versions can be added above this line>-------------------------------------------------
 		VersionPlusOne,
 		LatestVersion = VersionPlusOne - 1
diff --git a/Engine/Source/Runtime/MovieScene/Private/Channels/MovieSceneFloatChannel.cpp b/Engine/Source/Runtime/MovieScene/Private/Channels/MovieSceneFloatChannel.cpp
--- a/Engine/Source/Runtime/MovieScene/Private/Channels/MovieSceneFloatChannel.cpp
+++ b/Engine/Source/Runtime/MovieScene/Private/Channels/MovieSceneFloatChannel.cpp
@@ -4,7 +4,64 @@
 #include "MovieSceneFwd.h"
 #include "Channels/MovieSceneChannelProxy.h"
 #include "MovieSceneFrameMigration.h"
+#include "UObject/SequencerObjectVersion.h"
 
+bool FMovieSceneTangentData::Serialize(FArchive& Ar)
+{
+	Ar.UsingCustomVersion(FSequencerObjectVersion::GUID);
+	if (Ar.CustomVer(FSequencerObjectVersion::GUID) < FSequencerObjectVersion::SerializeFloatChannel)
+	{
+		return false;
+	}
+
+	// Serialization is handled manually to avoid the extra size overhead of UProperty tagging.
+	// Otherwise with many keys in a FMovieSceneTangentData the size can become quite large.
+	Ar << ArriveTangent;
+	Ar << LeaveTangent;
+	Ar << TangentWeightMode;
+	Ar << ArriveTangentWeight;
+	Ar << LeaveTangentWeight;
+
+	return true;
+}
+
+bool FMovieSceneTangentData::operator==(const FMovieSceneTangentData& TangentData) const
+{
+	return (ArriveTangent == TangentData.ArriveTangent) && (LeaveTangent == TangentData.LeaveTangent) && (TangentWeightMode == TangentData.TangentWeightMode) && (ArriveTangentWeight == TangentData.ArriveTangentWeight) && (LeaveTangentWeight == TangentData.LeaveTangentWeight);
+}
+
+bool FMovieSceneTangentData::operator!=(const FMovieSceneTangentData& Other) const
+{
+	return !(*this == Other);
+}
+
+bool FMovieSceneFloatValue::Serialize(FArchive& Ar)
+{
+	Ar.UsingCustomVersion(FSequencerObjectVersion::GUID);
+	if (Ar.CustomVer(FSequencerObjectVersion::GUID) < FSequencerObjectVersion::SerializeFloatChannel)
+	{
+		return false;
+	}
+
+	// Serialization is handled manually to avoid the extra size overhead of UProperty tagging.
+	// Otherwise with many keys in a FMovieSceneFloatValue the size can become quite large.
+	Ar << Value;
+	Ar << InterpMode;
+	Ar << TangentMode;
+	Ar << Tangent;
+
+	return true;
+}
+
+bool FMovieSceneFloatValue::operator==(const FMovieSceneFloatValue& FloatValue) const
+{
+	return (Value == FloatValue.Value) && (InterpMode == FloatValue.InterpMode) && (TangentMode == FloatValue.TangentMode) && (Tangent == FloatValue.Tangent);
+}
+
+bool FMovieSceneFloatValue::operator!=(const FMovieSceneFloatValue& Other) const
+{
+	return !(*this == Other);
+}
 
 bool FMovieSceneFloatChannel::SerializeFromMismatchedTag(const FPropertyTag& Tag, FStructuredArchive::FSlot Slot)
 {
diff --git a/Engine/Source/Runtime/MovieScene/Public/Channels/MovieSceneFloatChannel.h b/Engine/Source/Runtime/MovieScene/Public/Channels/MovieSceneFloatChannel.h
--- a/Engine/Source/Runtime/MovieScene/Public/Channels/MovieSceneFloatChannel.h
+++ b/Engine/Source/Runtime/MovieScene/Public/Channels/MovieSceneFloatChannel.h
@@ -13,7 +13,6 @@
 
 #include "MovieSceneFloatChannel.generated.h"
 
-
 USTRUCT()
 struct FMovieSceneTangentData
 {
@@ -27,6 +26,15 @@ struct FMovieSceneTangentData
 		, LeaveTangentWeight(0.f)
 	{}
 
+	bool Serialize(FArchive& Ar);
+	bool operator==(const FMovieSceneTangentData& Other) const;
+	bool operator!=(const FMovieSceneTangentData& Other) const;
+	friend FArchive& operator<<(FArchive& Ar, FMovieSceneTangentData& P)
+	{
+		P.Serialize(Ar);
+		return Ar;
+	}
+
 	/** If RCIM_Cubic, the arriving tangent at this key */
 	UPROPERTY(EditAnywhere, Category="Key")
 	float ArriveTangent;
@@ -49,6 +57,25 @@ struct FMovieSceneTangentData
 
 };
 
+template<>
+struct TIsPODType<FMovieSceneTangentData>
+{
+	enum { Value = true };
+};
+
+
+template<>
+struct TStructOpsTypeTraits<FMovieSceneTangentData>
+	: public TStructOpsTypeTraitsBase2<FMovieSceneTangentData>
+{
+	enum
+	{
+		WithSerializer = true,
+		WithCopy = false,
+		WithIdenticalViaEquality = true,
+	};
+};
+
 USTRUCT()
 struct FMovieSceneFloatValue
 {
@@ -62,6 +89,15 @@ struct FMovieSceneFloatValue
 		: Value(InValue), InterpMode(RCIM_Cubic), TangentMode(RCTM_Auto)
 	{}
 
+	bool Serialize(FArchive& Ar);
+	bool operator==(const FMovieSceneFloatValue& Other) const;
+	bool operator!=(const FMovieSceneFloatValue& Other) const;
+	friend FArchive& operator<<(FArchive& Ar, FMovieSceneFloatValue& P)
+	{
+		P.Serialize(Ar);
+		return Ar;
+	}
+
 	UPROPERTY(EditAnywhere, Category="Key")
 	float Value;
 
@@ -75,6 +111,25 @@ struct FMovieSceneFloatValue
 	FMovieSceneTangentData Tangent;
 };
 
+template<>
+struct TIsPODType<FMovieSceneFloatValue>
+{
+	enum { Value = true };
+};
+
+
+template<>
+struct TStructOpsTypeTraits<FMovieSceneFloatValue>
+	: public TStructOpsTypeTraitsBase2<FMovieSceneFloatValue>
+{
+	enum
+	{
+		WithSerializer = true,
+		WithCopy = false,
+		WithIdenticalViaEquality = true,
+	};
+};
+
 USTRUCT()
 struct MOVIESCENE_API FMovieSceneFloatChannel : public FMovieSceneChannel
 {
