Sequencer: Async load fixes for object path channels

  - Channel keys now store references to objects as both hard and soft references, which allows them to play nicely with async loading as well as keeping support for cross-level actor references

diff --git a/Engine/Source/Editor/MovieSceneTools/Private/TrackEditors/ObjectPropertyTrackEditor.cpp b/Engine/Source/Editor/MovieSceneTools/Private/TrackEditors/ObjectPropertyTrackEditor.cpp
--- a/Engine/Source/Editor/MovieSceneTools/Private/TrackEditors/ObjectPropertyTrackEditor.cpp
+++ b/Engine/Source/Editor/MovieSceneTools/Private/TrackEditors/ObjectPropertyTrackEditor.cpp
@@ -20,7 +20,7 @@ TSharedRef<ISequencerTrackEditor> FObjectPropertyTrackEditor::CreateTrackEditor(
 
 void FObjectPropertyTrackEditor::GenerateKeysFromPropertyChanged(const FPropertyChangedParams& PropertyChangedParams, FGeneratedTrackKeys& OutGeneratedKeys)
 {
-	TSoftObjectPtr<> Existing = PropertyChangedParams.GetPropertyValue<UObject*>();
+	UObject* Existing = PropertyChangedParams.GetPropertyValue<UObject*>();
 	OutGeneratedKeys.Add(FMovieSceneChannelValueSetter::Create<FMovieSceneObjectPathChannel>(0, Existing, true));
 }
 
diff --git a/Engine/Source/Editor/MovieSceneTools/Public/ClipboardTypes.h b/Engine/Source/Editor/MovieSceneTools/Public/ClipboardTypes.h
--- a/Engine/Source/Editor/MovieSceneTools/Public/ClipboardTypes.h
+++ b/Engine/Source/Editor/MovieSceneTools/Public/ClipboardTypes.h
@@ -9,6 +9,7 @@
 
 #include "Channels/MovieSceneFloatChannel.h"
 #include "Channels/MovieSceneEvent.h"
+#include "Channels/MovieSceneObjectPathChannel.h"
 #include "Sections/MovieSceneEventSection.h"
 #include "Sections/MovieSceneActorReferenceSection.h"
 #include "MovieSceneClipboard.h"
@@ -59,9 +60,9 @@ namespace MovieSceneClipboard
 	{
 		return "EventPayload";
 	}
-	template<> inline FName GetKeyTypeName<TSoftObjectPtr<UObject>>()
+	template<> inline FName GetKeyTypeName<FMovieSceneObjectPathChannelKeyValue>()
 	{
-		return "SoftObjectPtr";
+		return "MovieSceneObjectPathChannelKeyValue";
 	}
 	template<> inline FName GetKeyTypeName<FMovieSceneEvent>()
 	{
diff --git a/Engine/Source/Runtime/MovieScene/Private/Channels/MovieSceneObjectPathChannel.cpp b/Engine/Source/Runtime/MovieScene/Private/Channels/MovieSceneObjectPathChannel.cpp
--- a/Engine/Source/Runtime/MovieScene/Private/Channels/MovieSceneObjectPathChannel.cpp
+++ b/Engine/Source/Runtime/MovieScene/Private/Channels/MovieSceneObjectPathChannel.cpp
@@ -2,34 +2,56 @@
 
 #include "Channels/MovieSceneObjectPathChannel.h"
 
-bool FMovieSceneObjectPathChannel::Evaluate(FFrameTime InTime, UObject*& OutValue) const
+bool FMovieSceneObjectPathChannelKeyValue::SerializeFromMismatchedTag(const FPropertyTag& Tag, FStructuredArchive::FSlot Slot)
 {
-	if (Times.Num())
+	if (Tag.Type == NAME_SoftObjectProperty)
 	{
-		const int32 Index = FMath::Max(0, Algo::UpperBound(Times, InTime.FrameNumber)-1);
-		OutValue = Values[Index].LoadSynchronous();
+		FSoftObjectPtr OldProperty;
+		Slot << OldProperty;
+
+		SoftPtr = OldProperty.ToSoftObjectPath();
+		if (OldProperty.ToSoftObjectPath().GetSubPathString().Len() == 0)
+		{
+			// Forcibly load the old property so we can store it as a hard reference, only if it was not referencing an actor or other sub object
+			HardPtr = OldProperty.LoadSynchronous();
+		}
 		return true;
 	}
-	else if (!DefaultValue.IsNull())
+	return false;
+}
+
+FMovieSceneObjectPathChannelKeyValue& FMovieSceneObjectPathChannelKeyValue::operator=(UObject* NewObject)
+{
+	HardPtr = NewObject;
+	SoftPtr = NewObject;
+	return *this;
+}
+
+UObject* FMovieSceneObjectPathChannelKeyValue::Get() const
+{
+	if (!HardPtr && !SoftPtr.IsNull())
 	{
-		OutValue = DefaultValue.LoadSynchronous();
-		return true;
+		HardPtr = SoftPtr.Get();
+		if (!HardPtr)
+		{
+			HardPtr = SoftPtr.LoadSynchronous();
+		}
 	}
 
-	return false;
+	return HardPtr;
 }
 
-bool FMovieSceneObjectPathChannel::Evaluate(FFrameTime InTime, TSoftObjectPtr<>& OutValue) const
+bool FMovieSceneObjectPathChannel::Evaluate(FFrameTime InTime, UObject*& OutValue) const
 {
 	if (Times.Num())
 	{
 		const int32 Index = FMath::Max(0, Algo::UpperBound(Times, InTime.FrameNumber)-1);
-		OutValue = Values[Index];
+		OutValue = Values[Index].Get();
 		return true;
 	}
-	else if (!DefaultValue.IsNull())
+	else if (!DefaultValue.GetSoftPtr().IsNull())
 	{
-		OutValue = DefaultValue;
+		OutValue = DefaultValue.Get();
 		return true;
 	}
 
diff --git a/Engine/Source/Runtime/MovieScene/Public/Channels/MovieSceneObjectPathChannel.h b/Engine/Source/Runtime/MovieScene/Public/Channels/MovieSceneObjectPathChannel.h
--- a/Engine/Source/Runtime/MovieScene/Public/Channels/MovieSceneObjectPathChannel.h
+++ b/Engine/Source/Runtime/MovieScene/Public/Channels/MovieSceneObjectPathChannel.h
@@ -8,6 +8,66 @@
 #include "MovieSceneChannelTraits.h"
 #include "MovieSceneObjectPathChannel.generated.h"
 
+/**
+ * Key value type for object path channels that stores references to objects as both a hard and soft reference, to ensure compatability with both sub objects and async loading
+ */
+USTRUCT()
+struct MOVIESCENE_API FMovieSceneObjectPathChannelKeyValue
+{
+	GENERATED_BODY()
+
+	/** Default constructor */
+	FMovieSceneObjectPathChannelKeyValue()
+		: HardPtr(nullptr)
+	{}
+
+	/** Construction from an object pointer */
+	FMovieSceneObjectPathChannelKeyValue(UObject* InObject)
+		: SoftPtr(InObject)
+		, HardPtr(InObject)
+	{}
+
+	/**
+	 * Assignment from a raw object pointer
+	 */
+	FMovieSceneObjectPathChannelKeyValue& operator=(UObject* NewObject);
+
+public:
+
+	/**
+	 * Legacy conversion from a TSoftObjectPtr<>
+	 */
+	bool SerializeFromMismatchedTag(FPropertyTag const& Tag, FStructuredArchive::FSlot Slot);
+	
+	/**
+	 * Access the soft object pointer that this key should load
+	 */
+	const TSoftObjectPtr<>& GetSoftPtr() const
+	{
+		return SoftPtr;
+	}
+
+	/**
+	 * Attempt to find this object either by returning the internally kept raw pointer, or by resolving (but not loading) the soft object path
+	 */
+	UObject* Get() const;
+
+private:
+
+	/** Persistent storage of the object by path (which allows us to support cross-level actor references, for instance) */
+	UPROPERTY()
+	TSoftObjectPtr<UObject> SoftPtr;
+
+	/** Hard reference to the loaded object - relevant for any asset type which also hints the async loader to efficiently load the asset in advance */
+	UPROPERTY()
+	mutable UObject* HardPtr;
+};
+
+template<>
+struct TStructOpsTypeTraits<FMovieSceneObjectPathChannelKeyValue> : public TStructOpsTypeTraitsBase2<FMovieSceneObjectPathChannelKeyValue>
+{
+	enum { WithStructuredSerializeFromMismatchedTag = true };
+};
 
 USTRUCT()
 struct MOVIESCENE_API FMovieSceneObjectPathChannel : public FMovieSceneChannel
@@ -34,9 +94,9 @@ struct MOVIESCENE_API FMovieSceneObjectPathChannel : public FMovieSceneChannel
 	 *
 	 * @return An object that is able to manipulate this channel's data
 	 */
-	FORCEINLINE TMovieSceneChannelData<TSoftObjectPtr<>> GetData()
+	FORCEINLINE TMovieSceneChannelData<FMovieSceneObjectPathChannelKeyValue> GetData()
 	{
-		return TMovieSceneChannelData<TSoftObjectPtr<>>(&Times, &Values, &KeyHandles);
+		return TMovieSceneChannelData<FMovieSceneObjectPathChannelKeyValue>(&Times, &Values, &KeyHandles);
 	}
 
 	/**
@@ -44,9 +104,9 @@ struct MOVIESCENE_API FMovieSceneObjectPathChannel : public FMovieSceneChannel
 	 *
 	 * @return An object that is able to interrogate this channel's data
 	 */
-	FORCEINLINE TMovieSceneChannelData<const TSoftObjectPtr<>> GetData() const
+	FORCEINLINE TMovieSceneChannelData<const FMovieSceneObjectPathChannelKeyValue> GetData() const
 	{
-		return TMovieSceneChannelData<const TSoftObjectPtr<>>(&Times, &Values);
+		return TMovieSceneChannelData<const FMovieSceneObjectPathChannelKeyValue>(&Times, &Values);
 	}
 
 	/**
@@ -56,7 +116,6 @@ struct MOVIESCENE_API FMovieSceneObjectPathChannel : public FMovieSceneChannel
 	 * @param OutValue   A value to receive the result
 	 * @return true if the channel was evaluated successfully, false otherwise
 	 */
-	bool Evaluate(FFrameTime InTime, TSoftObjectPtr<>& OutValue) const;
 	bool Evaluate(FFrameTime InTime, UObject*& OutValue) const;
 
 public:
@@ -82,7 +141,7 @@ public:
 	 *
 	 * @param InDefaultValue The desired default value
 	 */
-	FORCEINLINE void SetDefault(const TSoftObjectPtr<>& InDefaultValue)
+	FORCEINLINE void SetDefault(UObject* InDefaultValue)
 	{
 		DefaultValue = InDefaultValue;
 	}
@@ -92,7 +151,7 @@ public:
 	 *
 	 * @return (Optional) The channel's default value
 	 */
-	FORCEINLINE const TSoftObjectPtr<>& GetDefault() const
+	FORCEINLINE const FMovieSceneObjectPathChannelKeyValue& GetDefault() const
 	{
 		return DefaultValue;
 	}
@@ -102,7 +161,7 @@ public:
 	 */
 	FORCEINLINE void RemoveDefault()
 	{
-		DefaultValue.Reset();
+		DefaultValue = nullptr;
 	}
 
 private:
@@ -114,10 +173,10 @@ private:
 	TArray<FFrameNumber> Times;
 
 	UPROPERTY(meta=(KeyValues))
-	TArray<TSoftObjectPtr<UObject>> Values;
+	TArray<FMovieSceneObjectPathChannelKeyValue> Values;
 
 	UPROPERTY()
-	TSoftObjectPtr<UObject> DefaultValue;
+	FMovieSceneObjectPathChannelKeyValue DefaultValue;
 
 	FMovieSceneKeyHandleMap KeyHandles;
 };
diff --git a/Engine/Source/Runtime/MovieSceneTracks/Private/Evaluation/MovieSceneObjectPropertyTemplate.cpp b/Engine/Source/Runtime/MovieSceneTracks/Private/Evaluation/MovieSceneObjectPropertyTemplate.cpp
--- a/Engine/Source/Runtime/MovieSceneTracks/Private/Evaluation/MovieSceneObjectPropertyTemplate.cpp
+++ b/Engine/Source/Runtime/MovieSceneTracks/Private/Evaluation/MovieSceneObjectPropertyTemplate.cpp
@@ -7,8 +7,8 @@
 
 struct FObjectPropertyExecToken : IMovieSceneExecutionToken
 {
-	FObjectPropertyExecToken(TSoftObjectPtr<>&& InValue)
-		: Value(MoveTemp(InValue))
+	FObjectPropertyExecToken(UObject* InValue)
+		: NewObjectValue(MoveTemp(InValue))
 	{}
 
 	virtual void Execute(const FMovieSceneContext& Context, const FMovieSceneEvaluationOperand& Operand, FPersistentEvaluationData& PersistentData, IMovieScenePlayer& Player) override
@@ -19,9 +19,6 @@ struct FObjectPropertyExecToken : IMovieSceneExecutionToken
 		FTrackInstancePropertyBindings* PropertyBindings = PropertyTrackData.PropertyBindings.Get();
 
 		check(PropertyBindings);
-
-		UObject* NewObjectValue = Value.LoadSynchronous();
-
 		for (TWeakObjectPtr<> WeakObject : Player.FindBoundObjects(Operand))
 		{
 			UObject* ObjectPtr = WeakObject.Get();
@@ -63,7 +60,7 @@ struct FObjectPropertyExecToken : IMovieSceneExecutionToken
 		}
 	}
 
-	TSoftObjectPtr<> Value;
+	UObject* NewObjectValue;
 };
 
 
@@ -80,9 +77,9 @@ void FMovieSceneObjectPropertyTemplate::SetupOverrides()
 
 void FMovieSceneObjectPropertyTemplate::Evaluate(const FMovieSceneEvaluationOperand& Operand, const FMovieSceneContext& Context, const FPersistentEvaluationData& PersistentData, FMovieSceneExecutionTokens& ExecutionTokens) const
 {
-	TSoftObjectPtr<> Ptr;
-	if (ObjectChannel.Evaluate(Context.GetTime(), Ptr))
+	UObject* Object = nullptr;
+	if (ObjectChannel.Evaluate(Context.GetTime(), Object))
 	{
-		ExecutionTokens.Add(FObjectPropertyExecToken(MoveTemp(Ptr)));
+		ExecutionTokens.Add(FObjectPropertyExecToken(Object));
 	}
 }
diff --git a/Engine/Source/Runtime/MovieSceneTracks/Private/Evaluation/MovieScenePrimitiveMaterialTemplate.cpp b/Engine/Source/Runtime/MovieSceneTracks/Private/Evaluation/MovieScenePrimitiveMaterialTemplate.cpp
--- a/Engine/Source/Runtime/MovieSceneTracks/Private/Evaluation/MovieScenePrimitiveMaterialTemplate.cpp
+++ b/Engine/Source/Runtime/MovieSceneTracks/Private/Evaluation/MovieScenePrimitiveMaterialTemplate.cpp
@@ -43,17 +43,15 @@ struct FSetMaterialTokenProducer : IMovieScenePreAnimatedTokenProducer
 struct FPrimitiveMaterialExecToken : IMovieSceneExecutionToken
 {
 	int32 MaterialIndex;
-	TSoftObjectPtr<> Value;
+	UMaterialInterface* NewMaterial;
 
-	FPrimitiveMaterialExecToken(int32 InMaterialIndex, TSoftObjectPtr<>&& InValue)
+	FPrimitiveMaterialExecToken(int32 InMaterialIndex, UMaterialInterface* InNewMaterial)
 		: MaterialIndex(InMaterialIndex)
-		, Value(MoveTemp(InValue))
+		, NewMaterial(InNewMaterial)
 	{}
 
 	virtual void Execute(const FMovieSceneContext& Context, const FMovieSceneEvaluationOperand& Operand, FPersistentEvaluationData& PersistentData, IMovieScenePlayer& Player) override
 	{
-		UMaterialInterface* NewMaterial = Cast<UMaterialInterface>(Value.LoadSynchronous());
-
 		for (TWeakObjectPtr<> WeakObject : Player.FindBoundObjects(Operand))
 		{
 			UPrimitiveComponent* PrimitiveComponent = Cast<UPrimitiveComponent>(WeakObject.Get());
@@ -79,9 +77,13 @@ FMovieScenePrimitiveMaterialTemplate::FMovieScenePrimitiveMaterialTemplate(const
 
 void FMovieScenePrimitiveMaterialTemplate::Evaluate(const FMovieSceneEvaluationOperand& Operand, const FMovieSceneContext& Context, const FPersistentEvaluationData& PersistentData, FMovieSceneExecutionTokens& ExecutionTokens) const
 {
-	TSoftObjectPtr<> Ptr;
+	UObject* Ptr = nullptr;
 	if (MaterialChannel.Evaluate(Context.GetTime(), Ptr))
 	{
-		ExecutionTokens.Add(FPrimitiveMaterialExecToken(MaterialIndex, MoveTemp(Ptr)));
+		// If the channel has been successfully evaluated, only assign the object if it's null, or a valid material interface
+		if (Ptr == nullptr || Ptr->IsA<UMaterialInterface>())
+		{
+			ExecutionTokens.Add(FPrimitiveMaterialExecToken(MaterialIndex, Cast<UMaterialInterface>(Ptr)));
+		}
 	}
 }
