Sequencer: Prevent forcibly synchronizing time when server/client times differ due to natural looping.

  - When the server looped back to the start of the sequence (or end, when playing in reverse), clients that were very close to the end of the sequence were detecting a synchronization difference and forcibly playing to the new time. This resulting in sweeping backwards through the entire sequence back to the start.
  - Server loop counts are now sent to clients to determine whether a difference in time genuinely needs synchronizing, or can be ignored due to an impending loop. This is computed by offsetting the server time by the difference in loop counts.
  - Forced synchronization events now only play forwards, and jump backwards (or the opposite, if plaing in reverse)
  - Fixed incorrect network handling of ::Stopped as ::Pause

diff --git a/Engine/Source/Runtime/MovieScene/Private/MovieSceneSequencePlayer.cpp b/Engine/Source/Runtime/MovieScene/Private/MovieSceneSequencePlayer.cpp
--- a/Engine/Source/Runtime/MovieScene/Private/MovieSceneSequencePlayer.cpp
+++ b/Engine/Source/Runtime/MovieScene/Private/MovieSceneSequencePlayer.cpp
@@ -64,6 +64,7 @@ void UMovieSceneSequencePlayer::UpdateNetworkSyncProperties()
 	{
 		NetSyncProps.LastKnownPosition = PlayPosition.GetCurrentPosition();
 		NetSyncProps.LastKnownStatus   = Status;
+		NetSyncProps.LastKnownNumLoops = CurrentNumLoops;
 	}
 }
 
@@ -75,6 +76,7 @@ void UMovieSceneSequencePlayer::GetLifetimeReplicatedProps(TArray<FLifetimePrope
 	DOREPLIFETIME(UMovieSceneSequencePlayer, bReversePlayback);
 	DOREPLIFETIME(UMovieSceneSequencePlayer, StartTime);
 	DOREPLIFETIME(UMovieSceneSequencePlayer, DurationFrames);
+	DOREPLIFETIME(UMovieSceneSequencePlayer, PlaybackSettings);
 }
 
 EMovieScenePlayerStatus::Type UMovieSceneSequencePlayer::GetPlaybackStatus() const
@@ -1019,9 +1021,7 @@ void UMovieSceneSequencePlayer::PostNetReceive()
 		return;
 	}
 
-	// Stop() is explicitly sent through RPC, so if there is a change of state from Playing to Stopped, it must be a pause
-	const bool bHasPaused         = NetSyncProps.LastKnownStatus == EMovieScenePlayerStatus::Stopped && Status == EMovieScenePlayerStatus::Playing;
-	const bool bHasStartedPlaying = NetSyncProps.LastKnownStatus == EMovieScenePlayerStatus::Playing && Status == EMovieScenePlayerStatus::Stopped;
+	const bool bHasStartedPlaying = NetSyncProps.LastKnownStatus == EMovieScenePlayerStatus::Playing && Status != EMovieScenePlayerStatus::Playing;
 	const bool bHasChangedStatus  = NetSyncProps.LastKnownStatus   != Status;
 	const bool bHasChangedTime    = NetSyncProps.LastKnownPosition != PlayPosition.GetCurrentPosition();
 
@@ -1066,27 +1066,48 @@ void UMovieSceneSequencePlayer::PostNetReceive()
 			PlayToFrame(NetSyncProps.LastKnownPosition);
 		}
 	}
-	else if (bHasPaused)
-	{
-		// Continue playing all the way up to the time that it has stopped, then stop playback
-		PlayToFrame(NetSyncProps.LastKnownPosition);
-		Pause();
-	}
 	else
 	{
-		static float     ThresholdS     = 0.2f;
-		const FFrameTime FrameThreshold = ThresholdS * PlayPosition.GetInputRate();
-		const FFrameTime Difference     = FMath::Abs(PlayPosition.GetCurrentPosition() - NetSyncProps.LastKnownPosition);
-
 		if (bHasChangedTime)
 		{
 			// Make sure the client time matches the server according to the client's current status
 			if (Status == EMovieScenePlayerStatus::Playing)
 			{
-				if (Difference > FrameThreshold)
+				// When the server has looped back to the start but a client is near the end (and is thus about to loop), we don't want to forcibly synchronize the time unless
+				// the *real* difference in time is above the threshold. We compute the real-time difference by adding SequenceDuration*LoopCountDifference to the server position:
+				//		start	srv_time																																clt_time		end
+				//		0		1		2		3		4		5		6		7		8		9		10		11		12		13		14		15		16		17		18		19		20
+				//		|		|																																		|				|
+				//
+				//		Let NetSyncProps.LastKnownNumLoops = 1, CurrentNumLoops = 0, bReversePlayback = false
+				//			=> LoopOffset = 1
+				//			   OffsetServerTime = srv_time + FrameDuration*LoopOffset = 1 + 20*1 = 21
+				//			   Difference = 21 - 18 = 3 frames
+				static float       ThresholdS       = 0.2f;
+				const int32        LoopOffset       = (NetSyncProps.LastKnownNumLoops - CurrentNumLoops) * (bReversePlayback ? -1 : 1);
+				const FFrameTime   FrameThreshold   = ThresholdS * PlayPosition.GetInputRate();
+				const FFrameTime   OffsetServerTime = NetSyncProps.LastKnownPosition + GetFrameDuration()*LoopOffset;
+				const FFrameTime   Difference       = FMath::Abs(PlayPosition.GetCurrentPosition() - OffsetServerTime);
+
+				if (bHasChangedStatus)
 				{
+					// If the status has changed forcibly play to the server position before setting the new status
 					PlayToFrame(NetSyncProps.LastKnownPosition);
 				}
+				else if (Difference > FrameThreshold)
+				{
+					// We're drastically out of sync with the server so we need to forcibly set the time.
+					// Play to the time only if it is further on in the sequence (in our play direction)
+					const bool bPlayToFrame = bReversePlayback ? NetSyncProps.LastKnownPosition < PlayPosition.GetCurrentPosition() : NetSyncProps.LastKnownPosition > PlayPosition.GetCurrentPosition();
+					if (bPlayToFrame)
+					{
+						PlayToFrame(NetSyncProps.LastKnownPosition);
+					}
+					else
+					{
+						JumpToFrame(NetSyncProps.LastKnownPosition);
+					}
+				}
 			}
 			else if (Status == EMovieScenePlayerStatus::Stopped)
 			{
@@ -1100,10 +1121,9 @@ void UMovieSceneSequencePlayer::PostNetReceive()
 
 		if (bHasChangedStatus)
 		{
-			// ::Stopped means Pause here, since Stop() is handled explicitly by RPC_OnStopEvent
 			switch (NetSyncProps.LastKnownStatus)
 			{
-			case EMovieScenePlayerStatus::Stopped:   Pause();  break;
+			case EMovieScenePlayerStatus::Paused:    Pause(); break;
 			case EMovieScenePlayerStatus::Playing:   Play();  break;
 			case EMovieScenePlayerStatus::Scrubbing: Scrub(); break;
 			}
diff --git a/Engine/Source/Runtime/MovieScene/Public/MovieSceneSequencePlayer.h b/Engine/Source/Runtime/MovieScene/Public/MovieSceneSequencePlayer.h
--- a/Engine/Source/Runtime/MovieScene/Public/MovieSceneSequencePlayer.h
+++ b/Engine/Source/Runtime/MovieScene/Public/MovieSceneSequencePlayer.h
@@ -65,6 +65,11 @@ struct FMovieSceneSequenceReplProperties
 {
 	GENERATED_BODY()
 
+	FMovieSceneSequenceReplProperties()
+		: LastKnownStatus(EMovieScenePlayerStatus::Stopped)
+		, LastKnownNumLoops(0)
+	{}
+
 	/** The last known position of the sequence on the server */
 	UPROPERTY()
 	FFrameTime LastKnownPosition;
@@ -72,6 +77,10 @@ struct FMovieSceneSequenceReplProperties
 	/** The last known playback status of the sequence on the server */
 	UPROPERTY()
 	TEnumAsByte<EMovieScenePlayerStatus::Type> LastKnownStatus;
+
+	/** The last known number of loops of the sequence on the server */
+	UPROPERTY()
+	int32 LastKnownNumLoops;
 };
 
 
@@ -636,7 +645,7 @@ protected:
 	TArray<FLatentAction> LatentActions;
 
 	/** Specific playback settings for the animation. */
-	UPROPERTY()
+	UPROPERTY(replicated)
 	FMovieSceneSequencePlaybackSettings PlaybackSettings;
 
 	/** The root template instance we're evaluating */
