Sequencer: Fixed object path channel holding a raw ptr to actors which may get forcibly destroyed after PIE

diff --git a/Engine/Source/Runtime/MovieScene/Private/Channels/MovieSceneObjectPathChannel.cpp b/Engine/Source/Runtime/MovieScene/Private/Channels/MovieSceneObjectPathChannel.cpp
--- a/Engine/Source/Runtime/MovieScene/Private/Channels/MovieSceneObjectPathChannel.cpp
+++ b/Engine/Source/Runtime/MovieScene/Private/Channels/MovieSceneObjectPathChannel.cpp
@@ -1,6 +1,7 @@
 // Copyright 1998-2019 Epic Games, Inc. All Rights Reserved.
 
 #include "Channels/MovieSceneObjectPathChannel.h"
+#include "Engine/World.h"
 
 bool FMovieSceneObjectPathChannelKeyValue::SerializeFromMismatchedTag(const FPropertyTag& Tag, FStructuredArchive::FSlot Slot)
 {
@@ -13,7 +14,13 @@ bool FMovieSceneObjectPathChannelKeyValue::SerializeFromMismatchedTag(const FPro
 		if (OldProperty.ToSoftObjectPath().GetSubPathString().Len() == 0)
 		{
 			// Forcibly load the old property so we can store it as a hard reference, only if it was not referencing an actor or other sub object
-			HardPtr = OldProperty.LoadSynchronous();
+			UObject* RawObject = OldProperty.LoadSynchronous();
+
+			// Do not store raw ptrs to actors or other objects that exist in worlds
+			if (RawObject && !RawObject->GetTypedOuter<UWorld>())
+			{
+				HardPtr = RawObject;
+			}
 		}
 		return true;
 	}
@@ -22,7 +29,16 @@ bool FMovieSceneObjectPathChannelKeyValue::SerializeFromMismatchedTag(const FPro
 
 FMovieSceneObjectPathChannelKeyValue& FMovieSceneObjectPathChannelKeyValue::operator=(UObject* NewObject)
 {
-	HardPtr = NewObject;
+	// Do not store raw ptrs to actors or other objects that exist in worlds
+	if (!NewObject || NewObject->GetTypedOuter<UWorld>())
+	{
+		HardPtr = nullptr;
+	}
+	else
+	{
+		HardPtr = NewObject;
+	}
+
 	SoftPtr = NewObject;
 	return *this;
 }
@@ -31,11 +47,18 @@ UObject* FMovieSceneObjectPathChannelKeyValue::Get() const
 {
 	if (!HardPtr && !SoftPtr.IsNull())
 	{
-		HardPtr = SoftPtr.Get();
-		if (!HardPtr)
+		UObject* ResolvedPtr = SoftPtr.Get();
+		if (!ResolvedPtr)
+		{
+			ResolvedPtr = SoftPtr.LoadSynchronous();
+		}
+
+		// Do not store raw ptrs to actors or other objects that exist in worlds
+		if (ResolvedPtr && !ResolvedPtr->GetTypedOuter<UWorld>())
 		{
-			HardPtr = SoftPtr.LoadSynchronous();
+			HardPtr = ResolvedPtr;
 		}
+		return ResolvedPtr;
 	}
 
 	return HardPtr;
