Sequencer: Added check to ensure that only valid child templates are set up

diff --git a/Engine/Source/Runtime/MovieScene/Private/Evaluation/MovieSceneEvaluationTemplateInstance.cpp b/Engine/Source/Runtime/MovieScene/Private/Evaluation/MovieSceneEvaluationTemplateInstance.cpp
--- a/Engine/Source/Runtime/MovieScene/Private/Evaluation/MovieSceneEvaluationTemplateInstance.cpp
+++ b/Engine/Source/Runtime/MovieScene/Private/Evaluation/MovieSceneEvaluationTemplateInstance.cpp
@@ -542,7 +542,7 @@ void FMovieSceneRootEvaluationTemplateInstance::CallSetupTearDown(const FMovieSc
 			{
 				Track->OnBeginEvaluation(PersistentDataProxy, Player);
 			}
-			else
+			else if (Track->HasChildTemplate(Key.SectionIndex))
 			{
 				PersistentDataProxy.SetSectionKey(Key);
 				Track->GetChildTemplate(Key.SectionIndex).OnBeginEvaluation(PersistentDataProxy, Player);
