Sequencer: Do not react to network replication if the sequence player is still pending initialization

diff --git a/Engine/Source/Runtime/MovieScene/Private/MovieSceneSequencePlayer.cpp b/Engine/Source/Runtime/MovieScene/Private/MovieSceneSequencePlayer.cpp
--- a/Engine/Source/Runtime/MovieScene/Private/MovieSceneSequencePlayer.cpp
+++ b/Engine/Source/Runtime/MovieScene/Private/MovieSceneSequencePlayer.cpp
@@ -870,7 +870,7 @@ void UMovieSceneSequencePlayer::RPC_ExplicitServerUpdateEvent_Implementation(EUp
 	// ~~~~~~~~~~~~~~~~~~~~~~~~~~~~~~~~~~~~~~~~~~~~~~~~~~~~~~~~~~~
 	// Handle an explicit jump/play/scrub command from the server.
 
-	if (HasAuthority())
+	if (HasAuthority() || !Sequence)
 	{
 		// Never run network sync operations on authoritative players
 		return;
