Sequencer: Reset loop count on stop so that it doesn't persist to the next call to play.

diff --git a/Engine/Source/Runtime/MovieScene/Private/MovieSceneSequencePlayer.cpp b/Engine/Source/Runtime/MovieScene/Private/MovieSceneSequencePlayer.cpp
--- a/Engine/Source/Runtime/MovieScene/Private/MovieSceneSequencePlayer.cpp
+++ b/Engine/Source/Runtime/MovieScene/Private/MovieSceneSequencePlayer.cpp
@@ -288,6 +288,9 @@ void UMovieSceneSequencePlayer::StopInternal(FFrameTime TimeToResetTo)
 
 		CurrentNumLoops = 0;
 
+		// Reset loop count on stop so that it doesn't persist to the next call to play
+		PlaybackSettings.LoopCount.Value = 0;
+
 		if (PlaybackSettings.bRestoreState)
 		{
 			RestorePreAnimatedState();
