Sequencer: Handle legacy conversions where DisplayRate is 0.

diff --git a/Engine/Source/Runtime/MovieScene/Private/MovieScene.cpp b/Engine/Source/Runtime/MovieScene/Private/MovieScene.cpp
--- a/Engine/Source/Runtime/MovieScene/Private/MovieScene.cpp
+++ b/Engine/Source/Runtime/MovieScene/Private/MovieScene.cpp
@@ -103,8 +103,8 @@ void UMovieScene::Serialize( FArchive& Ar )
 		else if (FixedFrameInterval_DEPRECATED == 1 / 120.0f) { DisplayRate = FCommonFrameRates::FPS_120(); }
 		else if (FixedFrameInterval_DEPRECATED != 0.f)
 		{
-			uint32 Numerator = FMath::RoundToInt(1.f / FixedFrameInterval_DEPRECATED);
-			DisplayRate = FFrameRate(Numerator, 1);
+			uint32 Numerator = FMath::RoundToInt(1000.f / FixedFrameInterval_DEPRECATED);
+			DisplayRate = FFrameRate(Numerator, 1000);
 		}
 		else
 		{
