Added initializer for member variable

diff --git a/Engine/Source/Runtime/MovieScene/Public/MovieSceneSequencePlayer.h b/Engine/Source/Runtime/MovieScene/Public/MovieSceneSequencePlayer.h
--- a/Engine/Source/Runtime/MovieScene/Public/MovieSceneSequencePlayer.h
+++ b/Engine/Source/Runtime/MovieScene/Public/MovieSceneSequencePlayer.h
@@ -66,7 +66,8 @@ struct FMovieSceneSequenceReplProperties
 	GENERATED_BODY()
 
 	FMovieSceneSequenceReplProperties()
-		: LastKnownNumLoops(0)
+		: LastKnownStatus(EMovieScenePlayerStatus::Stopped)
+		, LastKnownNumLoops(0)
 	{}
 
 	/** The last known position of the sequence on the server */
