Fix for linear interpolation to work with surrounding cubics by setting it's tangent in AutoSetTangents and only forcing linear if the next nex isn't cubicr.
Added a version to SequencerObjectVersion to keep the old behavior on load by breaking tangents of the cubic to have the previous one act like it's linear.  Also added a CVar in case we ever want the old behavior.
Also fixed a bug where the Key Details would not call AutoSetTangent, by adding a method to FMovieSceneChannel to call PostEdits on Data change.

diff --git a/Engine/Source/Runtime/Core/Public/UObject/SequencerObjectVersion.h b/Engine/Source/Runtime/Core/Public/UObject/SequencerObjectVersion.h
--- a/Engine/Source/Runtime/Core/Public/UObject/SequencerObjectVersion.h
+++ b/Engine/Source/Runtime/Core/Public/UObject/SequencerObjectVersion.h
@@ -42,6 +42,9 @@ struct CORE_API FSequencerObjectVersion
 		// Manual serialization of float channel
 		SerializeFloatChannel,
 
+		// Change the linear keys so they act the old way and interpolate always.
+		ModifyLinearKeysForOldInterp,
+
 		// -----<new versions can be added above this line>-------------------------------------------------
 		VersionPlusOne,
 		LatestVersion = VersionPlusOne - 1
diff --git a/Engine/Source/Runtime/MovieScene/Private/Channels/MovieSceneFloatChannel.cpp b/Engine/Source/Runtime/MovieScene/Private/Channels/MovieSceneFloatChannel.cpp
--- a/Engine/Source/Runtime/MovieScene/Private/Channels/MovieSceneFloatChannel.cpp
+++ b/Engine/Source/Runtime/MovieScene/Private/Channels/MovieSceneFloatChannel.cpp
@@ -5,6 +5,13 @@
 #include "Channels/MovieSceneChannelProxy.h"
 #include "MovieSceneFrameMigration.h"
 #include "UObject/SequencerObjectVersion.h"
+#include "HAL/ConsoleManager.h"
+
+static TAutoConsoleVariable<int32> CVarSequencerLinearCubicInterpolation(
+	TEXT("Sequencer.LinearCubicInterpolation"),
+	1,
+	TEXT("If 1 Linear Keys Act As Cubic Interpolation with Linear Tangents, if 0 Linear Key Forces Linear Interpolation to Next Key."),
+	ECVF_Default);
 
 bool FMovieSceneTangentData::Serialize(FArchive& Ar)
 {
@@ -183,14 +190,15 @@ static float EvalForTwoKeys(const FMovieSceneFloatValue& Key1, FFrameNumber Key1
 
 	float Diff = (float)(Key2Time - Key1Time).Value;
 	Diff /= DecimalRate;
-	
+	const int CheckBothLinear = CVarSequencerLinearCubicInterpolation->GetInt();
+
 	if (Diff > 0 && Key1.InterpMode != RCIM_Constant)
 	{
 		const float Alpha = ((float)(InTime - Key1Time).Value / DecimalRate) / Diff;
 		const float P0 = Key1.Value;
 		const float P3 = Key2.Value;
 
-		if (Key1.InterpMode == RCIM_Linear)
+		if (Key1.InterpMode == RCIM_Linear && (!CheckBothLinear || Key2.InterpMode != RCIM_Cubic))
 		{
 			return FMath::Lerp(P0, P3, Alpha);
 		}
@@ -521,6 +529,7 @@ bool FMovieSceneFloatChannel::Evaluate(FFrameTime InTime,  float& OutValue) cons
 	float Interp = 0.f;
 	int32 Index1 = INDEX_NONE, Index2 = INDEX_NONE;
 	MovieScene::EvaluateTime(Times, Params.Time, Index1, Index2, Interp);
+	const int CheckBothLinear = CVarSequencerLinearCubicInterpolation->GetInt();
 
 	if (Index1 == INDEX_NONE)
 	{
@@ -534,8 +543,13 @@ bool FMovieSceneFloatChannel::Evaluate(FFrameTime InTime,  float& OutValue) cons
 	{
 		FMovieSceneFloatValue Key1 = Values[Index1];
 		FMovieSceneFloatValue Key2 = Values[Index2];
-
-		switch (Key1.InterpMode)
+		TEnumAsByte<ERichCurveInterpMode> InterpMode = Key1.InterpMode;
+	    if(InterpMode == RCIM_Linear && (CheckBothLinear  && Key2.InterpMode == RCIM_Cubic))
+		{
+			InterpMode = RCIM_Cubic;
+		}
+		
+		switch (InterpMode)
 		{
 		case RCIM_Cubic:
 		{
@@ -667,7 +681,15 @@ void FMovieSceneFloatChannel::AutoSetTangents(float Tension)
 
 	{
 		FMovieSceneFloatValue& FirstValue = Values[0];
-		if (FirstValue.InterpMode == RCIM_Cubic && FirstValue.TangentMode == RCTM_Auto)
+		if (FirstValue.InterpMode == RCIM_Linear)
+		{
+			FirstValue.Tangent.TangentWeightMode = RCTWM_WeightedNone;
+			FMovieSceneFloatValue& NextKey = Values[1];
+			const float NextTimeDiff = FMath::Max<double>(KINDA_SMALL_NUMBER, Times[1].Value - Times[0].Value);
+			float NewTangent = (NextKey.Value - FirstValue.Value) / NextTimeDiff;
+			FirstValue.Tangent.LeaveTangent = NewTangent;
+		}
+		else if (FirstValue.InterpMode == RCIM_Cubic && FirstValue.TangentMode == RCTM_Auto)
 		{
 			FirstValue.Tangent.LeaveTangent = FirstValue.Tangent.ArriveTangent = 0.0f;
 			FirstValue.Tangent.TangentWeightMode = RCTWM_WeightedNone; 
@@ -676,7 +698,16 @@ void FMovieSceneFloatChannel::AutoSetTangents(float Tension)
 
 	{
 		FMovieSceneFloatValue& LastValue = Values.Last();
-		if (LastValue.InterpMode == RCIM_Cubic && LastValue.TangentMode == RCTM_Auto)
+		if (LastValue.InterpMode == RCIM_Linear)
+		{
+			LastValue.Tangent.TangentWeightMode = RCTWM_WeightedNone;
+			int32 Index = Values.Num() - 1;
+			FMovieSceneFloatValue& PrevKey = Values[Index-1];
+			const float PrevTimeDiff = FMath::Max<double>(KINDA_SMALL_NUMBER, Times[Index].Value - Times[Index - 1].Value);
+			float NewTangent = (LastValue.Value - LastValue.Value) / PrevTimeDiff;
+			LastValue.Tangent.ArriveTangent = NewTangent;
+		}
+		else if (LastValue.InterpMode == RCIM_Cubic && LastValue.TangentMode == RCTM_Auto)
 		{
 			LastValue.Tangent.LeaveTangent = LastValue.Tangent.ArriveTangent = 0.0f;
 			LastValue.Tangent.TangentWeightMode = RCTWM_WeightedNone;
@@ -691,7 +722,6 @@ void FMovieSceneFloatChannel::AutoSetTangents(float Tension)
 		if (ThisKey.InterpMode == RCIM_Cubic && ThisKey.TangentMode == RCTM_Auto)
 		{
 			FMovieSceneFloatValue& NextKey = Values[Index+1];
-
 			const float PrevToNextTimeDiff = FMath::Max<double>(KINDA_SMALL_NUMBER, Times[Index+1].Value - Times[Index-1].Value);
 
 			float NewTangent = 0.f;
@@ -702,6 +732,19 @@ void FMovieSceneFloatChannel::AutoSetTangents(float Tension)
 			ThisKey.Tangent.LeaveTangent = ThisKey.Tangent.ArriveTangent = NewTangent;
 			ThisKey.Tangent.TangentWeightMode = RCTWM_WeightedNone;
 		}
+		else if (ThisKey.InterpMode == RCIM_Linear)
+		{
+			ThisKey.Tangent.TangentWeightMode = RCTWM_WeightedNone; 
+			FMovieSceneFloatValue& NextKey = Values[Index + 1];
+
+			const float PrevTimeDiff = FMath::Max<double>(KINDA_SMALL_NUMBER, Times[Index].Value - Times[Index - 1].Value);
+			float NewTangent  = (ThisKey.Value - PrevKey.Value) / PrevTimeDiff;
+			ThisKey.Tangent.ArriveTangent = NewTangent;
+			
+			const float NextTimeDiff = FMath::Max<double>(KINDA_SMALL_NUMBER, Times[Index + 1].Value - Times[Index].Value);
+			NewTangent = (NextKey.Value - ThisKey.Value) / NextTimeDiff;
+			ThisKey.Tangent.LeaveTangent = NewTangent;
+		}
 		else if (PrevKey.InterpMode == RCIM_Constant || ThisKey.InterpMode == RCIM_Constant)
 		{
 			if (PrevKey.InterpMode != RCIM_Cubic)
@@ -711,6 +754,7 @@ void FMovieSceneFloatChannel::AutoSetTangents(float Tension)
 
 			ThisKey.Tangent.LeaveTangent  = 0.0f;
 		}
+		
 	}
 }
 
@@ -851,6 +895,11 @@ void FMovieSceneFloatChannel::Reset()
 	bHasDefaultValue = false;
 }
 
+void FMovieSceneFloatChannel::PostEditChange()
+{
+	AutoSetTangents();
+}
+
 void FMovieSceneFloatChannel::Offset(FFrameNumber DeltaPosition)
 {
 	GetData().Offset(DeltaPosition);
@@ -946,4 +995,42 @@ void FMovieSceneFloatChannel::AddKeys(const TArray<FFrameNumber>& InTimes, const
 	{
 		KeyHandles.AllocateHandle(Index);
 	}
-}
\ No newline at end of file
+}
+
+bool FMovieSceneFloatChannel::Serialize(FArchive& Ar)
+{
+	Ar.UsingCustomVersion(FSequencerObjectVersion::GUID);
+	return false;
+}
+
+void FMovieSceneFloatChannel::PostSerialize(const FArchive& Ar)
+{
+	if (Ar.CustomVer(FSequencerObjectVersion::GUID) < FSequencerObjectVersion::ModifyLinearKeysForOldInterp)
+	{
+		bool bNeedAutoSetAlso = false;
+		//we need to possibly modify cuvic tangents if we get a set of linear..cubic tangents so it works like it used to
+		if (Values.Num() >= 2)
+		{
+			for (int32 Index = 1; Index < Values.Num(); ++Index)
+			{
+				FMovieSceneFloatValue  PrevKey = Values[Index - 1];
+				FMovieSceneFloatValue& ThisKey = Values[Index];
+
+				if (ThisKey.InterpMode == RCIM_Cubic && PrevKey.TangentMode == RCIM_Linear)
+				{
+					ThisKey.Tangent.TangentWeightMode = RCTWM_WeightedNone;
+					ThisKey.TangentMode = RCTM_Break;
+					//leave next tangent will be set up if auto or user, just need to modify prev.
+					const float PrevTimeDiff = FMath::Max<double>(KINDA_SMALL_NUMBER, Times[Index].Value - Times[Index - 1].Value);
+					float NewTangent = (ThisKey.Value - PrevKey.Value) / PrevTimeDiff;
+					ThisKey.Tangent.ArriveTangent = NewTangent;
+					bNeedAutoSetAlso = true;
+				}
+			}
+		}
+		if (bNeedAutoSetAlso)
+		{
+			AutoSetTangents();
+		}
+	}
+}
diff --git a/Engine/Source/Runtime/MovieScene/Public/Channels/MovieSceneChannel.h b/Engine/Source/Runtime/MovieScene/Public/Channels/MovieSceneChannel.h
--- a/Engine/Source/Runtime/MovieScene/Public/Channels/MovieSceneChannel.h
+++ b/Engine/Source/Runtime/MovieScene/Public/Channels/MovieSceneChannel.h
@@ -140,4 +140,11 @@ struct MOVIESCENE_API FMovieSceneChannel
 	 */
 	virtual void ClearDefault()
 	{}
+
+	/**
+	 * Perfor a possibly heavy operation after an edit change 
+	 *
+	 */
+	virtual void PostEditChange() {}
+
 };
\ No newline at end of file
diff --git a/Engine/Source/Runtime/MovieScene/Public/Channels/MovieSceneFloatChannel.h b/Engine/Source/Runtime/MovieScene/Public/Channels/MovieSceneFloatChannel.h
--- a/Engine/Source/Runtime/MovieScene/Public/Channels/MovieSceneFloatChannel.h
+++ b/Engine/Source/Runtime/MovieScene/Public/Channels/MovieSceneFloatChannel.h
@@ -222,6 +222,7 @@ public:
 	virtual void Offset(FFrameNumber DeltaPosition) override;
 	virtual void Optimize(const FKeyDataOptimizationParams& InParameters) override;
 	virtual void ClearDefault() override;
+	virtual void PostEditChange() override;
 
 public:
 
@@ -256,6 +257,15 @@ public:
 
 public:
 
+	bool Serialize(FArchive& Ar);
+	void PostSerialize(const FArchive& Ar);
+	friend FArchive& operator<<(FArchive& Ar, FMovieSceneFloatChannel& Me)
+	{
+		Me.Serialize(Ar);
+		return Ar;
+	}
+
+
 	/** Serialize this float function from a mismatching property tag (FRichCurve) */
 	bool SerializeFromMismatchedTag(const FPropertyTag& Tag, FStructuredArchive::FSlot Slot);
 
@@ -349,7 +359,12 @@ public:
 template<>
 struct TStructOpsTypeTraits<FMovieSceneFloatChannel> : public TStructOpsTypeTraitsBase2<FMovieSceneFloatChannel>
 {
-	enum { WithStructuredSerializeFromMismatchedTag = true };
+	enum 
+	{ 
+		WithStructuredSerializeFromMismatchedTag = true, 
+	    WithSerializer = true,
+		WithPostSerialize = true,
+    };
 };
 
 template<>
diff --git a/Engine/Source/Runtime/MovieScene/Public/MovieSceneKeyProxy.h b/Engine/Source/Runtime/MovieScene/Public/MovieSceneKeyProxy.h
--- a/Engine/Source/Runtime/MovieScene/Public/MovieSceneKeyProxy.h
+++ b/Engine/Source/Runtime/MovieScene/Public/MovieSceneKeyProxy.h
@@ -89,6 +89,7 @@ void IMovieSceneKeyProxy::OnProxyValueChanged(TMovieSceneChannelHandle<ChannelTy
 			ChannelData.MoveKey(KeyIndex, InOutTime);
 			InSection->ExpandToFrame(InOutTime);
 		}
+		Channel->PostEditChange();
 	}
 }
 
