Fix for Linux Build break.

diff --git a/Engine/Source/Runtime/MovieScene/Private/Channels/MovieSceneFloatChannel.cpp b/Engine/Source/Runtime/MovieScene/Private/Channels/MovieSceneFloatChannel.cpp
--- a/Engine/Source/Runtime/MovieScene/Private/Channels/MovieSceneFloatChannel.cpp
+++ b/Engine/Source/Runtime/MovieScene/Private/Channels/MovieSceneFloatChannel.cpp
@@ -1016,7 +1016,7 @@ void FMovieSceneFloatChannel::PostSerialize(const FArchive& Ar)
 				FMovieSceneFloatValue  PrevKey = Values[Index - 1];
 				FMovieSceneFloatValue& ThisKey = Values[Index];
 
-				if (ThisKey.InterpMode == RCIM_Cubic && PrevKey.TangentMode == RCIM_Linear)
+				if (ThisKey.InterpMode == RCIM_Cubic && PrevKey.InterpMode == RCIM_Linear)
 				{
 					ThisKey.Tangent.TangentWeightMode = RCTWM_WeightedNone;
 					ThisKey.TangentMode = RCTM_Break;
