Fix for static analysis

diff --git a/Engine/Source/Runtime/MovieScene/Private/Channels/MovieSceneFloatChannel.cpp b/Engine/Source/Runtime/MovieScene/Private/Channels/MovieSceneFloatChannel.cpp
--- a/Engine/Source/Runtime/MovieScene/Private/Channels/MovieSceneFloatChannel.cpp
+++ b/Engine/Source/Runtime/MovieScene/Private/Channels/MovieSceneFloatChannel.cpp
@@ -704,7 +704,7 @@ void FMovieSceneFloatChannel::AutoSetTangents(float Tension)
 			int32 Index = Values.Num() - 1;
 			FMovieSceneFloatValue& PrevKey = Values[Index-1];
 			const float PrevTimeDiff = FMath::Max<double>(KINDA_SMALL_NUMBER, Times[Index].Value - Times[Index - 1].Value);
-			float NewTangent = (LastValue.Value - LastValue.Value) / PrevTimeDiff;
+			float NewTangent = (LastValue.Value - PrevKey.Value) / PrevTimeDiff;
 			LastValue.Tangent.ArriveTangent = NewTangent;
 		}
 		else if (LastValue.InterpMode == RCIM_Cubic && LastValue.TangentMode == RCTM_Auto)
