Sequencer: Check for valid time controller

diff --git a/Engine/Source/Runtime/MovieScene/Private/MovieSceneSequencePlayer.cpp b/Engine/Source/Runtime/MovieScene/Private/MovieSceneSequencePlayer.cpp
--- a/Engine/Source/Runtime/MovieScene/Private/MovieSceneSequencePlayer.cpp
+++ b/Engine/Source/Runtime/MovieScene/Private/MovieSceneSequencePlayer.cpp
@@ -286,7 +286,10 @@ void UMovieSceneSequencePlayer::StopInternal(FFrameTime TimeToResetTo)
 
 		// Put the cursor at the specified position
 		PlayPosition.Reset(TimeToResetTo);
-		TimeController->StopPlaying(GetCurrentTime());
+		if (TimeController.IsValid())
+		{
+			TimeController->StopPlaying(GetCurrentTime());
+		}
 
 		CurrentNumLoops = 0;
 
