Sequencer: Fixed evaluating sub sequences in isolation using the wrong time

diff --git a/Engine/Source/Runtime/MovieScene/Private/Evaluation/MovieSceneEvaluationTemplateInstance.cpp b/Engine/Source/Runtime/MovieScene/Private/Evaluation/MovieSceneEvaluationTemplateInstance.cpp
--- a/Engine/Source/Runtime/MovieScene/Private/Evaluation/MovieSceneEvaluationTemplateInstance.cpp
+++ b/Engine/Source/Runtime/MovieScene/Private/Evaluation/MovieSceneEvaluationTemplateInstance.cpp
@@ -226,7 +226,7 @@ void FMovieSceneRootEvaluationTemplateInstance::Evaluate(FMovieSceneContext Cont
 		return;
 	}
 
-	const FMovieSceneEvaluationGroup* GroupToEvaluate = SetupFrame(OverrideRootSequence, InOverrideRootID, Context);
+	const FMovieSceneEvaluationGroup* GroupToEvaluate = SetupFrame(OverrideRootSequence, InOverrideRootID, &Context);
 	if (!GroupToEvaluate)
 	{
 		CallSetupTearDown(Player);
@@ -268,7 +268,7 @@ FMovieSceneEvaluationPtrCache FMovieSceneRootEvaluationTemplateInstance::Constru
 	return FMovieSceneEvaluationPtrCache(RootOverridePath, OverrideRootSequence, TemplateStore.Get(), PreviousAndCurrentFrameSequenceIDs);
 }
 
-const FMovieSceneEvaluationGroup* FMovieSceneRootEvaluationTemplateInstance::SetupFrame(UMovieSceneSequence* OverrideRootSequence, FMovieSceneSequenceID InOverrideRootID, FMovieSceneContext Context)
+const FMovieSceneEvaluationGroup* FMovieSceneRootEvaluationTemplateInstance::SetupFrame(UMovieSceneSequence* OverrideRootSequence, FMovieSceneSequenceID InOverrideRootID, FMovieSceneContext* InOutContext)
 {
 	check(OverrideRootSequence);
 
@@ -276,7 +276,6 @@ const FMovieSceneEvaluationGroup* FMovieSceneRootEvaluationTemplateInstance::Set
 	RootOverridePath.Set(InOverrideRootID, GetHierarchy());
 
 	FMovieSceneEvaluationTemplate* OverrideRootTemplate = nullptr;
-	FMovieSceneSequenceTransform   RootOverrideTransform;
 
 	if (InOverrideRootID == MovieSceneSequenceID::Root)
 	{
@@ -288,7 +287,7 @@ const FMovieSceneEvaluationGroup* FMovieSceneRootEvaluationTemplateInstance::Set
 		OverrideRootTemplate = &TemplateStore->AccessTemplate(*OverrideRootSequence);
 		if (const FMovieSceneSubSequenceData* OverrideSubData = GetHierarchy().FindSubData(InOverrideRootID))
 		{
-			RootOverrideTransform = OverrideSubData->RootToSequenceTransform;
+			*InOutContext = InOutContext->Transform(OverrideSubData->RootToSequenceTransform, OverrideSubData->TickResolution);
 		}
 	}
 
@@ -303,14 +302,10 @@ const FMovieSceneEvaluationGroup* FMovieSceneRootEvaluationTemplateInstance::Set
 		FMovieSceneEvaluationTemplateGenerator(*OverrideRootSequence, *OverrideRootTemplate).Generate();
 	}
 
-	FFrameNumber RootOverrideTime = (Context.GetTime() * RootOverrideTransform).FloorToFrame();
-
-
-
 	FMovieSceneEvaluationField& EvaluationField = OverrideRootTemplate->EvaluationField;
 
 	// Get the range that we are evaluating in the root's space
-	TRange<FFrameNumber> ContextRange = Context.GetTraversedFrameNumberRange() * RootOverrideTransform;
+	TRange<FFrameNumber> ContextRange = InOutContext->GetTraversedFrameNumberRange();
 
 	// Verify and update the evaluation field for this range, returning the bounds of the overlapping field entries
 	TRange<int32> FieldRange = EvaluationField.ConditionallyCompileRange(ContextRange, OverrideRootSequence, *TemplateStore);
@@ -321,7 +316,7 @@ const FMovieSceneEvaluationGroup* FMovieSceneRootEvaluationTemplateInstance::Set
 
 	// The one that we want to evaluate is either the first or last index in the range.
 	// FieldRange is always of the form [First, Last+1)
-	int32 TemplateFieldIndex = Context.GetDirection() == EPlayDirection::Forwards ? FieldRange.GetUpperBoundValue() - 1 : FieldRange.GetLowerBoundValue();
+	int32 TemplateFieldIndex = InOutContext->GetDirection() == EPlayDirection::Forwards ? FieldRange.GetUpperBoundValue() - 1 : FieldRange.GetLowerBoundValue();
 	if (TemplateFieldIndex != INDEX_NONE)
 	{
 		// Set meta-data
diff --git a/Engine/Source/Runtime/MovieScene/Public/Evaluation/MovieSceneEvaluationTemplateInstance.h b/Engine/Source/Runtime/MovieScene/Public/Evaluation/MovieSceneEvaluationTemplateInstance.h
--- a/Engine/Source/Runtime/MovieScene/Public/Evaluation/MovieSceneEvaluationTemplateInstance.h
+++ b/Engine/Source/Runtime/MovieScene/Public/Evaluation/MovieSceneEvaluationTemplateInstance.h
@@ -157,10 +157,10 @@ private:
 	 *
 	 * @param OverrideRootSequence	Pointer to the sequence that is considered the root for this evaluation (that maps to InOverrideRootID)
 	 * @param OverrideRootID		The sequence ID of the currently considered root (normally MovieSceneSequenceID::Root unless Evaluate Sub-Sequences in Isolation is active)
-	 * @param Context				The evaluation context for this frame
+	 * @param InOutContext			The evaluation context for this frame, will be transformed to the correct space if InOverrideRootID is not Root
 	 * @return The evaluation group within the root template's evaluation field to evaluate, or nullptr if none could be found or compiled
 	 */
-	const FMovieSceneEvaluationGroup* SetupFrame(UMovieSceneSequence* OverrideRootSequence, FMovieSceneSequenceID InOverrideRootID, FMovieSceneContext Context);
+	const FMovieSceneEvaluationGroup* SetupFrame(UMovieSceneSequence* OverrideRootSequence, FMovieSceneSequenceID InOverrideRootID, FMovieSceneContext* InOutContext);
 
 	/**
 	 * Process entities that are newly evaluated, and those that are no longer being evaluated
