Sequencer: Add deprecated message for deprecated ufunctions

diff --git a/Engine/Source/Runtime/MovieScene/Public/MovieSceneSequencePlayer.h b/Engine/Source/Runtime/MovieScene/Public/MovieSceneSequencePlayer.h
--- a/Engine/Source/Runtime/MovieScene/Public/MovieSceneSequencePlayer.h
+++ b/Engine/Source/Runtime/MovieScene/Public/MovieSceneSequencePlayer.h
@@ -223,28 +223,28 @@ public:
 	 * @return The current playback position
 	 */
 	UE_DEPRECATED(4.20, "Please use GetCurrentTime instead")
-	UFUNCTION(BlueprintCallable, Category="Game|Cinematic")
+	UFUNCTION(BlueprintCallable, Category="Game|Cinematic", meta = (DeprecatedFunction, DeprecationMessage = "Please use GetCurrentTime instead"))
 	float GetPlaybackPosition() const { return GetCurrentTime().AsSeconds() - StartTime / PlayPosition.GetInputRate(); }
 
 	/**
 	 * Get the playback length of the sequence
 	 */
 	UE_DEPRECATED(4.20, "Please use GetDuration instead")
-	UFUNCTION(BlueprintCallable, Category="Game|Cinematic")
+	UFUNCTION(BlueprintCallable, Category="Game|Cinematic", meta = (DeprecatedFunction, DeprecationMessage = "Please use GetDuration instead"))
 	float GetLength() const;
 
 	/**
 	 * Get the offset within the level sequence to start playing
 	 */
 	UE_DEPRECATED(4.20, "Please use GetStartTime instead")
-	UFUNCTION(BlueprintCallable, Category="Game|Cinematic")
+	UFUNCTION(BlueprintCallable, Category="Game|Cinematic", meta = (DeprecatedFunction, DeprecationMessage = "Please use GetStartTime instead"))
 	float GetPlaybackStart() const { return StartTime / PlayPosition.GetInputRate(); }
 
 	/**
 	 * Get the offset within the level sequence to finish playing
 	 */
 	UE_DEPRECATED(4.20, "Please use GetEndTime instead")
-	UFUNCTION(BlueprintCallable, Category="Game|Cinematic")
+	UFUNCTION(BlueprintCallable, Category="Game|Cinematic", meta = (DeprecatedFunction, DeprecationMessage = "Please use GetEndTime instead"))
 	float GetPlaybackEnd() const { return (StartTime + DurationFrames) / PlayPosition.GetInputRate(); }
 
 	/**
@@ -253,7 +253,7 @@ public:
 	 * If the animation is currently playing, it will continue to do so from the new position
 	 */
 	UE_DEPRECATED(4.20, "Please use PlayToFrame instead")
-	UFUNCTION(BlueprintCallable, Category="Game|Cinematic")
+	UFUNCTION(BlueprintCallable, Category="Game|Cinematic", meta = (DeprecatedFunction, DeprecationMessage = "Please use PlayToFrame instead"))
 	void SetPlaybackPosition(float NewPlaybackPosition) { Status == EMovieScenePlayerStatus::Playing ? PlayToSeconds(NewPlaybackPosition + StartTime / PlayPosition.GetInputRate()) : JumpToSeconds(NewPlaybackPosition + StartTime / PlayPosition.GetInputRate()); }
 
 	/**
@@ -263,7 +263,7 @@ public:
 	 * @param	NewEndTime		The new ending time for playback.  Must be larger than the start time.
 	 */
 	UE_DEPRECATED(4.20, "Please use SetFrameRange or SetTimeRange instead")
-	UFUNCTION(BlueprintCallable, Category="Game|Cinematic")
+	UFUNCTION(BlueprintCallable, Category="Game|Cinematic", meta = (DeprecatedFunction, DeprecationMessage = "Please use SetFrameRange or SetTimeRange instead"))
 	void SetPlaybackRange( const float NewStartTime, const float NewEndTime );
 
 	/**
@@ -272,7 +272,7 @@ public:
 	 * This can be used to update sequencer repeatedly, as if in a scrubbing state
 	 */
 	UE_DEPRECATED(4.20, "Please use ScrubToTime instead")
-	UFUNCTION(BlueprintCallable, Category="Game|Cinematic")
+	UFUNCTION(BlueprintCallable, Category="Game|Cinematic", meta = (DeprecatedFunction, DeprecationMessage = "Please use ScrubToTime instead"))
 	void JumpToPosition(float NewPlaybackPosition) { ScrubToSeconds(NewPlaybackPosition); }
 
 public:
