Sequencer: Remove null sections on PostLoad.

diff --git a/Engine/Source/Runtime/MovieScene/Private/MovieSceneTrack.cpp b/Engine/Source/Runtime/MovieScene/Private/MovieSceneTrack.cpp
--- a/Engine/Source/Runtime/MovieScene/Private/MovieSceneTrack.cpp
+++ b/Engine/Source/Runtime/MovieScene/Private/MovieSceneTrack.cpp
@@ -44,6 +44,19 @@ void UMovieSceneTrack::PostLoad()
 	{
 		EvalOptions.bEvalNearestSection = EvalOptions.bEvaluateNearestSection_DEPRECATED;
 	}
+
+	// Remove any null sections
+	for (int32 SectionIndex = 0; SectionIndex < GetAllSections().Num(); )
+	{
+		if (GetAllSections()[SectionIndex] == nullptr)
+		{
+			RemoveSectionAt(SectionIndex);
+		}
+		else
+		{
+			++SectionIndex;
+		}
+	}
 }
 
 bool UMovieSceneTrack::IsPostLoadThreadSafe() const
