Sequencer: Use GetLastValidTime() instead of StartTime+Duration when looping playback reversed.

diff --git a/Engine/Source/Runtime/MovieScene/Private/MovieSceneSequencePlayer.cpp b/Engine/Source/Runtime/MovieScene/Private/MovieSceneSequencePlayer.cpp
--- a/Engine/Source/Runtime/MovieScene/Private/MovieSceneSequencePlayer.cpp
+++ b/Engine/Source/Runtime/MovieScene/Private/MovieSceneSequencePlayer.cpp
@@ -726,7 +726,7 @@ void UMovieSceneSequencePlayer::UpdateTimeCursorPosition_Internal(FFrameTime New
 	if (Method == EUpdatePositionMethod::Play && ShouldStopOrLoop(NewPosition))
 	{
 		// The actual start time taking into account reverse playback
-		FFrameNumber StartTimeWithReversed = bReversePlayback ? StartTime + Duration : StartTime;
+		FFrameNumber StartTimeWithReversed = bReversePlayback ? GetLastValidTime().FrameNumber : StartTime;
 
 		FFrameTime PositionRelativeToStart = NewPosition.FrameNumber - StartTimeWithReversed;
 
