Fix events occasionally being skipped at end of sequence playback and looping.

diff --git a/Engine/Source/Runtime/MovieScene/Private/MovieSceneSequencePlayer.cpp b/Engine/Source/Runtime/MovieScene/Private/MovieSceneSequencePlayer.cpp
--- a/Engine/Source/Runtime/MovieScene/Private/MovieSceneSequencePlayer.cpp
+++ b/Engine/Source/Runtime/MovieScene/Private/MovieSceneSequencePlayer.cpp
@@ -728,6 +728,9 @@ void UMovieSceneSequencePlayer::UpdateTimeCursorPosition_Internal(FFrameTime New
 		// The actual start time taking into account reverse playback
 		FFrameNumber StartTimeWithReversed = bReversePlayback ? GetLastValidTime().FrameNumber : StartTime;
 
+		// The actual end time taking into account reverse playback
+		FFrameTime EndTimeWithReversed = bReversePlayback ? StartTime : GetLastValidTime().FrameNumber;
+
 		FFrameTime PositionRelativeToStart = NewPosition.FrameNumber - StartTimeWithReversed;
 
 		const int32 NumTimesLooped    = FMath::Abs(PositionRelativeToStart.FrameNumber.Value / Duration);
@@ -738,6 +741,16 @@ void UMovieSceneSequencePlayer::UpdateTimeCursorPosition_Internal(FFrameTime New
 		{
 			CurrentNumLoops += NumTimesLooped;
 
+			// Finish evaluating any frames left in the current loop in case they have events attached
+			FFrameTime CurrentPosition = PlayPosition.GetCurrentPosition();
+			if ((bReversePlayback && CurrentPosition > EndTimeWithReversed) ||
+				(!bReversePlayback && CurrentPosition < EndTimeWithReversed))
+			{
+				FMovieSceneEvaluationRange Range = PlayPosition.PlayTo(EndTimeWithReversed);
+				UpdateMovieSceneInstance(Range, StatusOverride);
+			}
+
+
 			const FFrameTime Overplay       = FFrameTime(PositionRelativeToStart.FrameNumber.Value % Duration, PositionRelativeToStart.GetSubFrame());
 			FFrameTime NewFrameOffset;
 			
