Fix for crash when editing a sequencer blueprint and attempting to play in editor before recompiling the blueprint

diff --git a/Engine/Source/Runtime/MovieScene/Private/MovieSceneCommonHelpers.cpp b/Engine/Source/Runtime/MovieScene/Private/MovieSceneCommonHelpers.cpp
--- a/Engine/Source/Runtime/MovieScene/Private/MovieSceneCommonHelpers.cpp
+++ b/Engine/Source/Runtime/MovieScene/Private/MovieSceneCommonHelpers.cpp
@@ -423,7 +423,7 @@ FTrackInstancePropertyBindings::FPropertyAddress FTrackInstancePropertyBindings:
 
 	InPropertyPath.ParseIntoArray(PropertyNames, TEXT("."), true);
 
-	if( PropertyNames.Num() > 0 )
+	if(IsValid(&InObject) && PropertyNames.Num() > 0)
 	{
 		return FindPropertyRecursive( (void*)&InObject, InObject.GetClass(), PropertyNames, 0 );
 	}
