Sequencer: Set MarkedFrame properties as BlueprintReadWrite. AddMarkedFrame now returns the index to the newly added marked frame.

diff --git a/Engine/Source/Runtime/MovieScene/Private/MovieScene.cpp b/Engine/Source/Runtime/MovieScene/Private/MovieScene.cpp
--- a/Engine/Source/Runtime/MovieScene/Private/MovieScene.cpp
+++ b/Engine/Source/Runtime/MovieScene/Private/MovieScene.cpp
@@ -1043,7 +1043,7 @@ void UMovieScene::SetMarkedFrame(int32 InMarkIndex, FFrameNumber InFrameNumber)
 	MarkedFrames[InMarkIndex].FrameNumber = InFrameNumber;
 }
 
-void UMovieScene::AddMarkedFrame(const FMovieSceneMarkedFrame &InMarkedFrame)
+int32 UMovieScene::AddMarkedFrame(const FMovieSceneMarkedFrame &InMarkedFrame)
 {
 	FString NewLabel;
 
@@ -1084,6 +1084,8 @@ void UMovieScene::AddMarkedFrame(const FMovieSceneMarkedFrame &InMarkedFrame)
 	{
 		MarkedFrames[MarkedIndex].Label = NewLabel;
 	}
+
+	return MarkedIndex;
 }
 
 void UMovieScene::RemoveMarkedFrame(int32 RemoveIndex)
diff --git a/Engine/Source/Runtime/MovieScene/Public/MovieScene.h b/Engine/Source/Runtime/MovieScene/Public/MovieScene.h
--- a/Engine/Source/Runtime/MovieScene/Public/MovieScene.h
+++ b/Engine/Source/Runtime/MovieScene/Public/MovieScene.h
@@ -81,10 +81,10 @@ struct FMovieSceneMarkedFrame
 #endif
 	{}
 
-	UPROPERTY(EditAnywhere, Category = "Marked Frame")
+	UPROPERTY(EditAnywhere, BlueprintReadWrite, Category = "Marked Frame")
 	FFrameNumber FrameNumber;
 
-	UPROPERTY(EditAnywhere, Category = "Marked Frame")
+	UPROPERTY(EditAnywhere, BlueprintReadWrite, Category = "Marked Frame")
 	FString Label;
 
 #if WITH_EDITORONLY_DATA
@@ -763,8 +763,9 @@ public:
 	 * A unique label will be generated if the marked frame label is empty
 	 *
 	 * @InMarkedFrame The given user marked frame to add
+	 * @return The index to the newly added marked frame
 	 */
-	void AddMarkedFrame(const FMovieSceneMarkedFrame& InMarkedFrame);
+	int32 AddMarkedFrame(const FMovieSceneMarkedFrame& InMarkedFrame);
 
 	/*
 	 * Remove the user marked frame by index.
