Sequencer: Added warning about removing null sections

diff --git a/Engine/Source/Runtime/MovieScene/Private/MovieSceneTrack.cpp b/Engine/Source/Runtime/MovieScene/Private/MovieSceneTrack.cpp
--- a/Engine/Source/Runtime/MovieScene/Private/MovieSceneTrack.cpp
+++ b/Engine/Source/Runtime/MovieScene/Private/MovieSceneTrack.cpp
@@ -50,6 +50,7 @@ void UMovieSceneTrack::PostLoad()
 	{
 		if (GetAllSections()[SectionIndex] == nullptr)
 		{
+			UE_LOG(LogMovieScene, Warning, TEXT("Removing null section from %s:%s"), *GetPathName(), *GetDisplayName().ToString());
 			RemoveSectionAt(SectionIndex);
 		}
 		else
