Sequencer: Editor only

diff --git a/Engine/Source/Runtime/MovieScene/Private/MovieSceneTrack.cpp b/Engine/Source/Runtime/MovieScene/Private/MovieSceneTrack.cpp
--- a/Engine/Source/Runtime/MovieScene/Private/MovieSceneTrack.cpp
+++ b/Engine/Source/Runtime/MovieScene/Private/MovieSceneTrack.cpp
@@ -50,7 +50,9 @@ void UMovieSceneTrack::PostLoad()
 	{
 		if (GetAllSections()[SectionIndex] == nullptr)
 		{
+#if WITH_EDITOR
 			UE_LOG(LogMovieScene, Warning, TEXT("Removing null section from %s:%s"), *GetPathName(), *GetDisplayName().ToString());
+#endif
 			RemoveSectionAt(SectionIndex);
 		}
 		else
