Sequencer: Fixed a vulnerability that caused object path channels to end up with hard references to actors

diff --git a/Engine/Source/Runtime/MovieScene/Private/Channels/MovieSceneObjectPathChannel.cpp b/Engine/Source/Runtime/MovieScene/Private/Channels/MovieSceneObjectPathChannel.cpp
--- a/Engine/Source/Runtime/MovieScene/Private/Channels/MovieSceneObjectPathChannel.cpp
+++ b/Engine/Source/Runtime/MovieScene/Private/Channels/MovieSceneObjectPathChannel.cpp
@@ -27,6 +27,13 @@ bool FMovieSceneObjectPathChannelKeyValue::SerializeFromMismatchedTag(const FPro
 	return false;
 }
 
+FMovieSceneObjectPathChannelKeyValue::FMovieSceneObjectPathChannelKeyValue(UObject* InObject)
+	: SoftPtr(nullptr)
+	, HardPtr(nullptr)
+{
+	*this = InObject;
+}
+
 FMovieSceneObjectPathChannelKeyValue& FMovieSceneObjectPathChannelKeyValue::operator=(UObject* NewObject)
 {
 	// Do not store raw ptrs to actors or other objects that exist in worlds
diff --git a/Engine/Source/Runtime/MovieScene/Public/Channels/MovieSceneObjectPathChannel.h b/Engine/Source/Runtime/MovieScene/Public/Channels/MovieSceneObjectPathChannel.h
--- a/Engine/Source/Runtime/MovieScene/Public/Channels/MovieSceneObjectPathChannel.h
+++ b/Engine/Source/Runtime/MovieScene/Public/Channels/MovieSceneObjectPathChannel.h
@@ -22,10 +22,7 @@ struct MOVIESCENE_API FMovieSceneObjectPathChannelKeyValue
 	{}
 
 	/** Construction from an object pointer */
-	FMovieSceneObjectPathChannelKeyValue(UObject* InObject)
-		: SoftPtr(InObject)
-		, HardPtr(InObject)
-	{}
+	FMovieSceneObjectPathChannelKeyValue(UObject* InObject);
 
 	/**
 	 * Assignment from a raw object pointer
