AutoKey Fix
Auto Key now works like other packages to avoid overshoot by flattening tangents that are not between other keys and gradually flattening the tangent when it get's nearer to the min/max values.

diff --git a/Engine/Source/Runtime/MovieScene/Private/Channels/MovieSceneFloatChannel.cpp b/Engine/Source/Runtime/MovieScene/Private/Channels/MovieSceneFloatChannel.cpp
--- a/Engine/Source/Runtime/MovieScene/Private/Channels/MovieSceneFloatChannel.cpp
+++ b/Engine/Source/Runtime/MovieScene/Private/Channels/MovieSceneFloatChannel.cpp
@@ -13,6 +13,12 @@ static TAutoConsoleVariable<int32> CVarSequencerLinearCubicInterpolation(
 	TEXT("If 1 Linear Keys Act As Cubic Interpolation with Linear Tangents, if 0 Linear Key Forces Linear Interpolation to Next Key."),
 	ECVF_Default);
 
+static TAutoConsoleVariable<int32> CVarSequencerAutoTangentInterpolation(
+	TEXT("Sequencer.AutoTangentNew"),
+	1,
+	TEXT("If 1 Auto Tangent will use new algorithm to gradually flatten maximum/minimum keys, if 0 Auto Tangent will average all keys (pre 4.23 behavior)."),
+	ECVF_Default);
+
 bool FMovieSceneTangentData::Serialize(FArchive& Ar)
 {
 	Ar.UsingCustomVersion(FSequencerObjectVersion::GUID);
@@ -678,6 +684,7 @@ void FMovieSceneFloatChannel::AutoSetTangents(float Tension)
 	{
 		return;
 	}
+	const int UseNewAutoTangent = CVarSequencerAutoTangentInterpolation->GetInt();
 
 	{
 		FMovieSceneFloatValue& FirstValue = Values[0];
@@ -722,11 +729,40 @@ void FMovieSceneFloatChannel::AutoSetTangents(float Tension)
 		if (ThisKey.InterpMode == RCIM_Cubic && ThisKey.TangentMode == RCTM_Auto)
 		{
 			FMovieSceneFloatValue& NextKey = Values[Index+1];
-			const float PrevToNextTimeDiff = FMath::Max<double>(KINDA_SMALL_NUMBER, Times[Index+1].Value - Times[Index-1].Value);
-
 			float NewTangent = 0.f;
-			AutoCalcTangent(PrevKey.Value, ThisKey.Value, NextKey.Value, Tension, NewTangent);
-			NewTangent /= PrevToNextTimeDiff;
+			const float PrevToNextTimeDiff = FMath::Max<double>(KINDA_SMALL_NUMBER, Times[Index + 1].Value - Times[Index - 1].Value);
+
+			if (!UseNewAutoTangent)
+			{
+				AutoCalcTangent(PrevKey.Value, ThisKey.Value, NextKey.Value, Tension, NewTangent);
+				NewTangent /= PrevToNextTimeDiff;
+			}
+			else
+			{
+				// if key doesn't lie between we keep it flat(0.0).
+				if (ThisKey.Value >= PrevKey.Value && ThisKey.Value <= NextKey.Value ||
+					ThisKey.Value <= PrevKey.Value && ThisKey.Value >= NextKey.Value)
+				{
+					AutoCalcTangent(PrevKey.Value, ThisKey.Value, NextKey.Value, Tension, NewTangent);
+					NewTangent /= PrevToNextTimeDiff;
+					//if within 0 to 15% or 85% to 100% range we gradually weight tangent to zero
+					const float AverageToZeroRange = 0.85f;
+					const float ValDiff = FMath::Abs<float>(NextKey.Value - PrevKey.Value);
+					const float OurDiff = FMath::Abs<float>(ThisKey.Value - PrevKey.Value);
+					//ValDiff won't be zero ever due to previous check
+					float PercDiff = OurDiff / ValDiff;
+					if (PercDiff > AverageToZeroRange)
+					{
+						PercDiff = (PercDiff - AverageToZeroRange) / (1.0f - AverageToZeroRange);
+						NewTangent = NewTangent * (1.0f - PercDiff);
+					}
+					else if (PercDiff < (1.0f - AverageToZeroRange))
+					{
+						PercDiff = PercDiff  / (1.0f - AverageToZeroRange);
+						NewTangent = NewTangent * PercDiff;
+					}
+				}
+			}
 
 			// In 'auto' mode, arrive and leave tangents are always the same
 			ThisKey.Tangent.LeaveTangent = ThisKey.Tangent.ArriveTangent = NewTangent;
