CIS fixes for parentheses to if statement.

diff --git a/Engine/Source/Runtime/MovieScene/Private/Channels/MovieSceneFloatChannel.cpp b/Engine/Source/Runtime/MovieScene/Private/Channels/MovieSceneFloatChannel.cpp
--- a/Engine/Source/Runtime/MovieScene/Private/Channels/MovieSceneFloatChannel.cpp
+++ b/Engine/Source/Runtime/MovieScene/Private/Channels/MovieSceneFloatChannel.cpp
@@ -740,8 +740,8 @@ void FMovieSceneFloatChannel::AutoSetTangents(float Tension)
 			else
 			{
 				// if key doesn't lie between we keep it flat(0.0).
-				if (ThisKey.Value >= PrevKey.Value && ThisKey.Value <= NextKey.Value ||
-					ThisKey.Value <= PrevKey.Value && ThisKey.Value >= NextKey.Value)
+				if ( (ThisKey.Value >= PrevKey.Value && ThisKey.Value <= NextKey.Value) ||
+					(ThisKey.Value <= PrevKey.Value && ThisKey.Value >= NextKey.Value))
 				{
 					AutoCalcTangent(PrevKey.Value, ThisKey.Value, NextKey.Value, Tension, NewTangent);
 					NewTangent /= PrevToNextTimeDiff;
