Sequencer: Fallback to clock time if timecode is not synchronized or if the audio clock doesn't exist.

diff --git a/Engine/Source/Runtime/MovieScene/Private/MovieSceneTimeController.cpp b/Engine/Source/Runtime/MovieScene/Private/MovieSceneTimeController.cpp
--- a/Engine/Source/Runtime/MovieScene/Private/MovieSceneTimeController.cpp
+++ b/Engine/Source/Runtime/MovieScene/Private/MovieSceneTimeController.cpp
@@ -3,6 +3,7 @@
 #include "MovieSceneTimeController.h"
 #include "Engine/World.h"
 #include "Engine/Engine.h"
+#include "Engine/TimecodeProvider.h"
 #include "AudioDevice.h"
 #include "Misc/App.h"
 
@@ -104,15 +105,22 @@ double FMovieSceneTimeController_PlatformClock::GetCurrentTime() const
 double FMovieSceneTimeController_AudioClock::GetCurrentTime() const
 {
 	FAudioDevice* AudioDevice = GEngine ? GEngine->GetMainAudioDevice() : nullptr;
-	return AudioDevice ? AudioDevice->GetAudioClock() : 0.0;
+	return AudioDevice ? AudioDevice->GetAudioClock() : FPlatformTime::Seconds();
 }
 
 double FMovieSceneTimeController_TimecodeClock::GetCurrentTime() const
 {
-	FTimecode Timecode = FApp::GetTimecode();
-	FFrameRate FrameRate = FApp::GetTimecodeFrameRate();
-	FFrameNumber FrameNumber = Timecode.ToFrameNumber(FrameRate);
-	return FrameRate.AsSeconds(FrameNumber);
+	if (GEngine && GEngine->GetTimecodeProvider() && GEngine->GetTimecodeProvider()->GetSynchronizationState() == ETimecodeProviderSynchronizationState::Synchronized)
+	{
+		FTimecode Timecode = FApp::GetTimecode();
+		FFrameRate FrameRate = FApp::GetTimecodeFrameRate();
+		FFrameNumber FrameNumber = Timecode.ToFrameNumber(FrameRate);
+		return FrameRate.AsSeconds(FrameNumber);
+	}
+	else
+	{
+		return FPlatformTime::Seconds();
+	}
 }
 
 void FMovieSceneTimeController_Tick::OnStartPlaying(const FQualifiedFrameTime& InStartTime)
