Sequencer: Fix weighted tangent case where cubic solve are not between 0 and 1.

diff --git a/Engine/Source/Runtime/MovieScene/Private/Channels/MovieSceneFloatChannel.cpp b/Engine/Source/Runtime/MovieScene/Private/Channels/MovieSceneFloatChannel.cpp
--- a/Engine/Source/Runtime/MovieScene/Private/Channels/MovieSceneFloatChannel.cpp
+++ b/Engine/Source/Runtime/MovieScene/Private/Channels/MovieSceneFloatChannel.cpp
@@ -653,6 +653,12 @@ bool FMovieSceneFloatChannel::Evaluate(FFrameTime InTime,  float& OutValue) cons
 							}
 						}
 					}
+
+					if (NewInterp == TNumericLimits<float>::Lowest())
+					{
+						NewInterp = 0.f;
+					}
+
 				}
 				//now use NewInterp and adjusted tangents plugged into the Y (Value) part of the graph.
 				const float P0 = Key1.Value;
