Adding in equality operators to fix serialization issues with movie scene track and subseciton data.

diff --git a/Engine/Source/Runtime/MovieScene/Public/Evaluation/MovieSceneEvaluationTemplate.h b/Engine/Source/Runtime/MovieScene/Public/Evaluation/MovieSceneEvaluationTemplate.h
--- a/Engine/Source/Runtime/MovieScene/Public/Evaluation/MovieSceneEvaluationTemplate.h
+++ b/Engine/Source/Runtime/MovieScene/Public/Evaluation/MovieSceneEvaluationTemplate.h
@@ -75,9 +75,14 @@ struct FMovieSceneTrackFieldData
 		return true;
 	}
 
+	friend bool operator==(const FMovieSceneTrackFieldData& A, const FMovieSceneTrackFieldData& B)
+	{
+		return A.Field == B.Field;
+	}
+
 	TMovieSceneEvaluationTree<FMovieSceneTrackIdentifier> Field;
 };
-template<> struct TStructOpsTypeTraits<FMovieSceneTrackFieldData> : public TStructOpsTypeTraitsBase2<FMovieSceneTrackFieldData> { enum { WithSerializer = true }; };
+template<> struct TStructOpsTypeTraits<FMovieSceneTrackFieldData> : public TStructOpsTypeTraitsBase2<FMovieSceneTrackFieldData> { enum { WithSerializer = true, WithIdenticalViaEquality = true }; };
 
 /** Data that represents a single sub-section */
 USTRUCT()
@@ -94,6 +99,11 @@ struct FMovieSceneSubSectionData
 		return Ar << In.Section << In.ObjectBindingId << In.Flags;
 	}
 
+	friend bool operator==(const FMovieSceneSubSectionData& A, const FMovieSceneSubSectionData& B)
+	{
+		return A.Section == B.Section && A.ObjectBindingId == B.ObjectBindingId && A.Flags == B.Flags;
+	}
+
 	/** The sub section itself */
 	UPROPERTY()
 	TWeakObjectPtr<UMovieSceneSubSection> Section;
diff --git a/Engine/Source/Runtime/MovieScene/Public/Evaluation/MovieSceneEvaluationTree.h b/Engine/Source/Runtime/MovieScene/Public/Evaluation/MovieSceneEvaluationTree.h
--- a/Engine/Source/Runtime/MovieScene/Public/Evaluation/MovieSceneEvaluationTree.h
+++ b/Engine/Source/Runtime/MovieScene/Public/Evaluation/MovieSceneEvaluationTree.h
@@ -132,6 +132,14 @@ struct TEvaluationTreeEntryContainer
 		return Ar << In.Entries << In.Items;
 	}
 
+	/**
+	 * Equality operator
+	 */
+	friend bool operator==(const TEvaluationTreeEntryContainer<ElementType>& A, const TEvaluationTreeEntryContainer<ElementType>& B)
+	{
+		return A.Entries == B.Entries && A.Items == B.Items;
+	}
+
 private:
 
 	/**
@@ -156,6 +164,14 @@ private:
 			Ar << In.StartIndex << In.Size << In.Capacity;
 			return Ar;
 		}
+
+		/**
+		 * Equality operator
+		 */
+		friend bool operator==(const FEntry& A, const FEntry& B)
+		{
+			return A.StartIndex == B.StartIndex && A.Size == B.Size && A.Capacity == B.Capacity;
+		}
 	};
 
 	/** List of allocated entries for each allocated entry. Should only ever grow, never shrink. Shrinking would cause previously established handles to become invalid. */
@@ -265,6 +281,14 @@ struct FMovieSceneEvaluationTreeNode
 		return Ar << In.Range << In.Parent << In.ChildrenID << In.DataID;
 	}
 
+	/**
+	 * Equality operator
+	 */
+	friend bool operator==(const FMovieSceneEvaluationTreeNode& A, const FMovieSceneEvaluationTreeNode& B)
+	{
+		return A.Range == B.Range && A.Parent == B.Parent && A.ChildrenID == B.ChildrenID && A.DataID == B.DataID;
+	}
+
 	/** The time-range that this node represents */
 	TRange<FFrameNumber> Range;
 	/** Handle to the parent node */
@@ -631,6 +655,14 @@ struct TMovieSceneEvaluationTree : FMovieSceneEvaluationTree
 		return Ar << BaseTree << In.Data;
 	}
 
+	/**
+	 * Equality operator
+	 */
+	friend bool operator==(const TMovieSceneEvaluationTree<DataType>& A, const TMovieSceneEvaluationTree<DataType>& B)
+	{
+		return A.Data == B.Data && A.RootNode == B.RootNode && A.ChildNodes == B.ChildNodes;
+	}
+
 private:
 
 	/** Tree data container that corresponds to FMovieSceneEvaluationTreeNode::DataID */
