Copy/Pasting property tracks should overwrite equivalent tracks

diff --git a/Engine/Source/Runtime/MovieScene/Private/MovieScene.cpp b/Engine/Source/Runtime/MovieScene/Private/MovieScene.cpp
--- a/Engine/Source/Runtime/MovieScene/Private/MovieScene.cpp
+++ b/Engine/Source/Runtime/MovieScene/Private/MovieScene.cpp
@@ -593,6 +593,17 @@ bool UMovieScene::AddGivenTrack(UMovieSceneTrack* InTrack, const FGuid& ObjectGu
 	{
 		if (Binding.GetObjectGuid() == ObjectGuid)
 		{
+			// Tracks of the same class should be unique per name.
+			for (UMovieSceneTrack* Track : Binding.GetTracks())
+			{
+				if (Track->GetClass() == InTrack->GetClass() && Track->GetTrackName() == InTrack->GetTrackName())
+				{
+					// If a track of the same class and name exists, remove it so the new track replaces it
+					Binding.RemoveTrack(*Track);
+					break;
+				}
+			}
+
 			InTrack->Rename(nullptr, this);
 			check(InTrack);
 			Binding.AddTrack(*InTrack);
