Sequencer: Move OnStartedPlaying() to within first update in UpdateTimeCursorPosition(). This fixes an issue where setting the hud/player visibility was getting stomped on later.

diff --git a/Engine/Source/Runtime/MovieScene/Private/MovieSceneSequencePlayer.cpp b/Engine/Source/Runtime/MovieScene/Private/MovieSceneSequencePlayer.cpp
--- a/Engine/Source/Runtime/MovieScene/Private/MovieSceneSequencePlayer.cpp
+++ b/Engine/Source/Runtime/MovieScene/Private/MovieSceneSequencePlayer.cpp
@@ -39,6 +39,7 @@ UMovieSceneSequencePlayer::UMovieSceneSequencePlayer(const FObjectInitializer& I
 	, Status(EMovieScenePlayerStatus::Stopped)
 	, bReversePlayback(false)
 	, bIsEvaluating(false)
+	, bPendingOnStartedPlaying(false)
 	, Sequence(nullptr)
 	, StartTime(0)
 	, DurationFrames(0)
@@ -161,11 +162,10 @@ void UMovieSceneSequencePlayer::PlayInternal()
 			PreAnimatedState.EnableGlobalCapture();
 		}
 
+		bPendingOnStartedPlaying = true;
 		Status = EMovieScenePlayerStatus::Playing;
 		TimeController->StartPlaying(GetCurrentTime());
-
-		OnStartedPlaying();
-
+		
 		UMovieSceneSequence* MovieSceneSequence = RootTemplateInstance.GetSequence(MovieSceneSequenceID::Root);
 		UMovieScene*         MovieScene         = MovieSceneSequence ? MovieSceneSequence->GetMovieScene() : nullptr;
 
@@ -738,6 +738,12 @@ void UMovieSceneSequencePlayer::UpdateTimeCursorPosition_Internal(FFrameTime New
 		return;
 	}
 	
+	if (bPendingOnStartedPlaying)
+	{
+		OnStartedPlaying();
+		bPendingOnStartedPlaying = false;
+	}
+
 	if (Method == EUpdatePositionMethod::Play && ShouldStopOrLoop(NewPosition))
 	{
 		// The actual start time taking into account reverse playback
diff --git a/Engine/Source/Runtime/MovieScene/Public/MovieSceneSequencePlayer.h b/Engine/Source/Runtime/MovieScene/Public/MovieSceneSequencePlayer.h
--- a/Engine/Source/Runtime/MovieScene/Public/MovieSceneSequencePlayer.h
+++ b/Engine/Source/Runtime/MovieScene/Public/MovieSceneSequencePlayer.h
@@ -608,6 +608,9 @@ protected:
 	/** Set to true while evaluating to prevent reentrancy */
 	uint32 bIsEvaluating : 1;
 
+	/** Set to true to invoke OnStartedPlaying on first update tick for started playing */
+	uint32 bPendingOnStartedPlaying : 1;
+
 	/** The sequence to play back */
 	UPROPERTY(transient)
 	UMovieSceneSequence* Sequence;
@@ -681,4 +684,6 @@ private:
 	* Valid only if we've been ticked at least once since having a tick interval; otherwise set to -1.f
 	*/
 	float LastTickGameTimeSeconds;
+
+
 };
