Added pragma once and updated copyright notices

diff --git a/Engine/Source/Runtime/MovieScene/Private/Evaluation/MovieSceneMotionVectorSimulation.cpp b/Engine/Source/Runtime/MovieScene/Private/Evaluation/MovieSceneMotionVectorSimulation.cpp
--- a/Engine/Source/Runtime/MovieScene/Private/Evaluation/MovieSceneMotionVectorSimulation.cpp
+++ b/Engine/Source/Runtime/MovieScene/Private/Evaluation/MovieSceneMotionVectorSimulation.cpp
@@ -1,4 +1,4 @@
-// Copyright 1998-2017 Epic Games, Inc. All Rights Reserved.
+// Copyright 1998-2019 Epic Games, Inc. All Rights Reserved.
 
 #include "Evaluation/MovieSceneMotionVectorSimulation.h"
 #include "IMovieScenePlayer.h"
diff --git a/Engine/Source/Runtime/MovieScene/Private/Evaluation/MovieSceneMotionVectorSimulation.h b/Engine/Source/Runtime/MovieScene/Private/Evaluation/MovieSceneMotionVectorSimulation.h
--- a/Engine/Source/Runtime/MovieScene/Private/Evaluation/MovieSceneMotionVectorSimulation.h
+++ b/Engine/Source/Runtime/MovieScene/Private/Evaluation/MovieSceneMotionVectorSimulation.h
@@ -1,4 +1,6 @@
-// Copyright 1998-2017 Epic Games, Inc. All Rights Reserved.
+// Copyright 1998-2019 Epic Games, Inc. All Rights Reserved.
+
+#pragma once
 
 #include "UObject/ObjectKey.h"
 #include "Evaluation/IMovieSceneMotionVectorSimulation.h"
