Sequencer: ResolveLocalToRoot now resolves a root binding as a local space binding if the parent is not root.

This addresses a situation where you make a root binding like this:

MasterSequence
    SequenceA -> binding here to an object in SequenceB
    SequenceB

Then add another layer:
    MasterMasterSequence
        MasterSequence
            SequenceA, B

diff --git a/Engine/Source/Runtime/MovieScene/Private/MovieSceneObjectBindingID.cpp b/Engine/Source/Runtime/MovieScene/Private/MovieSceneObjectBindingID.cpp
--- a/Engine/Source/Runtime/MovieScene/Private/MovieSceneObjectBindingID.cpp
+++ b/Engine/Source/Runtime/MovieScene/Private/MovieSceneObjectBindingID.cpp
@@ -7,12 +7,25 @@ FMovieSceneObjectBindingID FMovieSceneObjectBindingID::ResolveLocalToRoot(FMovie
 {
 	FMovieSceneSequenceID NewSequenceID = FMovieSceneSequenceID(uint32(SequenceID));
 
-	if (Space == EMovieSceneObjectBindingSpace::Local && LocalSequenceID != MovieSceneSequenceID::Root)
+	// If the space is root, the binding was made as a root binding, but it's now being referred to as a local space binding, so resolve from the parent
+	bool bResolveFromParent = false;
+	if (Space == EMovieSceneObjectBindingSpace::Root)
+	{
+		const FMovieSceneSequenceHierarchyNode* CurrentNode = Hierarchy.FindNode(LocalSequenceID);
+		if (CurrentNode && CurrentNode->ParentID != MovieSceneSequenceID::Root && Hierarchy.FindNode(CurrentNode->ParentID))
+		{
+			LocalSequenceID = CurrentNode->ParentID;
+			bResolveFromParent = true;
+		}
+	}
+
+	if ((Space == EMovieSceneObjectBindingSpace::Local && LocalSequenceID != MovieSceneSequenceID::Root) || bResolveFromParent)
 	{
 		while (LocalSequenceID != MovieSceneSequenceID::Root)
 		{
 			const FMovieSceneSequenceHierarchyNode* CurrentNode = Hierarchy.FindNode(LocalSequenceID);
 			const FMovieSceneSubSequenceData* SubData = Hierarchy.FindSubData(LocalSequenceID);
+			
 			if (!ensureAlwaysMsgf(CurrentNode && SubData, TEXT("Malformed sequence hierarchy")))
 			{
 				return FMovieSceneObjectBindingID(Guid, NewSequenceID);
