Update copyright

diff --git a/Engine/Source/Runtime/MovieScene/Public/Evaluation/IMovieSceneMotionVectorSimulation.h b/Engine/Source/Runtime/MovieScene/Public/Evaluation/IMovieSceneMotionVectorSimulation.h
--- a/Engine/Source/Runtime/MovieScene/Public/Evaluation/IMovieSceneMotionVectorSimulation.h
+++ b/Engine/Source/Runtime/MovieScene/Public/Evaluation/IMovieSceneMotionVectorSimulation.h
@@ -1,4 +1,4 @@
-// Copyright 1998-2017 Epic Games, Inc. All Rights Reserved.
+// Copyright 1998-2019 Epic Games, Inc. All Rights Reserved.
 
 #pragma once
 
