Corrected clause in  IMovieSceneMotionVectorSimulation::IsEnabled

diff --git a/Engine/Source/Runtime/MovieScene/Private/Evaluation/MovieSceneMotionVectorSimulation.cpp b/Engine/Source/Runtime/MovieScene/Private/Evaluation/MovieSceneMotionVectorSimulation.cpp
--- a/Engine/Source/Runtime/MovieScene/Private/Evaluation/MovieSceneMotionVectorSimulation.cpp
+++ b/Engine/Source/Runtime/MovieScene/Private/Evaluation/MovieSceneMotionVectorSimulation.cpp
@@ -48,7 +48,7 @@ void IMovieSceneMotionVectorSimulation::EnableThisFrame(FPersistentEvaluationDat
 
 bool IMovieSceneMotionVectorSimulation::IsEnabled(const FPersistentEvaluationData& PersistentData, const FMovieSceneContext& Context)
 {
-	return !Context.IsSilent() && FSharedMotionSimulationData::IsEnabled(PersistentData) && FSharedMotionSimulationData::IsEnabled(PersistentData);
+	return !Context.IsSilent() && FSharedMotionSimulationData::IsEnabled(PersistentData) && FMotionVectorSimulation::IsEnabled();
 }
 
 FFrameTime IMovieSceneMotionVectorSimulation::GetSimulationTime(const FMovieSceneContext& Context)
