Back out changelist 6670633

diff --git a/Engine/Source/Runtime/MovieScene/Private/MovieSceneObjectBindingID.cpp b/Engine/Source/Runtime/MovieScene/Private/MovieSceneObjectBindingID.cpp
--- a/Engine/Source/Runtime/MovieScene/Private/MovieSceneObjectBindingID.cpp
+++ b/Engine/Source/Runtime/MovieScene/Private/MovieSceneObjectBindingID.cpp
@@ -7,25 +7,12 @@ FMovieSceneObjectBindingID FMovieSceneObjectBindingID::ResolveLocalToRoot(FMovie
 {
 	FMovieSceneSequenceID NewSequenceID = FMovieSceneSequenceID(uint32(SequenceID));
 
-	// If the space is root, the binding was made as a root binding, but it's now being referred to as a local space binding, so resolve from the parent
-	bool bResolveFromParent = false;
-	if (Space == EMovieSceneObjectBindingSpace::Root)
-	{
-		const FMovieSceneSequenceHierarchyNode* CurrentNode = Hierarchy.FindNode(LocalSequenceID);
-		if (CurrentNode && CurrentNode->ParentID != MovieSceneSequenceID::Root && Hierarchy.FindNode(CurrentNode->ParentID))
-		{
-			LocalSequenceID = CurrentNode->ParentID;
-			bResolveFromParent = true;
-		}
-	}
-
-	if ((Space == EMovieSceneObjectBindingSpace::Local && LocalSequenceID != MovieSceneSequenceID::Root) || bResolveFromParent)
+	if (Space == EMovieSceneObjectBindingSpace::Local && LocalSequenceID != MovieSceneSequenceID::Root)
 	{
 		while (LocalSequenceID != MovieSceneSequenceID::Root)
 		{
 			const FMovieSceneSequenceHierarchyNode* CurrentNode = Hierarchy.FindNode(LocalSequenceID);
 			const FMovieSceneSubSequenceData* SubData = Hierarchy.FindSubData(LocalSequenceID);
-			
 			if (!ensureAlwaysMsgf(CurrentNode && SubData, TEXT("Malformed sequence hierarchy")))
 			{
 				return FMovieSceneObjectBindingID(Guid, NewSequenceID);
