Sequencer: Reset LastTickGameTimeSeconds on stop. Fixes a regression in ticking sequencer separately.

diff --git a/Engine/Source/Runtime/MovieScene/Private/MovieSceneSequencePlayer.cpp b/Engine/Source/Runtime/MovieScene/Private/MovieSceneSequencePlayer.cpp
--- a/Engine/Source/Runtime/MovieScene/Private/MovieSceneSequencePlayer.cpp
+++ b/Engine/Source/Runtime/MovieScene/Private/MovieSceneSequencePlayer.cpp
@@ -293,6 +293,7 @@ void UMovieSceneSequencePlayer::StopInternal(FFrameTime TimeToResetTo)
 		}
 
 		CurrentNumLoops = 0;
+		LastTickGameTimeSeconds = -1;
 
 		// Reset loop count on stop so that it doesn't persist to the next call to play
 		PlaybackSettings.LoopCount.Value = 0;
