UMovieSceneFolder now prints the sequence name as part of the warning when it notifies you of automatically removed missing children.

diff --git a/Engine/Source/Runtime/MovieScene/Private/MovieSceneFolder.cpp b/Engine/Source/Runtime/MovieScene/Private/MovieSceneFolder.cpp
--- a/Engine/Source/Runtime/MovieScene/Private/MovieSceneFolder.cpp
+++ b/Engine/Source/Runtime/MovieScene/Private/MovieSceneFolder.cpp
@@ -173,7 +173,7 @@ void UMovieSceneFolder::PostLoad()
 				ChildMasterTracks.RemoveAt(ChildMasterTrackIndex);
 				ChildMasterTrackIndex--;
 
-				UE_LOG(LogMovieScene, Warning, TEXT("Folder (%s) contained a reference to a Master Track (%s) that no longer exists in the sequence, removing."), *GetFolderName().ToString(), *GetNameSafe(ChildTrack));
+				UE_LOG(LogMovieScene, Warning, TEXT("Folder (%s) in Sequence (%s) contained a reference to a Master Track (%s) that no longer exists in the sequence, removing."), *GetFolderName().ToString(), *OwningScene->GetPathName(), *GetNameSafe(ChildTrack));
 			}
 		}
 
@@ -186,7 +186,7 @@ void UMovieSceneFolder::PostLoad()
 				ChildObjectBindings.RemoveAt(ChildObjectBindingIndex);
 				ChildObjectBindingIndex--;
 
-				UE_LOG(LogMovieScene, Warning, TEXT("Folder (%s) contained a reference to an Object Binding (%s) that no longer exists in the sequence, removing."), *GetFolderName().ToString(), *ChildBinding.ToString());
+				UE_LOG(LogMovieScene, Warning, TEXT("Folder (%s) in Sequence (%s) contained a reference to an Object Binding (%s) that no longer exists in the sequence, removing."), *GetFolderName().ToString(), *OwningScene->GetPathName(), *ChildBinding.ToString());
 			}
 		}
 
@@ -204,7 +204,7 @@ void UMovieSceneFolder::PostLoad()
 				ChildFolders.RemoveAt(ChildFolderIndex);
 				ChildFolderIndex--;
 
-				UE_LOG(LogMovieScene, Warning, TEXT("Folder (%s) contained a reference to an Folder (%s) that exists in multiple places in the sequence, removing."), *GetFolderName().ToString(), *ChildFolders[ChildFolderIndex]->GetFolderName().ToString());
+				UE_LOG(LogMovieScene, Warning, TEXT("Folder (%s) in Sequence (%s) contained a reference to an Folder (%s) that exists in multiple places in the sequence, removing."), *GetFolderName().ToString(), *OwningScene->GetPathName(), *ChildFolders[ChildFolderIndex]->GetFolderName().ToString());
 			}
 		}
 	}
