Calculating Auto-Tangents for MovieSceneFloatChannels no longer resets the tangent handle values if the preceeding key is a constant key. This was legacy behavior before evaluation was changed to not take tangent values into account when evaluating constant keys.

diff --git a/Engine/Source/Runtime/MovieScene/Private/Channels/MovieSceneFloatChannel.cpp b/Engine/Source/Runtime/MovieScene/Private/Channels/MovieSceneFloatChannel.cpp
--- a/Engine/Source/Runtime/MovieScene/Private/Channels/MovieSceneFloatChannel.cpp
+++ b/Engine/Source/Runtime/MovieScene/Private/Channels/MovieSceneFloatChannel.cpp
@@ -787,15 +787,7 @@ void FMovieSceneFloatChannel::AutoSetTangents(float Tension)
 			NewTangent = (NextKey.Value - ThisKey.Value) / NextTimeDiff;
 			ThisKey.Tangent.LeaveTangent = NewTangent;
 		}
-		else if (PrevKey.InterpMode == RCIM_Constant || ThisKey.InterpMode == RCIM_Constant)
-		{
-			if (PrevKey.InterpMode != RCIM_Cubic)
-			{
-				ThisKey.Tangent.ArriveTangent = 0.f;
-			}
 
-			ThisKey.Tangent.LeaveTangent  = 0.0f;
-		}
 		
 	}
 }
