Sequencer: Corrected FMovieSceneChannelHandle::operator!=

diff --git a/Engine/Source/Runtime/MovieScene/Public/Channels/MovieSceneChannelHandle.h b/Engine/Source/Runtime/MovieScene/Public/Channels/MovieSceneChannelHandle.h
--- a/Engine/Source/Runtime/MovieScene/Public/Channels/MovieSceneChannelHandle.h
+++ b/Engine/Source/Runtime/MovieScene/Public/Channels/MovieSceneChannelHandle.h
@@ -48,16 +48,7 @@ public:
 	 */
 	friend bool operator!=(const FMovieSceneChannelHandle& A, const FMovieSceneChannelHandle& B)
 	{
-		if (A.ChannelIndex != B.ChannelIndex || A.ChannelTypeName != B.ChannelTypeName)
-		{
-			return false;
-		}
-
-		FMovieSceneChannelProxy* ProxyA = A.WeakChannelProxy.Pin().Get();
-		FMovieSceneChannelProxy* ProxyB = B.WeakChannelProxy.Pin().Get();
-
-		// Expired channel handles always compare falsey
-		return !ProxyA || !ProxyB || ProxyA != ProxyB;
+		return !(A == B);
 	}
 
 public:
