Sequencer: Attach tracks should not store/restore transform as that will be handled by DetachFromParent

diff --git a/Engine/Source/Runtime/MovieScene/Public/IMovieScenePlayer.h b/Engine/Source/Runtime/MovieScene/Public/IMovieScenePlayer.h
--- a/Engine/Source/Runtime/MovieScene/Public/IMovieScenePlayer.h
+++ b/Engine/Source/Runtime/MovieScene/Public/IMovieScenePlayer.h
@@ -277,6 +277,19 @@ public:
 		PreAnimatedState.SavePreAnimatedState(InTokenType, InProducer, ECapturePreAnimatedState::Entity, CaptureEntity);
 	}
 
+	/**
+	 * Attempt to save specific global state for the specified token state before it animates an object.
+	 * @note: Will only call IMovieSceneExecutionToken::CacheExistingState if no state has been previously cached for the specified token type
+	 *
+	 * @param InObject			The object to cache state for
+	 * @param InTokenType		Unique marker that identifies the originating token type
+	 * @param InProducer		Producer implementation that defines how to create the preanimated token, if it doesn't already exist
+	 */
+	FORCEINLINE void SaveGlobalPreAnimatedState(UObject& InObject, FMovieSceneAnimTypeID InTokenType, const IMovieScenePreAnimatedTokenProducer& InProducer)
+	{
+		PreAnimatedState.SavePreAnimatedState(InTokenType, InProducer, InObject, ECapturePreAnimatedState::Global, FMovieSceneEvaluationKey());
+	}
+
 	/**
 	 * Restore all pre-animated state
 	 */
diff --git a/Engine/Source/Runtime/MovieSceneTracks/Private/Evaluation/MovieScene3DAttachTemplate.cpp b/Engine/Source/Runtime/MovieSceneTracks/Private/Evaluation/MovieScene3DAttachTemplate.cpp
--- a/Engine/Source/Runtime/MovieSceneTracks/Private/Evaluation/MovieScene3DAttachTemplate.cpp
+++ b/Engine/Source/Runtime/MovieSceneTracks/Private/Evaluation/MovieScene3DAttachTemplate.cpp
@@ -211,8 +211,8 @@ struct F3DAttachExecutionToken
 			{
 				Player.SavePreAnimatedState(*SceneComponent, FMobilityTokenProducer::GetAnimTypeID(), FMobilityTokenProducer());
 
-				// Attachment can affect transform as well, so save transform state.
-				Player.SavePreAnimatedState(*SceneComponent, F3DTransformTokenProducer::GetAnimTypeID(), F3DTransformTokenProducer());
+				// Attachment can affect transform as well, so save global transform state.
+				Player.SaveGlobalPreAnimatedState(*SceneComponent, F3DTransformTokenProducer::GetAnimTypeID(), F3DTransformTokenProducer());
 
 				Player.SavePreAnimatedState(*SceneComponent, TMovieSceneAnimTypeID<F3DAttachExecutionToken>(), F3DAttachTokenProducer(AttachmentLocationRule, AttachmentRotationRule, AttachmentScaleRule, DetachmentLocationRule, DetachmentRotationRule, DetachmentScaleRule));
 
