Sequencer: Update comments for PlayToSeconds and PlayToMarkedFrame so that it matches the description of PlayToFrame

diff --git a/Engine/Source/Runtime/MovieScene/Public/MovieSceneSequencePlayer.h b/Engine/Source/Runtime/MovieScene/Public/MovieSceneSequencePlayer.h
--- a/Engine/Source/Runtime/MovieScene/Public/MovieSceneSequencePlayer.h
+++ b/Engine/Source/Runtime/MovieScene/Public/MovieSceneSequencePlayer.h
@@ -371,7 +371,8 @@ public:
 
 
 	/**
-	 * Play the sequence from the current time, to the specified time in seconds
+	 * Low-level call to set the current time of the player by evaluating only the specified time. Will not trigger any events.
+	 * Does not alter the persistent playback status of the player (IsPlaying).
 	 *
 	 * @param TimeInSeconds   The desired time in seconds
 	 */
@@ -379,7 +380,8 @@ public:
 	void PlayToSeconds(float TimeInSeconds);
 
 	/**
-	 * Scrub the sequence from the current time, to the specified time in seconds
+	 * Low-level call to set the current time of the player by evaluating only the specified time. Will not trigger any events.
+	 * Does not alter the persistent playback status of the player (IsPlaying).
 	 *
 	 * @param TimeInSeconds   The desired time in seconds
 	 */
@@ -387,7 +389,8 @@ public:
 	void ScrubToSeconds(float TimeInSeconds);
 
 	/**
-	 * Jump to the specified time in seconds, without evaluating the sequence in between the current and desired time (as if in a paused state)
+	 * Low-level call to set the current time of the player by evaluating only the specified time, as if scrubbing the timeline. Will trigger only events that exist at the specified time.
+	 * Does not alter the persistent playback status of the player (IsPlaying).
 	 *
 	 * @param TimeInSeconds   The desired time in seconds
 	 */
@@ -396,7 +399,8 @@ public:
 
 
 	/**
-	 * Play the sequence from the current time, to the specified marked frame by label
+	 * Low-level call to set the current time of the player to the marked frame by label by evaluating only the specified time. Will not trigger any events.
+	 * Does not alter the persistent playback status of the player (IsPlaying).
 	 *
 	 * @param InLabel   The desired marked frame label to play to
 	 * @return Whether the marked frame was found
@@ -405,7 +409,8 @@ public:
 	bool PlayToMarkedFrame(const FString& InLabel);
 
 	/**
-	 * Scrub the sequence from the current time, to the specified marked frame by label
+	 * Low-level call to set the current time of the player to the marked frame by label by evaluating only the specified time. Will not trigger any events.
+	 * Does not alter the persistent playback status of the player (IsPlaying).
 	 *
 	 * @param InLabel   The desired marked frame label to scrub to
 	 * @return Whether the marked frame was found
@@ -414,7 +419,8 @@ public:
 	bool ScrubToMarkedFrame(const FString& InLabel);
 
 	/**
-	 * Jump to the specified marked frame by label, without evaluating the sequence in between the current and desired time (as if in a paused state)
+	 * Low-level call to set the current time of the player to the marked frame by label by evaluating only the specified time, as if scrubbing the timeline. Will trigger only events that exist at the specified time.
+	 * Does not alter the persistent playback status of the player (IsPlaying).
 	 *
 	 * @param InLabel   The desired marked frame label to jump to
 	 * @return Whether the marked frame was found
