Sequencer: Fix pause condition needing to reset LastTickGameTimeSeconds.

diff --git a/Engine/Source/Runtime/MovieScene/Private/MovieSceneSequencePlayer.cpp b/Engine/Source/Runtime/MovieScene/Private/MovieSceneSequencePlayer.cpp
--- a/Engine/Source/Runtime/MovieScene/Private/MovieSceneSequencePlayer.cpp
+++ b/Engine/Source/Runtime/MovieScene/Private/MovieSceneSequencePlayer.cpp
@@ -221,6 +221,8 @@ void UMovieSceneSequencePlayer::Pause()
 		Status = EMovieScenePlayerStatus::Paused;
 		TimeController->StopPlaying(GetCurrentTime());
 
+		LastTickGameTimeSeconds = -1;
+
 		// Evaluate the sequence at its current time, with a status of 'stopped' to ensure that animated state pauses correctly. (ie. audio sounds should stop/pause)
 		{
 			bIsEvaluating = true;
