Sequencer: Fix incorrect enum initialization

diff --git a/Engine/Source/Runtime/MovieScene/Public/Evaluation/Blending/MovieSceneBlendType.h b/Engine/Source/Runtime/MovieScene/Public/Evaluation/Blending/MovieSceneBlendType.h
--- a/Engine/Source/Runtime/MovieScene/Public/Evaluation/Blending/MovieSceneBlendType.h
+++ b/Engine/Source/Runtime/MovieScene/Public/Evaluation/Blending/MovieSceneBlendType.h
@@ -10,8 +10,6 @@
 UENUM()
 enum class EMovieSceneBlendType : uint8
 {
-	None = 0 UMETA(Hidden),
-
 	/** Blends all other weighted values together as an average of the total weight */
 	Absolute			= 0x1,
 	/** Applies this value as a sum total of all other additives */
