Sequencer: Fix stop condition logic

diff --git a/Engine/Source/Runtime/MovieScene/Private/MovieSceneSequencePlayer.cpp b/Engine/Source/Runtime/MovieScene/Private/MovieSceneSequencePlayer.cpp
--- a/Engine/Source/Runtime/MovieScene/Private/MovieSceneSequencePlayer.cpp
+++ b/Engine/Source/Runtime/MovieScene/Private/MovieSceneSequencePlayer.cpp
@@ -281,7 +281,7 @@ void UMovieSceneSequencePlayer::StopAtCurrentTime()
 
 void UMovieSceneSequencePlayer::StopInternal(FFrameTime TimeToResetTo)
 {
-	if (IsPlaying() || IsPaused() || RootTemplateInstance.IsValid())
+	if (IsPlaying() || IsPaused())
 	{
 		if (bIsEvaluating)
 		{
@@ -309,7 +309,10 @@ void UMovieSceneSequencePlayer::StopInternal(FFrameTime TimeToResetTo)
 			RestorePreAnimatedState();
 		}
 
-		RootTemplateInstance.Finish(*this);
+		if (RootTemplateInstance.IsValid())
+		{
+			RootTemplateInstance.Finish(*this);
+		}
 
 		if (OldMaxTickRate.IsSet())
 		{
@@ -326,10 +329,13 @@ void UMovieSceneSequencePlayer::StopInternal(FFrameTime TimeToResetTo)
 
 		OnStopped();
 
-		UMovieSceneSequence* MovieSceneSequence = RootTemplateInstance.GetSequence(MovieSceneSequenceID::Root);
-		if (MovieSceneSequence)
+		if (RootTemplateInstance.IsValid())
 		{
-			UE_LOG(LogMovieScene, Verbose, TEXT("Stop - MovieSceneSequence: %s"), *MovieSceneSequence->GetName());
+			UMovieSceneSequence* MovieSceneSequence = RootTemplateInstance.GetSequence(MovieSceneSequenceID::Root);
+			if (MovieSceneSequence)
+			{
+				UE_LOG(LogMovieScene, Verbose, TEXT("Stop - MovieSceneSequence: %s"), *MovieSceneSequence->GetName());
+			}
 		}
 
 		if (OnStop.IsBound())
@@ -342,6 +348,7 @@ void UMovieSceneSequencePlayer::StopInternal(FFrameTime TimeToResetTo)
 void UMovieSceneSequencePlayer::GoToEndAndStop()
 {
 	FFrameTime Time = GetLastValidTime();
+	Status = EMovieScenePlayerStatus::Playing;
 	JumpToFrame(Time);
 	StopInternal(Time);
 }
