Sequencer: Short-circuit segment lookup if there is only a single segment

diff --git a/Engine/Source/Runtime/MovieScene/Private/Evaluation/MovieSceneEvaluationTrack.cpp b/Engine/Source/Runtime/MovieScene/Private/Evaluation/MovieSceneEvaluationTrack.cpp
--- a/Engine/Source/Runtime/MovieScene/Private/Evaluation/MovieSceneEvaluationTrack.cpp
+++ b/Engine/Source/Runtime/MovieScene/Private/Evaluation/MovieSceneEvaluationTrack.cpp
@@ -394,6 +394,14 @@ void FMovieSceneEvaluationTrack::Initialize(FMovieSceneSegmentIdentifier Segment
 
 void FMovieSceneEvaluationTrack::DefaultInitialize(FMovieSceneSegmentIdentifier SegmentID, const FMovieSceneEvaluationOperand& Operand, FMovieSceneContext Context, FPersistentEvaluationData& PersistentData, IMovieScenePlayer& Player) const
 {
+	// If we only have a single child template and a single segment, short-circuit the lookup into the impl array to avoid the unnecessary cache misses
+	if (ChildTemplates.Num() == 1 && Segments.Num() == 1)
+	{
+		PersistentData.DeriveSectionKey(0);
+		ChildTemplates[0]->Initialize(Operand, Context, PersistentData, Player);
+		return;
+	}
+
 	const FMovieSceneSegment& Segment = Segments[SegmentID];
 	for (const FSectionEvaluationData& EvalData : Segment.Impls)
 	{
@@ -438,6 +446,19 @@ void FMovieSceneEvaluationTrack::DefaultEvaluate(FMovieSceneSegmentIdentifier Se
 
 void FMovieSceneEvaluationTrack::EvaluateStatic(FMovieSceneSegmentIdentifier SegmentID, const FMovieSceneEvaluationOperand& Operand, FMovieSceneContext Context, const FPersistentEvaluationData& PersistentData, FMovieSceneExecutionTokens& ExecutionTokens) const
 {
+	// If we only have a single child template and a single segment, short-circuit the lookup into the impl array to avoid the unnecessary cache misses
+	if (ChildTemplates.Num() == 1 && Segments.Num() == 1)
+	{
+		const FMovieSceneEvalTemplate& Template = ChildTemplates[0].GetValue();
+
+		PersistentData.DeriveSectionKey(0);
+		ExecutionTokens.SetCurrentScope(FMovieSceneEvaluationScope(PersistentData.GetSectionKey(), Template.GetCompletionMode()));
+		ExecutionTokens.SetContext(Context);
+
+		Template.Evaluate(Operand, Context, PersistentData, ExecutionTokens);
+		return;
+	}
+
 	int32 SortedIndex = Segments.GetSortedIndex(SegmentID);
 	if (SortedIndex == INDEX_NONE)
 	{
diff --git a/Engine/Source/Runtime/MovieScene/Public/Evaluation/MovieSceneEvaluationTrack.h b/Engine/Source/Runtime/MovieScene/Public/Evaluation/MovieSceneEvaluationTrack.h
--- a/Engine/Source/Runtime/MovieScene/Public/Evaluation/MovieSceneEvaluationTrack.h
+++ b/Engine/Source/Runtime/MovieScene/Public/Evaluation/MovieSceneEvaluationTrack.h
@@ -93,6 +93,11 @@ struct FMovieSceneEvaluationTrackSegments
 	/** Add a new segment to the container */
 	FMovieSceneSegmentIdentifier Add(FMovieSceneSegment&& In);
 
+	int32 Num() const
+	{
+		return SegmentIdentifierToIndex.Num();
+	}
+
 private:
 
 	/** Array of indices into SortedSegments where each FMovieSceneSegmentIdentifier represents and index into SegmentIdentifierToIndex. Never shuffled until the container is reset. */
