Sequencer: Prevent pre-animated state being restored on spawnables that are already destroyed.

This change introduces a new flag to evaluation templates: whether to prioritize tear down or not. Normally tear down strictly occurs in the reverse order of setup, however this is not desireable for spawned objects since we end up restoring a bunch of state such as transforms, mobility, attachments etc on objects that are going to be destroyed that same evaluation. Spawn tracks now specify this prioritized tear down which means that the object will always be destroyed first, and any subsequent calls to RestorePreAnimatedState will early out because the object is now pending destruction.

diff --git a/Engine/Source/Runtime/MovieScene/Private/Compilation/MovieSceneCompiler.cpp b/Engine/Source/Runtime/MovieScene/Private/Compilation/MovieSceneCompiler.cpp
--- a/Engine/Source/Runtime/MovieScene/Private/Compilation/MovieSceneCompiler.cpp
+++ b/Engine/Source/Runtime/MovieScene/Private/Compilation/MovieSceneCompiler.cpp
@@ -593,7 +593,8 @@ void FMovieSceneCompiler::PopulateMetaData(FCompiledGroupResult& OutResult, cons
 	OutResult.MetaData.Reset();
 
 	// Add all the init tracks first
-	uint32 SortOrder = 0;
+	uint16 SetupIndex    = 0;
+	uint16 TearDownIndex = 0;
 	for (const FCompileOnTheFlyData& CompileData : SortedCompileData)
 	{
 		if (!CompileData.bRequiresInit)
@@ -601,16 +602,21 @@ void FMovieSceneCompiler::PopulateMetaData(FCompiledGroupResult& OutResult, cons
 			continue;
 		}
 
+		const bool bPriorityTearDown = CompileData.Track->HasTearDownPriority();
+
 		FMovieSceneEvaluationFieldSegmentPtr SegmentPtr = CompileData.Segment;
 
 		// Add the track key
 		FMovieSceneEvaluationKey TrackKey(SegmentPtr.SequenceID, SegmentPtr.TrackIdentifier);
-		OutResult.MetaData.ActiveEntities.Add(FMovieSceneOrderedEvaluationKey{ TrackKey, SortOrder++ });
+		OutResult.MetaData.ActiveEntities.Add(FMovieSceneOrderedEvaluationKey{ TrackKey, SetupIndex++, (bPriorityTearDown ? TearDownIndex : uint16(MAX_uint16-TearDownIndex)) });
+		++TearDownIndex;
 
 		for (FSectionEvaluationData EvalData : CompileData.Track->GetSegment(SegmentPtr.SegmentID).Impls)
 		{
 			FMovieSceneEvaluationKey SectionKey = TrackKey.AsSection(EvalData.ImplIndex);
-			OutResult.MetaData.ActiveEntities.Add(FMovieSceneOrderedEvaluationKey{ SectionKey, SortOrder++ });
+			OutResult.MetaData.ActiveEntities.Add(FMovieSceneOrderedEvaluationKey{ SectionKey, SetupIndex++, (bPriorityTearDown ? TearDownIndex : uint16(MAX_uint16 - TearDownIndex)) });
+
+			++TearDownIndex;
 		}
 	}
 
@@ -622,16 +628,21 @@ void FMovieSceneCompiler::PopulateMetaData(FCompiledGroupResult& OutResult, cons
 			continue;
 		}
 
+		const bool bPriorityTearDown = CompileData.Track->HasTearDownPriority();
+
 		FMovieSceneEvaluationFieldSegmentPtr SegmentPtr = CompileData.Segment;
 
 		// Add the track key
 		FMovieSceneEvaluationKey TrackKey(SegmentPtr.SequenceID, SegmentPtr.TrackIdentifier);
-		OutResult.MetaData.ActiveEntities.Add(FMovieSceneOrderedEvaluationKey{ TrackKey, SortOrder++ });
+		OutResult.MetaData.ActiveEntities.Add(FMovieSceneOrderedEvaluationKey{ TrackKey, SetupIndex++, (bPriorityTearDown ? TearDownIndex : uint16(MAX_uint16 - TearDownIndex)) });
+		++TearDownIndex;
 
 		for (FSectionEvaluationData EvalData : CompileData.Track->GetSegment(SegmentPtr.SegmentID).Impls)
 		{
 			FMovieSceneEvaluationKey SectionKey = TrackKey.AsSection(EvalData.ImplIndex);
-			OutResult.MetaData.ActiveEntities.Add(FMovieSceneOrderedEvaluationKey{ SectionKey, SortOrder++ });
+			OutResult.MetaData.ActiveEntities.Add(FMovieSceneOrderedEvaluationKey{ SectionKey, SetupIndex++, (bPriorityTearDown ? TearDownIndex : uint16(MAX_uint16 - TearDownIndex)) });
+
+			++TearDownIndex;
 		}
 	}
 
diff --git a/Engine/Source/Runtime/MovieScene/Private/Evaluation/MovieSceneEvaluationField.cpp b/Engine/Source/Runtime/MovieScene/Private/Evaluation/MovieSceneEvaluationField.cpp
--- a/Engine/Source/Runtime/MovieScene/Private/Evaluation/MovieSceneEvaluationField.cpp
+++ b/Engine/Source/Runtime/MovieScene/Private/Evaluation/MovieSceneEvaluationField.cpp
@@ -311,8 +311,7 @@ void FMovieSceneEvaluationMetaData::DiffEntities(const FMovieSceneEvaluationMeta
 			++LastFrameKeys;
 		}
 
-		// Expired keys are torn down in reverse
-		Algo::SortBy(*ExpiredKeys, [](const FMovieSceneOrderedEvaluationKey& In){ return uint32(-1) - In.EvaluationIndex; });
+		Algo::SortBy(*ExpiredKeys, &FMovieSceneOrderedEvaluationKey::TearDownIndex);
 	}
 
 	// Add any remaining new entities
@@ -324,7 +323,7 @@ void FMovieSceneEvaluationMetaData::DiffEntities(const FMovieSceneEvaluationMeta
 			++ThisFrameKeys;
 		}
 
-		Algo::SortBy(*NewKeys, &FMovieSceneOrderedEvaluationKey::EvaluationIndex);
+		Algo::SortBy(*NewKeys, &FMovieSceneOrderedEvaluationKey::SetupIndex);
 	}
 }
 
diff --git a/Engine/Source/Runtime/MovieScene/Private/Evaluation/MovieSceneEvaluationTrack.cpp b/Engine/Source/Runtime/MovieScene/Private/Evaluation/MovieSceneEvaluationTrack.cpp
--- a/Engine/Source/Runtime/MovieScene/Private/Evaluation/MovieSceneEvaluationTrack.cpp
+++ b/Engine/Source/Runtime/MovieScene/Private/Evaluation/MovieSceneEvaluationTrack.cpp
@@ -115,6 +115,7 @@ FMovieSceneEvaluationTrack::FMovieSceneEvaluationTrack()
 	, SourceTrack(nullptr)
 	, bEvaluateInPreroll(true)
 	, bEvaluateInPostroll(true)
+	, bTearDownPriority(false)
 {
 }
 
@@ -125,6 +126,7 @@ FMovieSceneEvaluationTrack::FMovieSceneEvaluationTrack(const FGuid& InObjectBind
 	, SourceTrack(nullptr)
 	, bEvaluateInPreroll(true)
 	, bEvaluateInPostroll(true)
+	, bTearDownPriority(false)
 {
 }
 
diff --git a/Engine/Source/Runtime/MovieScene/Private/Evaluation/MovieScenePreAnimatedState.cpp b/Engine/Source/Runtime/MovieScene/Private/Evaluation/MovieScenePreAnimatedState.cpp
--- a/Engine/Source/Runtime/MovieScene/Private/Evaluation/MovieScenePreAnimatedState.cpp
+++ b/Engine/Source/Runtime/MovieScene/Private/Evaluation/MovieScenePreAnimatedState.cpp
@@ -1,6 +1,8 @@
 // Copyright 1998-2019 Epic Games, Inc. All Rights Reserved.
 
 #include "Evaluation/MovieScenePreAnimatedState.h"
+#include "GameFramework/Actor.h"
+#include "Components/ActorComponent.h"
 
 DECLARE_CYCLE_STAT(TEXT("Save Pre Animated State"), MovieSceneEval_SavePreAnimatedState, STATGROUP_MovieSceneEval);
 
@@ -32,6 +34,21 @@ namespace MovieSceneImpl
 	{
 		if (Object)
 		{
+			if (AActor* Actor = Cast<AActor>(Object))
+			{
+				if (Actor->IsActorBeingDestroyed())
+				{
+					return;
+				}
+			}
+			else if (UActorComponent* Component = Cast<UActorComponent>(Object))
+			{
+				if (Component->IsBeingDestroyed())
+				{
+					return;
+				}
+			}
+
 			if (Token.OptionalEntityToken.IsValid())
 			{
 				Token.OptionalEntityToken->RestoreState(*Object, Player);
diff --git a/Engine/Source/Runtime/MovieScene/Public/Evaluation/MovieSceneEvaluationField.h b/Engine/Source/Runtime/MovieScene/Public/Evaluation/MovieSceneEvaluationField.h
--- a/Engine/Source/Runtime/MovieScene/Public/Evaluation/MovieSceneEvaluationField.h
+++ b/Engine/Source/Runtime/MovieScene/Public/Evaluation/MovieSceneEvaluationField.h
@@ -150,7 +150,10 @@ struct FMovieSceneOrderedEvaluationKey
 	FMovieSceneEvaluationKey Key;
 
 	UPROPERTY()
-	uint32 EvaluationIndex;
+	uint16 SetupIndex;
+
+	UPROPERTY()
+	uint16 TearDownIndex;
 };
 
 /** Informational meta-data that applies to a given time range */
diff --git a/Engine/Source/Runtime/MovieScene/Public/Evaluation/MovieSceneEvaluationTrack.h b/Engine/Source/Runtime/MovieScene/Public/Evaluation/MovieSceneEvaluationTrack.h
--- a/Engine/Source/Runtime/MovieScene/Public/Evaluation/MovieSceneEvaluationTrack.h
+++ b/Engine/Source/Runtime/MovieScene/Public/Evaluation/MovieSceneEvaluationTrack.h
@@ -283,6 +283,22 @@ public:
 		return bEvaluateInPostroll;
 	}
 
+	/**
+	 * Tell this track to prioritize its tear down over other tracks, regardless of evaluation priority
+	 */
+	void PrioritizeTearDown()
+	{
+		bTearDownPriority = true;
+	}
+
+	/**
+	 * Check whether this track has tear down priority or not
+	 */
+	bool HasTearDownPriority() const
+	{
+		return bTearDownPriority;
+	}
+
 public:
 
 	/**
@@ -577,6 +593,10 @@ private:
 	/** Whether this track is evaluated in postroll */
 	UPROPERTY()
 	uint32 bEvaluateInPostroll : 1;
+
+	/** Whether track should be given priority when being torn down */
+	UPROPERTY()
+	uint32 bTearDownPriority : 1;
 };
 
 template<> struct TStructOpsTypeTraits<FMovieSceneEvaluationTrack> : public TStructOpsTypeTraitsBase2<FMovieSceneEvaluationTrack> { enum { WithPostSerialize = true, WithCopy = false }; };
diff --git a/Engine/Source/Runtime/MovieSceneTracks/Private/Tracks/MovieSceneSpawnTrack.cpp b/Engine/Source/Runtime/MovieSceneTracks/Private/Tracks/MovieSceneSpawnTrack.cpp
--- a/Engine/Source/Runtime/MovieSceneTracks/Private/Tracks/MovieSceneSpawnTrack.cpp
+++ b/Engine/Source/Runtime/MovieSceneTracks/Private/Tracks/MovieSceneSpawnTrack.cpp
@@ -125,6 +125,8 @@ void UMovieSceneSpawnTrack::PostCompile(FMovieSceneEvaluationTrack& OutTrack, co
 	OutTrack.SetEvaluationGroup(IMovieSceneTracksModule::GetEvaluationGroupName(EBuiltInEvaluationGroup::SpawnObjects));
 	// Set priority to highest possible
 	OutTrack.SetEvaluationPriority(GetEvaluationPriority());
+
+	OutTrack.PrioritizeTearDown();
 }
 
 #if WITH_EDITORONLY_DATA
