Sequencer: Increased inline allocation size for execution tokens to 128 bytes

diff --git a/Engine/Source/Runtime/MovieScene/Private/Evaluation/Blending/MovieSceneAccumulatedBlendState.cpp b/Engine/Source/Runtime/MovieScene/Private/Evaluation/Blending/MovieSceneAccumulatedBlendState.cpp
--- a/Engine/Source/Runtime/MovieScene/Private/Evaluation/Blending/MovieSceneAccumulatedBlendState.cpp
+++ b/Engine/Source/Runtime/MovieScene/Private/Evaluation/Blending/MovieSceneAccumulatedBlendState.cpp
@@ -16,7 +16,7 @@ void FMovieSceneAccumulatedBlendState::Consolidate(TMap<FMovieSceneBlendingKey,
 		{
 			if (UObject* Obj = WeakObj.Get())
 			{
-				for (TInlineValue<FTokenEntry>& Token : TokensToBlend)
+				for (TInlineValue<FTokenEntry, 128>& Token : TokensToBlend)
 				{
 					FMovieSceneBlendingKey Key = { Obj, Token->GetActuatorID() };
 
@@ -34,7 +34,7 @@ void FMovieSceneAccumulatedBlendState::Consolidate(TMap<FMovieSceneBlendingKey,
 
 void FMovieSceneAccumulatedBlendState::Consolidate(TMap<FMovieSceneBlendingKey, FActuatorTokenStackPtr>& InOutBlendState)
 {
-	for (TInlineValue<FTokenEntry>& Token : TokensToBlend)
+	for (TInlineValue<FTokenEntry, 128>& Token : TokensToBlend)
 	{
 		FMovieSceneBlendingKey Key = { nullptr, Token->GetActuatorID() };
 
diff --git a/Engine/Source/Runtime/MovieScene/Public/Evaluation/Blending/MovieSceneAccumulatedBlendState.h b/Engine/Source/Runtime/MovieScene/Public/Evaluation/Blending/MovieSceneAccumulatedBlendState.h
--- a/Engine/Source/Runtime/MovieScene/Public/Evaluation/Blending/MovieSceneAccumulatedBlendState.h
+++ b/Engine/Source/Runtime/MovieScene/Public/Evaluation/Blending/MovieSceneAccumulatedBlendState.h
@@ -177,5 +177,5 @@ private:
 	};
 
 	/** Array of all tokens that have been added this frame */
-	TArray<TInlineValue<FTokenEntry>> TokensToBlend;
+	TArray<TInlineValue<FTokenEntry, 128>> TokensToBlend;
 };
diff --git a/Engine/Source/Runtime/MovieScene/Public/Evaluation/MovieSceneEvalTemplate.h b/Engine/Source/Runtime/MovieScene/Public/Evaluation/MovieSceneEvalTemplate.h
--- a/Engine/Source/Runtime/MovieScene/Public/Evaluation/MovieSceneEvalTemplate.h
+++ b/Engine/Source/Runtime/MovieScene/Public/Evaluation/MovieSceneEvalTemplate.h
@@ -191,7 +191,7 @@ protected:
 USTRUCT()
 struct FMovieSceneEvalTemplatePtr
 #if CPP
-	: TInlineValue<FMovieSceneEvalTemplate, 48>
+	: TInlineValue<FMovieSceneEvalTemplate, 128>
 #endif
 {
 	GENERATED_BODY()
