Sequencer: Fixed rounding issue caused by transforming compile and clamp ranges with timescales < 1

diff --git a/Engine/Source/Runtime/MovieScene/Private/Compilation/MovieSceneCompiler.cpp b/Engine/Source/Runtime/MovieScene/Private/Compilation/MovieSceneCompiler.cpp
--- a/Engine/Source/Runtime/MovieScene/Private/Compilation/MovieSceneCompiler.cpp
+++ b/Engine/Source/Runtime/MovieScene/Private/Compilation/MovieSceneCompiler.cpp
@@ -42,20 +42,46 @@ struct FGatherParameters
 		SubParams.LocalCompileRange         = SubParams.RootCompileRange * SubData.RootToSequenceTransform;
 		SubParams.LocalClampRange           = SubParams.RootClampRange   * SubData.RootToSequenceTransform;
 
-		if (SubParams.LocalCompileRange.IsEmpty())
+		SubParams.AccountForRounding();
+
+		return SubParams;
+	}
+
+	void AccountForRounding()
+	{
+		if (LocalCompileRange.IsEmpty())
 		{
 			// Ensure that the compile range is not empty by extending the upper bound by one frame.
 			// This can happen when there is a scale < 1 on a sub section, and the root-space compile range is only 1 frame wide.
 			// We know that neither bound is Open by this point because such ranges would not be considered empty
-			SubParams.LocalCompileRange.SetUpperBoundValue(SubParams.LocalCompileRange.GetUpperBoundValue() + 1);
+			LocalCompileRange.SetUpperBoundValue(LocalCompileRange.GetUpperBoundValue() + 1);
+		}
+
+		// Due to rounding caused by inner timescales, it is possible that the compile range can end up not overlapping the clamp range, which will result in corruption.
+		// Specifically this will happen when the root compile and clamp ranges overlap by < 0.5/InnerScale frames since frame numbers are floored when being transformed to local space.
+		// For example, a transformation of 0.5 scale to a root compile and clamp range of [-100, 101) and [100, 200) would yield [-50, 50) and [50, 100) which do not overlap.
+		// To alleviate this, we treat the clamp range as the authoritative range (it is always set to the boundaries for the eval field, whereas the compile range is simply used as a guide to know where to compile)
+		if (!LocalCompileRange.Overlaps(LocalClampRange))
+		{
+			if (LocalClampRange.GetLowerBound().IsClosed() && LocalCompileRange.GetUpperBound().IsClosed() && LocalCompileRange.GetUpperBoundValue() <= LocalClampRange.GetLowerBoundValue())
+			{
+				// Force inclusive bounds
+				LocalCompileRange.SetUpperBound(LocalClampRange.GetLowerBoundValue());
+			}
+			if (LocalClampRange.GetUpperBound().IsClosed() && LocalCompileRange.GetLowerBound().IsClosed() && LocalCompileRange.GetLowerBoundValue() >= LocalClampRange.GetUpperBoundValue())
+			{
+				// Force inclusive bounds
+				LocalCompileRange.SetLowerBound(LocalClampRange.GetLowerBoundValue());
+			}
 		}
-		return SubParams;
 	}
 
 	void SetClampRange(TRange<FFrameNumber> InNewRootClampRange)
 	{
 		RootClampRange  = InNewRootClampRange;
 		LocalClampRange = InNewRootClampRange * RootToSequenceTransform;
+
+		AccountForRounding();
 	}
 
 	/** Clamp the specified range to the current clamp range (in root space) */
