Sequencer: Protect against assigning a sequence that would produce circular dependencies. For example, if the current level sequence is NewLevelSequence, this prevents the user from setting a subsection in it to NewLevelSequence as well.

diff --git a/Engine/Source/Runtime/MovieScene/Private/Sections/MovieSceneSubSection.cpp b/Engine/Source/Runtime/MovieScene/Private/Sections/MovieSceneSubSection.cpp
--- a/Engine/Source/Runtime/MovieScene/Private/Sections/MovieSceneSubSection.cpp
+++ b/Engine/Source/Runtime/MovieScene/Private/Sections/MovieSceneSubSection.cpp
@@ -1,12 +1,14 @@
 // Copyright 1998-2019 Epic Games, Inc. All Rights Reserved.
 
 #include "Sections/MovieSceneSubSection.h"
+#include "Tracks/MovieSceneSubTrack.h"
 #include "MovieScene.h"
 #include "MovieSceneTrack.h"
 #include "MovieSceneSequence.h"
 #include "MovieSceneTimeHelpers.h"
 #include "Evaluation/MovieSceneEvaluationTemplate.h"
 #include "Misc/FrameRate.h"
+#include "Logging/MessageLog.h"
 
 TWeakObjectPtr<UMovieSceneSubSection> UMovieSceneSubSection::TheRecordingSection;
 
@@ -196,8 +198,33 @@ AActor* UMovieSceneSubSection::GetActorToRecord()
 }
 
 #if WITH_EDITOR
+void UMovieSceneSubSection::PreEditChange(UProperty* PropertyAboutToChange)
+{
+	if (PropertyAboutToChange && PropertyAboutToChange->GetFName() == GET_MEMBER_NAME_CHECKED(UMovieSceneSubSection, SubSequence))
+	{
+		// Store the current subsequence in case it needs to be restored in PostEditChangeProperty because the new value would introduce a circular dependency
+		PreviousSubSequence = SubSequence;
+	}
+
+	return Super::PreEditChange(PropertyAboutToChange);
+}
+
 void UMovieSceneSubSection::PostEditChangeProperty(FPropertyChangedEvent& PropertyChangedEvent)
 {
+	if (PropertyChangedEvent.Property && PropertyChangedEvent.Property->GetFName() == GET_MEMBER_NAME_CHECKED(UMovieSceneSubSection, SubSequence))
+	{
+		UMovieSceneSubTrack* TrackOuter = Cast<UMovieSceneSubTrack>(GetOuter());
+		if (SubSequence && TrackOuter && TrackOuter->ContainsSequence(*SubSequence, true, this))
+		{
+			UE_LOG(LogMovieScene, Error, TEXT("Invalid level sequence %s. There could be a circular dependency."), *SubSequence->GetDisplayName().ToString());
+
+			// Restore to the previous sub sequence because there was a circular dependency
+			SubSequence = PreviousSubSequence;
+		}
+
+		PreviousSubSequence = nullptr;
+	}
+
 	Super::PostEditChangeProperty(PropertyChangedEvent);
 
 	// recreate runtime instance when sequence is changed
diff --git a/Engine/Source/Runtime/MovieScene/Private/Tracks/MovieSceneSubTrack.cpp b/Engine/Source/Runtime/MovieScene/Private/Tracks/MovieSceneSubTrack.cpp
--- a/Engine/Source/Runtime/MovieScene/Private/Tracks/MovieSceneSubTrack.cpp
+++ b/Engine/Source/Runtime/MovieScene/Private/Tracks/MovieSceneSubTrack.cpp
@@ -96,12 +96,23 @@ UMovieSceneSubSection* UMovieSceneSubTrack::AddSequenceToRecord()
 	return NewSection;
 }
 
-bool UMovieSceneSubTrack::ContainsSequence(const UMovieSceneSequence& Sequence, bool Recursively) const
+bool UMovieSceneSubTrack::ContainsSequence(const UMovieSceneSequence& Sequence, bool Recursively, const UMovieSceneSection* SectionToSkip) const
 {
+	UMovieSceneSequence* OuterSequence = GetTypedOuter<UMovieSceneSequence>();
+	if (OuterSequence == &Sequence)
+	{
+		return true;
+	}
+
 	for (const auto& Section : Sections)
 	{
 		const auto SubSection = CastChecked<UMovieSceneSubSection>(Section);
 
+		if (SubSection == SectionToSkip)
+		{
+			continue;
+		}
+
 		// is the section referencing the sequence?
 		const UMovieSceneSequence* SubSequence = SubSection->GetSequence();
 
diff --git a/Engine/Source/Runtime/MovieScene/Public/Sections/MovieSceneSubSection.h b/Engine/Source/Runtime/MovieScene/Public/Sections/MovieSceneSubSection.h
--- a/Engine/Source/Runtime/MovieScene/Public/Sections/MovieSceneSubSection.h
+++ b/Engine/Source/Runtime/MovieScene/Public/Sections/MovieSceneSubSection.h
@@ -96,6 +96,7 @@ public:
 	virtual void PostLoad() override;
 
 #if WITH_EDITOR
+	virtual void PreEditChange(UProperty* PropertyAboutToChange) override;
 	virtual void PostEditChangeProperty(FPropertyChangedEvent& PropertyChangedEvent) override;
 
 	/** Delegate to fire when our sequence is changed in the property editor */
@@ -186,5 +187,8 @@ protected:
 #if WITH_EDITOR
 	/** Delegate to fire when our sequence is changed in the property editor */
 	FOnSequenceChanged OnSequenceChangedDelegate;
+
+	/* Previous sub sequence, restored if changed sub sequence is invalid*/
+	UMovieSceneSequence* PreviousSubSequence;
 #endif
 };
diff --git a/Engine/Source/Runtime/MovieScene/Public/Tracks/MovieSceneSubTrack.h b/Engine/Source/Runtime/MovieScene/Public/Tracks/MovieSceneSubTrack.h
--- a/Engine/Source/Runtime/MovieScene/Public/Tracks/MovieSceneSubTrack.h
+++ b/Engine/Source/Runtime/MovieScene/Public/Tracks/MovieSceneSubTrack.h
@@ -51,9 +51,10 @@ public:
 	 *
 	 * @param Sequence The sequence to find.
 	 * @param Recursively Whether to search for the sequence in sub-sequences.
+	 * @param SectionToSkip Skip this section when searching the track (ie. the section is already set to this sequence). 
 	 * @return true if the sequence is in this track, false otherwise.
 	 */
-	bool ContainsSequence(const UMovieSceneSequence& Sequence, bool Recursively = false) const;
+	bool ContainsSequence(const UMovieSceneSequence& Sequence, bool Recursively = false, const UMovieSceneSection* SectionToSkip = nullptr) const;
 
 	/**
 	 * Add a new sequence to record
