Sequencer: Ensure that movie scene, section and track are always RF_Transactional. In one case, a movie scene happened to be not transactional and that caused stale bindings which crashed the tree view rebuild.

diff --git a/Engine/Source/Runtime/MovieScene/Private/MovieScene.cpp b/Engine/Source/Runtime/MovieScene/Private/MovieScene.cpp
--- a/Engine/Source/Runtime/MovieScene/Private/MovieScene.cpp
+++ b/Engine/Source/Runtime/MovieScene/Private/MovieScene.cpp
@@ -63,6 +63,13 @@ bool UMovieScene::IsPostLoadThreadSafe() const
 	return true;
 }
 
+void UMovieScene::PostInitProperties()
+{
+	SetFlags(RF_Transactional);
+
+	Super::PostInitProperties();
+}
+
 void UMovieScene::Serialize( FArchive& Ar )
 {
 	Ar.UsingCustomVersion(FMovieSceneEvaluationCustomVersion::GUID);
diff --git a/Engine/Source/Runtime/MovieScene/Private/MovieSceneSection.cpp b/Engine/Source/Runtime/MovieScene/Private/MovieSceneSection.cpp
--- a/Engine/Source/Runtime/MovieScene/Private/MovieSceneSection.cpp
+++ b/Engine/Source/Runtime/MovieScene/Private/MovieSceneSection.cpp
@@ -42,6 +42,8 @@ UMovieSceneSection::UMovieSceneSection(const FObjectInitializer& ObjectInitializ
 
 void UMovieSceneSection::PostInitProperties()
 {
+	SetFlags(RF_Transactional);
+
 	// Propagate sub object flags from our outer (track) to ourselves. This is required for sections that are stored on blueprints (archetypes) so that they can be referenced in worlds.
 	if (GetOuter()->HasAnyFlags(RF_ClassDefaultObject|RF_ArchetypeObject))
 	{
diff --git a/Engine/Source/Runtime/MovieScene/Private/MovieSceneTrack.cpp b/Engine/Source/Runtime/MovieScene/Private/MovieSceneTrack.cpp
--- a/Engine/Source/Runtime/MovieScene/Private/MovieSceneTrack.cpp
+++ b/Engine/Source/Runtime/MovieScene/Private/MovieSceneTrack.cpp
@@ -27,6 +27,8 @@ UMovieSceneTrack::UMovieSceneTrack(const FObjectInitializer& InInitializer)
 
 void UMovieSceneTrack::PostInitProperties()
 {
+	SetFlags(RF_Transactional);
+
 	// Propagate sub object flags from our outer (movie scene) to ourselves. This is required for tracks that are stored on blueprints (archetypes) so that they can be referenced in worlds.
 	if (GetOuter()->HasAnyFlags(RF_ClassDefaultObject|RF_ArchetypeObject))
 	{
diff --git a/Engine/Source/Runtime/MovieScene/Public/MovieScene.h b/Engine/Source/Runtime/MovieScene/Public/MovieScene.h
--- a/Engine/Source/Runtime/MovieScene/Public/MovieScene.h
+++ b/Engine/Source/Runtime/MovieScene/Public/MovieScene.h
@@ -251,6 +251,7 @@ public:
 	/**~ UObject implementation */
 	virtual void Serialize( FArchive& Ar ) override;
 	virtual bool IsPostLoadThreadSafe() const override;
+	virtual void PostInitProperties() override;
 
 public:
 
