Sequencer: Fix blueprint components not restoring properly when autotracked with "Allow Level Sequencer Edits Only". Swap out cached pointers in pre animated state when they change with OnObjectsReplaced.

diff --git a/Engine/Source/Editor/Sequencer/Private/Sequencer.cpp b/Engine/Source/Editor/Sequencer/Private/Sequencer.cpp
--- a/Engine/Source/Editor/Sequencer/Private/Sequencer.cpp
+++ b/Engine/Source/Editor/Sequencer/Private/Sequencer.cpp
@@ -364,6 +364,15 @@ void FSequencer::InitSequencer(const FSequencerInitParams& InitParams, const TSh
 		AcquiredResources.Add([=] { GEditor->OnBlueprintCompiled().Remove(OnBlueprintCompiledHandle); });
 	}
 
+	{
+		FDelegateHandle OnObjectsReplacedHandle = GEditor->OnObjectsReplaced().AddLambda([&](const TMap<UObject*, UObject*>& ReplacementMap)
+		{
+			PreAnimatedState.OnObjectsReplaced(ReplacementMap);
+		});
+		AcquiredResources.Add([=] { GEditor->OnObjectsReplaced().Remove(OnObjectsReplacedHandle); });
+	}
+
+
 	{
 		ISequenceRecorder* Recorder = FModuleManager::Get().GetModulePtr<ISequenceRecorder>("SequenceRecorder");
 		
diff --git a/Engine/Source/Runtime/MovieScene/Private/Evaluation/MovieScenePreAnimatedState.cpp b/Engine/Source/Runtime/MovieScene/Private/Evaluation/MovieScenePreAnimatedState.cpp
--- a/Engine/Source/Runtime/MovieScene/Private/Evaluation/MovieScenePreAnimatedState.cpp
+++ b/Engine/Source/Runtime/MovieScene/Private/Evaluation/MovieScenePreAnimatedState.cpp
@@ -406,6 +406,34 @@ void FMovieScenePreAnimatedState::DiscardAndRemoveEntityTokensForObject(UObject&
 	}
 }
 
+void FMovieScenePreAnimatedState::OnObjectsReplaced(const TMap<UObject*, UObject*>& ReplacementMap)
+{
+	for (auto Iter = ReplacementMap.CreateConstIterator(); Iter; ++Iter)
+	{
+		UObject* OldObject = Iter->Key;
+		UObject* NewObject = Iter->Value;
+
+		FObjectKey OldKey = FObjectKey(OldObject);
+		if (OldObject && NewObject && ObjectTokens.Contains(OldKey))
+		{
+			FObjectKey NewKey = FObjectKey(NewObject);
+
+			ObjectTokens.Add(NewKey, MoveTemp(ObjectTokens[OldKey]));
+			ObjectTokens[NewKey].SetPayload(NewObject);
+
+			ObjectTokens.Remove(OldKey);
+
+			for (auto& Pair : EntityToAnimatedObjects)
+			{
+				if (Pair.Value.Contains(OldKey))
+				{
+					Pair.Value.Add(NewKey);
+					Pair.Value.Remove(OldKey);
+				}
+			}
+		}
+	}
+}
 
 /** Explicit, exported template instantiations */
 template struct MOVIESCENE_API TMovieSceneSavedTokens<IMovieScenePreAnimatedTokenPtr>;
diff --git a/Engine/Source/Runtime/MovieScene/Public/Evaluation/MovieScenePreAnimatedState.h b/Engine/Source/Runtime/MovieScene/Public/Evaluation/MovieScenePreAnimatedState.h
--- a/Engine/Source/Runtime/MovieScene/Public/Evaluation/MovieScenePreAnimatedState.h
+++ b/Engine/Source/Runtime/MovieScene/Public/Evaluation/MovieScenePreAnimatedState.h
@@ -147,6 +147,8 @@ struct MOVIESCENE_API TMovieSceneSavedTokens
 	 */
 	void Reset();
 
+	void SetPayload(PayloadType&& InPayload) { Payload = MoveTemp(InPayload); }
+
 private:
 
 	/** Array defining how whether (and how) particular entities have evaluated */
@@ -283,6 +285,11 @@ public:
 	 */
 	MOVIESCENE_API void DiscardAndRemoveEntityTokensForObject(UObject& Object);
 
+	/**
+	 * Called when objects have been replaced so that pre animated state can swap out to the new objects
+	 */
+	MOVIESCENE_API void OnObjectsReplaced(const TMap<UObject*, UObject*>& ReplacementMap);
+
 public:
 
 	void SetCaptureEntity(FMovieSceneEvaluationKey InEntity, EMovieSceneCompletionMode InCompletionMode)
