Sequencer: Fully consider GetEffectiveTimeDilation when calculating sequence time instead of just MatineeTimeDilation.

 - This fixes the sequence not considering TimeScale or DemoPlayTimeDilation

diff --git a/Engine/Source/Runtime/MovieScene/Private/MovieSceneSequencePlayer.cpp b/Engine/Source/Runtime/MovieScene/Private/MovieSceneSequencePlayer.cpp
--- a/Engine/Source/Runtime/MovieScene/Private/MovieSceneSequencePlayer.cpp
+++ b/Engine/Source/Runtime/MovieScene/Private/MovieSceneSequencePlayer.cpp
@@ -686,7 +686,7 @@ void UMovieSceneSequencePlayer::Update(const float DeltaSeconds)
 
 	if (IsPlaying())
 	{
-		// Delta seconds has already been multiplied by MatineeTimeDilation at this point, so don't pass that through to Tick
+		// Delta seconds has already been multiplied by GetEffectiveTimeDilation at this point, so don't pass that through to Tick
 		float PlayRate = bReversePlayback ? -PlaybackSettings.PlayRate : PlaybackSettings.PlayRate;
 
 		float DeltaTimeForFunction = DeltaSeconds;
@@ -700,7 +700,7 @@ void UMovieSceneSequencePlayer::Update(const float DeltaSeconds)
 
 		if (World)
 		{
-			PlayRate *= World->GetWorldSettings()->MatineeTimeDilation;
+			PlayRate *= World->GetWorldSettings()->GetEffectiveTimeDilation();
 		}
 
 		FFrameTime NewTime = TimeController->RequestCurrentTime(GetCurrentTime(), PlayRate);
