Sequencer: Ping compensation for replicated sequence playback

Hitch correction can be verrified by seeing the "Correcting de-synced play position for sequence..." message in the log.
For more information more often, enable "log LogMovieSceneRepl VeryVerbose"

#note This is a copy from my original changelist that was submitted in 7167459 which was deadended

diff --git a/Engine/Source/Runtime/MovieScene/Private/MovieSceneSequencePlayer.cpp b/Engine/Source/Runtime/MovieScene/Private/MovieSceneSequencePlayer.cpp
--- a/Engine/Source/Runtime/MovieScene/Private/MovieSceneSequencePlayer.cpp
+++ b/Engine/Source/Runtime/MovieScene/Private/MovieSceneSequencePlayer.cpp
@@ -9,8 +9,11 @@
 #include "Misc/RuntimeErrors.h"
 #include "Net/UnrealNetwork.h"
 #include "Engine/NetDriver.h"
+#include "Engine/NetConnection.h"
+#include "GameFramework/PlayerController.h"
+#include "GameFramework/PlayerState.h"
 
-DEFINE_LOG_CATEGORY_STATIC(LogMovieSceneRepl, Warning, All);
+DEFINE_LOG_CATEGORY_STATIC(LogMovieSceneRepl, Log, All);
 
 bool FMovieSceneSequenceLoopCount::SerializeFromMismatchedTag( const FPropertyTag& Tag, FStructuredArchive::FSlot Slot )
 {
@@ -1067,12 +1070,24 @@ void UMovieSceneSequencePlayer::PostNetReceive()
 		return;
 	}
 
+	float PingMs = 0.f;
+
+	UWorld* PlayWorld = GetPlaybackWorld();
+	if (PlayWorld)
+	{
+		UNetDriver* NetDriver = PlayWorld->GetNetDriver();
+		if (NetDriver && NetDriver->ServerConnection && NetDriver->ServerConnection->PlayerController && NetDriver->ServerConnection->PlayerController->PlayerState)
+		{
+			PingMs = NetDriver->ServerConnection->PlayerController->PlayerState->ExactPing * (bReversePlayback ? -1.f : 1.f);
+		}
+	}
+
 	const bool bHasStartedPlaying = NetSyncProps.LastKnownStatus == EMovieScenePlayerStatus::Playing && Status != EMovieScenePlayerStatus::Playing;
 	const bool bHasChangedStatus  = NetSyncProps.LastKnownStatus   != Status;
 	const bool bHasChangedTime    = NetSyncProps.LastKnownPosition != PlayPosition.GetCurrentPosition();
 
+	const FFrameTime PingLag      = (PingMs/1000.f) * PlayPosition.GetInputRate();
 	const FFrameTime LagThreshold = 0.2f * PlayPosition.GetInputRate();
-	const FFrameTime LagDisparity = FMath::Abs(PlayPosition.GetCurrentPosition() - NetSyncProps.LastKnownPosition);
 
 	if (!bHasChangedStatus && !bHasChangedTime)
 	{
@@ -1106,10 +1121,11 @@ void UMovieSceneSequencePlayer::PostNetReceive()
 		// Initiate playback from our current position
 		PlayInternal();
 
+		const FFrameTime LagDisparity = FMath::Abs(PlayPosition.GetCurrentPosition() - (NetSyncProps.LastKnownPosition + PingLag));
 		if (LagDisparity > LagThreshold)
 		{
 			// Synchronize to the server time as best we can if there is a large disparity
-			PlayToFrame(NetSyncProps.LastKnownPosition);
+			PlayToFrame(NetSyncProps.LastKnownPosition + PingLag);
 		}
 	}
 	else
@@ -1129,29 +1145,44 @@ void UMovieSceneSequencePlayer::PostNetReceive()
 				//			=> LoopOffset = 1
 				//			   OffsetServerTime = srv_time + FrameDuration*LoopOffset = 1 + 20*1 = 21
 				//			   Difference = 21 - 18 = 3 frames
-				static float       ThresholdS       = 0.2f;
 				const int32        LoopOffset       = (NetSyncProps.LastKnownNumLoops - CurrentNumLoops) * (bReversePlayback ? -1 : 1);
-				const FFrameTime   FrameThreshold   = ThresholdS * PlayPosition.GetInputRate();
-				const FFrameTime   OffsetServerTime = NetSyncProps.LastKnownPosition + GetFrameDuration()*LoopOffset;
+				const FFrameTime   OffsetServerTime = (NetSyncProps.LastKnownPosition + PingLag) + GetFrameDuration()*LoopOffset;
 				const FFrameTime   Difference       = FMath::Abs(PlayPosition.GetCurrentPosition() - OffsetServerTime);
 
 				if (bHasChangedStatus)
 				{
 					// If the status has changed forcibly play to the server position before setting the new status
-					PlayToFrame(NetSyncProps.LastKnownPosition);
+					PlayToFrame(NetSyncProps.LastKnownPosition + PingLag);
 				}
-				else if (Difference > FrameThreshold)
+				else if (Difference > LagThreshold + PingLag)
 				{
+#if !NO_LOGGING
+					if (UE_LOG_ACTIVE(LogMovieSceneRepl, Log))
+					{
+						FFrameTime CurrentTime  = PlayPosition.GetCurrentPosition();
+						FString    SequenceName = RootTemplateInstance.GetSequence(MovieSceneSequenceID::Root)->GetName();
+
+						AActor* Actor = GetTypedOuter<AActor>();
+						if (Actor->GetWorld()->GetNetMode() == NM_Client)
+						{
+							SequenceName += FString::Printf(TEXT(" (client %d)"), GPlayInEditorID - 1);
+						}
+
+						UE_LOG(LogMovieSceneRepl, Log, TEXT("Correcting de-synced play position for sequence %s %s @ frame %d, subframe %f. Server is %s @ frame %d, subframe %f. Client ping is %.2fms."),
+							*SequenceName, *UEnum::GetValueAsString(TEXT("MovieScene.EMovieScenePlayerStatus"), Status.GetValue()), CurrentTime.FrameNumber.Value, CurrentTime.GetSubFrame(),
+							*UEnum::GetValueAsString(TEXT("MovieScene.EMovieScenePlayerStatus"), NetSyncProps.LastKnownStatus.GetValue()), NetSyncProps.LastKnownPosition.FrameNumber.Value, NetSyncProps.LastKnownPosition.GetSubFrame(), PingMs);
+					}
+#endif
 					// We're drastically out of sync with the server so we need to forcibly set the time.
 					// Play to the time only if it is further on in the sequence (in our play direction)
 					const bool bPlayToFrame = bReversePlayback ? NetSyncProps.LastKnownPosition < PlayPosition.GetCurrentPosition() : NetSyncProps.LastKnownPosition > PlayPosition.GetCurrentPosition();
 					if (bPlayToFrame)
 					{
-						PlayToFrame(NetSyncProps.LastKnownPosition);
+						PlayToFrame(NetSyncProps.LastKnownPosition + PingLag);
 					}
 					else
 					{
-						JumpToFrame(NetSyncProps.LastKnownPosition);
+						JumpToFrame(NetSyncProps.LastKnownPosition + PingLag);
 					}
 				}
 			}
