Sequencer: Added cvar for controlling net playback synchronization threshold

Sequencer.NetSyncThreshold defaults to 200ms, and defines the point at which desync is deemed unacceptable and a force sync is required

diff --git a/Engine/Source/Runtime/MovieScene/Private/MovieSceneSequencePlayer.cpp b/Engine/Source/Runtime/MovieScene/Private/MovieSceneSequencePlayer.cpp
--- a/Engine/Source/Runtime/MovieScene/Private/MovieSceneSequencePlayer.cpp
+++ b/Engine/Source/Runtime/MovieScene/Private/MovieSceneSequencePlayer.cpp
@@ -15,6 +15,13 @@
 
 DEFINE_LOG_CATEGORY_STATIC(LogMovieSceneRepl, Log, All);
 
+float GSequencerNetSyncThresholdMS = 200;
+static FAutoConsoleVariableRef CVarSequencerNetSyncThresholdMS(
+	TEXT("Sequencer.NetSyncThreshold"),
+	GSequencerNetSyncThresholdMS,
+	TEXT("(Default: 200ms. Defines the threshold at which clients and servers must be forcibly re-synced during playback.")
+	);
+
 bool FMovieSceneSequenceLoopCount::SerializeFromMismatchedTag( const FPropertyTag& Tag, FStructuredArchive::FSlot Slot )
 {
 	if (Tag.Type == NAME_IntProperty)
@@ -1087,7 +1094,7 @@ void UMovieSceneSequencePlayer::PostNetReceive()
 	const bool bHasChangedTime    = NetSyncProps.LastKnownPosition != PlayPosition.GetCurrentPosition();
 
 	const FFrameTime PingLag      = (PingMs/1000.f) * PlayPosition.GetInputRate();
-	const FFrameTime LagThreshold = 0.2f * PlayPosition.GetInputRate();
+	const FFrameTime LagThreshold = (GSequencerNetSyncThresholdMS * 0.001f) * PlayPosition.GetInputRate();
 
 	if (!bHasChangedStatus && !bHasChangedTime)
 	{
