Moved the LexToString for TRange<FFrameNumber> to a common area so we can re-use it elsewhere.

diff --git a/Engine/Source/Runtime/MovieScene/Private/Tests/MovieSceneTestsCommon.h b/Engine/Source/Runtime/MovieScene/Private/Tests/MovieSceneTestsCommon.h
--- a/Engine/Source/Runtime/MovieScene/Private/Tests/MovieSceneTestsCommon.h
+++ b/Engine/Source/Runtime/MovieScene/Private/Tests/MovieSceneTestsCommon.h
@@ -4,5 +4,3 @@
 
 #include "Math/Range.h"
 #include "Containers/UnrealString.h"
-
-MOVIESCENE_API FString LexToString(const TRange<FFrameNumber>& InRange);
diff --git a/Engine/Source/Runtime/MovieScene/Private/Tests/MovieSceneTransformTests.cpp b/Engine/Source/Runtime/MovieScene/Private/Tests/MovieSceneTransformTests.cpp
--- a/Engine/Source/Runtime/MovieScene/Private/Tests/MovieSceneTransformTests.cpp
+++ b/Engine/Source/Runtime/MovieScene/Private/Tests/MovieSceneTransformTests.cpp
@@ -6,6 +6,7 @@
 #include "Containers/ArrayView.h"
 #include "Misc/AutomationTest.h"
 #include "MovieSceneTestsCommon.h"
+#include "MovieSceneTimeHelpers.h"
 
 #define LOCTEXT_NAMESPACE "MovieSceneTransformTests"
 
@@ -28,26 +29,6 @@ bool IsEqual(TRange<FFrameNumber> A, TRange<FFrameNumber> B)
 	return IsEqual(A.GetLowerBound(), B.GetLowerBound()) && IsEqual(A.GetUpperBound(), B.GetUpperBound());
 }
 
-FString LexToString(const TRange<FFrameNumber>& InRange)
-{
-	TRangeBound<FFrameNumber> SourceLower = InRange.GetLowerBound();
-	TRangeBound<FFrameNumber> SourceUpper = InRange.GetUpperBound();
-
-	return *FString::Printf(TEXT("%s-%s"),
-		SourceLower.IsOpen() ? 
-			TEXT("[...") : 
-			SourceLower.IsInclusive() ?
-				*FString::Printf(TEXT("[%i"), SourceLower.GetValue().Value) :
-				*FString::Printf(TEXT("(%i"), SourceLower.GetValue().Value),
-
-		SourceUpper.IsOpen() ? 
-			TEXT("...]") : 
-			SourceUpper.IsInclusive() ?
-				*FString::Printf(TEXT("%i]"), SourceUpper.GetValue().Value) :
-				*FString::Printf(TEXT("%i)"), SourceUpper.GetValue().Value)
-		);
-}
-
 bool TestTransform(FAutomationTestBase& Test, FMovieSceneSequenceTransform Transform, TArrayView<TRange<FFrameNumber>> InSource, TArrayView<TRange<FFrameNumber>> InExpected, const TCHAR* TestName)
 {
 	check(InSource.Num() == InExpected.Num());
diff --git a/Engine/Source/Runtime/MovieScene/Public/MovieSceneTimeHelpers.h b/Engine/Source/Runtime/MovieScene/Public/MovieSceneTimeHelpers.h
--- a/Engine/Source/Runtime/MovieScene/Public/MovieSceneTimeHelpers.h
+++ b/Engine/Source/Runtime/MovieScene/Public/MovieSceneTimeHelpers.h
@@ -189,4 +189,24 @@ inline FFrameTime ClampToDiscreteRange(FFrameTime InTime, const TRange<FFrameNum
 
 
 
-} // namespace MovieScene
\ No newline at end of file
+} // namespace MovieScene
+
+inline FString LexToString(const TRange<FFrameNumber>& InRange)
+{
+	TRangeBound<FFrameNumber> SourceLower = InRange.GetLowerBound();
+	TRangeBound<FFrameNumber> SourceUpper = InRange.GetUpperBound();
+
+	return *FString::Printf(TEXT("%s-%s"),
+		SourceLower.IsOpen() ?
+		TEXT("[-inf") :
+		SourceLower.IsInclusive() ?
+		*FString::Printf(TEXT("[%i"), SourceLower.GetValue().Value) :
+		*FString::Printf(TEXT("(%i"), SourceLower.GetValue().Value),
+
+		SourceUpper.IsOpen() ?
+		TEXT("+inf]") :
+		SourceUpper.IsInclusive() ?
+		*FString::Printf(TEXT("%i]"), SourceUpper.GetValue().Value) :
+		*FString::Printf(TEXT("%i)"), SourceUpper.GetValue().Value)
+	);
+}
