CIS Fixes.

diff --git a/Engine/Source/Runtime/MovieScene/Private/Evaluation/MovieSceneEvaluationTree.cpp b/Engine/Source/Runtime/MovieScene/Private/Evaluation/MovieSceneEvaluationTree.cpp
--- a/Engine/Source/Runtime/MovieScene/Private/Evaluation/MovieSceneEvaluationTree.cpp
+++ b/Engine/Source/Runtime/MovieScene/Private/Evaluation/MovieSceneEvaluationTree.cpp
@@ -5,7 +5,6 @@
 #include "Misc/AutomationTest.h"
 #include "Compilation/MovieSceneSegmentCompiler.h"
 #include "Compilation/MovieSceneCompilerRules.h"
-#include "Tests/MovieSceneTestsCommon.h"
 #include "Evaluation/MovieSceneEvaluationField.h"
 #include "MovieSceneCommonHelpers.h"
 #include "UObject/Package.h"
diff --git a/Engine/Source/Runtime/MovieScene/Private/Tests/MovieSceneCompilerTests.cpp b/Engine/Source/Runtime/MovieScene/Private/Tests/MovieSceneCompilerTests.cpp
--- a/Engine/Source/Runtime/MovieScene/Private/Tests/MovieSceneCompilerTests.cpp
+++ b/Engine/Source/Runtime/MovieScene/Private/Tests/MovieSceneCompilerTests.cpp
@@ -5,12 +5,12 @@
 #include "Compilation/MovieSceneCompiler.h"
 #include "Compilation/MovieSceneSegmentCompiler.h"
 #include "Compilation/MovieSceneCompilerRules.h"
-#include "MovieSceneTestsCommon.h"
 #include "Evaluation/MovieSceneEvaluationTrack.h"
 #include "Evaluation/MovieSceneEvaluationField.h"
 #include "Evaluation/MovieSceneSequenceTemplateStore.h"
 #include "Algo/Find.h"
 #include "UObject/Package.h"
+#include "MovieSceneTimeHelpers.h"
 
 #if WITH_DEV_AUTOMATION_TESTS
 
diff --git a/Engine/Source/Runtime/MovieScene/Private/Tests/MovieSceneSegmentCompilerTests.cpp b/Engine/Source/Runtime/MovieScene/Private/Tests/MovieSceneSegmentCompilerTests.cpp
--- a/Engine/Source/Runtime/MovieScene/Private/Tests/MovieSceneSegmentCompilerTests.cpp
+++ b/Engine/Source/Runtime/MovieScene/Private/Tests/MovieSceneSegmentCompilerTests.cpp
@@ -6,10 +6,10 @@
 #include "Compilation/MovieSceneCompiler.h"
 #include "Compilation/MovieSceneSegmentCompiler.h"
 #include "Compilation/MovieSceneCompilerRules.h"
-#include "MovieSceneTestsCommon.h"
 #include "Evaluation/MovieSceneEvaluationTrack.h"
 #include "Evaluation/MovieSceneEvaluationField.h"
 #include "UObject/Package.h"
+#include "MovieSceneTimeHelpers.h"
 
 #if WITH_DEV_AUTOMATION_TESTS
 
diff --git a/Engine/Source/Runtime/MovieScene/Private/Tests/MovieSceneTestsCommon.h b/Engine/Source/Runtime/MovieScene/Private/Tests/MovieSceneTestsCommon.h
deleted file mode 100644
diff --git a/Engine/Source/Runtime/MovieScene/Private/Tests/MovieSceneTransformTests.cpp b/Engine/Source/Runtime/MovieScene/Private/Tests/MovieSceneTransformTests.cpp
--- a/Engine/Source/Runtime/MovieScene/Private/Tests/MovieSceneTransformTests.cpp
+++ b/Engine/Source/Runtime/MovieScene/Private/Tests/MovieSceneTransformTests.cpp
@@ -5,7 +5,6 @@
 #include "Evaluation/MovieSceneSectionParameters.h"
 #include "Containers/ArrayView.h"
 #include "Misc/AutomationTest.h"
-#include "MovieSceneTestsCommon.h"
 #include "MovieSceneTimeHelpers.h"
 
 #define LOCTEXT_NAMESPACE "MovieSceneTransformTests"
