UMovieSceneSequencePlayer now relies on the external owner to call Stop() when the player is being torn down. BeginDestroy() is too late in the lifecycle (called immediately before deletion by GC) for UMovieSceneSequencePlayer to call Stop() on itself.

Stop has many potential side effects as it may try to restore state of external actors that this player is affecting. It is not safe to call this during BeginDestroy as the validity of object pointers is very turbulent during a GC pass. Now, the Level Sequence Actor (and ActorSequenceComponent) detect EndPlay() and call Stop (which is a noop if it's already stopped). EndPlay() is called when the original call to Destroy(...) goes out to tear down the owning actor at which point the validity of object pointers is still valid as we don't get GC'd until after EndPlay is called on destroyed items.

This fixes an issue where during a seamless level transition, Stop() would sometimes fail to disable Cinematic Mode on the Player Controller (which survives the transition) as the World would have been destroyed at that point and we use the world to get the local player controller.

diff --git a/Engine/Plugins/MovieScene/ActorSequence/Source/ActorSequence/Private/ActorSequenceComponent.cpp b/Engine/Plugins/MovieScene/ActorSequence/Source/ActorSequence/Private/ActorSequenceComponent.cpp
--- a/Engine/Plugins/MovieScene/ActorSequence/Source/ActorSequence/Private/ActorSequenceComponent.cpp
+++ b/Engine/Plugins/MovieScene/ActorSequence/Source/ActorSequence/Private/ActorSequenceComponent.cpp
@@ -37,6 +37,19 @@ void UActorSequenceComponent::BeginPlay()
 	}
 }
 
+void UActorSequenceComponent::EndPlay(const EEndPlayReason::Type EndPlayReason)
+{
+	if (SequencePlayer)
+	{
+		// Stop the internal sequence player during EndPlay when it's safer
+		// to modify other actor's state during state restoration.
+		SequencePlayer->Stop();
+	}
+	
+	Super::EndPlay(EndPlayReason);
+}
+
+
 void UActorSequenceComponent::TickComponent(float DeltaSeconds, enum ELevelTick TickType, FActorComponentTickFunction *ThisTickFunction)
 {
 	Super::TickComponent(DeltaSeconds, TickType, ThisTickFunction);
diff --git a/Engine/Plugins/MovieScene/ActorSequence/Source/ActorSequence/Public/ActorSequenceComponent.h b/Engine/Plugins/MovieScene/ActorSequence/Source/ActorSequence/Public/ActorSequenceComponent.h
--- a/Engine/Plugins/MovieScene/ActorSequence/Source/ActorSequence/Public/ActorSequenceComponent.h
+++ b/Engine/Plugins/MovieScene/ActorSequence/Source/ActorSequence/Public/ActorSequenceComponent.h
@@ -35,6 +35,7 @@ public:
 	
 	virtual void PostInitProperties() override;
 	virtual void BeginPlay() override;
+	virtual void EndPlay(const EEndPlayReason::Type EndPlayReason) override;
 	virtual void TickComponent(float DeltaTime, enum ELevelTick TickType, FActorComponentTickFunction *ThisTickFunction) override;
 
 protected:
diff --git a/Engine/Source/Runtime/LevelSequence/Private/LevelSequenceActor.cpp b/Engine/Source/Runtime/LevelSequence/Private/LevelSequenceActor.cpp
--- a/Engine/Source/Runtime/LevelSequence/Private/LevelSequenceActor.cpp
+++ b/Engine/Source/Runtime/LevelSequence/Private/LevelSequenceActor.cpp
@@ -142,7 +142,15 @@ void ALevelSequenceActor::BeginPlay()
 
 void ALevelSequenceActor::EndPlay(const EEndPlayReason::Type EndPlayReason)
 {
-	GetWorld()->LevelSequenceActors.Remove(this);
+	if (SequencePlayer)
+	{
+		// Stop may modify a lot of actor state so it needs to be called
+		// during EndPlay (when Actors + World are still valid) instead
+		// of waiting for the UObject to be destroyed by GC.
+		SequencePlayer->Stop();
+	}
+
+ 	GetWorld()->LevelSequenceActors.Remove(this);
 
 	Super::EndPlay(EndPlayReason);
 }
diff --git a/Engine/Source/Runtime/MovieScene/Private/MovieSceneSequencePlayer.cpp b/Engine/Source/Runtime/MovieScene/Private/MovieSceneSequencePlayer.cpp
--- a/Engine/Source/Runtime/MovieScene/Private/MovieSceneSequencePlayer.cpp
+++ b/Engine/Source/Runtime/MovieScene/Private/MovieSceneSequencePlayer.cpp
@@ -950,23 +950,13 @@ TArray<FMovieSceneObjectBindingID> UMovieSceneSequencePlayer::GetObjectBindings(
 	return ObjectBindings;
 }
 
-void UMovieSceneSequencePlayer::BeginDestroy()
-{
-	Stop();
-
-	if (GEngine && OldMaxTickRate.IsSet())
-	{
-		GEngine->SetMaxFPS(OldMaxTickRate.GetValue());
-	}
-
-	Super::BeginDestroy();
-}
-
 UWorld* UMovieSceneSequencePlayer::GetPlaybackWorld() const
 {
 	UObject* PlaybackContext = GetPlaybackContext();
 	return PlaybackContext ? PlaybackContext->GetWorld() : nullptr;
-}bool UMovieSceneSequencePlayer::HasAuthority() const
+}
+
+bool UMovieSceneSequencePlayer::HasAuthority() const
 {
 	AActor* Actor = GetTypedOuter<AActor>();
 	return Actor && Actor->HasAuthority() && !IsPendingKillOrUnreachable();
diff --git a/Engine/Source/Runtime/MovieScene/Public/MovieSceneSequencePlayer.h b/Engine/Source/Runtime/MovieScene/Public/MovieSceneSequencePlayer.h
--- a/Engine/Source/Runtime/MovieScene/Public/MovieSceneSequencePlayer.h
+++ b/Engine/Source/Runtime/MovieScene/Public/MovieSceneSequencePlayer.h
@@ -549,7 +549,6 @@ protected:
 	virtual IMovieScenePlaybackClient* GetPlaybackClient() override { return PlaybackClient ? &*PlaybackClient : nullptr; }
 
 	/*~ Begin UObject interface */
-	virtual void BeginDestroy() override;
 	virtual bool IsSupportedForNetworking() const { return true; }
 	virtual bool IsDestructionThreadSafe() const override { return false; }
 	virtual int32 GetFunctionCallspace(UFunction* Function, FFrame* Stack) override;
