FTransformData build fix

diff --git a/Engine/Source/Runtime/MovieScene/Public/TransformData.h b/Engine/Source/Runtime/MovieScene/Public/TransformData.h
--- a/Engine/Source/Runtime/MovieScene/Public/TransformData.h
+++ b/Engine/Source/Runtime/MovieScene/Public/TransformData.h
@@ -3,6 +3,8 @@
 #pragma once
 
 #include "CoreMinimal.h"
+#include "Components/SceneComponent.h"
+
 /**
 * Stores information about a transform for the purpose of adding keys to a transform section
 */
