Sequencer: Expose blend type and completion mode

diff --git a/Engine/Source/Runtime/MovieScene/Public/Evaluation/Blending/MovieSceneBlendType.h b/Engine/Source/Runtime/MovieScene/Public/Evaluation/Blending/MovieSceneBlendType.h
--- a/Engine/Source/Runtime/MovieScene/Public/Evaluation/Blending/MovieSceneBlendType.h
+++ b/Engine/Source/Runtime/MovieScene/Public/Evaluation/Blending/MovieSceneBlendType.h
@@ -7,9 +7,10 @@
 #include "MovieSceneBlendType.generated.h"
 
 /** Movie scene blend type enumeration */
-UENUM()
+UENUM(BlueprintType)
 enum class EMovieSceneBlendType : uint8
 {
+	Invalid UMETA(Hidden),
 	/** Blends all other weighted values together as an average of the total weight */
 	Absolute			= 0x1,
 	/** Applies this value as a sum total of all other additives */
@@ -19,7 +20,7 @@ enum class EMovieSceneBlendType : uint8
 };
 
 /** Optional blend type structure */
-USTRUCT()
+USTRUCT(BlueprintType)
 struct FOptionalMovieSceneBlendType
 {
 	GENERATED_BODY()
diff --git a/Engine/Source/Runtime/MovieScene/Public/MovieSceneSection.h b/Engine/Source/Runtime/MovieScene/Public/MovieSceneSection.h
--- a/Engine/Source/Runtime/MovieScene/Public/MovieSceneSection.h
+++ b/Engine/Source/Runtime/MovieScene/Public/MovieSceneSection.h
@@ -24,7 +24,7 @@ struct FMovieSceneEvalTemplatePtr;
 template<typename> class TArrayView;
 
 /** Enumeration specifying how to handle state when this section is no longer evaluated */
-UENUM()
+UENUM(BlueprintType)
 enum class EMovieSceneCompletionMode : uint8
 {
 	KeepState,
@@ -286,9 +286,28 @@ public:
 	 */
 	MOVIESCENE_API virtual TOptional<TRange<FFrameNumber> > GetAutoSizeRange() const;
 
+	/**
+	 * Gets this section's completion mode
+	 */
+	UFUNCTION(BlueprintPure, Category = "Movie Scene Section")
+	EMovieSceneCompletionMode GetCompletionMode() const
+	{
+		return EvalOptions.CompletionMode;
+	}
+
+	/*
+	 * Sets this section's completion mode
+	 */
+	UFUNCTION(BlueprintCallable, Category = "Movie Scene Section")
+	void SetCompletionMode(EMovieSceneCompletionMode InCompletionMode)
+	{
+		EvalOptions.CompletionMode = InCompletionMode;
+	}
+
 	/**
 	 * Gets this section's blend type
 	 */
+	UFUNCTION(BlueprintPure, Category = "Movie Scene Section")
 	FOptionalMovieSceneBlendType GetBlendType() const
 	{
 		return BlendType;
@@ -297,6 +316,7 @@ public:
 	/**
 	 * Sets this section's blend type
 	 */
+	UFUNCTION(BlueprintCallable, Category = "Movie Scene Section")
 	MOVIESCENE_API virtual void SetBlendType(EMovieSceneBlendType InBlendType)
 	{
 		if (GetSupportedBlendTypes().Contains(InBlendType))
