Sequencer: Fix pasting a spawnable doesn't correctly evaluate. The sequence signature wasn't getting invalidated after swapping the tracks from the copyable binding to the spawnable.

diff --git a/Engine/Source/Runtime/MovieScene/Private/MovieScene.cpp b/Engine/Source/Runtime/MovieScene/Private/MovieScene.cpp
--- a/Engine/Source/Runtime/MovieScene/Private/MovieScene.cpp
+++ b/Engine/Source/Runtime/MovieScene/Private/MovieScene.cpp
@@ -1017,6 +1017,8 @@ void UMovieScene::ReplaceBinding(const FGuid& OldGuid, const FGuid& NewGuid, con
 
 void UMovieScene::ReplaceBinding(const FGuid& BindingToReplaceGuid, const FMovieSceneBinding& NewBinding)
 {
+	Modify();
+
 	FMovieSceneBinding* Binding = ObjectBindings.FindByPredicate([BindingToReplaceGuid](const FMovieSceneBinding& CheckedBinding) { return CheckedBinding.GetObjectGuid() == BindingToReplaceGuid; });
 	if (Binding)
 	{
