Sequencer: Ensure that restore state doesn't happen during evaluation

diff --git a/Engine/Source/Runtime/MovieScene/Private/Evaluation/MovieSceneEvaluationTemplateInstance.cpp b/Engine/Source/Runtime/MovieScene/Private/Evaluation/MovieSceneEvaluationTemplateInstance.cpp
--- a/Engine/Source/Runtime/MovieScene/Private/Evaluation/MovieSceneEvaluationTemplateInstance.cpp
+++ b/Engine/Source/Runtime/MovieScene/Private/Evaluation/MovieSceneEvaluationTemplateInstance.cpp
@@ -217,6 +217,8 @@ void FMovieSceneRootEvaluationTemplateInstance::Evaluate(FMovieSceneContext Cont
 {
 	SCOPE_CYCLE_COUNTER(MovieSceneEval_EntireEvaluationCost);
 
+	TGuardValue<bool> Guard(Player.PreAnimatedState.bLocked, true);
+
 	Swap(ThisFrameMetaData, LastFrameMetaData);
 	ThisFrameMetaData.Reset();
 
diff --git a/Engine/Source/Runtime/MovieScene/Private/Evaluation/MovieScenePreAnimatedState.cpp b/Engine/Source/Runtime/MovieScene/Private/Evaluation/MovieScenePreAnimatedState.cpp
--- a/Engine/Source/Runtime/MovieScene/Private/Evaluation/MovieScenePreAnimatedState.cpp
+++ b/Engine/Source/Runtime/MovieScene/Private/Evaluation/MovieScenePreAnimatedState.cpp
@@ -336,6 +336,8 @@ void FMovieScenePreAnimatedState::RestorePreAnimatedState(IMovieScenePlayer& Pla
 
 void FMovieScenePreAnimatedState::RestorePreAnimatedStateImpl(IMovieScenePlayer& Player, const FMovieSceneEvaluationKey& Key, TOptional<TFunctionRef<bool(FMovieSceneAnimTypeID)>> InFilter)
 {
+	ensureMsgf(!bLocked, TEXT("Movie scene evaluation caused a re-entrant restore state"));
+
 	auto* AnimatedObjects = EntityToAnimatedObjects.Find(Key);
 	if (!AnimatedObjects)
 	{
diff --git a/Engine/Source/Runtime/MovieScene/Public/Evaluation/MovieScenePreAnimatedState.h b/Engine/Source/Runtime/MovieScene/Public/Evaluation/MovieScenePreAnimatedState.h
--- a/Engine/Source/Runtime/MovieScene/Public/Evaluation/MovieScenePreAnimatedState.h
+++ b/Engine/Source/Runtime/MovieScene/Public/Evaluation/MovieScenePreAnimatedState.h
@@ -174,6 +174,7 @@ public:
 	 */
 	FMovieScenePreAnimatedState()
 		: MasterTokens(MovieSceneImpl::FNull())
+		, bLocked(false)
 	{
 		DefaultGlobalCaptureMode = ECapturePreAnimatedState::None;
 	}
@@ -306,6 +307,11 @@ public:
 		EntityToAnimatedObjects.FindOrAdd(EntityKey).Add(FObjectKey());
 	}
 
+public:
+
+	/** Locked during evaluation to prevent reentrant restore state */
+	bool bLocked;
+
 protected:
 
 	MOVIESCENE_API void RestorePreAnimatedStateImpl(IMovieScenePlayer& Player, const FMovieSceneEvaluationKey& Key, TOptional<TFunctionRef<bool(FMovieSceneAnimTypeID)>> InFilter);
