Fix CIS

diff --git a/Engine/Source/Runtime/MovieScene/Public/Evaluation/MovieScenePreAnimatedState.h b/Engine/Source/Runtime/MovieScene/Public/Evaluation/MovieScenePreAnimatedState.h
--- a/Engine/Source/Runtime/MovieScene/Public/Evaluation/MovieScenePreAnimatedState.h
+++ b/Engine/Source/Runtime/MovieScene/Public/Evaluation/MovieScenePreAnimatedState.h
@@ -173,8 +173,8 @@ public:
 	 * Default construction
 	 */
 	FMovieScenePreAnimatedState()
-		: MasterTokens(MovieSceneImpl::FNull())
-		, bLocked(false)
+		: bLocked(false)
+		, MasterTokens(MovieSceneImpl::FNull())
 	{
 		DefaultGlobalCaptureMode = ECapturePreAnimatedState::None;
 	}
