Sequencer: Back out 8889048 temporarily

diff --git a/Engine/Source/Runtime/MovieScene/Private/Evaluation/MovieScenePreAnimatedState.cpp b/Engine/Source/Runtime/MovieScene/Private/Evaluation/MovieScenePreAnimatedState.cpp
--- a/Engine/Source/Runtime/MovieScene/Private/Evaluation/MovieScenePreAnimatedState.cpp
+++ b/Engine/Source/Runtime/MovieScene/Private/Evaluation/MovieScenePreAnimatedState.cpp
@@ -336,8 +336,6 @@ void FMovieScenePreAnimatedState::RestorePreAnimatedState(IMovieScenePlayer& Pla
 
 void FMovieScenePreAnimatedState::RestorePreAnimatedStateImpl(IMovieScenePlayer& Player, const FMovieSceneEvaluationKey& Key, TOptional<TFunctionRef<bool(FMovieSceneAnimTypeID)>> InFilter)
 {
-	ensureMsgf(!bLocked, TEXT("Movie scene evaluation caused a re-entrant restore state"));
-
 	auto* AnimatedObjects = EntityToAnimatedObjects.Find(Key);
 	if (!AnimatedObjects)
 	{
