Sequencer: Fix ensure for delayed restore

diff --git a/Engine/Source/Runtime/MovieScene/Private/Evaluation/MovieSceneEvaluationTemplateInstance.cpp b/Engine/Source/Runtime/MovieScene/Private/Evaluation/MovieSceneEvaluationTemplateInstance.cpp
--- a/Engine/Source/Runtime/MovieScene/Private/Evaluation/MovieSceneEvaluationTemplateInstance.cpp
+++ b/Engine/Source/Runtime/MovieScene/Private/Evaluation/MovieSceneEvaluationTemplateInstance.cpp
@@ -156,6 +156,7 @@ FMovieSceneRootEvaluationTemplateInstance::FMovieSceneRootEvaluationTemplateInst
 	: RootSequence(nullptr)
 	, RootID(MovieSceneSequenceID::Root)
 	, TemplateStore(MakeShared<FMovieSceneSequencePrecompiledTemplateStore>())
+	, bEvaluating(false)
 {
 }
 
@@ -215,9 +216,11 @@ void FMovieSceneRootEvaluationTemplateInstance::Finish(IMovieScenePlayer& Player
 
 void FMovieSceneRootEvaluationTemplateInstance::Evaluate(FMovieSceneContext Context, IMovieScenePlayer& Player, FMovieSceneSequenceID InOverrideRootID)
 {
+	ensureMsgf(!bEvaluating, TEXT("Movie scene evaluation caused a re-entrant evaluation"));
+
 	SCOPE_CYCLE_COUNTER(MovieSceneEval_EntireEvaluationCost);
 
-	TGuardValue<bool> Guard(Player.PreAnimatedState.bLocked, true);
+	TGuardValue<bool> Guard(bEvaluating, true);
 
 	Swap(ThisFrameMetaData, LastFrameMetaData);
 	ThisFrameMetaData.Reset();
diff --git a/Engine/Source/Runtime/MovieScene/Public/Evaluation/MovieScenePreAnimatedState.h b/Engine/Source/Runtime/MovieScene/Public/Evaluation/MovieScenePreAnimatedState.h
--- a/Engine/Source/Runtime/MovieScene/Public/Evaluation/MovieScenePreAnimatedState.h
+++ b/Engine/Source/Runtime/MovieScene/Public/Evaluation/MovieScenePreAnimatedState.h
@@ -173,8 +173,7 @@ public:
 	 * Default construction
 	 */
 	FMovieScenePreAnimatedState()
-		: bLocked(false)
-		, MasterTokens(MovieSceneImpl::FNull())
+		: MasterTokens(MovieSceneImpl::FNull())
 	{
 		DefaultGlobalCaptureMode = ECapturePreAnimatedState::None;
 	}
@@ -307,11 +306,6 @@ public:
 		EntityToAnimatedObjects.FindOrAdd(EntityKey).Add(FObjectKey());
 	}
 
-public:
-
-	/** Locked during evaluation to prevent reentrant restore state */
-	bool bLocked;
-
 protected:
 
 	MOVIESCENE_API void RestorePreAnimatedStateImpl(IMovieScenePlayer& Player, const FMovieSceneEvaluationKey& Key, TOptional<TFunctionRef<bool(FMovieSceneAnimTypeID)>> InFilter);
