Sequencer: Fix ensure for delayed restore

diff --git a/Engine/Source/Runtime/MovieScene/Public/Evaluation/MovieSceneEvaluationTemplateInstance.h b/Engine/Source/Runtime/MovieScene/Public/Evaluation/MovieSceneEvaluationTemplateInstance.h
--- a/Engine/Source/Runtime/MovieScene/Public/Evaluation/MovieSceneEvaluationTemplateInstance.h
+++ b/Engine/Source/Runtime/MovieScene/Public/Evaluation/MovieSceneEvaluationTemplateInstance.h
@@ -208,6 +208,9 @@ private:
 
 	/** Execution tokens that are used to apply animated state */
 	FMovieSceneExecutionTokens ExecutionTokens;
+
+	/** True when in Evaluate */
+	bool bEvaluating;
 };
 
 template<>
