Allow sub-sequence tracks to generate an evaluation template.

diff --git a/Engine/Source/Runtime/MovieScene/Private/Compilation/MovieSceneEvaluationTemplateGenerator.cpp b/Engine/Source/Runtime/MovieScene/Private/Compilation/MovieSceneEvaluationTemplateGenerator.cpp
--- a/Engine/Source/Runtime/MovieScene/Private/Compilation/MovieSceneEvaluationTemplateGenerator.cpp
+++ b/Engine/Source/Runtime/MovieScene/Private/Compilation/MovieSceneEvaluationTemplateGenerator.cpp
@@ -66,14 +66,8 @@ void FMovieSceneEvaluationTemplateGenerator::ProcessTrack(const UMovieSceneTrack
 		return;
 	}
 
-	// Deal with sub tracks specifically
-	if (const UMovieSceneSubTrack* SubTrack = Cast<const UMovieSceneSubTrack>(&Track))
-	{
-		ProcessSubTrack(*SubTrack, ObjectBindingId);
-	}
-
 	// If the ledger already contains this track, just update the uncompiled track field
-	else if (FMovieSceneTrackIdentifier TrackID = Template.GetLedger().FindTrack(Track.GetSignature()))
+	if (FMovieSceneTrackIdentifier TrackID = Template.GetLedger().FindTrack(Track.GetSignature()))
 	{
 		CompiledSignatures.Add(Track.GetSignature());
 
@@ -90,6 +84,12 @@ void FMovieSceneEvaluationTemplateGenerator::ProcessTrack(const UMovieSceneTrack
 		Args.DefaultCompletionMode = SourceSequence.DefaultCompletionMode;
 		Track.GenerateTemplate(Args);
 	}
+
+	// Deal with sub tracks specifically
+	if (const UMovieSceneSubTrack* SubTrack = Cast<const UMovieSceneSubTrack>(&Track))
+	{
+		ProcessSubTrack(*SubTrack, ObjectBindingId);
+	}
 }
 
 
