Sequencer: Add support for grouping/linking sections so dragging maintains the relative offset

diff --git a/Engine/Source/Editor/Sequencer/Private/Sequencer.cpp b/Engine/Source/Editor/Sequencer/Private/Sequencer.cpp
--- a/Engine/Source/Editor/Sequencer/Private/Sequencer.cpp
+++ b/Engine/Source/Editor/Sequencer/Private/Sequencer.cpp
@@ -7576,6 +7576,83 @@ bool FSequencer::IsNodeLocked() const
 	return NumSections > 0;
 }
 
+void FSequencer::GroupSelectedSections()
+{
+	UMovieScene* MovieScene = GetFocusedMovieSceneSequence()->GetMovieScene();
+	if (MovieScene->IsReadOnly())
+	{
+		ShowReadOnlyError();
+		return;
+	}
+
+	const FScopedTransaction Transaction(LOCTEXT("GroupSelectedSections", "Group Selected Sections"));
+
+	TArray<UMovieSceneSection*> Sections;
+	for (TWeakObjectPtr<UMovieSceneSection> WeakSection : Selection.GetSelectedSections())
+	{
+		UMovieSceneSection* Section = WeakSection.Get();
+		// We do not want to group sections that are infinite, as they should not be moveable
+		if (Section && (Section->HasStartFrame() || Section->HasEndFrame()))
+		{
+			Sections.Add(Section);
+		}
+	}
+
+	MovieScene->GroupSections(Sections);
+}
+
+bool FSequencer::CanGroupSelectedSections() const
+{
+	int32 GroupableSections = 0;
+	for (TWeakObjectPtr<UMovieSceneSection> WeakSection : Selection.GetSelectedSections())
+	{
+		UMovieSceneSection* Section = WeakSection.Get();
+		// We do not want to group sections that are infinite, as they should not be moveable
+		if (Section && (Section->HasStartFrame() || Section->HasEndFrame()))
+		{
+			if (++GroupableSections >= 2)
+			{
+				return true;
+			}
+		}
+	}
+	return false;
+}
+
+void FSequencer::UngroupSelectedSections()
+{
+	UMovieScene* MovieScene = GetFocusedMovieSceneSequence()->GetMovieScene();
+	if (MovieScene->IsReadOnly())
+	{
+		ShowReadOnlyError();
+		return;
+	}
+
+	const FScopedTransaction Transaction(LOCTEXT("UngroupSelectedSections", "Ungroup Selected Sections"));
+
+	for (TWeakObjectPtr<UMovieSceneSection> WeakSection : Selection.GetSelectedSections())
+	{
+		if (WeakSection.IsValid())
+		{
+			MovieScene->UngroupSection(*WeakSection.Get());
+		}
+	}
+}
+
+bool FSequencer::CanUngroupSelectedSections() const
+{
+	UMovieScene* MovieScene = GetFocusedMovieSceneSequence()->GetMovieScene();
+
+	for (TWeakObjectPtr<UMovieSceneSection> WeakSection : Selection.GetSelectedSections())
+	{
+		if (WeakSection.IsValid() && MovieScene->IsSectionInGroup(*WeakSection.Get()))
+		{
+			return true;
+		}
+	}
+	return false;
+}
+
 void FSequencer::SaveSelectedNodesSpawnableState()
 {
 	UMovieScene* MovieScene = GetFocusedMovieSceneSequence()->GetMovieScene();
diff --git a/Engine/Source/Editor/Sequencer/Private/Sequencer.h b/Engine/Source/Editor/Sequencer/Private/Sequencer.h
--- a/Engine/Source/Editor/Sequencer/Private/Sequencer.h
+++ b/Engine/Source/Editor/Sequencer/Private/Sequencer.h
@@ -506,6 +506,14 @@ public:
 	void ToggleNodeLocked();
 	bool IsNodeLocked() const;
 
+	/** Called when a user executes the Group menu item */
+	void GroupSelectedSections();
+	bool CanGroupSelectedSections() const;
+
+	/** Called when a user executes the Ungroup menu item */
+	void UngroupSelectedSections();
+	bool CanUngroupSelectedSections() const;
+
 	/** Called when a user executes the set key time for selected keys */
 	bool CanSetKeyTime() const;
 	void SetKeyTime();
diff --git a/Engine/Source/Editor/Sequencer/Private/SequencerContextMenus.cpp b/Engine/Source/Editor/Sequencer/Private/SequencerContextMenus.cpp
--- a/Engine/Source/Editor/Sequencer/Private/SequencerContextMenus.cpp
+++ b/Engine/Source/Editor/Sequencer/Private/SequencerContextMenus.cpp
@@ -253,6 +253,13 @@ void FSectionContextMenu::PopulateMenu(FMenuBuilder& MenuBuilder)
 	// Copy a reference to the context menu by value into each lambda handler to ensure the type stays alive until the menu is closed
 	TSharedRef<FSectionContextMenu> Shared = AsShared();
 
+	// Clean SectionGroups to prevent any potential stale references from affecting the context menu entries
+	Sequencer->GetFocusedMovieSceneSequence()->GetMovieScene()->CleanSectionGroups();
+	
+	// These are potentially expensive checks in large sequences, and won't change while context menu is open
+	const bool bCanGroup = Sequencer->CanGroupSelectedSections();
+	const bool bCanUngroup = Sequencer->CanUngroupSelectedSections();
+
 	ISequencerModule& SequencerModule = FModuleManager::LoadModuleChecked<ISequencerModule>("Sequencer");
 
 	for (auto& Pair : ChannelsByType)
@@ -386,6 +393,30 @@ void FSectionContextMenu::PopulateMenu(FMenuBuilder& MenuBuilder)
 			EUserInterfaceActionType::ToggleButton
 		);
 
+		MenuBuilder.AddMenuEntry(
+			LOCTEXT("GroupSections", "Group"),
+			LOCTEXT("GroupSectionsTooltip", "Group selected sections together so that when any section is moved, all sections in that group move together."),
+			FSlateIcon(),
+			FUIAction(
+				FExecuteAction::CreateSP(Sequencer, &FSequencer::GroupSelectedSections),
+				FCanExecuteAction::CreateLambda([bCanGroup] { return bCanGroup; })
+			),
+			NAME_None,
+			EUserInterfaceActionType::Button
+		);
+
+		MenuBuilder.AddMenuEntry(
+			LOCTEXT("UngroupSections", "Ungroup"),
+			LOCTEXT("UngroupSectionsTooltip", "Ungroup selected sections"),
+			FSlateIcon(),
+			FUIAction(
+				FExecuteAction::CreateSP(Sequencer, &FSequencer::UngroupSelectedSections),
+				FCanExecuteAction::CreateLambda([bCanUngroup] { return bCanUngroup; })
+			),
+			NAME_None,
+			EUserInterfaceActionType::Button
+		);
+
 		// @todo Sequencer this should delete all selected sections
 		// delete/selection needs to be rethought in general
 		MenuBuilder.AddMenuEntry(
diff --git a/Engine/Source/Editor/Sequencer/Private/Tools/EditToolDragOperations.cpp b/Engine/Source/Editor/Sequencer/Private/Tools/EditToolDragOperations.cpp
--- a/Engine/Source/Editor/Sequencer/Private/Tools/EditToolDragOperations.cpp
+++ b/Engine/Source/Editor/Sequencer/Private/Tools/EditToolDragOperations.cpp
@@ -3,6 +3,7 @@
 #include "Tools/EditToolDragOperations.h"
 #include "ISequencer.h"
 #include "MovieSceneTrack.h"
+#include "MovieSceneSequence.h"
 #include "Sequencer.h"
 #include "SequencerSettings.h"
 #include "SequencerCommonHelpers.h"
@@ -563,10 +564,28 @@ FMoveKeysAndSections::FMoveKeysAndSections(FSequencer& InSequencer, const TSet<F
 	// We support partially infinite (infinite on one side) sections however.
 	for (const TWeakObjectPtr<UMovieSceneSection>& WeakSection : InSelectedSections)
 	{
-		const UMovieSceneSection* Section = WeakSection.Get();
-		if (Section->HasStartFrame() || Section->HasEndFrame())
+		const UMovieSceneSection* SelectedSection = WeakSection.Get();
+		if (SelectedSection->HasStartFrame() || SelectedSection->HasEndFrame())
 		{
-			Sections.Add(WeakSection);
+			Sections.AddUnique(WeakSection);
+
+			UMovieScene* MovieScene = InSequencer.GetFocusedMovieSceneSequence()->GetMovieScene();
+			if (MovieScene)
+			{
+				// If the section is in a group, we also want to add the sections it is grouped with
+				const FMovieSceneSectionGroup* SectionGroup = MovieScene->GetSectionGroup(*SelectedSection);
+				if (SectionGroup)
+				{
+					for (TWeakObjectPtr<UMovieSceneSection> WeakGroupedSection : *SectionGroup)
+					{
+						// Verify sections are still valid, and are not infinite.
+						if (WeakGroupedSection.IsValid() && (WeakGroupedSection->HasStartFrame() || WeakGroupedSection->HasEndFrame()))
+						{
+							Sections.AddUnique(WeakGroupedSection);
+						}
+					}
+				}
+			}
 		}
 	}
 }
diff --git a/Engine/Source/Runtime/MovieScene/Private/MovieScene.cpp b/Engine/Source/Runtime/MovieScene/Private/MovieScene.cpp
--- a/Engine/Source/Runtime/MovieScene/Private/MovieScene.cpp
+++ b/Engine/Source/Runtime/MovieScene/Private/MovieScene.cpp
@@ -3,6 +3,7 @@
 #include "MovieScene.h"
 #include "MovieSceneTrack.h"
 #include "MovieSceneFolder.h"
+#include "MovieSceneSequence.h"
 #include "Evaluation/MovieSceneEvaluationCustomVersion.h"
 #include "Compilation/MovieSceneSegmentCompiler.h"
 #include "UObject/SequencerObjectVersion.h"
@@ -528,6 +529,159 @@ void UMovieScene::SetPlaybackRangeLocked(bool bLocked)
 #endif
 
 
+#if WITH_EDITORONLY_DATA
+
+bool FMovieSceneSectionGroup::Contains(const UMovieSceneSection& Section) const
+{
+	return Sections.Contains(&Section);
+}
+
+void FMovieSceneSectionGroup::Add(UMovieSceneSection& Section)
+{
+	Sections.AddUnique(&Section);
+}
+
+void FMovieSceneSectionGroup::Append(const FMovieSceneSectionGroup& SectionGroup)
+{
+	if (SectionGroup == *this)
+	{
+		return;
+	}
+
+	// Append the groups using AddUnique to prevent duplicates
+	Sections.Reserve(Sections.Num() + SectionGroup.Num());
+	for (TWeakObjectPtr<UMovieSceneSection> Section : SectionGroup)
+	{
+		if (Section.IsValid())
+		{
+			Sections.AddUnique(Section);
+		}
+	}
+}
+
+void FMovieSceneSectionGroup::Remove(const UMovieSceneSection& Section)
+{
+	Sections.RemoveSingle(&Section);
+}
+
+void FMovieSceneSectionGroup::Clean()
+{
+	Sections.RemoveAll([](TWeakObjectPtr<UMovieSceneSection>& Section) { return !(Section.IsValid()); });
+}
+
+
+bool UMovieScene::IsSectionInGroup(const UMovieSceneSection& InSection) const
+{
+	for (const FMovieSceneSectionGroup& Group : SectionGroups)
+	{
+		if (Group.Contains(InSection))
+		{
+			return true;
+		}
+	}
+	return false;
+}
+
+void UMovieScene::GroupSections(const TArray<UMovieSceneSection*> InSections)
+{
+	if (InSections.Num() < 2)
+	{
+		return;
+	}
+
+	Modify();
+
+	FMovieSceneSectionGroup* Group = nullptr;
+
+	// Check if the first section is already in a group, and if so use it as the target group
+	for (FMovieSceneSectionGroup& ExistingGroup : SectionGroups)
+	{
+		if (ExistingGroup.Contains(*InSections[0]))
+		{
+			Group = &ExistingGroup;
+			break;
+		}
+	}
+
+	// If we didn't find a group, create a new one.
+	if (!Group)
+	{
+		int32 Index = SectionGroups.AddDefaulted(1);
+		Group = &SectionGroups[Index];
+		Group->Add(*InSections[0]);
+	}
+
+	// Add the remaining sections
+	for (int32 Index = 1; Index < InSections.Num(); ++Index)
+	{
+		UMovieSceneSection* Section = InSections[Index];
+
+		// Check if the section is already in a group, and merge if needed
+		for (FMovieSceneSectionGroup& ExistingGroup : SectionGroups)
+		{
+			// Skip checking the working group. If the section is in it, there is nothing to do.
+			if (&ExistingGroup == Group)
+			{
+				continue;
+			}
+
+			if (ExistingGroup.Contains(*Section))
+			{
+				// If the section is already in a group, merge our new group in to it, and work on that group instead
+				// This has fewer edge cases than merging the found group in to the working group.
+				ExistingGroup.Append(*Group);
+				SectionGroups.RemoveSingle(*Group);
+				Group = &ExistingGroup;
+				break;
+			}
+		}
+
+		// If the section was not added as part of a merged group, add it now.
+		if (!Group->Contains(*Section))
+		{
+			Group->Add(*Section);
+		}
+	}
+}
+
+void UMovieScene::UngroupSection(const UMovieSceneSection& InSection)
+{
+	for (FMovieSceneSectionGroup& ExistingGroup : SectionGroups)
+	{
+		if (ExistingGroup.Contains(InSection))
+		{
+			Modify();
+			ExistingGroup.Remove(InSection);
+			break;
+		}
+	}
+	CleanSectionGroups();
+}
+
+const FMovieSceneSectionGroup* UMovieScene::GetSectionGroup(const UMovieSceneSection& InSection) const
+{
+	for (const FMovieSceneSectionGroup& Group : SectionGroups)
+	{
+		if (Group.Contains(InSection))
+		{
+			return &Group;
+		}
+	}
+	return nullptr;
+}
+
+void UMovieScene::CleanSectionGroups()
+{
+	for (FMovieSceneSectionGroup& SectionGroup : SectionGroups)
+	{
+		SectionGroup.Clean();
+	}
+
+	SectionGroups.RemoveAll([](FMovieSceneSectionGroup& SectionGroup) { return SectionGroup.Num() < 2; });
+}
+#endif
+
+
 TArray<UMovieSceneSection*> UMovieScene::GetAllSections() const
 {
 	TArray<UMovieSceneSection*> OutSections;
@@ -943,6 +1097,10 @@ void UMovieScene::PostLoad()
 		MarkedFrames.Add(FMovieSceneMarkedFrame(MarkedFrame));
 	}
 	EditorData.MarkedFrames_DEPRECATED.Empty();
+
+
+	// Clean any section groups which might refer to sections which were not serialized
+	CleanSectionGroups();
 #endif
 
 	Super::PostLoad();
@@ -974,6 +1132,9 @@ void UMovieScene::PreSave(const class ITargetPlatform* TargetPlatform)
 			It.RemoveCurrent();
 		}
 	}
+
+	// Clean any section groups which contain stale references
+	CleanSectionGroups();
 #endif
 }
 
diff --git a/Engine/Source/Runtime/MovieScene/Public/MovieScene.h b/Engine/Source/Runtime/MovieScene/Public/MovieScene.h
--- a/Engine/Source/Runtime/MovieScene/Public/MovieScene.h
+++ b/Engine/Source/Runtime/MovieScene/Public/MovieScene.h
@@ -179,6 +179,62 @@ struct FMovieSceneObjectBindingIDs
 	TArray<FMovieSceneObjectBindingID> IDs;
 };
 
+#if WITH_EDITORONLY_DATA
+/**
+ *
+ */
+USTRUCT()
+struct FMovieSceneSectionGroup
+{
+	GENERATED_BODY()
+public:
+	/* @return Whether the section is part of this group */
+	MOVIESCENE_API bool Contains(const UMovieSceneSection& Section) const;
+	
+	/* Add the section to this group */
+	MOVIESCENE_API void Add(UMovieSceneSection& Section);
+
+	/* Remove the section from this group */
+	MOVIESCENE_API void Remove(const UMovieSceneSection& Section);
+
+	/* Add all members of a group to this group */
+	MOVIESCENE_API void Append(const FMovieSceneSectionGroup& SectionGroup);
+
+	/* Removes any sections which the pointers are stale or otherwise not valid */
+	MOVIESCENE_API void Clean();
+
+	int32 Num() const { return Sections.Num(); }
+
+protected:
+	UPROPERTY()
+	TArray<TWeakObjectPtr<UMovieSceneSection> > Sections;
+
+public:
+	/**
+	 * Comparison operators
+	 * We only need these for being stored in a container, to check if it's the same object.
+	 * Not intended for direct use.
+	 */
+	FORCEINLINE bool operator==(const FMovieSceneSectionGroup& Other) const
+	{
+		return &Sections == &Other.Sections;
+	}
+	FORCEINLINE bool operator!=(const FMovieSceneSectionGroup& Other) const
+	{
+		return &Sections != &Other.Sections;;
+	}
+
+	/**
+	 * DO NOT USE DIRECTLY
+	 * STL-like iterators to enable range-based for loop support.
+	 */
+	FORCEINLINE TArray<TWeakObjectPtr<UMovieSceneSection> >::RangedForIteratorType      begin() { return Sections.begin(); }
+	FORCEINLINE TArray<TWeakObjectPtr<UMovieSceneSection> >::RangedForConstIteratorType begin() const { return Sections.begin(); }
+	FORCEINLINE TArray<TWeakObjectPtr<UMovieSceneSection> >::RangedForIteratorType      end() { return Sections.end(); }
+	FORCEINLINE TArray<TWeakObjectPtr<UMovieSceneSection> >::RangedForConstIteratorType end() const { return Sections.end(); }
+};
+#endif
+
 /**
  * Implements a movie scene asset.
  */
@@ -760,6 +816,31 @@ public:
 		EditorData = InEditorData;
 	}
 
+	/*
+	 * @return Whether the section is in a group
+	 */
+	bool IsSectionInGroup(const UMovieSceneSection& InSection) const;
+
+	/*
+	 * Create a group containing InSections, merging any existing groups the sections are in
+	 */
+	void GroupSections(const TArray<UMovieSceneSection*> InSections);
+
+	/*
+	 * Remove InSection from any group it currently is in
+	 */
+	void UngroupSection(const UMovieSceneSection& InSection);
+
+	/*
+	 * @return The group containing the InSection, or a nullptr if InSection is not grouped.
+	 */
+	const FMovieSceneSectionGroup* GetSectionGroup(const UMovieSceneSection& InSection) const;
+
+	/*
+	 * Cleans stale UMovieSceneSection pointers, and removes any section groups which are no longer valid, e.g. contain less that two valid sections
+	 */
+	void CleanSectionGroups();
+
 	/** The timecode at which this movie scene section is based (ie. when it was recorded) */
 	UPROPERTY()
 	FMovieSceneTimecodeSource TimecodeSource;
@@ -969,6 +1050,10 @@ private:
 	UPROPERTY()
 	TArray<FString> MuteNodes;
 
+	/** Groups of sections which should maintain the same relative offset */
+	UPROPERTY()
+	TArray<FMovieSceneSectionGroup> SectionGroups;
+
 private:
 
 	UPROPERTY()
