Sequencer: Fix CIS

diff --git a/Engine/Source/Runtime/MovieScene/Public/MovieScene.h b/Engine/Source/Runtime/MovieScene/Public/MovieScene.h
--- a/Engine/Source/Runtime/MovieScene/Public/MovieScene.h
+++ b/Engine/Source/Runtime/MovieScene/Public/MovieScene.h
@@ -179,9 +179,8 @@ struct FMovieSceneObjectBindingIDs
 	TArray<FMovieSceneObjectBindingID> IDs;
 };
 
-#if WITH_EDITORONLY_DATA
 /**
- *
+ * Structure that represents a group of sections
  */
 USTRUCT()
 struct FMovieSceneSectionGroup
@@ -233,7 +232,6 @@ public:
 	FORCEINLINE TArray<TWeakObjectPtr<UMovieSceneSection> >::RangedForIteratorType      end() { return Sections.end(); }
 	FORCEINLINE TArray<TWeakObjectPtr<UMovieSceneSection> >::RangedForConstIteratorType end() const { return Sections.end(); }
 };
-#endif
 
 /**
  * Implements a movie scene asset.
