Movie scene player now has a list of binding overrides.

diff --git a/Engine/Source/Runtime/MovieScene/Private/Evaluation/MovieSceneEvaluationState.cpp b/Engine/Source/Runtime/MovieScene/Private/Evaluation/MovieSceneEvaluationState.cpp
--- a/Engine/Source/Runtime/MovieScene/Private/Evaluation/MovieSceneEvaluationState.cpp
+++ b/Engine/Source/Runtime/MovieScene/Private/Evaluation/MovieSceneEvaluationState.cpp
@@ -184,6 +184,7 @@ void FMovieSceneObjectCache::SetSequence(UMovieSceneSequence& InSequence, FMovie
 
 void FMovieSceneObjectCache::UpdateBindings(const FGuid& InGuid, IMovieScenePlayer& Player)
 {
+	// Invalidate existing bindings, we're going to rebuild them.
 	FBoundObjects* Bindings = &BoundObjects.FindOrAdd(InGuid);
 	Bindings->Objects.Reset();
 
@@ -196,42 +197,63 @@ void FMovieSceneObjectCache::UpdateBindings(const FGuid& InGuid, IMovieScenePlay
 	}
 	ChildBindings.Remove(InGuid);
 
-	// Find the sequence for the current
+	// Find the sequence for this cache.
 	UMovieSceneSequence* Sequence = WeakSequence.Get();
 	if (!Sequence)
 	{
 		return;
 	}
 
-	const bool bUseParentsAsContext = Sequence->AreParentContextsSignificant();
-	
-	UObject* Context = Player.GetPlaybackContext();
-
-	const FMovieScenePossessable* Possessable = Sequence->GetMovieScene()->FindPossessable(InGuid);
-	if (Possessable)
+	// If we have overrides for this binding, ask the player to find it for us (most probably in a different cache
+	// for a different sequence).
+	// TODO-lchabant: we could technically end up in a circular override that creates an infinite loop...
+	const FMovieSceneEvaluationOperand Operand(SequenceID, InGuid);
+	if (const FMovieSceneEvaluationOperand* OverrideOperand = Player.BindingOverrides.Find(Operand))
+	{
+		const TArrayView<TWeakObjectPtr<>> OverrideBoundObjects = Player.FindBoundObjects(*OverrideOperand);
+		Bindings->Objects.Append(OverrideBoundObjects.GetData(), OverrideBoundObjects.Num());
+	}
+	else
 	{
-		UObject* ResolutionContext = Context;
+		const bool bUseParentsAsContext = Sequence->AreParentContextsSignificant();
 
-		// Because these are ordered parent-first, the parent must have already been bound, if it exists
-		if (Possessable->GetParent().IsValid())
-		{
-			TArrayView<TWeakObjectPtr<>> ParentBoundObjects = FindBoundObjects(Possessable->GetParent(), Player);
+		UObject* Context = Player.GetPlaybackContext();
 
-			ChildBindings.FindOrAdd(Possessable->GetParent()).AddUnique(InGuid);
+		const FMovieScenePossessable* Possessable = Sequence->GetMovieScene()->FindPossessable(InGuid);
+		if (Possessable)
+		{
+			UObject* ResolutionContext = Context;
 
-			// Refresh bindings in case of map changes
-			Bindings = BoundObjects.Find(InGuid);
-			for (TWeakObjectPtr<> Parent : ParentBoundObjects)
+			// Because these are ordered parent-first, the parent must have already been bound, if it exists
+			if (Possessable->GetParent().IsValid())
 			{
-				if (bUseParentsAsContext)
+				TArrayView<TWeakObjectPtr<>> ParentBoundObjects = FindBoundObjects(Possessable->GetParent(), Player);
+
+				ChildBindings.FindOrAdd(Possessable->GetParent()).AddUnique(InGuid);
+
+				// Refresh bindings in case of map changes
+				Bindings = BoundObjects.Find(InGuid);
+				for (TWeakObjectPtr<> Parent : ParentBoundObjects)
 				{
-					ResolutionContext = Parent.Get();
-					if (!ResolutionContext)
+					if (bUseParentsAsContext)
 					{
-						continue;
+						ResolutionContext = Parent.Get();
+						if (!ResolutionContext)
+						{
+							continue;
+						}
 					}
-				}
 
+					TArray<UObject*, TInlineAllocator<1>> FoundObjects;
+					Player.ResolveBoundObjects(InGuid, SequenceID, *Sequence, ResolutionContext, FoundObjects);
+					for (UObject* Object : FoundObjects)
+					{
+						Bindings->Objects.Add(Object);
+					}
+				}
+			}
+			else
+			{
 				TArray<UObject*, TInlineAllocator<1>> FoundObjects;
 				Player.ResolveBoundObjects(InGuid, SequenceID, *Sequence, ResolutionContext, FoundObjects);
 				for (UObject* Object : FoundObjects)
@@ -242,38 +264,29 @@ void FMovieSceneObjectCache::UpdateBindings(const FGuid& InGuid, IMovieScenePlay
 		}
 		else
 		{
-			TArray<UObject*, TInlineAllocator<1>> FoundObjects;
-			Player.ResolveBoundObjects(InGuid, SequenceID, *Sequence, ResolutionContext, FoundObjects);
-			for (UObject* Object : FoundObjects)
-			{
-				Bindings->Objects.Add(Object);
-			}
-		}
-	}
-	else
-	{
-		// Probably a spawnable then (or an phantom)
-		bool bUseDefault = true;
+			// Probably a spawnable then (or an phantom)
+			bool bUseDefault = true;
 
-		// Allow external overrides for spawnables
-		const IMovieScenePlaybackClient* PlaybackClient = Player.GetPlaybackClient();
-		if (PlaybackClient)
-		{
-			TArray<UObject*, TInlineAllocator<1>> FoundObjects;
-			bUseDefault = PlaybackClient->RetrieveBindingOverrides(InGuid, SequenceID, FoundObjects);
-			for (UObject* Object : FoundObjects)
+			// Allow external overrides for spawnables
+			const IMovieScenePlaybackClient* PlaybackClient = Player.GetPlaybackClient();
+			if (PlaybackClient)
 			{
-				Bindings->Objects.Add(Object);
+				TArray<UObject*, TInlineAllocator<1>> FoundObjects;
+				bUseDefault = PlaybackClient->RetrieveBindingOverrides(InGuid, SequenceID, FoundObjects);
+				for (UObject* Object : FoundObjects)
+				{
+					Bindings->Objects.Add(Object);
+				}
 			}
-		}
 
-		// If we have no overrides, or they want to allow the default spawnable, do that now
-		if (bUseDefault)
-		{
-			UObject* SpawnedObject = Player.GetSpawnRegister().FindSpawnedObject(InGuid, SequenceID).Get();
-			if (SpawnedObject)
+			// If we have no overrides, or they want to allow the default spawnable, do that now
+			if (bUseDefault)
 			{
-				Bindings->Objects.Add(SpawnedObject);
+				UObject* SpawnedObject = Player.GetSpawnRegister().FindSpawnedObject(InGuid, SequenceID).Get();
+				if (SpawnedObject)
+				{
+					Bindings->Objects.Add(SpawnedObject);
+				}
 			}
 		}
 	}
diff --git a/Engine/Source/Runtime/MovieScene/Public/IMovieScenePlayer.h b/Engine/Source/Runtime/MovieScene/Public/IMovieScenePlayer.h
--- a/Engine/Source/Runtime/MovieScene/Public/IMovieScenePlayer.h
+++ b/Engine/Source/Runtime/MovieScene/Public/IMovieScenePlayer.h
@@ -360,6 +360,9 @@ public:
 	/** Motion Vector Simulation */
 	TInlineValue<IMovieSceneMotionVectorSimulation> MotionVectorSimulation;
 
+	/** List of binding overrides to use for the sequence */
+	TMap<FMovieSceneEvaluationOperand, FMovieSceneEvaluationOperand> BindingOverrides;
+
 private:
 
 	/** Null register that asserts on use */
diff --git a/Engine/Source/Runtime/MovieSceneTracks/Private/Evaluation/MovieSceneSpawnTemplate.cpp b/Engine/Source/Runtime/MovieSceneTracks/Private/Evaluation/MovieSceneSpawnTemplate.cpp
--- a/Engine/Source/Runtime/MovieSceneTracks/Private/Evaluation/MovieSceneSpawnTemplate.cpp
+++ b/Engine/Source/Runtime/MovieSceneTracks/Private/Evaluation/MovieSceneSpawnTemplate.cpp
@@ -47,6 +47,13 @@ struct FSpawnObjectToken : IMovieSceneExecutionToken
 	{
 		MOVIESCENE_DETAILED_SCOPE_CYCLE_COUNTER(MovieSceneEval_SpawnTrack_TokenExecute)
 
+		if (const FMovieSceneEvaluationOperand* OperandOverride = Player.BindingOverrides.Find(Operand))
+		{
+			// Don't do anything if this operand was overriden... someone else will take care of it (either another spawn track, or
+			// some possessable).
+			return;
+		}
+
 		bool bHasSpawnedObject = Player.GetSpawnRegister().FindSpawnedObject(Operand.ObjectBindingID, Operand.SequenceID).Get() != nullptr;
 		
 		// Check binding overrides to see if this spawnable has been overridden, and whether it allows the default spawnable to exist
