Sequencer: Fix const_cast

diff --git a/Engine/Source/Runtime/MovieScene/Private/MovieScene.cpp b/Engine/Source/Runtime/MovieScene/Private/MovieScene.cpp
--- a/Engine/Source/Runtime/MovieScene/Private/MovieScene.cpp
+++ b/Engine/Source/Runtime/MovieScene/Private/MovieScene.cpp
@@ -561,7 +561,7 @@ void FMovieSceneSectionGroup::Append(const FMovieSceneSectionGroup& SectionGroup
 
 void FMovieSceneSectionGroup::Remove(const UMovieSceneSection& Section)
 {
-	Sections.RemoveSingle(&Section);
+	Sections.RemoveSingle(const_cast<UMovieSceneSection*>(&Section));
 }
 
 void FMovieSceneSectionGroup::Clean()
