TRange<FFrameNumber> LexToString now uses "," to split the upper/lower bounds instead of "-" to avoid confusion with negative upper bound values.

diff --git a/Engine/Source/Runtime/MovieScene/Public/MovieSceneTimeHelpers.h b/Engine/Source/Runtime/MovieScene/Public/MovieSceneTimeHelpers.h
--- a/Engine/Source/Runtime/MovieScene/Public/MovieSceneTimeHelpers.h
+++ b/Engine/Source/Runtime/MovieScene/Public/MovieSceneTimeHelpers.h
@@ -196,7 +196,7 @@ inline FString LexToString(const TRange<FFrameNumber>& InRange)
 	TRangeBound<FFrameNumber> SourceLower = InRange.GetLowerBound();
 	TRangeBound<FFrameNumber> SourceUpper = InRange.GetUpperBound();
 
-	return *FString::Printf(TEXT("%s-%s"),
+	return *FString::Printf(TEXT("%s,%s"),
 		SourceLower.IsOpen() ?
 		TEXT("[-inf") :
 		SourceLower.IsInclusive() ?
