Sequencer: Add support for looping subsequences

This change introduces a new form of time transformation called "time warping", which is used in looping sequences. You can transform time from a root sequence into a looping sequence and, assuming you know which loop number you are in, you can convert from a local/sub-sequence time back into a root time.

This new warping transform is used everywhere the old linear time offset/scale transform was used by making the sequence time transform struct use both, and the shot sections are getting a UI upgrade to display looping information.

Looping is left as an opt-in thing on tracks that have sub-sequences, though, because it changes what resizing a section means.

diff --git a/Engine/Plugins/MovieScene/SequencerScripting/Source/SequencerScripting/Private/ExtensionLibraries/MovieSceneSectionExtensions.cpp b/Engine/Plugins/MovieScene/SequencerScripting/Source/SequencerScripting/Private/ExtensionLibraries/MovieSceneSectionExtensions.cpp
--- a/Engine/Plugins/MovieScene/SequencerScripting/Source/SequencerScripting/Private/ExtensionLibraries/MovieSceneSectionExtensions.cpp
+++ b/Engine/Plugins/MovieScene/SequencerScripting/Source/SequencerScripting/Private/ExtensionLibraries/MovieSceneSectionExtensions.cpp
@@ -343,7 +343,7 @@ int32 UMovieSceneSectionExtensions::GetParentSequenceFrame(UMovieSceneSubSection
 		
 	for (int32 SectionIndex = 0; SectionIndex < SubSectionChain.Num(); ++SectionIndex)
 	{
-		LocalFrameTime = LocalFrameTime * SubSectionChain[SectionIndex]->OuterToInnerTransform().Inverse();
+		LocalFrameTime = LocalFrameTime * SubSectionChain[SectionIndex]->OuterToInnerTransform().InverseLinearOnly();
 	}
 
 	FFrameRate ParentDisplayRate = ParentSequence->GetMovieScene()->GetDisplayRate();
diff --git a/Engine/Source/Editor/Sequencer/Private/SSequencer.cpp b/Engine/Source/Editor/Sequencer/Private/SSequencer.cpp
--- a/Engine/Source/Editor/Sequencer/Private/SSequencer.cpp
+++ b/Engine/Source/Editor/Sequencer/Private/SSequencer.cpp
@@ -711,6 +711,23 @@ void SSequencer::Construct(const FArguments& InArgs, TSharedRef<FSequencer> InSe
 									.MinDesiredWidth(this, &SSequencer::GetPlayTimeMinDesiredWidth)
 								]
 							]
+							+ SHorizontalBox::Slot()
+							.AutoWidth()
+							.VAlign(VAlign_Center)
+							.HAlign(HAlign_Right)
+							.Padding(FMargin(CommonPadding + 2.0, 0.f, 0.f, 0.f))
+							[
+								SNew(SBorder)
+								.BorderImage(nullptr)
+								[
+									// Current loop index, if any
+									SAssignNew(LoopIndexDisplay, STextBlock)
+									.Text_Lambda([this]() -> FText {
+										uint32 LoopIndex = SequencerPtr.Pin()->GetLocalLoopIndex();
+										return (LoopIndex != FMovieSceneTimeWarping::InvalidWarpCount) ? FText::AsNumber(LoopIndex + 1) : FText();
+									})
+								]
+							]
 						]
 					]
 
diff --git a/Engine/Source/Editor/Sequencer/Private/SSequencer.h b/Engine/Source/Editor/Sequencer/Private/SSequencer.h
--- a/Engine/Source/Editor/Sequencer/Private/SSequencer.h
+++ b/Engine/Source/Editor/Sequencer/Private/SSequencer.h
@@ -15,6 +15,7 @@
 #include "Framework/Commands/UICommandList.h"
 #include "Widgets/Input/NumericTypeInterface.h"
 #include "Widgets/Input/SSpinBox.h"
+#include "Widgets/Text/STextBlock.h"
 #include "Sequencer.h"
 
 class FActorDragDropGraphEdOp;
@@ -595,9 +596,12 @@ private:
 	/** The search widget for filtering curves in the Curve Editor tree. */
 	TSharedPtr<SWidget> CurveEditorSearchBox;
 
-	/** The current playback time display.*/
+	/** The current playback time display. */
 	TSharedPtr<STemporarilyFocusedSpinBox<double>> PlayTimeDisplay;
 
+	/** The current loop display for when editing a looping sub-sequence. */
+	TSharedPtr<STextBlock> LoopIndexDisplay;
+
 	/** The sequencer tree view responsible for the outliner and track areas */
 	TSharedPtr<SSequencerTreeView> TreeView;
 
diff --git a/Engine/Source/Editor/Sequencer/Private/Sequencer.cpp b/Engine/Source/Editor/Sequencer/Private/Sequencer.cpp
--- a/Engine/Source/Editor/Sequencer/Private/Sequencer.cpp
+++ b/Engine/Source/Editor/Sequencer/Private/Sequencer.cpp
@@ -506,6 +506,8 @@ FSequencer::FSequencer()
 	, bUpdatingSequencerSelection( false )
 	, bUpdatingExternalSelection( false )
 	, bNeedsEvaluate(false)
+	, LocalLoopIndexOnBeginScrubbing(FMovieSceneTimeWarping::InvalidWarpCount)
+	, LocalLoopIndexOffsetDuringScrubbing(0)
 {
 	Selection.GetOnOutlinerNodeSelectionChanged().AddRaw(this, &FSequencer::OnSelectedOutlinerNodesChanged);
 	Selection.GetOnNodesWithSelectedKeysOrSectionsChanged().AddRaw(this, &FSequencer::OnSelectedOutlinerNodesChanged);
@@ -821,6 +823,7 @@ void FSequencer::ResetToNewRootSequence(UMovieSceneSequence& NewSequence)
 	RootTemplateInstance.Initialize(NewSequence, *this);
 
 	RootToLocalTransform = FMovieSceneSequenceTransform();
+	RootToLocalLoopCounter = FMovieSceneWarpCounter();
 
 	ResetPerMovieSceneData();
 	SequencerWidget->ResetBreadcrumbs();
@@ -1047,8 +1050,20 @@ void FSequencer::PopToSequenceInstance(FMovieSceneSequenceIDRef SequenceID)
 
 void FSequencer::UpdateSubSequenceData()
 {
+	const bool bIsScrubbing = GetPlaybackStatus() == EMovieScenePlayerStatus::Scrubbing;
+	const bool bIsSubSequenceWarping = RootToLocalTransform.NestedTransforms.Num() > 0 && RootToLocalTransform.NestedTransforms.Last().IsWarping();
+	const bool bIsScrubbingWarpingSubSequence = bIsScrubbing && bIsSubSequenceWarping;
+
 	SubSequenceRange = TRange<FFrameNumber>::Empty();
 	RootToLocalTransform = FMovieSceneSequenceTransform();
+	if (!bIsScrubbingWarpingSubSequence)
+	{
+		RootToLocalLoopCounter = FMovieSceneWarpCounter();
+	}
+	// else: we're scrubbing, and we don't want to increase/decrease the loop index quite yet,
+	// because that would mess up time transforms. This would be because the mouse would still be
+	// before/after the current loop, and therefore would already add/subtract more than a full
+	// loop's time to the current time, so we don't need the loop counter to change yet.
 
 	// Find the parent sub section and set up the sub sequence range, if necessary
 	if (ActiveTemplateIDs.Num() <= 1)
@@ -1062,6 +1077,50 @@ void FSequencer::UpdateSubSequenceData()
 	{
 		SubSequenceRange = SubSequenceData->PlayRange.Value;
 		RootToLocalTransform = SubSequenceData->RootToSequenceTransform;
+
+		const FQualifiedFrameTime RootTime = GetGlobalTime();
+		if (!bIsScrubbingWarpingSubSequence)
+		{
+			FFrameTime LocalTime;
+			RootToLocalTransform.TransformTime(RootTime.Time, LocalTime, RootToLocalLoopCounter);
+		}
+		else
+		{
+			// If we are scrubbing _and_ the current sequence is warping, we need to do some custom stuff.
+			const FFrameNumber PlayRangeSize = SubSequenceData->PlayRange.Value.Size<FFrameNumber>();
+			const FFrameNumber PlayRangeUpperBound = SubSequenceData->PlayRange.Value.GetUpperBoundValue();
+			const FFrameNumber PlayRangeLowerBound = SubSequenceData->PlayRange.Value.GetLowerBoundValue();
+			
+			ensure(LocalLoopIndexOnBeginScrubbing != FMovieSceneTimeWarping::InvalidWarpCount);
+			ensure(RootToLocalLoopCounter.WarpCounts.Num() > 0);
+
+			// Compute the new local time based on the specific loop that we had when we started scrubbing.
+			FMovieSceneSequenceTransform RootToLocalTransformWithoutLeafLooping = RootToLocalTransform;
+			FMovieSceneNestedSequenceTransform LeafLooping = RootToLocalTransformWithoutLeafLooping.NestedTransforms.Pop();
+			FFrameTime LocalTimeWithLastLoopUnwarped = RootTime.Time * RootToLocalTransformWithoutLeafLooping;
+			LocalTimeWithLastLoopUnwarped = LocalTimeWithLastLoopUnwarped * LeafLooping.LinearTransform;
+			if (LeafLooping.IsWarping())
+			{
+				LeafLooping.Warping.TransformTimeSpecific(
+						LocalTimeWithLastLoopUnwarped, LocalLoopIndexOnBeginScrubbing, LocalTimeWithLastLoopUnwarped);
+			}
+
+			// Now figure out if we're in a next/previous loop because we scrubbed past the lower/upper bound
+			// of the loop. Note, again, that we only compute the new loop index for UI display purposes at this
+			// point (see comment at the beginning of this method). We will commit to the new loop indices
+			// once we're done scrubbing.
+			LocalLoopIndexOffsetDuringScrubbing = 0;
+			while (LocalTimeWithLastLoopUnwarped >= PlayRangeUpperBound)
+			{
+				LocalTimeWithLastLoopUnwarped = LocalTimeWithLastLoopUnwarped - PlayRangeSize;
+				++LocalLoopIndexOffsetDuringScrubbing;
+			}
+			while (LocalTimeWithLastLoopUnwarped <= PlayRangeLowerBound)
+			{
+				LocalTimeWithLastLoopUnwarped = LocalTimeWithLastLoopUnwarped + PlayRangeSize;
+				--LocalLoopIndexOffsetDuringScrubbing;
+			}
+		}
 	}
 }
 
@@ -1695,7 +1754,7 @@ void FSequencer::BakeTransform()
 			for (FFrameTime EvalTime = InFrame; EvalTime < OutFrame; EvalTime += Interval)
 			{
 				FFrameNumber KeyTime = FFrameRate::Snap(EvalTime, Resolution, SnapRate).FloorToFrame();
-				FMovieSceneEvaluationRange Range = PlayPosition.JumpTo(KeyTime * RootToLocalTransform.Inverse());
+				FMovieSceneEvaluationRange Range = PlayPosition.JumpTo(KeyTime * RootToLocalTransform.InverseLinearOnly());
 				EvaluateInternal(Range);
 				
 				USceneComponent* Parent = nullptr;
@@ -2474,14 +2533,28 @@ bool FSequencer::GetAutoSetTrackDefaults() const
 
 FQualifiedFrameTime FSequencer::GetLocalTime() const
 {
-	FFrameRate FocusedResolution = GetFocusedTickResolution();
-	FFrameTime CurrentPosition   = PlayPosition.GetCurrentPosition();
+	const FFrameRate FocusedResolution = GetFocusedTickResolution();
+	const FFrameTime CurrentPosition   = PlayPosition.GetCurrentPosition();
 
-	FFrameTime RootTime = ConvertFrameTime(CurrentPosition, PlayPosition.GetInputRate(), PlayPosition.GetOutputRate());
+	const FFrameTime RootTime = ConvertFrameTime(CurrentPosition, PlayPosition.GetInputRate(), PlayPosition.GetOutputRate());
 	return FQualifiedFrameTime(RootTime * RootToLocalTransform, FocusedResolution);
 }
 
 
+uint32 FSequencer::GetLocalLoopIndex() const
+{
+	if (RootToLocalLoopCounter.WarpCounts.Num() == 0)
+	{
+		return FMovieSceneTimeWarping::InvalidWarpCount;
+	}
+	else
+	{
+		const bool bIsScrubbing = GetPlaybackStatus() == EMovieScenePlayerStatus::Scrubbing;
+		return RootToLocalLoopCounter.WarpCounts.Last() + (bIsScrubbing ? LocalLoopIndexOffsetDuringScrubbing : 0);
+	}
+}
+
+
 FQualifiedFrameTime FSequencer::GetGlobalTime() const
 {
 	FFrameTime RootTime = ConvertFrameTime(PlayPosition.GetCurrentPosition(), PlayPosition.GetInputRate(), PlayPosition.GetOutputRate());
@@ -2526,7 +2599,7 @@ void FSequencer::SetLocalTimeDirectly(FFrameTime NewTime)
 	}
 
 	// Transform the time to the root time-space
-	SetGlobalTime(NewTime * RootToLocalTransform.Inverse());
+	SetGlobalTime(NewTime * RootToLocalTransform.InverseFromWarp(RootToLocalLoopCounter));
 
 	if (PreviousFocusedWidget.IsValid())
 	{
@@ -2737,7 +2810,7 @@ void FSequencer::RenderMovieInternal(TRange<FFrameNumber> Range, bool bSetFrameO
 
 		if (const FMovieSceneSubSequenceData* SubSequenceData = RootTemplateInstance.GetHierarchy().FindSubData(GetFocusedTemplateID()))
 		{
-			Range = Range * SubSequenceData->RootToSequenceTransform.Inverse();
+			Range = Range * SubSequenceData->RootToSequenceTransform.InverseLinearOnly();
 		}
 	}
 
@@ -3680,7 +3753,7 @@ FReply FSequencer::OnPlay(bool bTogglePlay)
 
 		if (GetLocalTime().Time <= MinInclusiveTime || GetLocalTime().Time >= MaxInclusiveTime)
 		{
-			FFrameTime NewGlobalTime = (PlaybackSpeed > 0 ? MinInclusiveTime : MaxInclusiveTime) * RootToLocalTransform.Inverse();
+			FFrameTime NewGlobalTime = (PlaybackSpeed > 0 ? MinInclusiveTime : MaxInclusiveTime) * RootToLocalTransform.InverseFromWarp(RootToLocalLoopCounter);
 			SetGlobalTime(NewGlobalTime);
 		}
 
@@ -3999,7 +4072,9 @@ void FSequencer::SetLocalTimeLooped(FFrameTime NewLocalTime)
 {
 	TOptional<EMovieScenePlayerStatus::Type> NewPlaybackStatus;
 
-	FFrameTime NewGlobalTime = NewLocalTime * RootToLocalTransform.Inverse();
+	const FMovieSceneSequenceTransform LocalToRootTransform = RootToLocalTransform.InverseFromWarp(RootToLocalLoopCounter);
+
+	FFrameTime NewGlobalTime = NewLocalTime * LocalToRootTransform;
 
 	TRange<FFrameNumber> TimeBounds = GetTimeBounds();
 
@@ -4018,7 +4093,7 @@ void FSequencer::SetLocalTimeLooped(FFrameTime NewLocalTime)
 		{
 			if (NewLocalTime < MinInclusiveTime || NewLocalTime > MaxInclusiveTime)
 			{
-				NewGlobalTime = (PlaybackSpeed > 0 ? MinInclusiveTime : MaxInclusiveTime) * RootToLocalTransform.Inverse();
+				NewGlobalTime = (PlaybackSpeed > 0 ? MinInclusiveTime : MaxInclusiveTime) * LocalToRootTransform;
 
 				bResetPosition = true;
 				bHasJumped = true;
@@ -4042,7 +4117,7 @@ void FSequencer::SetLocalTimeLooped(FFrameTime NewLocalTime)
 		// Stop if we hit the playback range end
 		if (bReachedEnd)
 		{
-			NewGlobalTime = (PlaybackSpeed > 0 ? MaxInclusiveTime : MinInclusiveTime) * RootToLocalTransform.Inverse();
+			NewGlobalTime = (PlaybackSpeed > 0 ? MaxInclusiveTime : MinInclusiveTime) * LocalToRootTransform;
 			NewPlaybackStatus = EMovieScenePlayerStatus::Stopped;
 		}
 		// Constrain to the play range if necessary
@@ -4051,7 +4126,7 @@ void FSequencer::SetLocalTimeLooped(FFrameTime NewLocalTime)
 			// Clamp to bound or jump back if necessary
 			if (NewLocalTime < MinInclusiveTime || NewLocalTime >= MaxInclusiveTime)
 			{
-				NewGlobalTime = (PlaybackSpeed > 0 ? MinInclusiveTime : MaxInclusiveTime) * RootToLocalTransform.Inverse();
+				NewGlobalTime = (PlaybackSpeed > 0 ? MinInclusiveTime : MaxInclusiveTime) * LocalToRootTransform;
 
 				bResetPosition = true;
 			}
@@ -4062,7 +4137,7 @@ void FSequencer::SetLocalTimeLooped(FFrameTime NewLocalTime)
 			FFrameTime WorkingMin = (WorkingRange.GetLowerBoundValue() * LocalTickResolution).CeilToFrame();
 			FFrameTime WorkingMax = (WorkingRange.GetUpperBoundValue() * LocalTickResolution).FloorToFrame();
 
-			NewGlobalTime = FMath::Clamp(NewLocalTime, WorkingMin, WorkingMax) * RootToLocalTransform.Inverse();
+			NewGlobalTime = FMath::Clamp(NewLocalTime, WorkingMin, WorkingMax) * LocalToRootTransform;
 
 			bResetPosition = true;
 			NewPlaybackStatus = EMovieScenePlayerStatus::Stopped;
@@ -4310,6 +4385,9 @@ void FSequencer::OnBeginScrubbing()
 	SetPlaybackStatus(EMovieScenePlayerStatus::Scrubbing);
 	SequencerWidget->RegisterActiveTimerForPlayback();
 
+	LocalLoopIndexOnBeginScrubbing = GetLocalLoopIndex();
+	LocalLoopIndexOffsetDuringScrubbing = 0;
+
 	OnBeginScrubbingDelegate.Broadcast();
 }
 
@@ -4320,6 +4398,9 @@ void FSequencer::OnEndScrubbing()
 	AutoscrubOffset.Reset();
 	StopAutoscroll();
 
+	LocalLoopIndexOnBeginScrubbing = FMovieSceneTimeWarping::InvalidWarpCount;
+	LocalLoopIndexOffsetDuringScrubbing = 0;
+
 	OnEndScrubbingDelegate.Broadcast();
 }
 
@@ -8676,7 +8757,7 @@ void FSequencer::Pause()
 		FFrameRate          FocusedDisplayRate = GetFocusedDisplayRate();
 
 		// Snap to the focused play rate
-		FFrameTime RootPosition  = FFrameRate::Snap(LocalTime.Time, LocalTime.Rate, FocusedDisplayRate) * RootToLocalTransform.Inverse();
+		FFrameTime RootPosition  = FFrameRate::Snap(LocalTime.Time, LocalTime.Rate, FocusedDisplayRate) * RootToLocalTransform.InverseFromWarp(RootToLocalLoopCounter);
 
 		// Convert the root position from tick resolution time base (the output rate), to the play position input rate
 		FFrameTime InputPosition = ConvertFrameTime(RootPosition, PlayPosition.GetOutputRate(), PlayPosition.GetInputRate());
@@ -8744,7 +8825,7 @@ void FSequencer::StepToNextShot()
 	FMovieSceneSequenceID OuterSequenceID = ActiveTemplateIDs[ActiveTemplateIDs.Num()-2];
 	UMovieSceneSequence* Sequence = RootTemplateInstance.GetSequence(OuterSequenceID);
 
-	FFrameTime CurrentTime = SubSequenceRange.GetLowerBoundValue() * RootToLocalTransform.Inverse();
+	FFrameTime CurrentTime = SubSequenceRange.GetLowerBoundValue() * RootToLocalTransform.InverseFromWarp(RootToLocalLoopCounter);
 
 	UMovieSceneSubSection* NextShot = Cast<UMovieSceneSubSection>(FindNextOrPreviousShot(Sequence, CurrentTime.FloorToFrame(), true));
 	if (!NextShot)
@@ -8777,7 +8858,7 @@ void FSequencer::StepToPreviousShot()
 	FMovieSceneSequenceID OuterSequenceID = ActiveTemplateIDs[ActiveTemplateIDs.Num() - 2];
 	UMovieSceneSequence* Sequence = RootTemplateInstance.GetSequence(OuterSequenceID);
 
-	FFrameTime CurrentTime = SubSequenceRange.GetLowerBoundValue() * RootToLocalTransform.Inverse();
+	FFrameTime CurrentTime = SubSequenceRange.GetLowerBoundValue() * RootToLocalTransform.InverseFromWarp(RootToLocalLoopCounter);
 	UMovieSceneSubSection* PreviousShot = Cast<UMovieSceneSubSection>(FindNextOrPreviousShot(Sequence, CurrentTime.FloorToFrame(), false));
 	if (!PreviousShot)
 	{
diff --git a/Engine/Source/Editor/Sequencer/Private/Sequencer.h b/Engine/Source/Editor/Sequencer/Private/Sequencer.h
--- a/Engine/Source/Editor/Sequencer/Private/Sequencer.h
+++ b/Engine/Source/Editor/Sequencer/Private/Sequencer.h
@@ -683,6 +683,7 @@ public:
 	virtual bool GetAutoSetTrackDefaults() const override;
 	virtual FQualifiedFrameTime GetLocalTime() const override;
 	virtual FQualifiedFrameTime GetGlobalTime() const override;
+	virtual uint32 GetLocalLoopIndex() const override;
 	virtual void SetLocalTime(FFrameTime Time, ESnapTimeMode SnapTimeMode = ESnapTimeMode::STM_None) override;
 	virtual void SetLocalTimeDirectly(FFrameTime NewTime) override;
 	virtual void SetGlobalTime(FFrameTime Time) override;
@@ -1083,8 +1084,12 @@ private:
 	 */
 	TArray<bool> ActiveTemplateStates;
 
+	/** Time transformation from the root sequence to the currently edited sequence. */
 	FMovieSceneSequenceTransform RootToLocalTransform;
 
+	/** Current loop of the current sub-sequence, if we are in a looping sub-sequence. */
+	FMovieSceneWarpCounter RootToLocalLoopCounter;
+
 	/** The time range target to be viewed */
 	TRange<double> TargetViewRange;
 
@@ -1127,6 +1132,12 @@ private:
 	/** Current play position */
 	FMovieScenePlaybackPosition PlayPosition;
 
+	/** Local loop index at the time we began scrubbing */
+	uint32 LocalLoopIndexOnBeginScrubbing;
+
+	/** Local loop index to add for the purposes of displaying it in the UI */
+	uint32 LocalLoopIndexOffsetDuringScrubbing;
+
 	/** The playback speed */
 	float PlaybackSpeed;
 
diff --git a/Engine/Source/Editor/Sequencer/Public/ISequencer.h b/Engine/Source/Editor/Sequencer/Public/ISequencer.h
--- a/Engine/Source/Editor/Sequencer/Public/ISequencer.h
+++ b/Engine/Source/Editor/Sequencer/Public/ISequencer.h
@@ -323,6 +323,8 @@ public:
 	 */
 	virtual FQualifiedFrameTime GetGlobalTime() const = 0;
 
+	virtual uint32 GetLocalLoopIndex() const = 0;
+
 	/**
 	 * Sets the cursor position relative to the currently focused sequence
 	 *
diff --git a/Engine/Source/Editor/UnrealEd/Private/Fbx/FbxMainExport.cpp b/Engine/Source/Editor/UnrealEd/Private/Fbx/FbxMainExport.cpp
--- a/Engine/Source/Editor/UnrealEd/Private/Fbx/FbxMainExport.cpp
+++ b/Engine/Source/Editor/UnrealEd/Private/Fbx/FbxMainExport.cpp
@@ -1599,7 +1599,7 @@ int32 FFbxExporter::FLevelSequenceAnimTrackAdapter::GetStartFrame() const
 {
 	FFrameRate TickResolution = MovieScene->GetTickResolution();
 	FFrameRate DisplayRate = MovieScene->GetDisplayRate();
-	return FFrameRate::TransformTime(FFrameTime(MovieScene::DiscreteInclusiveLower(MovieScene->GetPlaybackRange()) * RootToLocalTransform.Inverse()), TickResolution, DisplayRate).RoundToFrame().Value;
+	return FFrameRate::TransformTime(FFrameTime(MovieScene::DiscreteInclusiveLower(MovieScene->GetPlaybackRange()) * RootToLocalTransform.InverseLinearOnly()), TickResolution, DisplayRate).RoundToFrame().Value;
 }
 
 int32 FFbxExporter::FLevelSequenceAnimTrackAdapter::GetLength() const
@@ -1616,7 +1616,7 @@ void FFbxExporter::FLevelSequenceAnimTrackAdapter::UpdateAnimation( int32 LocalF
 	FFrameRate DisplayRate = MovieScene->GetDisplayRate();
 
 	FFrameTime LocalTime = FFrameRate::TransformTime(FFrameTime(LocalFrame), DisplayRate, TickResolution);
-	FFrameTime GlobalTime = LocalTime * RootToLocalTransform.Inverse();
+	FFrameTime GlobalTime = LocalTime * RootToLocalTransform.InverseLinearOnly();
 
 	FMovieSceneContext Context = FMovieSceneContext(FMovieSceneEvaluationRange(GlobalTime, TickResolution), MovieScenePlayer->GetPlaybackStatus()).SetHasJumped(true);
 
@@ -2516,7 +2516,7 @@ void FFbxExporter::ExportChannelToFbxCurve(FbxAnimCurve& InFbxCurve, const FMovi
 
 		FbxTime FbxTime;
 		FbxAnimCurveKey FbxKey;
-		const double KeyTimeSeconds = GetExportOptions()->bExportLocalTime ? KeyTime / TickResolution : (KeyTime * RootToLocalTransform.Inverse()) / TickResolution;
+		const double KeyTimeSeconds = GetExportOptions()->bExportLocalTime ? KeyTime / TickResolution : (KeyTime * RootToLocalTransform.InverseLinearOnly()) / TickResolution;
 
 		FbxTime.SetSecondDouble(KeyTimeSeconds);
 
@@ -2656,7 +2656,7 @@ void FFbxExporter::ExportLevelSequence3DTransformTrack( FbxNode* FbxNode, UMovie
 		FbxCurveRotZ->KeyModifyBegin();
 
 		int32 LocalStartFrame = FFrameRate::TransformTime(FFrameTime(MovieScene::DiscreteInclusiveLower(InPlaybackRange)), TickResolution, DisplayRate).RoundToFrame().Value;
-		int32 StartFrame = FFrameRate::TransformTime(FFrameTime(MovieScene::DiscreteInclusiveLower(InPlaybackRange) * RootToLocalTransform.Inverse()), TickResolution, DisplayRate).RoundToFrame().Value;
+		int32 StartFrame = FFrameRate::TransformTime(FFrameTime(MovieScene::DiscreteInclusiveLower(InPlaybackRange) * RootToLocalTransform.InverseLinearOnly()), TickResolution, DisplayRate).RoundToFrame().Value;
 		int32 AnimationLength = FFrameRate::TransformTime(FFrameTime(FFrameNumber(MovieScene::DiscreteSize(InPlaybackRange))), TickResolution, DisplayRate).RoundToFrame().Value;
 
 		for (int32 FrameCount = 0; FrameCount <= AnimationLength; ++FrameCount)
diff --git a/Engine/Source/Runtime/MovieScene/Private/Compilation/MovieSceneCompiler.cpp b/Engine/Source/Runtime/MovieScene/Private/Compilation/MovieSceneCompiler.cpp
--- a/Engine/Source/Runtime/MovieScene/Private/Compilation/MovieSceneCompiler.cpp
+++ b/Engine/Source/Runtime/MovieScene/Private/Compilation/MovieSceneCompiler.cpp
@@ -17,6 +17,21 @@
 DECLARE_CYCLE_STAT(TEXT("Full Compile"),  MovieSceneEval_CompileFull,  STATGROUP_MovieSceneEval);
 DECLARE_CYCLE_STAT(TEXT("Compile Range"), MovieSceneEval_CompileRange, STATGROUP_MovieSceneEval);
 
+/** Parameter structure used for keeping sub-sequence information we need for compilation */
+struct FMovieSceneSubSequenceGatherData
+{
+	FMovieSceneSubSequenceGatherData() : HierarchicalBias(0) {}
+
+	FMovieSceneSubSequenceGatherData(const FMovieSceneSubSequenceData& SubData)
+		: RootToSequenceTransform(SubData.RootToSequenceTransform)
+		, HierarchicalBias(SubData.HierarchicalBias)
+	{
+	}
+
+	FMovieSceneSequenceTransform RootToSequenceTransform;
+	int32 HierarchicalBias;
+};
+
 /** Parameter structure used for gathering entities for a given time or range */
 struct FGatherParameters
 {
@@ -32,15 +47,23 @@ struct FGatherParameters
 		, HierarchicalBias(0)
 	{}
 
-	FGatherParameters CreateForSubData(const FMovieSceneSubSequenceData& SubData) const
+	FGatherParameters CreateForSubGatherData(const FMovieSceneSubSequenceGatherData& SubGatherData) const
+	{
+		return CreateForSubGatherData(SubGatherData, RootCompileRange, RootClampRange, FMovieSceneTimeWarping::InvalidWarpCount);
+	}
+
+	FGatherParameters CreateForSubGatherData(const FMovieSceneSubSequenceGatherData& SubGatherData, TRange<FFrameNumber> InRootCompileRange, TRange<FFrameNumber> InRootClampRange, uint32 LoopIndex) const
 	{
 		FGatherParameters SubParams = *this;
 
-		SubParams.RootToSequenceTransform   = SubData.RootToSequenceTransform;
-		SubParams.HierarchicalBias          = SubData.HierarchicalBias;
+		SubParams.RootToSequenceTransform   = SubGatherData.RootToSequenceTransform;
+		SubParams.HierarchicalBias          = SubGatherData.HierarchicalBias;
+
+		SubParams.RootCompileRange          = InRootCompileRange;
+		SubParams.RootClampRange            = InRootClampRange;
 
-		SubParams.LocalCompileRange         = SubParams.RootCompileRange * SubData.RootToSequenceTransform;
-		SubParams.LocalClampRange           = SubParams.RootClampRange   * SubData.RootToSequenceTransform;
+		SubParams.LocalCompileRange         = SubGatherData.RootToSequenceTransform.TransformRangeConstrained(InRootCompileRange);
+		SubParams.LocalClampRange           = SubGatherData.RootToSequenceTransform.TransformRangeConstrained(InRootClampRange);
 
 		if (SubParams.LocalCompileRange.IsEmpty())
 		{
@@ -49,13 +72,16 @@ struct FGatherParameters
 			// We know that neither bound is Open by this point because such ranges would not be considered empty
 			SubParams.LocalCompileRange.SetUpperBoundValue(SubParams.LocalCompileRange.GetUpperBoundValue() + 1);
 		}
+
+		SubParams.SequenceLoopCounter.AddWarpingLevel(LoopIndex);
+
 		return SubParams;
 	}
 
 	void SetClampRange(TRange<FFrameNumber> InNewRootClampRange)
 	{
 		RootClampRange  = InNewRootClampRange;
-		LocalClampRange = InNewRootClampRange * RootToSequenceTransform;
+		LocalClampRange = RootToSequenceTransform.TransformRangeConstrained(InNewRootClampRange);
 	}
 
 	/** Clamp the specified range to the current clamp range (in root space) */
@@ -87,6 +113,9 @@ struct FGatherParameters
 	/** Transform from the root time-space to the current sequence's time-space */
 	FMovieSceneSequenceTransform RootToSequenceTransform;
 
+	/** When we are looping the current sequence, which loop is this */
+	FMovieSceneWarpCounter SequenceLoopCounter;
+
 	/** Current accumulated hierarchical bias */
 	int32 HierarchicalBias;
 };
@@ -428,7 +457,7 @@ void FMovieSceneCompiler::GatherCompileOnTheFlyData(UMovieSceneSequence& InSeque
 
 	for ( ; SubSectionIt && SubSectionIt.Range().Overlaps(CompileClampIntersection); ++SubSectionIt)
 	{
-		TRange<FFrameNumber> ThisSegmentRangeRoot = Params.ClampRoot(SubSectionIt.Range() * Params.RootToSequenceTransform.Inverse());
+		TRange<FFrameNumber> ThisSegmentRangeRoot = Params.ClampRoot(SubSectionIt.Range() * Params.RootToSequenceTransform.InverseFromWarp(Params.SequenceLoopCounter));
 		if (ThisSegmentRangeRoot.IsEmpty())
 		{
 			continue;
@@ -459,19 +488,19 @@ void FMovieSceneCompiler::GatherCompileOnTheFlyData(UMovieSceneSequence& InSeque
 
 				SubSectionGatherParams.Flags = SubSectionData.Flags;
 
-				GatherCompileDataForSubSection(*SubSection, SubSectionData.ObjectBindingId, SubSectionGatherParams, OutData);
+				GatherCompileDataForSubSection(InSequence, *SubSection, SubSectionData.ObjectBindingId, SubSectionGatherParams, OutData);
 			}
 		}
 
 		if (!bAnySubSections)
 		{
 			// Intersect the unique range in the tree with the current overlapping empty range to constrict the resulting compile range in the case where this is a gap between sub sections
-			OutData.EmptySpace.AddTimeRange(Params.ClampRoot(SubSectionIt.Range() * Params.RootToSequenceTransform.Inverse()));
+			OutData.EmptySpace.AddTimeRange(Params.ClampRoot(SubSectionIt.Range() * Params.RootToSequenceTransform.InverseFromWarp(Params.SequenceLoopCounter)));
 		}
 	}
 }
 
-void FMovieSceneCompiler::GatherCompileDataForSubSection(const UMovieSceneSubSection& SubSection, const FGuid& InObjectBindingId, const FGatherParameters& Params, FMovieSceneGatheredCompilerData& OutData)
+void FMovieSceneCompiler::GatherCompileDataForSubSection(const UMovieSceneSequence& InOuterSequence, const UMovieSceneSubSection& SubSection, const FGuid& InObjectBindingId, const FGatherParameters& Params, FMovieSceneGatheredCompilerData& OutData)
 {
 	UMovieSceneSequence* SubSequence = SubSection.GetSequence();
 	if (!SubSequence)
@@ -492,22 +521,100 @@ void FMovieSceneCompiler::GatherCompileDataForSubSection(const UMovieSceneSubSec
 	Params.RootPath.Push(UnAccumulatedSequenceID);
 
 	// Find/add sub data in the root template
-	TOptional<FGatherParameters> SubParams;
+	TOptional<FMovieSceneSubSequenceGatherData> SubGatherData;
 
 	{
 		const FMovieSceneSubSequenceData* CompilationSubData = GetOrCreateSubSequenceData(InnerSequenceID, ParentSequenceID, SubSection, InObjectBindingId, Params.RootHierarchy);
 		check(CompilationSubData);
 
-		SubParams = Params.CreateForSubData(*CompilationSubData);
+		SubGatherData = FMovieSceneSubSequenceGatherData(*CompilationSubData);
 		// Any code after this point may reallocate the root hierarchy, so CompilationSubData cannot be used
 	}
 
-	GatherCompileOnTheFlyData(*SubSequence, SubParams.GetValue(), OutData);
+	if (!SubSection.Parameters.bCanLoop)
+	{
+		// The section isn't looping, so we can just compile the sub-sequence into the parent sequence's compiled data.
+		FGatherParameters SubParams = Params.CreateForSubGatherData(SubGatherData.GetValue());
+		GatherCompileOnTheFlyData(*SubSequence, SubParams, OutData);
+	}
+	else
+	{
+		// The section is looping so we need to compile it as many times as we have loops.
+		const FMovieSceneSequenceTransform SequenceToRootTransform = Params.RootToSequenceTransform.InverseFromWarp(Params.SequenceLoopCounter);
+		const float RootToSubSequenceTimeScale = SubGatherData.GetValue().RootToSequenceTransform.GetTimeScale();
+		const float SubSequenceToRootTimeScale = (RootToSubSequenceTimeScale != 0.f) ? 1.0f / RootToSubSequenceTimeScale : 1.f;
+		
+		const TRange<FFrameNumber> SubSequencePlaybackRange = SubSequence->GetMovieScene()->GetPlaybackRange();
+		const FFrameNumber SubSequenceLength = SubSequencePlaybackRange.Size<FFrameNumber>();
+
+		const FFrameTime RootLoopLength = FMath::Max(
+			FFrameTime(SubSequenceLength - SubSection.Parameters.StartFrameOffset - SubSection.Parameters.EndFrameOffset) * SubSequenceToRootTimeScale,
+			FFrameTime(1));
+		const FFrameTime FirstRootLoopLength = FMath::Max(
+			RootLoopLength - SubSection.Parameters.FirstLoopStartFrameOffset * SubSequenceToRootTimeScale,
+			FFrameTime(1));
+
+		const TOptional<FFrameNumber> LocalSectionEndTime = GetLoopingSubSectionEndTime(InOuterSequence, SubSection, Params);
+		if (LocalSectionEndTime.IsSet() && !SubSection.SectionRange.GetLowerBound().IsOpen())
+		{
+			uint32 LoopCount = 0;
+			FFrameTime CurRootRangeStart = MovieScene::DiscreteInclusiveLower(SubSection.SectionRange.GetLowerBound()) * SequenceToRootTransform;
+			TRange<FFrameNumber> CurRootRange(CurRootRangeStart.FloorToFrame(), (CurRootRangeStart + FirstRootLoopLength).FloorToFrame());
+			const FFrameNumber RootSectionEndTime = (LocalSectionEndTime.GetValue() * SequenceToRootTransform).FloorToFrame();
+
+			while (CurRootRange.GetLowerBoundValue() < RootSectionEndTime)
+			{
+				if (CurRootRange.Overlaps(Params.RootCompileRange) && CurRootRange.Overlaps(Params.RootClampRange))
+				{
+					FGatherParameters CurLoopParams = Params.CreateForSubGatherData(
+							SubGatherData.GetValue(), 
+							Params.RootCompileRange, CurRootRange,
+							LoopCount);
+					GatherCompileOnTheFlyData(*SubSequence, CurLoopParams, OutData);
+				}
+
+				CurRootRangeStart = CurRootRange.GetUpperBoundValue();
+				CurRootRange = TRange<FFrameNumber>(CurRootRangeStart.FloorToFrame(), (CurRootRangeStart + RootLoopLength).FloorToFrame());
+				if (CurRootRange.GetUpperBoundValue() > RootSectionEndTime)
+					CurRootRange.SetUpperBoundValue(RootSectionEndTime);
+				++LoopCount;
+			}
+		}
+		// Faced with the cosmic horror or infinites, we choose to shield our sanity and skip this sub-section.
+		// (it either has an open-ended start time, which means we needed to loop since before time began, which means
+		//  we don't know where loops are in the present... or it means the section and root sequence have open-ended
+		//  end times, which means we would need to compile loops forever)
+		// TODO-ludovic: error reporting?
+	}
 
 	// Pop the path off the root path
 	Params.RootPath.Pop();
 }
 
+TOptional<FFrameNumber> FMovieSceneCompiler::GetLoopingSubSectionEndTime(const UMovieSceneSequence& InRootSequence, const UMovieSceneSubSection& SubSection, const FGatherParameters& Params)
+{
+	TRangeBound<FFrameNumber> SectionRangeEnd = SubSection.SectionRange.GetUpperBound();
+	if (!SectionRangeEnd.IsOpen())
+	{
+		return MovieScene::DiscreteExclusiveUpper(SectionRangeEnd);
+	}
+
+	// This section is open ended... we don't want to compile its sub-sequence in an infinite loop so we'll bound
+	// that by the playback end of is own sequence.
+	if (const UMovieScene* MovieScene = InRootSequence.GetMovieScene())
+	{
+		const TRange<FFrameNumber> PlaybackRange = MovieScene->GetPlaybackRange();
+		if (!PlaybackRange.GetUpperBound().IsOpen())
+		{
+			return MovieScene::DiscreteExclusiveUpper(PlaybackRange.GetUpperBound());
+		}
+	}
+
+	// Sadly, the root sequence is also open ended, so we effectively would need to loop the sub-sequence
+	// indefinitely... we don't support that yet.
+	return TOptional<FFrameNumber>();
+}
+
 const FMovieSceneSubSequenceData* FMovieSceneCompiler::GetOrCreateSubSequenceData(FMovieSceneSequenceID InnerSequenceID, FMovieSceneSequenceID ParentSequenceID, const UMovieSceneSubSection& SubSection, const FGuid& InObjectBindingId, FMovieSceneSequenceHierarchy& InOutHierarchy)
 {
 	// Find/add sub data in the root template
@@ -529,8 +636,12 @@ const FMovieSceneSubSequenceData* FMovieSceneCompiler::GetOrCreateSubSequenceDat
 	const FMovieSceneSubSequenceData* ParentSubData = ParentSequenceID != MovieSceneSequenceID::Root ? InOutHierarchy.FindSubData(ParentSequenceID) : nullptr;
 	if (ParentSubData)
 	{
-		TRange<FFrameNumber> ParentPlayRangeChildSpace = ParentSubData->PlayRange.Value * NewSubData.RootToSequenceTransform;
-		NewSubData.PlayRange = TRange<FFrameNumber>::Intersection(ParentPlayRangeChildSpace, NewSubData.PlayRange.Value);
+		if (!NewSubData.RootToSequenceTransform.IsWarping())
+		{
+			TRange<FFrameNumber> ParentPlayRangeChildSpace = ParentSubData->PlayRange.Value * NewSubData.RootToSequenceTransform.LinearTransform;
+			NewSubData.PlayRange = TRange<FFrameNumber>::Intersection(ParentPlayRangeChildSpace, NewSubData.PlayRange.Value);
+		}
+		// else: the sub-sequence is looping so we'll probably need the whole playback range.
 
 		// Accumulate parent transform
 		NewSubData.RootToSequenceTransform = NewSubData.RootToSequenceTransform * ParentSubData->RootToSequenceTransform;
@@ -552,10 +663,10 @@ void FMovieSceneCompiler::GatherCompileDataForTrack(FMovieSceneEvaluationTrack&
 		return Track.HasChildTemplate(EvalData.ImplIndex) && Track.GetChildTemplate(EvalData.ImplIndex).RequiresInitialization();
 	};
 
-	FMovieSceneSequenceTransform SequenceToRootTransform  = Params.RootToSequenceTransform.Inverse();
+	FMovieSceneSequenceTransform SequenceToRootTransform  = Params.RootToSequenceTransform.InverseFromWarp(Params.SequenceLoopCounter);
 	FMovieSceneSequenceID        CurrentSequenceID        = Params.RootPath.Remap(MovieSceneSequenceID::Root);
 	TRange<FFrameNumber>         CompileClampIntersection = TRange<FFrameNumber>::Intersection(Params.LocalCompileRange, Params.LocalClampRange);
-
+	
 	FMovieSceneEvaluationTreeRangeIterator TrackIter = Track.IterateFrom(CompileClampIntersection.GetLowerBound());
 	for ( ; TrackIter && TrackIter.Range().Overlaps(CompileClampIntersection); ++TrackIter)
 	{
@@ -563,7 +674,7 @@ void FMovieSceneCompiler::GatherCompileDataForTrack(FMovieSceneEvaluationTrack&
 		if (!SegmentID.IsValid())
 		{
 			// No segment at this time, so just report the time range of the empty space.
-			TRange<FFrameNumber> ClampedEmptyTrackSpaceRoot = Params.ClampRoot(TrackIter.Range() * SequenceToRootTransform);
+			TRange<FFrameNumber> ClampedEmptyTrackSpaceRoot = Params.ClampRoot(TrackIter.Range() * SequenceToRootTransform.LinearTransform);
 			OutData.EmptySpace.AddTimeRange(ClampedEmptyTrackSpaceRoot);
 		}
 		else
@@ -579,7 +690,7 @@ void FMovieSceneCompiler::GatherCompileDataForTrack(FMovieSceneEvaluationTrack&
 			Data.bRequiresInit = ThisSegment.Impls.ContainsByPredicate(RequiresInit);
 
 			TRange<FFrameNumber> SegmentTrackIntersection = TRange<FFrameNumber>::Intersection(ThisSegment.Range, TrackIter.Range());
-			TRange<FFrameNumber> IntersectionRange        = Params.ClampRoot(SegmentTrackIntersection * SequenceToRootTransform);
+			TRange<FFrameNumber> IntersectionRange        = Params.ClampRoot(SegmentTrackIntersection * SequenceToRootTransform.LinearTransform);
 			if (!IntersectionRange.IsEmpty())
 			{
 				OutData.Tracks.Add(IntersectionRange, Data);
diff --git a/Engine/Source/Runtime/MovieScene/Private/Evaluation/MovieSceneEvaluationField.cpp b/Engine/Source/Runtime/MovieScene/Private/Evaluation/MovieSceneEvaluationField.cpp
--- a/Engine/Source/Runtime/MovieScene/Private/Evaluation/MovieSceneEvaluationField.cpp
+++ b/Engine/Source/Runtime/MovieScene/Private/Evaluation/MovieSceneEvaluationField.cpp
@@ -356,8 +356,25 @@ bool FMovieSceneEvaluationMetaData::IsDirty(const FMovieSceneSequenceHierarchy&
 
 			if (OutSubRangeToInvalidate)
 			{
-				TRange<FFrameNumber> DirtyRange = SubData ? TRange<FFrameNumber>::Hull(TRange<FFrameNumber>::Hull(SubData->PreRollRange.Value, SubData->PlayRange.Value), SubData->PostRollRange.Value) * SubData->RootToSequenceTransform.Inverse() : TRange<FFrameNumber>::All();
-				*OutSubRangeToInvalidate = TRange<FFrameNumber>::Hull(*OutSubRangeToInvalidate, DirtyRange);
+				if (SubData)
+				{
+					if (!SubData->RootToSequenceTransform.IsWarping())
+					{
+						TRange<FFrameNumber> FullSubPlayRange = TRange<FFrameNumber>::Hull(TRange<FFrameNumber>::Hull(SubData->PreRollRange.Value, SubData->PlayRange.Value), SubData->PostRollRange.Value);
+						TRange<FFrameNumber> DirtyRange = FullSubPlayRange * SubData->RootToSequenceTransform.InverseLinearOnly();
+						*OutSubRangeToInvalidate = TRange<FFrameNumber>::Hull(*OutSubRangeToInvalidate, DirtyRange);
+					}
+					else
+					{
+						TRange<FFrameNumber> DirtyRange = TRange<FFrameNumber>::All();
+						*OutSubRangeToInvalidate = TRange<FFrameNumber>::Hull(*OutSubRangeToInvalidate, DirtyRange);
+					}
+				}
+				else
+				{
+					TRange<FFrameNumber> DirtyRange = TRange<FFrameNumber>::All();
+					*OutSubRangeToInvalidate = TRange<FFrameNumber>::Hull(*OutSubRangeToInvalidate, DirtyRange);
+				}
 			}
 		}
 	}
diff --git a/Engine/Source/Runtime/MovieScene/Private/Evaluation/MovieSceneSequenceHierarchy.cpp b/Engine/Source/Runtime/MovieScene/Private/Evaluation/MovieSceneSequenceHierarchy.cpp
--- a/Engine/Source/Runtime/MovieScene/Private/Evaluation/MovieSceneSequenceHierarchy.cpp
+++ b/Engine/Source/Runtime/MovieScene/Private/Evaluation/MovieSceneSequenceHierarchy.cpp
@@ -29,22 +29,40 @@ FMovieSceneSubSequenceData::FMovieSceneSubSequenceData(const UMovieSceneSubSecti
 	checkf(MovieScenePtr, TEXT("Attempting to construct sub sequence data with a null sequence."));
 
 	TickResolution = MovieScenePtr->GetTickResolution();
+	FullPlayRange.Value = MovieScenePtr->GetPlaybackRange();
 
-	TRange<FFrameNumber> SubRange = InSubSection.GetTrueRange();
-	checkf(SubRange.GetLowerBound().IsClosed() && SubRange.GetUpperBound().IsClosed(), TEXT("Use of open (infinite) bounds with sub sections is not supported."));
+	TRange<FFrameNumber> SubSectionRange = InSubSection.GetTrueRange();
+	checkf(SubSectionRange.GetLowerBound().IsClosed() && SubSectionRange.GetUpperBound().IsClosed(), TEXT("Use of open (infinite) bounds with sub sections is not supported."));
 
+	// Get the transform from the given section to its inner sequence.
+	// Note that FMovieSceneCompiler will accumulate RootToSequenceTransform for us a bit later so that it ends up
+	// being truly the full transform.
 	OuterToInnerTransform = RootToSequenceTransform = InSubSection.OuterToInnerTransform();
 
-	PlayRange.Value = SubRange * RootToSequenceTransform;
+	if (!InSubSection.Parameters.bCanLoop)
+	{
+		PlayRange.Value = SubSectionRange * RootToSequenceTransform.LinearTransform;
+		UnwarpedPlayRange.Value = PlayRange.Value;
+	}
+	else
+	{
+		// If we're looping, there's a good chance we need the entirety of the sub-sequence to be compiled.
+		PlayRange.Value = MovieScenePtr->GetPlaybackRange();
+		PlayRange.Value.SetLowerBoundValue(PlayRange.Value.GetLowerBoundValue() + InSubSection.Parameters.StartFrameOffset);
+		PlayRange.Value.SetUpperBoundValue(FMath::Max(
+			PlayRange.Value.GetUpperBoundValue() - InSubSection.Parameters.EndFrameOffset,
+			PlayRange.Value.GetLowerBoundValue() + 1));
+		UnwarpedPlayRange.Value = RootToSequenceTransform.TransformRangeUnwarped(SubSectionRange);
+	}
 
 	// Make sure pre/postroll *ranges* are in the inner sequence's time space. Pre/PostRollFrames are in the outer sequence space.
 	if (InSubSection.GetPreRollFrames() > 0)
 	{
-		PreRollRange = MovieScene::MakeDiscreteRangeFromUpper( TRangeBound<FFrameNumber>::FlipInclusion(SubRange.GetLowerBound()), InSubSection.GetPreRollFrames() ) * RootToSequenceTransform;
+		PreRollRange = MovieScene::MakeDiscreteRangeFromUpper( TRangeBound<FFrameNumber>::FlipInclusion(SubSectionRange.GetLowerBound()), InSubSection.GetPreRollFrames() ) * RootToSequenceTransform.LinearTransform;
 	}
 	if (InSubSection.GetPostRollFrames() > 0)
 	{
-		PostRollRange = MovieScene::MakeDiscreteRangeFromLower( TRangeBound<FFrameNumber>::FlipInclusion(SubRange.GetUpperBound()), InSubSection.GetPostRollFrames() ) * RootToSequenceTransform;
+		PostRollRange = MovieScene::MakeDiscreteRangeFromLower( TRangeBound<FFrameNumber>::FlipInclusion(SubSectionRange.GetUpperBound()), InSubSection.GetPostRollFrames() ) * RootToSequenceTransform.LinearTransform;
 	}
 }
 
diff --git a/Engine/Source/Runtime/MovieScene/Private/Sections/MovieSceneSubSection.cpp b/Engine/Source/Runtime/MovieScene/Private/Sections/MovieSceneSubSection.cpp
--- a/Engine/Source/Runtime/MovieScene/Private/Sections/MovieSceneSubSection.cpp
+++ b/Engine/Source/Runtime/MovieScene/Private/Sections/MovieSceneSubSection.cpp
@@ -30,7 +30,7 @@ FMovieSceneSequenceTransform UMovieSceneSubSection::OuterToInnerTransform() cons
 		return FMovieSceneSequenceTransform();
 	}
 
-	UMovieScene*         MovieScenePtr = SequencePtr->GetMovieScene();
+	UMovieScene* MovieScenePtr = SequencePtr->GetMovieScene();
 
 	TRange<FFrameNumber> SubRange = GetRange();
 	if (!MovieScenePtr || SubRange.GetLowerBound().IsOpen())
@@ -38,21 +38,73 @@ FMovieSceneSequenceTransform UMovieSceneSubSection::OuterToInnerTransform() cons
 		return FMovieSceneSequenceTransform();
 	}
 
-	const FFrameNumber InnerStartTime = MovieScene::DiscreteInclusiveLower(MovieScenePtr->GetPlaybackRange()) + Parameters.StartFrameOffset;
-	const FFrameNumber OuterStartTime = MovieScene::DiscreteInclusiveLower(SubRange);
-
 	const FFrameRate   InnerFrameRate = MovieScenePtr->GetTickResolution();
 	const FFrameRate   OuterFrameRate = GetTypedOuter<UMovieScene>()->GetTickResolution();
-
 	const float        FrameRateScale = (OuterFrameRate == InnerFrameRate) ? 1.f : (InnerFrameRate / OuterFrameRate).AsDecimal();
 
-	return
+	const FFrameNumber StartOffset = Parameters.StartFrameOffset;
+	const FFrameNumber EndOffset = Parameters.bCanLoop ? Parameters.EndFrameOffset : 0;
+
+	const TRange<FFrameNumber> MovieScenePlaybackRange = GetValidatedInnerPlaybackRange(Parameters, *MovieScenePtr);
+	const FFrameNumber InnerStartTime = MovieScene::DiscreteInclusiveLower(MovieScenePlaybackRange);
+	const FFrameNumber OuterStartTime = MovieScene::DiscreteInclusiveLower(SubRange);
+
+	// This is the transform for the "placement" (position and scaling) of the sub-sequence.
+	FMovieSceneTimeTransform LinearTransform =
 		// Inner play offset
-		FMovieSceneSequenceTransform(InnerStartTime)
+		FMovieSceneTimeTransform(InnerStartTime)
 		// Inner play rate
-		* FMovieSceneSequenceTransform(0, Parameters.TimeScale * FrameRateScale)
+		* FMovieSceneTimeTransform(0, Parameters.TimeScale * FrameRateScale)
 		// Outer section start time
-		* FMovieSceneSequenceTransform(-OuterStartTime);
+		* FMovieSceneTimeTransform(-OuterStartTime);
+	
+	if (!Parameters.bCanLoop)
+	{
+		return FMovieSceneSequenceTransform(LinearTransform);
+	}
+	else
+	{
+		const FFrameNumber InnerEndTime = MovieScene::DiscreteExclusiveUpper(MovieScenePlaybackRange);
+		const FMovieSceneTimeWarping LoopingTransform(InnerStartTime, InnerEndTime);
+		LinearTransform = FMovieSceneTimeTransform(Parameters.FirstLoopStartFrameOffset) * LinearTransform;
+
+		FMovieSceneSequenceTransform Result;
+		Result.NestedTransforms.Add(FMovieSceneNestedSequenceTransform(LinearTransform, LoopingTransform));
+		return Result;
+	}
+}
+
+bool UMovieSceneSubSection::GetValidatedInnerPlaybackRange(TRange<FFrameNumber>& OutInnerPlaybackRange) const
+{
+	UMovieSceneSequence* SequencePtr = GetSequence();
+	if (SequencePtr != nullptr)
+	{
+		UMovieScene* MovieScenePtr = SequencePtr->GetMovieScene();
+		if (MovieScenePtr != nullptr)
+		{
+			OutInnerPlaybackRange = GetValidatedInnerPlaybackRange(Parameters, *MovieScenePtr);
+			return true;
+		}
+	}
+	return false;
+}
+
+TRange<FFrameNumber> UMovieSceneSubSection::GetValidatedInnerPlaybackRange(const FMovieSceneSectionParameters& SubSectionParameters, const UMovieScene& InnerMovieScene)
+{
+	const TRange<FFrameNumber> InnerPlaybackRange = InnerMovieScene.GetPlaybackRange();
+	TRangeBound<FFrameNumber> ValidatedLowerBound = InnerPlaybackRange.GetLowerBound();
+	TRangeBound<FFrameNumber> ValidatedUpperBound = InnerPlaybackRange.GetUpperBound();
+	if (ValidatedLowerBound.IsClosed() && ValidatedUpperBound.IsClosed())
+	{
+		const FFrameRate TickResolution = InnerMovieScene.GetTickResolution();
+		const FFrameRate DisplayRate = InnerMovieScene.GetDisplayRate();
+		const FFrameNumber OneFrameInTicks = FFrameRate::TransformTime(FFrameTime(1), DisplayRate, TickResolution).FloorToFrame();
+
+		ValidatedLowerBound.SetValue(ValidatedLowerBound.GetValue() + SubSectionParameters.StartFrameOffset);
+		ValidatedUpperBound.SetValue(FMath::Max(ValidatedUpperBound.GetValue() - SubSectionParameters.EndFrameOffset, ValidatedLowerBound.GetValue() + OneFrameInTicks));
+		return TRange<FFrameNumber>(ValidatedLowerBound, ValidatedUpperBound);
+	}
+	return InnerPlaybackRange;
 }
 
 FString UMovieSceneSubSection::GetPathNameInMovieScene() const
@@ -264,7 +316,9 @@ TOptional<TRange<FFrameNumber> > UMovieSceneSubSection::GetAutoSizeRange() const
 {
 	if (SubSequence && SubSequence->GetMovieScene())
 	{
-		FMovieSceneSequenceTransform InnerToOuter = OuterToInnerTransform().Inverse();
+		// We probably want to just auto-size the section to the sub-sequence's scaled playback range... if this section
+		// is looping, however, it's hard to know what we want to do.
+		FMovieSceneTimeTransform InnerToOuter = OuterToInnerTransform().InverseLinearOnly();
 		UMovieScene* InnerMovieScene = SubSequence->GetMovieScene();
 
 		FFrameTime IncAutoStartTime = FFrameTime(MovieScene::DiscreteInclusiveLower(InnerMovieScene->GetPlaybackRange())) * InnerToOuter;
diff --git a/Engine/Source/Runtime/MovieScene/Private/Tests/MovieSceneCompilerTests.cpp b/Engine/Source/Runtime/MovieScene/Private/Tests/MovieSceneCompilerTests.cpp
--- a/Engine/Source/Runtime/MovieScene/Private/Tests/MovieSceneCompilerTests.cpp
+++ b/Engine/Source/Runtime/MovieScene/Private/Tests/MovieSceneCompilerTests.cpp
@@ -331,4 +331,90 @@ bool FMovieSceneCompilerEmptySpaceOnTheFlyTest::RunTest(const FString& Parameter
 
 	return true;
 }
-#endif // WITH_DEV_AUTOMATION_TESTS
\ No newline at end of file
+
+
+#define _TEST_ASSERT(cond, msgfmt, ...)\
+	if (!ensure((cond)))\
+	{\
+		AddError(FString::Printf(TEXT(msgfmt), __VA_ARGS__));\
+		return false;\
+	}
+
+// This should use TestEqual methods on the test if those methods actually returned the true/false outcome...
+#define _TEST_ASSERT_EQUAL(actual, expected)\
+	if (!ensure(actual == expected))\
+	{\
+		AddError(FString::Printf(TEXT("Expected '%s' but actual value was '%s'."), *LexToString(expected), *LexToString(actual)));\
+		return false;\
+	}
+
+IMPLEMENT_SIMPLE_AUTOMATION_TEST(FMovieSceneCompilerSubSequencesTest, "System.Engine.Sequencer.Compiler.SubSequences", EAutomationTestFlags::EditorContext | EAutomationTestFlags::EngineFilter)
+bool FMovieSceneCompilerSubSequencesTest::RunTest(const FString& Parameters)
+{
+	struct FTest
+	{
+		TArray<TRange<FFrameNumber>> CompileRanges;
+	};
+
+	struct FTemplateStore : public IMovieSceneSequenceTemplateStore
+	{
+		FTemplateStore(const TArray<UTestMovieSceneSequence*>& InExpectedSequences)
+			: ExpectedSequences(InExpectedSequences)
+		{
+			Templates.SetNum(InExpectedSequences.Num());
+		}
+
+		FMovieSceneEvaluationTemplate& AccessTemplate(UMovieSceneSequence& InSequence) override
+		{
+			for (size_t i = 0; i < ExpectedSequences.Num(); ++i)
+			{
+				if (ExpectedSequences[i] == &InSequence)
+				{
+					return Templates[i];
+				}
+			}
+			check(false);
+			return Templates[0];
+		}
+
+	private:
+		TArray<UTestMovieSceneSequence*> ExpectedSequences;
+		TArray<FMovieSceneEvaluationTemplate> Templates;
+	};
+
+	UTestMovieSceneSequence* RootSequence = NewObject<UTestMovieSceneSequence>(GetTransientPackage());
+	UTestMovieSceneSubTrack* RootSubTrack = RootSequence->MovieScene->AddMasterTrack<UTestMovieSceneSubTrack>();
+	
+	UTestMovieSceneSequence* Shot1Sequence = NewObject<UTestMovieSceneSequence>(GetTransientPackage());
+	Shot1Sequence->GetMovieScene()->SetPlaybackRange(0, 100);
+	
+	UTestMovieSceneSubSection* Shot1SubSection = NewObject<UTestMovieSceneSubSection>(RootSubTrack);
+	Shot1SubSection->SetRange(TRange<FFrameNumber>(0, 100));
+	Shot1SubSection->SetSequence(Shot1Sequence);
+	RootSubTrack->SectionArray.Add(Shot1SubSection);
+	
+	UTestMovieSceneTrack* Shot1Track = Shot1Sequence->MovieScene->AddMasterTrack<UTestMovieSceneTrack>();
+	UTestMovieSceneSection* Shot1Section = NewObject<UTestMovieSceneSection>(Shot1Track);
+	Shot1Section->SetRange(TRange<FFrameNumber>(0, 60));
+	Shot1Track->SectionArray.Add(Shot1Section);
+	
+	FMovieSceneSequenceID Shot1SequenceID = Shot1SubSection->GetSequenceID();
+
+	{
+		FTemplateStore Store({ RootSequence, Shot1Sequence });
+		FMovieSceneCompiler::Compile(*RootSequence, Store);
+
+		const FMovieSceneEvaluationTemplate& RootTemplate = Store.AccessTemplate(*RootSequence);
+		_TEST_ASSERT_EQUAL(RootTemplate.EvaluationField.GetRanges().Num(), 3);
+		_TEST_ASSERT_EQUAL(RootTemplate.EvaluationField.GetRange(1), TRange<FFrameNumber>(0, 60));
+		const FMovieSceneEvaluationGroup& EvalGroup = RootTemplate.EvaluationField.GetGroup(1);
+		_TEST_ASSERT_EQUAL(EvalGroup.SegmentPtrLUT[0].SequenceID, Shot1SequenceID);
+	}
+
+	return true;
+}
+
+#undef _TEST_ASSERT
+#undef _TEST_ASSERT_EQUAL
+
+#endif // WITH_DEV_AUTOMATION_TESTS
diff --git a/Engine/Source/Runtime/MovieScene/Private/Tests/MovieSceneSegmentCompilerTests.cpp b/Engine/Source/Runtime/MovieScene/Private/Tests/MovieSceneSegmentCompilerTests.cpp
--- a/Engine/Source/Runtime/MovieScene/Private/Tests/MovieSceneSegmentCompilerTests.cpp
+++ b/Engine/Source/Runtime/MovieScene/Private/Tests/MovieSceneSegmentCompilerTests.cpp
@@ -604,45 +604,3 @@ bool FMovieSceneCompilerTreeIteratorBoundsTest::RunTest(const FString& Parameter
 }
 
 #endif // WITH_DEV_AUTOMATION_TESTS
-
-
-
-FMovieSceneEvalTemplatePtr UTestMovieSceneTrack::CreateTemplateForSection(const UMovieSceneSection& InSection) const
-{
-	return FTestMovieSceneEvalTemplate();
-}
-
-FMovieSceneTrackSegmentBlenderPtr UTestMovieSceneTrack::GetTrackSegmentBlender() const
-{
-	struct FSegmentBlender : FMovieSceneTrackSegmentBlender
-	{
-		bool bHighPass;
-		bool bEvaluateNearest;
-
-		FSegmentBlender(bool bInHighPassFilter, bool bInEvaluateNearest)
-			: bHighPass(bInHighPassFilter)
-			, bEvaluateNearest(bInEvaluateNearest)
-		{
-			bCanFillEmptySpace = bEvaluateNearest;
-		}
-
-		virtual void Blend(FSegmentBlendData& BlendData) const
-		{
-			if (bHighPass)
-			{
-				MovieSceneSegmentCompiler::ChooseLowestRowIndex(BlendData);
-			}
-
-			// Always sort by priority
-			Algo::SortBy(BlendData, [](const FMovieSceneSectionData& In){ return In.Section->GetRowIndex(); });
-		}
-
-		virtual TOptional<FMovieSceneSegment> InsertEmptySpace(const TRange<FFrameNumber>& Range, const FMovieSceneSegment* PreviousSegment, const FMovieSceneSegment* NextSegment) const
-		{
-			return bEvaluateNearest ? MovieSceneSegmentCompiler::EvaluateNearestSegment(Range, PreviousSegment, NextSegment) : TOptional<FMovieSceneSegment>();
-		}
-	};
-
-	// Evaluate according to bEvalNearestSection preference
-	return FSegmentBlender(bHighPassFilter, EvalOptions.bCanEvaluateNearestSection && EvalOptions.bEvalNearestSection);
-}
diff --git a/Engine/Source/Runtime/MovieScene/Private/Tests/MovieSceneTestObjects.cpp b/Engine/Source/Runtime/MovieScene/Private/Tests/MovieSceneTestObjects.cpp
new file mode 100644
--- /dev/null
+++ b/Engine/Source/Runtime/MovieScene/Private/Tests/MovieSceneTestObjects.cpp
@@ -0,0 +1,46 @@
+// Copyright 1998-2019 Epic Games, Inc. All Rights Reserved.
+
+#include "MovieSceneTestObjects.h"
+#include "Compilation/MovieSceneCompiler.h"
+#include "Compilation/MovieSceneCompilerRules.h"
+#include "Compilation/MovieSceneSegmentCompiler.h"
+
+FMovieSceneEvalTemplatePtr UTestMovieSceneTrack::CreateTemplateForSection(const UMovieSceneSection& InSection) const
+{
+	return FTestMovieSceneEvalTemplate();
+}
+
+FMovieSceneTrackSegmentBlenderPtr UTestMovieSceneTrack::GetTrackSegmentBlender() const
+{
+	struct FSegmentBlender : FMovieSceneTrackSegmentBlender
+	{
+		bool bHighPass;
+		bool bEvaluateNearest;
+
+		FSegmentBlender(bool bInHighPassFilter, bool bInEvaluateNearest)
+			: bHighPass(bInHighPassFilter)
+			, bEvaluateNearest(bInEvaluateNearest)
+		{
+			bCanFillEmptySpace = bEvaluateNearest;
+		}
+
+		virtual void Blend(FSegmentBlendData& BlendData) const
+		{
+			if (bHighPass)
+			{
+				MovieSceneSegmentCompiler::ChooseLowestRowIndex(BlendData);
+			}
+
+			// Always sort by priority
+			Algo::SortBy(BlendData, [](const FMovieSceneSectionData& In) { return In.Section->GetRowIndex(); });
+		}
+
+		virtual TOptional<FMovieSceneSegment> InsertEmptySpace(const TRange<FFrameNumber>& Range, const FMovieSceneSegment* PreviousSegment, const FMovieSceneSegment* NextSegment) const
+		{
+			return bEvaluateNearest ? MovieSceneSegmentCompiler::EvaluateNearestSegment(Range, PreviousSegment, NextSegment) : TOptional<FMovieSceneSegment>();
+		}
+	};
+
+	// Evaluate according to bEvalNearestSection preference
+	return FSegmentBlender(bHighPassFilter, EvalOptions.bCanEvaluateNearestSection && EvalOptions.bEvalNearestSection);
+}
diff --git a/Engine/Source/Runtime/MovieScene/Private/Tests/MovieSceneTestObjects.h b/Engine/Source/Runtime/MovieScene/Private/Tests/MovieSceneTestObjects.h
--- a/Engine/Source/Runtime/MovieScene/Private/Tests/MovieSceneTestObjects.h
+++ b/Engine/Source/Runtime/MovieScene/Private/Tests/MovieSceneTestObjects.h
@@ -7,6 +7,8 @@
 #include "MovieSceneSection.h"
 #include "MovieSceneSequence.h"
 #include "Evaluation/MovieSceneEvalTemplate.h"
+#include "Sections/MovieSceneSubSection.h"
+#include "Tracks/MovieSceneSubTrack.h"
 #include "MovieSceneTestObjects.generated.h"
 
 USTRUCT()
@@ -62,3 +64,22 @@ public:
 	UPROPERTY()
 	UMovieScene* MovieScene;
 };
+
+UCLASS(MinimalAPI)
+class UTestMovieSceneSubTrack : public UMovieSceneSubTrack
+{
+public:
+	GENERATED_BODY()
+
+	virtual const TArray<UMovieSceneSection*>& GetAllSections() const override { return SectionArray; }
+
+	UPROPERTY()
+	TArray<UMovieSceneSection*> SectionArray;
+};
+
+UCLASS(MinimalAPI)
+class UTestMovieSceneSubSection : public UMovieSceneSubSection
+{
+	GENERATED_BODY()
+};
+
diff --git a/Engine/Source/Runtime/MovieScene/Private/Tests/MovieSceneTransformTests.cpp b/Engine/Source/Runtime/MovieScene/Private/Tests/MovieSceneTransformTests.cpp
--- a/Engine/Source/Runtime/MovieScene/Private/Tests/MovieSceneTransformTests.cpp
+++ b/Engine/Source/Runtime/MovieScene/Private/Tests/MovieSceneTransformTests.cpp
@@ -9,6 +9,7 @@
 
 #define LOCTEXT_NAMESPACE "MovieSceneTransformTests"
 
+// Range equality.
 bool IsEqual(TRangeBound<FFrameNumber> A, TRangeBound<FFrameNumber> B)
 {
 	if (A.IsOpen() || B.IsOpen())
@@ -23,12 +24,66 @@ bool IsEqual(TRangeBound<FFrameNumber> A, TRangeBound<FFrameNumber> B)
 	return A.GetValue() == B.GetValue();
 }
 
+// Range equality.
 bool IsEqual(TRange<FFrameNumber> A, TRange<FFrameNumber> B)
 {
 	return IsEqual(A.GetLowerBound(), B.GetLowerBound()) && IsEqual(A.GetUpperBound(), B.GetUpperBound());
 }
 
-bool TestTransform(FAutomationTestBase& Test, FMovieSceneSequenceTransform Transform, TArrayView<TRange<FFrameNumber>> InSource, TArrayView<TRange<FFrameNumber>> InExpected, const TCHAR* TestName)
+// Frame number equality.
+bool IsEqual(FFrameNumber A, FFrameNumber B)
+{
+	return A.Value == B.Value;
+}
+
+// Frame time equality.
+bool IsEqual(FFrameTime A, FFrameTime B)
+{
+	return IsEqual(A.FrameNumber, B.FrameNumber) && FMath::IsNearlyEqual(A.GetSubFrame(), B.GetSubFrame());
+}
+
+// Most time transformations are not "round" so they return a frame time that must be rounded down to a frame number,
+// except for time warping which doesn't stretch anything and returns a frame number.
+template<typename TTransform>
+FFrameNumber TransformToFrameNumber(TTransform Transform, FFrameNumber Value)
+{
+	return (Value * Transform).FloorToFrame();
+}
+template<>
+FFrameNumber TransformToFrameNumber(FMovieSceneTimeWarping Transform, FFrameNumber Value)
+{
+	return Value * Transform;
+}
+
+// Generic method for testing the transform of frames and times.
+template<typename TTransform>
+bool TestTransform(FAutomationTestBase& Test, TTransform Transform, TArrayView<FFrameNumber> InSource, TArrayView<FFrameNumber> InExpected, const TCHAR* TestName)
+{
+	check(InSource.Num() == InExpected.Num());
+
+	bool bSuccess = true;
+	for (int32 Index = 0; Index < InSource.Num(); ++Index)
+	{
+		FFrameNumber Result = TransformToFrameNumber(Transform, InSource[Index]);
+		if (!IsEqual(Result, InExpected[Index]))
+		{
+			Test.AddError(FString::Printf(TEXT("Test '%s' failed (Index %d). Transform %s did not apply correctly (%s != %s)"),
+				TestName,
+				Index,
+				*LexToString(Transform),
+				*LexToString(Result),
+				*LexToString(InExpected[Index])));
+
+			bSuccess = false;
+		}
+	}
+
+	return bSuccess;
+}
+
+// A variant of the above method for testing the transform of ranges.
+template<typename TTransform>
+bool TestTransform(FAutomationTestBase& Test, TTransform Transform, TArrayView<TRange<FFrameNumber>> InSource, TArrayView<TRange<FFrameNumber>> InExpected, const TCHAR* TestName)
 {
 	check(InSource.Num() == InExpected.Num());
 
@@ -38,12 +93,10 @@ bool TestTransform(FAutomationTestBase& Test, FMovieSceneSequenceTransform Trans
 		TRange<FFrameNumber> Result = InSource[Index] * Transform;
 		if (!IsEqual(Result, InExpected[Index]))
 		{
-			Test.AddError(FString::Printf(TEXT("Test '%s' failed (Index %d). Transform (Scale %.3f, Offset %i+%.3f) did not apply correctly (%s != %s)"),
+			Test.AddError(FString::Printf(TEXT("Test '%s' failed (Index %d). Transform %s did not apply correctly (%s != %s)"),
 				TestName,
 				Index,
-				Transform.TimeScale,
-				Transform.Offset.FrameNumber.Value,
-				Transform.Offset.GetSubFrame(),
+				*LexToString(Transform),
 				*LexToString(Result),
 				*LexToString(InExpected[Index])));
 
@@ -61,8 +114,139 @@ FMovieSceneSequenceTransform TransformRange(FFrameNumber StartA, FFrameNumber En
 	return FMovieSceneSequenceTransform(StartB, Scale) * FMovieSceneSequenceTransform(-StartA);
 }
 
-IMPLEMENT_SIMPLE_AUTOMATION_TEST(FMovieSceneSubSectionCoreTransformsTest, "System.Engine.Sequencer.Core.Transforms", EAutomationTestFlags::EditorContext | EAutomationTestFlags::EngineFilter)
-bool FMovieSceneSubSectionCoreTransformsTest::RunTest(const FString& Parameters)
+// Linear transform tests
+IMPLEMENT_SIMPLE_AUTOMATION_TEST(FMovieSceneSubSectionCoreLinearTransformsTest, 
+		"System.Engine.Sequencer.Core.LinearTransforms", 
+		EAutomationTestFlags::EditorContext | EAutomationTestFlags::EngineFilter)
+bool FMovieSceneSubSectionCoreLinearTransformsTest::RunTest(const FString& Parameters)
+{
+	FFrameNumber SourceTimes[] = {
+		FFrameNumber(500),
+		FFrameNumber(525)
+	};
+
+	bool bSuccess = true;
+
+	{
+		FFrameNumber ExpectedTimes[] = {
+			FFrameNumber(500),
+			FFrameNumber(525)
+		};
+		FMovieSceneTimeTransform Transform(0);
+		bSuccess = TestTransform(*this, Transform, SourceTimes, ExpectedTimes, TEXT("IdentityTransform")) && bSuccess;
+
+		Transform = Transform.Inverse();
+		bSuccess = TestTransform(*this, Transform, ExpectedTimes, SourceTimes, TEXT("IdentityTransformInverse")) && bSuccess;
+	}
+
+	{
+		FFrameNumber ExpectedTimes[] = {
+			FFrameNumber(1000),
+			FFrameNumber(1050)
+		};
+		FMovieSceneTimeTransform Transform(0, 2.f);
+		bSuccess = TestTransform(*this, Transform, SourceTimes, ExpectedTimes, TEXT("OffsetTransform")) && bSuccess;
+
+		Transform = Transform.Inverse();
+		bSuccess = TestTransform(*this, Transform, ExpectedTimes, SourceTimes, TEXT("OffsetTransformInverse")) && bSuccess;
+	}
+
+	{
+		FFrameNumber ExpectedTimes[] = {
+			FFrameNumber(0),
+			FFrameNumber(50)
+		};
+		FMovieSceneTimeTransform Transform(-1000, 2.f);
+		bSuccess = TestTransform(*this, Transform, SourceTimes, ExpectedTimes, TEXT("OffsetAndScaleTransform")) && bSuccess;
+
+		Transform = Transform.Inverse();
+		bSuccess = TestTransform(*this, Transform, ExpectedTimes, SourceTimes, TEXT("OffsetAndScaleTransformInverse")) && bSuccess;
+	}
+
+	{
+		FFrameNumber ExpectedTimes[] = {
+			FFrameNumber(0),
+			FFrameNumber(50)
+		};
+		FMovieSceneTimeTransform Transform = FMovieSceneTimeTransform(0, 2.f) * FMovieSceneTimeTransform(-500);
+		bSuccess = TestTransform(*this, Transform, SourceTimes, ExpectedTimes, TEXT("OffsetAndScaleTransformObtainedFromMultiplication")) && bSuccess;
+	}
+
+	return bSuccess;
+}
+
+// Warping transform tests
+IMPLEMENT_SIMPLE_AUTOMATION_TEST(FMovieSceneSubSectionCoreWarpTransformsTest,
+		"System.Engine.Sequencer.Core.WarpTransforms",
+		EAutomationTestFlags::EditorContext | EAutomationTestFlags::EngineFilter)
+bool FMovieSceneSubSectionCoreWarpTransformsTest::RunTest(const FString& Parameters)
+{
+	FFrameNumber SourceTimes[] = {
+		FFrameNumber(0),
+		FFrameNumber(25),
+		FFrameNumber(50),
+		FFrameNumber(60),
+		FFrameNumber(120)
+	};
+
+	bool bSuccess = true;
+
+	{
+		FFrameNumber ExpectedTimes[] = {
+			FFrameNumber(0),
+			FFrameNumber(25),
+			FFrameNumber(0),
+			FFrameNumber(10),
+			FFrameNumber(20)
+		};
+		FMovieSceneTimeWarping Warping(0, 50);
+		bSuccess = TestTransform(*this, Warping, SourceTimes, ExpectedTimes, TEXT("SimpleWarping")) && bSuccess;
+
+		FMovieSceneTimeTransform Transform = Warping.InverseFromWarp(0);
+		bSuccess = TestTransform(*this, Transform, MakeArrayView(ExpectedTimes, 2), MakeArrayView(SourceTimes, 2), TEXT("SimpleWarpingInverseLoop0")) && bSuccess;
+
+		FMovieSceneTimeTransform Transform2 = Warping.InverseFromWarp(1);
+		bSuccess = TestTransform(*this, Transform2, MakeArrayView(ExpectedTimes + 2, 2), MakeArrayView(SourceTimes + 2, 2), TEXT("SimpleWarpingInverseLoop1")) && bSuccess;
+
+		FMovieSceneTimeTransform Transform3 = Warping.InverseFromWarp(2);
+		bSuccess = TestTransform(*this, Transform3, MakeArrayView(ExpectedTimes + 4, 1), MakeArrayView(SourceTimes + 4, 1), TEXT("SimpleWarpingInverseLoop2")) && bSuccess;
+	}
+
+	SourceTimes[0] = FFrameNumber(3);
+	SourceTimes[1] = FFrameNumber(28);
+	SourceTimes[2] = FFrameNumber(53);
+	SourceTimes[3] = FFrameNumber(63);
+	SourceTimes[4] = FFrameNumber(123);
+
+	{
+		FFrameNumber ExpectedTimes[] {
+			FFrameNumber(3),
+			FFrameNumber(28),
+			FFrameNumber(14),
+			FFrameNumber(24),
+			FFrameNumber(6)
+		};
+		FMovieSceneTimeWarping Warping(3, 42);
+		bSuccess = TestTransform(*this, Warping, SourceTimes, ExpectedTimes, TEXT("WarpingWithTrim")) && bSuccess;
+
+		FMovieSceneTimeTransform Transform = Warping.InverseFromWarp(0);
+		bSuccess = TestTransform(*this, Transform, MakeArrayView(ExpectedTimes, 2), MakeArrayView(SourceTimes, 2), TEXT("WarpingWithTrimInverseLoop0")) && bSuccess;
+
+		FMovieSceneTimeTransform Transform2 = Warping.InverseFromWarp(1);
+		bSuccess = TestTransform(*this, Transform2, MakeArrayView(ExpectedTimes + 2, 2), MakeArrayView(SourceTimes + 2, 2), TEXT("WarpingWithTrimInverseLoop1")) && bSuccess;
+
+		FMovieSceneTimeTransform Transform3 = Warping.InverseFromWarp(3);  // We lapsed one full loop
+		bSuccess = TestTransform(*this, Transform3, MakeArrayView(ExpectedTimes + 4, 1), MakeArrayView(SourceTimes + 4, 1), TEXT("WarpingWithTrimInverseLoop2")) && bSuccess;
+	}
+
+	return bSuccess;
+}
+
+// Sequence transform tests 
+IMPLEMENT_SIMPLE_AUTOMATION_TEST(FMovieSceneSubSectionCoreSequenceTransformsTest,
+		"System.Engine.Sequencer.Core.SequenceTransforms",
+		EAutomationTestFlags::EditorContext | EAutomationTestFlags::EngineFilter)
+bool FMovieSceneSubSectionCoreSequenceTransformsTest::RunTest(const FString& Parameters)
 {
 	// We test using ranges since that implicitly tests frame number transformation as well
 	static const TRangeBound<FFrameNumber> OpenBound;
@@ -86,7 +270,7 @@ bool FMovieSceneSubSectionCoreTransformsTest::RunTest(const FString& Parameters)
 			InfiniteRange, OpenLowerRange, OpenUpperRange, ClosedRange
 		};
 		
-		bSuccess = TestTransform(*this, IdentityTransform, SourceRanges, Expected, TEXT("IdentityTransform")) && bSuccess;
+		bSuccess = TestTransform(*this, IdentityTransform.LinearTransform, SourceRanges, Expected, TEXT("IdentityTransform")) && bSuccess;
 	}
 
 	{
@@ -97,7 +281,7 @@ bool FMovieSceneSubSectionCoreTransformsTest::RunTest(const FString& Parameters)
 			InfiniteRange, TRange<FFrameNumber>(OpenBound, FFrameNumber(300)), TRange<FFrameNumber>(FFrameNumber(200), OpenBound), TRange<FFrameNumber>(200, 300)
 		};
 
-		bSuccess = TestTransform(*this, Transform, SourceRanges, Expected, TEXT("Simple Translation")) && bSuccess;
+		bSuccess = TestTransform(*this, Transform.LinearTransform, SourceRanges, Expected, TEXT("Simple Translation")) && bSuccess;
 	}
 
 	{
@@ -110,7 +294,7 @@ bool FMovieSceneSubSectionCoreTransformsTest::RunTest(const FString& Parameters)
 			InfiniteRange, TRange<FFrameNumber>(OpenBound, FFrameNumber(1000)), TRange<FFrameNumber>(FFrameNumber(-200), OpenBound), TRange<FFrameNumber>(-200, 1000)
 		};
 
-		bSuccess = TestTransform(*this, Transform, SourceRanges, Expected, TEXT("Simple Translation + half speed")) && bSuccess;
+		bSuccess = TestTransform(*this, Transform.LinearTransform, SourceRanges, Expected, TEXT("Simple Translation + half speed")) && bSuccess;
 	}
 
 	{
@@ -134,7 +318,7 @@ bool FMovieSceneSubSectionCoreTransformsTest::RunTest(const FString& Parameters)
 			AddError(FString::Printf(TEXT("Accumulated transform does not have the same effect as separate transformations (%i+%.5f != %i+%.5f)"), AccumValue.FrameNumber.Value, AccumValue.GetSubFrame(), SeedValue.FrameNumber.Value, SeedValue.GetSubFrame()));
 		}
 
-		FMovieSceneSequenceTransform InverseTransform = AccumulatedTransform.Inverse();
+		FMovieSceneSequenceTransform InverseTransform = AccumulatedTransform.InverseLinearOnly();
 
 		FFrameTime InverseValue = AccumValue * InverseTransform;
 		if (InverseValue != 10)
diff --git a/Engine/Source/Runtime/MovieScene/Public/Compilation/MovieSceneCompiler.h b/Engine/Source/Runtime/MovieScene/Public/Compilation/MovieSceneCompiler.h
--- a/Engine/Source/Runtime/MovieScene/Public/Compilation/MovieSceneCompiler.h
+++ b/Engine/Source/Runtime/MovieScene/Public/Compilation/MovieSceneCompiler.h
@@ -86,7 +86,9 @@ private:
 
 	static void GatherCompileOnTheFlyData(UMovieSceneSequence& InSequence, const FGatherParameters& Params, FMovieSceneGatheredCompilerData& OutData);
 
-	static void GatherCompileDataForSubSection(const UMovieSceneSubSection& SubSection, const FGuid& InObjectBindingId, const FGatherParameters& Params, FMovieSceneGatheredCompilerData& OutData);
+	static void GatherCompileDataForSubSection(const UMovieSceneSequence& InOuterSequence, const UMovieSceneSubSection& SubSection, const FGuid& InObjectBindingId, const FGatherParameters& Params, FMovieSceneGatheredCompilerData& OutData);
+
+	static TOptional<FFrameNumber> GetLoopingSubSectionEndTime(const UMovieSceneSequence& InRootSequence, const UMovieSceneSubSection& SubSection, const FGatherParameters& Params);
 
 	static void GatherCompileDataForTrack(FMovieSceneEvaluationTrack& Track, FMovieSceneTrackIdentifier TrackID, const FGatherParameters& Params, FMovieSceneGatheredCompilerData& OutData);
 
diff --git a/Engine/Source/Runtime/MovieScene/Public/Evaluation/MovieScenePlayback.h b/Engine/Source/Runtime/MovieScene/Public/Evaluation/MovieScenePlayback.h
--- a/Engine/Source/Runtime/MovieScene/Public/Evaluation/MovieScenePlayback.h
+++ b/Engine/Source/Runtime/MovieScene/Public/Evaluation/MovieScenePlayback.h
@@ -266,9 +266,46 @@ public:
 	FMovieSceneContext Transform(const FMovieSceneSequenceTransform& InTransform, FFrameRate NewFrameRate) const
 	{
 		FMovieSceneContext NewContext = *this;
-		NewContext.EvaluationRange = EvaluationRange * InTransform;
 		NewContext.RootToSequenceTransform = NewContext.RootToSequenceTransform * InTransform;
 		NewContext.CurrentFrameRate = NewFrameRate;
+
+		NewContext.EvaluationRange = InTransform.TransformRangeUnwarped(EvaluationRange);
+		if (InTransform.IsWarping())
+		{
+			// If we have some looping, the transformed range might extend past the end of a loop and into
+			// the beginning of another. In that case, technically, the evaluation range ends up being a
+			// discontinuous range... for instance: [X, LoopEnd) + [LoopStart, Y)
+			// It can be even more than 2 ranges if the loop is short enough and the original evaluation
+			// range spans more than one loop!
+			// TODO: For now, we just take the "last part" of this discontinuous range but we should really 
+			//       actually make contexts have an array of evaluation ranges.
+			TRangeBound<FFrameTime> UpperEvalutionRangeBound = NewContext.EvaluationRange.GetUpperBound();
+			const FMovieSceneNestedSequenceTransform& LeafTransform = InTransform.NestedTransforms.Last();
+			if (UpperEvalutionRangeBound.IsClosed() && LeafTransform.IsWarping())
+			{
+				const FFrameNumber LeafWarpLength = LeafTransform.Warping.Length();
+				// Below: use strictly greater than comparison so that if the evalution range ends on the
+				//        end of a loop, we correctly evaluate up to there, instead of doing a 0-width
+				//        evaluation of the first frame of the loop.
+				while (UpperEvalutionRangeBound.GetValue() > LeafTransform.Warping.End)
+				{
+					UpperEvalutionRangeBound.SetValue(UpperEvalutionRangeBound.GetValue() - LeafWarpLength);
+				}
+
+				// We can end up with the upper bound being greater than the lower bound in the case outlined
+				// in the comment above (where we overlap a loop boundary). This is where we would split things
+				// up in multiple discontinous ranges, but as previously mentioned, we only take the last part
+				// for now.
+				TRangeBound<FFrameTime> LowerEvaluationRangeBound = NewContext.EvaluationRange.GetLowerBound();
+				if (LowerEvaluationRangeBound.IsClosed() && UpperEvalutionRangeBound.GetValue() < LowerEvaluationRangeBound.GetValue())
+				{
+					LowerEvaluationRangeBound.SetValue(LeafTransform.Warping.Start);
+				}
+
+				NewContext.EvaluationRange = TRange<FFrameTime>(LowerEvaluationRangeBound, UpperEvalutionRangeBound);
+			}
+		}
+
 		return NewContext;
 	}
 
diff --git a/Engine/Source/Runtime/MovieScene/Public/Evaluation/MovieSceneSectionParameters.h b/Engine/Source/Runtime/MovieScene/Public/Evaluation/MovieSceneSectionParameters.h
--- a/Engine/Source/Runtime/MovieScene/Public/Evaluation/MovieSceneSectionParameters.h
+++ b/Engine/Source/Runtime/MovieScene/Public/Evaluation/MovieSceneSectionParameters.h
@@ -4,7 +4,6 @@
 
 #include "CoreMinimal.h"
 #include "UObject/ObjectMacros.h"
-#include "Evaluation/MovieSceneSequenceTransform.h"
 #include "MovieSceneSectionParameters.generated.h"
 
 USTRUCT(BlueprintType)
@@ -15,6 +14,8 @@ struct FMovieSceneSectionParameters
 	/** Default constructor */
 	FMovieSceneSectionParameters()
 		: StartFrameOffset(0)
+		, EndFrameOffset(0)
+		, FirstLoopStartFrameOffset(0)
 		, TimeScale(1.0f)
 		, HierarchicalBias(100)
 		, StartOffset_DEPRECATED(0.f)
@@ -24,15 +25,27 @@ struct FMovieSceneSectionParameters
 
 public:
 	/** Number of frames (in display rate) to skip at the beginning of the sub-sequence. */
-	UPROPERTY(BlueprintReadWrite, EditAnywhere, BlueprintReadWrite, Category = "Clipping")
+	UPROPERTY(BlueprintReadWrite, EditAnywhere, Category="Clipping")
 	FFrameNumber StartFrameOffset;
 
+	/** Whether this section supports looping the sub-sequence. */
+	UPROPERTY(BlueprintReadWrite, EditAnywhere, Category="Clipping")
+	bool bCanLoop = false;
+
+	/** Number of frames (in display rate) to skip at the beginning of the sub-sequence. */
+	UPROPERTY(BlueprintReadWrite, EditAnywhere, Category="Clipping", meta=(EditCondition="bCanLoop"))
+	FFrameNumber EndFrameOffset;
+
+	/** Number of frames (in display rate) to offset the first loop of the sub-sequence. */
+	UPROPERTY(BlueprintReadWrite, EditAnywhere, Category="Clipping", meta=(EditCondition="bCanLoop"))
+	FFrameNumber FirstLoopStartFrameOffset;
+
 	/** Playback time scaling factor. */
-	UPROPERTY(BlueprintReadWrite, EditAnywhere, BlueprintReadWrite, Category="Timing")
+	UPROPERTY(BlueprintReadWrite, EditAnywhere, Category="Timing")
 	float TimeScale;
 
 	/** Hierachical bias. Higher bias will take precedence. */
-	UPROPERTY(BlueprintReadWrite, EditAnywhere, BlueprintReadWrite, Category="Sequence")
+	UPROPERTY(BlueprintReadWrite, EditAnywhere, Category="Sequence")
 	int32 HierarchicalBias;
 
 	UPROPERTY()
diff --git a/Engine/Source/Runtime/MovieScene/Public/Evaluation/MovieSceneSequenceHierarchy.h b/Engine/Source/Runtime/MovieScene/Public/Evaluation/MovieSceneSequenceHierarchy.h
--- a/Engine/Source/Runtime/MovieScene/Public/Evaluation/MovieSceneSequenceHierarchy.h
+++ b/Engine/Source/Runtime/MovieScene/Public/Evaluation/MovieSceneSequenceHierarchy.h
@@ -66,11 +66,24 @@ struct FMovieSceneSubSequenceData
 	/** This sequence's deterministic sequence ID. Used in editor to reduce the risk of collisions on recompilation. */ 
 	UPROPERTY()
 	FMovieSceneSequenceID DeterministicSequenceID;
-
+	
 	/** This sub sequence's playback range according to its parent sub section. Clamped recursively during template generation */
 	UPROPERTY()
 	FMovieSceneFrameRange PlayRange;
 
+	/** The sub-sequence's full playback range, in its own local time space. */
+	UPROPERTY()
+	FMovieSceneFrameRange FullPlayRange;
+
+	/**
+	 * The play range of the parent section, without any warping involved.
+	 * That means that, for a sub-sequence playing with an initial offset of 50 and looping 3 times,
+	 * this play range will start 50 frames after layRange's lower bound, and extend much past PlayRange's 
+	 * upper bound (3 times longer).
+	 */
+	UPROPERTY()
+	FMovieSceneFrameRange UnwarpedPlayRange;
+
 	/** The sequence preroll range considering the start offset */
 	UPROPERTY()
 	FMovieSceneFrameRange PreRollRange;
diff --git a/Engine/Source/Runtime/MovieScene/Public/Evaluation/MovieSceneSequenceTransform.h b/Engine/Source/Runtime/MovieScene/Public/Evaluation/MovieSceneSequenceTransform.h
--- a/Engine/Source/Runtime/MovieScene/Public/Evaluation/MovieSceneSequenceTransform.h
+++ b/Engine/Source/Runtime/MovieScene/Public/Evaluation/MovieSceneSequenceTransform.h
@@ -5,17 +5,147 @@
 #include "CoreMinimal.h"
 #include "UObject/ObjectMacros.h"
 #include "Misc/FrameTime.h"
+#include "MovieSceneTimeTransform.h"
+#include "MovieSceneTimeWarping.h"
 #include "MovieSceneSequenceTransform.generated.h"
 
+
+/**
+ * Converts a range from one type of bounds to another. The output bounds type must be implicitely 
+ * constructable from the input bounds type.
+ */
+template<typename InBoundType, typename OutBoundType>
+inline TRange<OutBoundType> ConvertRange(const TRange<InBoundType>& Range)
+{
+	const TRangeBound<InBoundType> SourceLower = Range.GetLowerBound();
+	TRangeBound<OutBoundType> DestLower = SourceLower.IsOpen() ?
+		TRangeBound<OutBoundType>() :
+		SourceLower.IsInclusive() ?
+			TRangeBound<OutBoundType>::Inclusive(SourceLower.GetValue()) :
+			TRangeBound<OutBoundType>::Exclusive(SourceLower.GetValue());
+
+	const TRangeBound<InBoundType> SourceUpper = Range.GetUpperBound();
+	TRangeBound<OutBoundType> DestUpper = SourceUpper.IsOpen() ?
+		TRangeBound<OutBoundType>() :
+		SourceUpper.IsInclusive() ?
+			TRangeBound<OutBoundType>::Inclusive(SourceUpper.GetValue()) :
+			TRangeBound<OutBoundType>::Exclusive(SourceUpper.GetValue());
+
+	return TRange<OutBoundType>(DestLower, DestUpper);
+}
+
+// Specialization of ConvertRange for round down FFrameTime to FFrameNumber.
+template<>
+inline TRange<FFrameNumber> ConvertRange(const TRange<FFrameTime>& Range)
+{
+	const TRangeBound<FFrameTime> SourceLower = Range.GetLowerBound();
+	TRangeBound<FFrameNumber> DestLower = SourceLower.IsOpen() ?
+		TRangeBound<FFrameNumber>() :
+		SourceLower.IsInclusive() ?
+			TRangeBound<FFrameNumber>::Inclusive(SourceLower.GetValue().FloorToFrame()) :
+			TRangeBound<FFrameNumber>::Exclusive(SourceLower.GetValue().FloorToFrame());
+
+	const TRangeBound<FFrameTime> SourceUpper = Range.GetUpperBound();
+	TRangeBound<FFrameNumber> DestUpper = SourceUpper.IsOpen() ?
+		TRangeBound<FFrameNumber>() :
+		SourceUpper.IsInclusive() ?
+			TRangeBound<FFrameNumber>::Inclusive(SourceUpper.GetValue().FloorToFrame()) :
+			TRangeBound<FFrameNumber>::Exclusive(SourceUpper.GetValue().FloorToFrame());
+
+	return TRange<FFrameNumber>(DestLower, DestUpper);
+}
+
+/**
+ * Struct that tracks warp counts in a way that works with the FMovieSceneTimeWarping struct.
+ */
+USTRUCT()
+struct FMovieSceneWarpCounter
+{
+	GENERATED_BODY()
+
+	void AddWarpingLevel(uint32 WarpCount)
+	{
+		if (WarpCount != FMovieSceneTimeWarping::InvalidWarpCount)
+		{
+			WarpCounts.Add(WarpCount);
+		}
+		else if (WarpCounts.Num() > 0)
+		{
+			if (WarpCounts.Last() != FMovieSceneTimeWarping::InvalidWarpCount)
+			{
+				WarpCounts.Add(FMovieSceneTimeWarping::InvalidWarpCount);
+			}
+		}
+		// else: ignore, non-warping sub-sequences have their transforms compressed into a single
+		//       linear transform until the point we meet the first looping sub-sequence.
+	}
+
+	void AddNonWarpingLevel()
+	{
+		AddWarpingLevel(FMovieSceneTimeWarping::InvalidWarpCount);
+	}
+
+	UPROPERTY()
+	TArray<uint32> WarpCounts;
+};
+
+/**
+ * Time transform information for a nested sequence.
+ */
+USTRUCT()
+struct FMovieSceneNestedSequenceTransform
+{
+	GENERATED_BODY()
+
+	FMovieSceneNestedSequenceTransform()
+	{}
+
+	FMovieSceneNestedSequenceTransform(FMovieSceneTimeTransform InLinearTransform)
+		: LinearTransform(InLinearTransform)
+	{}
+
+	FMovieSceneNestedSequenceTransform(FMovieSceneTimeTransform InLinearTransform, FMovieSceneTimeWarping InWarping)
+		: LinearTransform(InLinearTransform)
+		, Warping(InWarping)
+	{}
+
+	bool IsWarping() const { return Warping.IsValid(); }
+
+	/**
+	 * Linear time transform for this sub-sequence.
+	 */
+	UPROPERTY()
+	FMovieSceneTimeTransform LinearTransform;
+
+	/**
+	 * Time warping information for this sub-sequence.
+	 */
+	UPROPERTY()
+	FMovieSceneTimeWarping Warping;
+
+	friend bool operator==(const FMovieSceneNestedSequenceTransform& A, const FMovieSceneNestedSequenceTransform& B)
+	{
+		return A.LinearTransform == B.LinearTransform && A.Warping == B.Warping;
+	}
+
+	friend bool operator!=(const FMovieSceneNestedSequenceTransform& A, const FMovieSceneNestedSequenceTransform& B)
+	{
+		return A.LinearTransform != B.LinearTransform || A.Warping != B.Warping;
+	}
+
+	FMovieSceneTimeTransform InverseLinearOnly() const
+	{
+		return LinearTransform.Inverse();
+	}
+
+	FMovieSceneTimeTransform InverseFromWarp(uint32 WarpCount) const
+	{
+		return LinearTransform.Inverse() * Warping.InverseFromWarp(WarpCount);
+	}
+};
+
 /**
  * Movie scene sequence transform class that transforms from one time-space to another.
- *
- * @note The transform can be thought of as the top row of a 2x2 matrix, where the bottom row is the identity:
- * 			| TimeScale	Offset	|
- *			| 0			1		|
- *
- * As such, traditional matrix mathematics can be applied to transform between different sequence's time-spaces.
- * Transforms apply offset first, then time scale.
  */
 USTRUCT()
 struct FMovieSceneSequenceTransform
@@ -26,8 +156,6 @@ struct FMovieSceneSequenceTransform
 	 * Default construction to the identity transform
 	 */
 	FMovieSceneSequenceTransform()
-		: TimeScale(1.f)
-		, Offset(0)
 	{}
 
 	/**
@@ -36,37 +164,262 @@ struct FMovieSceneSequenceTransform
 	 * @param InOffset 			The offset to translate by
 	 * @param InTimeScale 		The timescale. For instance, if a sequence is playing twice as fast, pass 2.f
 	 */
-	FMovieSceneSequenceTransform(FFrameTime InOffset, float InTimeScale = 1.f)
-		: TimeScale(InTimeScale)
-		, Offset(InOffset)
+	explicit FMovieSceneSequenceTransform(FFrameTime InOffset, float InTimeScale = 1.f)
+		: LinearTransform(InOffset, InTimeScale)
+	{}
+
+	/**
+	 * Construction from a linear time transform.
+	 *
+	 * @param InLinearTransform	The linear transform
+	 */
+	FMovieSceneSequenceTransform(FMovieSceneTimeTransform InLinearTransform)
+		: LinearTransform(InLinearTransform)
 	{}
 
 	friend bool operator==(const FMovieSceneSequenceTransform& A, const FMovieSceneSequenceTransform& B)
 	{
-		return A.TimeScale == B.TimeScale && A.Offset == B.Offset;
+		return A.LinearTransform == B.LinearTransform && A.NestedTransforms == B.NestedTransforms;
 	}
 
 	friend bool operator!=(const FMovieSceneSequenceTransform& A, const FMovieSceneSequenceTransform& B)
 	{
-		return A.TimeScale != B.TimeScale || A.Offset != B.Offset;
+		return A.LinearTransform != B.LinearTransform || A.NestedTransforms != B.NestedTransforms;
+	}
+
+	/**
+	 * Returns whether this sequence transform includes any time warping.
+	 */
+	bool IsWarping() const
+	{
+		// Since we compress any linear transforms, we know that there will be warping as soon as we
+		// use the nested transform array.
+		return NestedTransforms.Num() > 0;
+	}
+
+	/**
+	 * Returns whether this sequence transform is purely linear (i.e. doesn't involve time warping).
+	 */
+	bool IsLinear() const
+	{
+		return NestedTransforms.Num() == 0;
 	}
 
 	/**
-	 * Retrieve the inverse of this transform
+	 * Returns the total time-scale of this transform.
 	 */
-	FMovieSceneSequenceTransform Inverse() const
+	float GetTimeScale() const
 	{
-		const FFrameTime NewOffset = -Offset/TimeScale;
-		return FMovieSceneSequenceTransform(NewOffset, 1.f/TimeScale);
+		float TimeScale = LinearTransform.TimeScale;
+		for (const FMovieSceneNestedSequenceTransform& NestedTransform : NestedTransforms)
+		{
+			if (NestedTransform.LinearTransform.TimeScale != 1.f)
+			{
+				TimeScale *= NestedTransform.LinearTransform.TimeScale;
+			}
+		}
+		return TimeScale;
+	}
+
+	/**
+	 * Transforms the given time, returning the transformed time along with a series of loop indices that
+	 * indicate which loops each nested level is in.
+	 */
+	void TransformTime(FFrameTime InTime, FFrameTime& OutTime, FMovieSceneWarpCounter& OutWarpCounter) const
+	{
+		OutWarpCounter.WarpCounts.Reset();
+
+		OutTime = InTime * LinearTransform;
+		
+		if (NestedTransforms.Num() > 0)
+		{
+			for (const FMovieSceneNestedSequenceTransform& NestedTransform : NestedTransforms)
+			{
+				OutTime = OutTime * NestedTransform.LinearTransform;
+
+				if (NestedTransform.IsWarping())
+				{
+					uint32 WarpIndex;
+					NestedTransform.Warping.TransformTime(OutTime, OutTime, WarpIndex);
+					OutWarpCounter.WarpCounts.Add(WarpIndex);
+				}
+				else
+				{
+					OutWarpCounter.WarpCounts.Add(FMovieSceneTimeWarping::InvalidWarpCount);
+				}
+			}
+		}
+	}
+
+	/**
+	 * Transforms the given range using the "pure" warping transformation.
+	 * See FMovieSceneTimeWarping for more information.
+	 */
+	TRange<FFrameTime> TransformRangePure(const TRange<FFrameTime>& Range) const
+	{
+		TRange<FFrameTime> Result = Range * LinearTransform;
+		for (const FMovieSceneNestedSequenceTransform& NestedTransform : NestedTransforms)
+		{
+			Result = Result * NestedTransform.LinearTransform;
+			if (NestedTransform.IsWarping())
+			{
+				Result = NestedTransform.Warping.TransformRangePure(Result);
+			}
+		}
+		return Result;
+	}
+
+	/**
+	 * Transforms the given range using the "unwarped" warping transformation. 
+	 * See FMovieSceneTimeWarping for more information.
+	 */
+	TRange<FFrameTime> TransformRangeUnwarped(const TRange<FFrameTime>& Range) const
+	{
+		TRange<FFrameTime> Result = Range * LinearTransform;
+		for (const FMovieSceneNestedSequenceTransform& NestedTransform : NestedTransforms)
+		{
+			Result = Result * NestedTransform.LinearTransform;
+			if (NestedTransform.IsWarping())
+			{
+				Result = NestedTransform.Warping.TransformRangeUnwarped(Result);
+			}
+		}
+		return Result;
+	}
+
+	/**
+	 * Transforms the given range using the "constrained" warping transformation. 
+	 * See FMovieSceneTimeWarping for more information.
+	 */
+	TRange<FFrameTime> TransformRangeConstrained(const TRange<FFrameTime>& Range) const
+	{
+		TRange<FFrameTime> Result = Range * LinearTransform;
+		for (const FMovieSceneNestedSequenceTransform& NestedTransform : NestedTransforms)
+		{
+			Result = Result * NestedTransform.LinearTransform;
+			if (NestedTransform.IsWarping())
+			{
+				Result = NestedTransform.Warping.TransformRangeConstrained(Result);
+			}
+		}
+		return Result;
+	}
+
+	/**
+	 * Transforms the given range using the "pure" warping transformation.
+	 * See FMovieSceneTimeWarping for more information.
+	 */
+	TRange<FFrameNumber> TransformRangePure(const TRange<FFrameNumber>& Range) const
+	{
+		TRange<FFrameTime> TimeRange = ConvertRange<FFrameNumber, FFrameTime>(Range);
+		TimeRange = TransformRangePure(TimeRange);
+		return ConvertRange<FFrameTime, FFrameNumber>(TimeRange);
+	}
+
+	/**
+	 * Transforms the given range using the "unwarped" warping transformation. 
+	 * See FMovieSceneTimeWarping for more information.
+	 */
+	TRange<FFrameNumber> TransformRangeUnwarped(const TRange<FFrameNumber>& Range) const
+	{
+		TRange<FFrameTime> TimeRange = ConvertRange<FFrameNumber, FFrameTime>(Range);
+		TimeRange = TransformRangeUnwarped(TimeRange);
+		return ConvertRange<FFrameTime, FFrameNumber>(TimeRange);
+	}
+
+	/**
+	 * Transforms the given range using the "constrained" warping transformation. 
+	 * See FMovieSceneTimeWarping for more information.
+	 */
+	TRange<FFrameNumber> TransformRangeConstrained(const TRange<FFrameNumber>& Range) const
+	{
+		TRange<FFrameTime> TimeRange = ConvertRange<FFrameNumber, FFrameTime>(Range);
+		TimeRange = TransformRangeConstrained(TimeRange);
+		return ConvertRange<FFrameTime, FFrameNumber>(TimeRange);
+	}
+
+	/**
+	 * Retrieve the inverse of the linear part of this transform.
+	 */
+	FMovieSceneTimeTransform InverseLinearOnly() const
+	{
+		return LinearTransform.Inverse();
+	}
+
+	/**
+	 * Retrieve the inverse of this transform assuming the local times would all belong to
+	 * the first warp inside each nested time range.
+	 */
+	FMovieSceneTimeTransform InverseFromAllFirstWarps() const
+	{
+		const size_t NestedTransformsSize = NestedTransforms.Num();
+		if (NestedTransformsSize == 0)
+		{
+			return LinearTransform.Inverse();
+		}
+		else
+		{
+			TArray<uint32> WarpCounts;
+			WarpCounts.Reserve(NestedTransforms.Num());
+			for (const FMovieSceneNestedSequenceTransform& NestedTransform : NestedTransforms)
+			{
+				WarpCounts.Add(NestedTransform.IsWarping() ? 0 : FMovieSceneTimeWarping::InvalidWarpCount);
+			}
+			return InverseFromWarp(WarpCounts);
+		}
+	}
+
+	/**
+	 * Retrieve the inverse of this transform assuming the local times would belong to the
+	 * n'th warp inside the time range. As such, this method requires that the warp counts
+	 * in the counter object is equal to the number of nested loops in this transform.
+	 */
+	FMovieSceneTimeTransform InverseFromWarp(const FMovieSceneWarpCounter& WarpCounter) const
+	{
+		return InverseFromWarp(WarpCounter.WarpCounts);
+	}
+
+	/**
+	 * Retrieve the inverse of this transform assuming the local times would belong to the
+	 * n'th warp inside the time range. As such, this method takes as parameter an array with
+	 * as many items as there currently have nested loops in this transform.
+	 */
+	FMovieSceneTimeTransform InverseFromWarp(const TArrayView<const uint32>& WarpCounts) const
+	{
+		const size_t NestedTransformsSize = NestedTransforms.Num();
+		check(WarpCounts.Num() == NestedTransformsSize);
+
+		if (NestedTransformsSize > 0)
+		{
+			// Start accumulating the inverse transforms in reverse order.
+			const FMovieSceneNestedSequenceTransform& LastNesting = NestedTransforms.Last();
+			FMovieSceneTimeTransform Result = LastNesting.InverseFromWarp(WarpCounts.Last());
+
+			// If there's a tail linear transform, the WarpCounts array has one less item than the
+			// NestedTransforms array, and we need to iterate with different indices that are offset by 1.
+			// Otherwise, we can iterate with the same index.
+			size_t Index = NestedTransformsSize - 1;
+			while (Index > 0)
+			{
+				--Index;
+				
+				const uint32 CurWarpCount = WarpCounts[Index];
+				const FMovieSceneNestedSequenceTransform& CurNesting = NestedTransforms[Index];
+				Result = Result * CurNesting.InverseFromWarp(CurWarpCount);
+			}
+			return Result;
+		}
+		else
+		{
+			return FMovieSceneTimeTransform(LinearTransform.Inverse());
+		}
 	}
 
-	/** The sequence's time scale (or play rate) */
 	UPROPERTY()
-	float TimeScale;
+	FMovieSceneTimeTransform LinearTransform;
 
-	/** Scalar frame offset applied before the scale */
 	UPROPERTY()
-	FFrameTime Offset;
+	TArray<FMovieSceneNestedSequenceTransform> NestedTransforms;
 };
 
 /**
@@ -77,13 +430,27 @@ struct FMovieSceneSequenceTransform
  */
 inline FFrameTime operator*(FFrameTime InTime, const FMovieSceneSequenceTransform& RHS)
 {
-	// Avoid floating point conversion when in the same time-space
-	if (RHS.TimeScale == 1.f)
+	const uint32 NestedTransformsSize = RHS.NestedTransforms.Num();
+	if (NestedTransformsSize == 0)
 	{
-		return InTime + RHS.Offset;
+		return InTime * RHS.LinearTransform;
+	}
+	else
+	{
+		FFrameTime OutTime = InTime * RHS.LinearTransform;
+		for (const FMovieSceneNestedSequenceTransform& NestedTransform : RHS.NestedTransforms)
+		{
+			if (NestedTransform.IsWarping())
+			{
+				OutTime = OutTime * NestedTransform.LinearTransform * NestedTransform.Warping;
+			}
+			else
+			{
+				OutTime = OutTime * NestedTransform.LinearTransform;
+			}
+		}
+		return OutTime;
 	}
-
-	return RHS.Offset + InTime * RHS.TimeScale;
 }
 
 /**
@@ -98,55 +465,6 @@ inline FFrameTime& operator*=(FFrameTime& InTime, const FMovieSceneSequenceTrans
 	return InTime;
 }
 
-/**
- * Transform a time range by a sequence transform
- *
- * @param LHS 				The time range to transform
- * @param RHS 				The transform
- */
-template<typename T>
-TRange<T> operator*(const TRange<T>& LHS, const FMovieSceneSequenceTransform& RHS)
-{
-	TRangeBound<T> SourceLower = LHS.GetLowerBound();
-	TRangeBound<T> TransformedLower =
-		SourceLower.IsOpen() ? 
-			TRangeBound<T>() : 
-			SourceLower.IsInclusive() ?
-				TRangeBound<T>::Inclusive(SourceLower.GetValue() * RHS) :
-				TRangeBound<T>::Exclusive(SourceLower.GetValue() * RHS);
-
-	TRangeBound<T> SourceUpper = LHS.GetUpperBound();
-	TRangeBound<T> TransformedUpper =
-		SourceUpper.IsOpen() ? 
-			TRangeBound<T>() : 
-			SourceUpper.IsInclusive() ?
-				TRangeBound<T>::Inclusive(SourceUpper.GetValue() * RHS) :
-				TRangeBound<T>::Exclusive(SourceUpper.GetValue() * RHS);
-
-	return TRange<T>(TransformedLower, TransformedUpper);
-}
-
-inline TRange<FFrameNumber> operator*(const TRange<FFrameNumber>& LHS, const FMovieSceneSequenceTransform& RHS)
-{
-	TRangeBound<FFrameNumber> SourceLower = LHS.GetLowerBound();
-	TRangeBound<FFrameNumber> TransformedLower =
-		SourceLower.IsOpen() ? 
-			TRangeBound<FFrameNumber>() : 
-			SourceLower.IsInclusive() ?
-				TRangeBound<FFrameNumber>::Inclusive((SourceLower.GetValue() * RHS).FloorToFrame()) :
-				TRangeBound<FFrameNumber>::Exclusive((SourceLower.GetValue() * RHS).FloorToFrame());
-
-	TRangeBound<FFrameNumber> SourceUpper = LHS.GetUpperBound();
-	TRangeBound<FFrameNumber> TransformedUpper =
-		SourceUpper.IsOpen() ? 
-			TRangeBound<FFrameNumber>() : 
-			SourceUpper.IsInclusive() ?
-				TRangeBound<FFrameNumber>::Inclusive((SourceUpper.GetValue() * RHS).FloorToFrame()) :
-				TRangeBound<FFrameNumber>::Exclusive((SourceUpper.GetValue() * RHS).FloorToFrame());
-
-	return TRange<FFrameNumber>(TransformedLower, TransformedUpper);
-}
-
 /**
  * Transform a time range by a sequence transform
  *
@@ -166,14 +484,63 @@ TRange<T>& operator*=(TRange<T>& LHS, const FMovieSceneSequenceTransform& RHS)
  */
 inline FMovieSceneSequenceTransform operator*(const FMovieSceneSequenceTransform& LHS, const FMovieSceneSequenceTransform& RHS)
 {
-	// The matrix multiplication occurs as follows:
-	//
-	// | TimeScaleA	, OffsetA	|	.	| TimeScaleB, OffsetB	|
-	// | 0			, 1			|		| 0			, 1			|
-
-	const FFrameTime ScaledOffsetRHS = LHS.TimeScale == 1.f ? RHS.Offset : RHS.Offset*LHS.TimeScale;
-	return FMovieSceneSequenceTransform(
-		LHS.Offset + ScaledOffsetRHS,		// New Offset
-		LHS.TimeScale * RHS.TimeScale		// New TimeScale
-		);
+	if (!LHS.IsWarping())
+	{
+		if (!RHS.IsWarping())
+		{
+			// None of the transforms are warping... we can combine them into another linear transform.
+			return FMovieSceneSequenceTransform(LHS.LinearTransform * RHS.LinearTransform);
+		}
+		else
+		{
+			// LHS is linear, but RHS is warping. Since transforms are supposed to apply from right to left,
+			// We need to append LHS at the "bottom" of RHS, i.e. add a new nested transform that's LHS.
+			FMovieSceneSequenceTransform Result(RHS);
+			Result.NestedTransforms.Add(FMovieSceneNestedSequenceTransform(LHS.LinearTransform));
+			return Result;
+		}
+	}
+	else // LHS is warping
+	{
+		if (!RHS.IsWarping())
+		{
+			// RHS isn't warping, but LHS is, so we combine the linear transform parts and start looping
+			// from there.
+			FMovieSceneSequenceTransform Result;
+			Result.LinearTransform = LHS.LinearTransform * RHS.LinearTransform;
+			Result.NestedTransforms = LHS.NestedTransforms;
+			return Result;
+		}
+		else
+		{
+			// Both are looping, we need to combine them. Usually, a warping transform doesn't use its linear part,
+			// because whatever linear placement/scaling it has would be in the linear part of the nested transform
+			// struct.
+			FMovieSceneSequenceTransform Result(RHS);
+			const bool bHasOnlyNested = LHS.LinearTransform.Offset == 0 && LHS.LinearTransform.TimeScale == 1.f;
+			ensure(bHasOnlyNested);
+			if (!bHasOnlyNested)
+			{
+				Result.NestedTransforms.Add(FMovieSceneNestedSequenceTransform(LHS.LinearTransform));
+			}
+			Result.NestedTransforms.Append(LHS.NestedTransforms);
+			return Result;
+		}
+	}
+}
+
+/** Convert a FMovieSceneSequenceTransform into a string */
+inline FString LexToString(const FMovieSceneSequenceTransform& InTransform)
+{
+	if (InTransform.NestedTransforms.Num() == 0)
+	{
+		return LexToString(InTransform.LinearTransform);
+	}
+	else
+	{
+		return *FString::Printf(TEXT("%s (+%d nested loops)"),
+				*LexToString(InTransform.LinearTransform),
+				InTransform.NestedTransforms.Num());
+	}
 }
+
diff --git a/Engine/Source/Runtime/MovieScene/Public/Evaluation/MovieSceneTimeTransform.h b/Engine/Source/Runtime/MovieScene/Public/Evaluation/MovieSceneTimeTransform.h
new file mode 100644
--- /dev/null
+++ b/Engine/Source/Runtime/MovieScene/Public/Evaluation/MovieSceneTimeTransform.h
@@ -0,0 +1,187 @@
+// Copyright 1998-2019 Epic Games, Inc. All Rights Reserved.
+
+#pragma once
+
+#include "CoreMinimal.h"
+#include "UObject/ObjectMacros.h"
+#include "Misc/FrameTime.h"
+#include "MovieSceneTimeTransform.generated.h"
+
+/**
+ * Movie scene sequence transform class that transforms from one time-space to another.
+ *
+ * @note The transform can be thought of as the top row of a 2x2 matrix, where the bottom row is the identity:
+ * 			| TimeScale	Offset	|
+ *			| 0			1		|
+ *
+ * As such, traditional matrix mathematics can be applied to transform between different sequence's time-spaces.
+ *
+ * Transforms apply time scale first, and then offset.
+ */
+USTRUCT()
+struct FMovieSceneTimeTransform
+{
+	GENERATED_BODY()
+
+	/**
+	 * Default construction to the identity transform
+	 */
+	FMovieSceneTimeTransform()
+		: TimeScale(1.f)
+		, Offset(0)
+	{}
+
+	/**
+	 * Construction from an offset, and a scale
+	 *
+	 * @param InOffset 			The offset to translate by
+	 * @param InTimeScale 		The timescale. For instance, if a sequence is playing twice as fast, pass 2.f
+	 */
+	explicit FMovieSceneTimeTransform(FFrameTime InOffset, float InTimeScale = 1.f)
+		: TimeScale(InTimeScale)
+		, Offset(InOffset)
+	{}
+
+	friend bool operator==(const FMovieSceneTimeTransform& A, const FMovieSceneTimeTransform& B)
+	{
+		return A.TimeScale == B.TimeScale && A.Offset == B.Offset;
+	}
+
+	friend bool operator!=(const FMovieSceneTimeTransform& A, const FMovieSceneTimeTransform& B)
+	{
+		return A.TimeScale != B.TimeScale || A.Offset != B.Offset;
+	}
+
+	/**
+	 * Retrieve the inverse of this transform
+	 */
+	FMovieSceneTimeTransform Inverse() const
+	{
+		const FFrameTime NewOffset = -Offset / TimeScale;
+		return FMovieSceneTimeTransform(NewOffset, 1.f / TimeScale);
+	}
+
+	/** The sequence's time scale (or play rate) */
+	UPROPERTY()
+	float TimeScale;
+
+	/** Scalar frame offset applied after the scale */
+	UPROPERTY()
+	FFrameTime Offset;
+};
+
+/**
+ * Transform a time by a sequence transform
+ *
+ * @param InTime 			The time to transform
+ * @param RHS 				The transform
+ */
+inline FFrameTime operator*(FFrameTime InTime, const FMovieSceneTimeTransform& RHS)
+{
+	// Avoid floating point conversion when in the same time-space
+	if (RHS.TimeScale == 1.f)
+	{
+		return InTime + RHS.Offset;
+	}
+
+	return InTime * RHS.TimeScale + RHS.Offset;
+}
+
+/**
+ * Transform a time by a sequence transform
+ *
+ * @param InTime 			The time to transform
+ * @param RHS 				The transform
+ */
+inline FFrameTime& operator*=(FFrameTime& InTime, const FMovieSceneTimeTransform& RHS)
+{
+	InTime = InTime * RHS;
+	return InTime;
+}
+
+/**
+ * Transform a time range by a sequence transform
+ *
+ * @param LHS 				The time range to transform
+ * @param RHS 				The transform
+ */
+template<typename T>
+TRange<T> operator*(const TRange<T>& LHS, const FMovieSceneTimeTransform& RHS)
+{
+	TRangeBound<T> SourceLower = LHS.GetLowerBound();
+	TRangeBound<T> TransformedLower =
+		SourceLower.IsOpen() ? 
+			TRangeBound<T>() : 
+			SourceLower.IsInclusive() ?
+				TRangeBound<T>::Inclusive(SourceLower.GetValue() * RHS) :
+				TRangeBound<T>::Exclusive(SourceLower.GetValue() * RHS);
+
+	TRangeBound<T> SourceUpper = LHS.GetUpperBound();
+	TRangeBound<T> TransformedUpper =
+		SourceUpper.IsOpen() ? 
+			TRangeBound<T>() : 
+			SourceUpper.IsInclusive() ?
+				TRangeBound<T>::Inclusive(SourceUpper.GetValue() * RHS) :
+				TRangeBound<T>::Exclusive(SourceUpper.GetValue() * RHS);
+
+	return TRange<T>(TransformedLower, TransformedUpper);
+}
+
+inline TRange<FFrameNumber> operator*(const TRange<FFrameNumber>& LHS, const FMovieSceneTimeTransform& RHS)
+{
+	TRangeBound<FFrameNumber> SourceLower = LHS.GetLowerBound();
+	TRangeBound<FFrameNumber> TransformedLower =
+		SourceLower.IsOpen() ? 
+			TRangeBound<FFrameNumber>() : 
+			SourceLower.IsInclusive() ?
+				TRangeBound<FFrameNumber>::Inclusive((SourceLower.GetValue() * RHS).FloorToFrame()) :
+				TRangeBound<FFrameNumber>::Exclusive((SourceLower.GetValue() * RHS).FloorToFrame());
+
+	TRangeBound<FFrameNumber> SourceUpper = LHS.GetUpperBound();
+	TRangeBound<FFrameNumber> TransformedUpper =
+		SourceUpper.IsOpen() ? 
+			TRangeBound<FFrameNumber>() : 
+			SourceUpper.IsInclusive() ?
+				TRangeBound<FFrameNumber>::Inclusive((SourceUpper.GetValue() * RHS).FloorToFrame()) :
+				TRangeBound<FFrameNumber>::Exclusive((SourceUpper.GetValue() * RHS).FloorToFrame());
+
+	return TRange<FFrameNumber>(TransformedLower, TransformedUpper);
+}
+
+/**
+ * Transform a time range by a sequence transform
+ *
+ * @param InTime 			The time range to transform
+ * @param RHS 				The transform
+ */
+template<typename T>
+TRange<T>& operator*=(TRange<T>& LHS, const FMovieSceneTimeTransform& RHS)
+{
+	LHS = LHS * RHS;
+	return LHS;
+}
+
+/**
+ * Multiply 2 transforms together, resulting in a single transform that gets from RHS parent to LHS space
+ * @note Transforms apply from right to left
+ */
+inline FMovieSceneTimeTransform operator*(const FMovieSceneTimeTransform& LHS, const FMovieSceneTimeTransform& RHS)
+{
+	// The matrix multiplication occurs as follows:
+	//
+	// | TimeScaleA	, OffsetA	|	.	| TimeScaleB, OffsetB	|
+	// | 0			, 1			|		| 0			, 1			|
+	//
+	const FFrameTime ScaledOffsetRHS = (LHS.TimeScale == 1.f) ? RHS.Offset : RHS.Offset * LHS.TimeScale;
+	return FMovieSceneTimeTransform(
+		LHS.Offset + ScaledOffsetRHS,		// New Offset
+		RHS.TimeScale * LHS.TimeScale		// New TimeScale
+		);
+}
+
+/** Convert a FFrameTime into a string */
+inline FString LexToString(const FMovieSceneTimeTransform& InTransform)
+{
+	return *FString::Printf(TEXT("[ %+i+%.3f ; x%.3f ]"), 
+			InTransform.Offset.FrameNumber.Value, InTransform.Offset.GetSubFrame(), InTransform.TimeScale);
+}
diff --git a/Engine/Source/Runtime/MovieScene/Public/Evaluation/MovieSceneTimeWarping.h b/Engine/Source/Runtime/MovieScene/Public/Evaluation/MovieSceneTimeWarping.h
new file mode 100644
--- /dev/null
+++ b/Engine/Source/Runtime/MovieScene/Public/Evaluation/MovieSceneTimeWarping.h
@@ -0,0 +1,280 @@
+// Copyright 1998-2019 Epic Games, Inc. All Rights Reserved.
+
+#pragma once
+
+#include "CoreMinimal.h"
+#include "MovieSceneTimeTransform.h"
+#include "UObject/ObjectMacros.h"
+#include "Misc/FrameTime.h"
+#include "MovieSceneTimeWarping.generated.h"
+
+/**
+ * Transform time by warping it around from end to start. This is mostly useful for
+ * looping things.
+ */
+USTRUCT()
+struct FMovieSceneTimeWarping
+{
+	static const uint32 InvalidWarpCount = (uint32)(-1);
+
+	GENERATED_BODY()
+
+	FMovieSceneTimeWarping()
+		: Start(0)
+		, End(0)
+	{}
+
+	explicit FMovieSceneTimeWarping(FFrameNumber InStart, FFrameNumber InEnd)
+		: Start(InStart)
+		, End(InEnd)
+	{
+		check(Start <= End);
+	}
+
+	/**
+	 * Returns the length of the warping.
+	 */
+	FFrameNumber Length() const { return End - Start; }
+
+	/**
+	 * Returns whether this warping transform is doing anything.
+	 */
+	bool IsValid() const { return End > Start; }
+
+	/**
+	 * Returns a range that encompasses the whole warping time span.
+	 */
+	TRange<FFrameTime> GetRange() const
+	{
+		const TRangeBound<FFrameTime> OutLower(TRangeBound<FFrameTime>::Inclusive(Start));
+		const TRangeBound<FFrameTime> OutUpper(TRangeBound<FFrameTime>::Exclusive(End));
+		return TRange<FFrameTime>(OutLower, OutUpper);
+	}
+
+	/**
+	 * Returns a transformation that takes us from a local time into a given loop back to the root time.
+	 */
+	FMovieSceneTimeTransform InverseFromWarp(uint32 WarpCount) const
+	{
+		check(IsValid() || WarpCount == InvalidWarpCount);
+
+		if (WarpCount == InvalidWarpCount || WarpCount == 0)
+		{
+			return FMovieSceneTimeTransform();
+		}
+		else
+		{
+			const FFrameNumber WarpLength = Length();
+			const FFrameNumber WarpsOffset = WarpLength * WarpCount;
+			return FMovieSceneTimeTransform(WarpsOffset);
+		}
+	}
+
+	/**
+	 * Transforms the given frame and returns the warped frame time along with the warp index we ended up in.
+	 */
+	void TransformFrame(FFrameNumber InFrame, FFrameNumber& OutFrame, uint32& OutWarpIndex) const
+	{
+		checkSlow(IsValid());
+		const FFrameNumber WarpLength = Length();
+		OutWarpIndex = 0;
+		FFrameNumber TempFrame = InFrame - Start;
+		while (TempFrame >= WarpLength)
+		{
+			TempFrame -= WarpLength;
+			++OutWarpIndex;
+		}
+		OutFrame = TempFrame + Start;
+	}
+
+	/**
+	 * Transforms the given time and returns the warped time along with the warp index we ended up in.
+	 */
+	void TransformTime(FFrameTime InTime, FFrameTime& OutTime, uint32& OutWarpIndex) const
+	{
+		checkSlow(IsValid());
+		const FFrameTime WarpLength = Length();
+		OutWarpIndex = 0;
+		FFrameTime TempTime = InTime - Start;
+		while (TempTime >= WarpLength)
+		{
+			TempTime = TempTime - WarpLength;
+			++OutWarpIndex;
+		}
+		OutTime = TempTime + Start;
+	}
+
+	/**
+	 * Transforms the given time by warping it by a specific warp count, regardless of how many warps are
+	 * needed in theory to stay in the warping range.
+	 */
+	void TransformTimeSpecific(FFrameTime InTime, uint32 WarpCount, FFrameTime& OutTime) const
+	{
+		checkSlow(IsValid());
+		const FFrameTime WarpLength = Length();
+		FFrameTime TempTime = InTime - Start;
+		for (uint32 WarpIndex = 0; WarpIndex < WarpCount; ++WarpIndex)
+		{
+			TempTime = TempTime - WarpLength;
+		}
+		OutTime = TempTime + Start;
+	}
+
+	/** 
+	 * Transforms the given range in a "naive" way, i.e. its lower and upper bounds are transformed
+	 * independently by this time warping. This means that the output range could be "inside-out"... that is,
+	 * the output lower bound could be greater than the output upper bound, like, for instance, in the case
+	 * of the input range starting near the end of a loop, and ending near the beginning of another loop.
+	 */
+	TRange<FFrameTime> TransformRangePure(const TRange<FFrameTime>& Range) const;
+
+	/**
+	 * Transforms the given range by time-warping its lower bound, and computing an upper bound so that the
+	 * size of the input range is preserved. So if the input range starts in the middle of a loop and lasts
+	 * 3 loops, the output range will start in the middle of the first loop, and last 3 times as long as
+	 * the Length of this time warp.
+	 */
+	TRange<FFrameTime> TransformRangeUnwarped(const TRange<FFrameTime>& Range) const;
+
+	/**
+	 * Transforms the given range by time-warping it and figuring out if it "covered" a full loop or not.
+	 * If it covered a full loop, return a full-loop range, i.e. a range that goes from StartOffset to 
+	 * (Length - EndOffset).
+	 * If it didn't cover a full loop, return the transformed range, which should be a subset of the full
+	 * loop range mentioned above.
+	 */
+	TRange<FFrameTime> TransformRangeConstrained(const TRange<FFrameTime>& Range) const;
+
+	friend bool operator==(const FMovieSceneTimeWarping& A, const FMovieSceneTimeWarping& B)
+	{
+		return A.Start == B.Start && A.End == B.End;
+	}
+
+	friend bool operator!=(const FMovieSceneTimeWarping& A, const FMovieSceneTimeWarping& B)
+	{
+		return A.Start != B.Start || A.End != B.End;
+	}
+
+	UPROPERTY()
+	FFrameNumber Start;
+
+	UPROPERTY()
+	FFrameNumber End;
+};
+
+inline FFrameNumber operator*(FFrameNumber InFrame, const FMovieSceneTimeWarping& RHS)
+{
+	checkSlow(RHS.IsValid());
+	const FFrameNumber WarpLength = RHS.Length();
+	return (InFrame - RHS.Start) % WarpLength + RHS.Start;
+}
+
+inline FFrameTime operator*(FFrameTime InTime, const FMovieSceneTimeWarping& RHS)
+{
+	checkSlow(RHS.IsValid());
+	const FFrameNumber WarpLength = RHS.Length();
+	return (InTime - RHS.Start) % WarpLength + RHS.Start;
+}
+
+inline FFrameNumber& operator*=(FFrameNumber& InFrame, const FMovieSceneTimeWarping& RHS)
+{
+	InFrame = InFrame * RHS;
+	return InFrame;
+}
+
+inline FFrameTime& operator*=(FFrameTime& InTime, const FMovieSceneTimeWarping& RHS)
+{
+	InTime = InTime * RHS;
+	return InTime;
+}
+
+inline TRange<FFrameTime> FMovieSceneTimeWarping::TransformRangePure(const TRange<FFrameTime>& Range) const
+{
+	checkSlow(IsValid());
+
+	const TRangeBound<FFrameTime> SourceLower = Range.GetLowerBound();
+	const TRangeBound<FFrameTime> TransformedLower =
+		SourceLower.IsOpen() ?
+			TRangeBound<FFrameTime>() :
+			SourceLower.IsInclusive() ?
+				TRangeBound<FFrameTime>::Inclusive(SourceLower.GetValue() * (*this)) :
+				TRangeBound<FFrameTime>::Exclusive(SourceLower.GetValue() * (*this));
+
+	const TRangeBound<FFrameTime> SourceUpper = Range.GetUpperBound();
+	const TRangeBound<FFrameTime> TransformedUpper =
+		SourceUpper.IsOpen() ?
+			TRangeBound<FFrameTime>() :
+			SourceUpper.IsInclusive() ?
+				TRangeBound<FFrameTime>::Inclusive(SourceUpper.GetValue() * (*this)) :
+				TRangeBound<FFrameTime>::Exclusive(SourceUpper.GetValue() * (*this));
+
+	return TRange<FFrameTime>(TransformedLower, TransformedUpper);
+}
+
+inline TRange<FFrameTime> FMovieSceneTimeWarping::TransformRangeUnwarped(const TRange<FFrameTime>& Range) const
+{
+	checkSlow(IsValid());
+
+	if (!Range.GetLowerBound().IsOpen() && !Range.GetUpperBound().IsOpen())
+	{
+		const TRangeBound<FFrameTime> SourceLower = Range.GetLowerBound();
+		const TRangeBound<FFrameTime> TransformedLower =
+			SourceLower.IsInclusive() ?
+				TRangeBound<FFrameTime>::Inclusive(SourceLower.GetValue() * (*this)) :
+				TRangeBound<FFrameTime>::Exclusive(SourceLower.GetValue() * (*this));
+
+		const FFrameTime RangeSize = Range.Size<FFrameTime>();
+		const FFrameTime TransformedUpperValue = TransformedLower.GetValue() + RangeSize;
+		const TRangeBound<FFrameTime> TransformedUpper =
+			Range.GetUpperBound().IsInclusive() ?
+				TRangeBound<FFrameTime>::Inclusive(TransformedUpperValue) :
+				TRangeBound<FFrameTime>::Exclusive(TransformedUpperValue);
+
+		return TRange<FFrameTime>(TransformedLower, TransformedUpper);
+	}
+	else
+	{
+		return TransformRangePure(Range);
+	}
+}
+
+inline TRange<FFrameTime> FMovieSceneTimeWarping::TransformRangeConstrained(const TRange<FFrameTime>& Range) const
+{
+	checkSlow(IsValid());
+	
+	if (!Range.GetLowerBound().IsOpen() && !Range.GetUpperBound().IsOpen())
+	{
+		const FFrameNumber WarpLength = Length();
+		const FFrameTime RangeSize = Range.Size<FFrameTime>();
+		if (RangeSize >= WarpLength)
+		{
+			// If the range is longer than a loop, we could, in theory, end up with a disjointed range, like, for
+			// example, a range that starts in the first loop and ends in the second loop. Regardless of whether
+			// we actually finished a full loop or not, we need here to take the "hull" of that, which is always
+			// the looping range inside our local space.
+			return GetRange();
+		}
+
+		const TRange<FFrameTime> UnwarpedTransformedRange = TransformRangeUnwarped(Range);
+		if (UnwarpedTransformedRange.GetUpperBoundValue() > End || UnwarpedTransformedRange.GetLowerBoundValue() < Start)
+		{
+			// The input range is going over a warp boundary, so we return the whole warping range.
+			return GetRange();
+		}
+
+		// The input range fits inside a single loop, so we can return that.
+		return UnwarpedTransformedRange;
+	}
+	else
+	{
+		// One of the bounds is open, so we are looping from or to infinitiy. That encompasses our whole
+		// looping range.
+		return GetRange();
+	}
+}
+
+/** Convert a FMovieSceneTimeWarping into a string */
+inline FString LexToString(const FMovieSceneTimeWarping& InWarping)
+{
+	return *FString::Printf(TEXT("[ %+i ⟳ %+i]"), InWarping.Start.Value, InWarping.End.Value);
+}
diff --git a/Engine/Source/Runtime/MovieScene/Public/Sections/MovieSceneSubSection.h b/Engine/Source/Runtime/MovieScene/Public/Sections/MovieSceneSubSection.h
--- a/Engine/Source/Runtime/MovieScene/Public/Sections/MovieSceneSubSection.h
+++ b/Engine/Source/Runtime/MovieScene/Public/Sections/MovieSceneSubSection.h
@@ -9,9 +9,10 @@
 #include "Curves/KeyHandle.h"
 #include "MovieSceneSequenceID.h"
 #include "MovieSceneSection.h"
+#include "Evaluation/MovieSceneEvaluationOperand.h"
 #include "Evaluation/MovieSceneSectionParameters.h"
 #include "Evaluation/MovieSceneSequenceHierarchy.h"
-#include "Evaluation/MovieSceneEvaluationOperand.h"
+#include "Evaluation/MovieSceneSequenceTransform.h"
 #include "MovieSceneSubSection.generated.h"
 
 class UMovieSceneSequence;
@@ -68,10 +69,22 @@ public:
 public:
 
 	/**
-	 * Get the transform that converts time from this section's time-base to its inner sequence's
+	 * Gets the transform that converts time from this section's time-base to its inner sequence's
 	 */
 	FMovieSceneSequenceTransform OuterToInnerTransform() const;
 
+	/**
+	 * Gets the playrange of the inner sequence, in the inner sequence's time space, trimmed with any start/end offsets,
+	 * and validated to make sure we get at least a 1-frame long playback range (e.g. in the case where excessive
+	 * trimming results in an invalid range).
+	 */
+	bool GetValidatedInnerPlaybackRange(TRange<FFrameNumber>& OutInnerPlaybackRange) const;
+
+	/**
+	 * Helper function used by the above method, but accessible for other uses like track editors.
+	 */
+	static TRange<FFrameNumber> GetValidatedInnerPlaybackRange(const FMovieSceneSectionParameters& SubSectionParameters, const UMovieScene& InnerMovieScene);
+
 	/**
 	 * Sets the sequence played by this section.
 	 *
@@ -144,7 +157,7 @@ public:
 
 public:
 
-	UPROPERTY(BlueprintReadWrite, EditAnywhere, BlueprintReadWrite, Category="General", meta=(ShowOnlyInnerProperties))
+	UPROPERTY(BlueprintReadWrite, EditAnywhere, Category="General", meta=(ShowOnlyInnerProperties))
 	FMovieSceneSectionParameters Parameters;
 
 private:
diff --git a/Engine/Source/Runtime/MovieSceneTracks/Private/Evaluation/MovieSceneEventTemplate.cpp b/Engine/Source/Runtime/MovieSceneTracks/Private/Evaluation/MovieSceneEventTemplate.cpp
--- a/Engine/Source/Runtime/MovieSceneTracks/Private/Evaluation/MovieSceneEventTemplate.cpp
+++ b/Engine/Source/Runtime/MovieSceneTracks/Private/Evaluation/MovieSceneEventTemplate.cpp
@@ -412,7 +412,7 @@ void FMovieSceneEventSectionTemplate::EvaluateSwept(const FMovieSceneEvaluationO
 	const int32 Last = bBackwards ? 0 : KeyTimes.Num() - 1;
 	const int32 Inc = bBackwards ? -1 : 1;
 
-	const float PositionInSeconds = Context.GetTime() * Context.GetRootToSequenceTransform().Inverse() / Context.GetFrameRate();
+	const float PositionInSeconds = Context.GetTime() * Context.GetRootToSequenceTransform().InverseLinearOnly() / Context.GetFrameRate();
 
 	if (bBackwards)
 	{
@@ -486,7 +486,7 @@ void FMovieSceneEventTriggerTemplate::EvaluateSwept(const FMovieSceneEvaluationO
 	const int32 Last = bBackwards ? 0 : EventTimes.Num() - 1;
 	const int32 Inc = bBackwards ? -1 : 1;
 
-	const float PositionInSeconds = Context.GetTime() * Context.GetRootToSequenceTransform().Inverse() / Context.GetFrameRate();
+	const float PositionInSeconds = Context.GetTime() * Context.GetRootToSequenceTransform().InverseLinearOnly() / Context.GetFrameRate();
 
 	if (bBackwards)
 	{
