Sequencer: Add string conversion for sequence IDs.

diff --git a/Engine/Source/Runtime/MovieScene/Public/MovieSceneSequenceID.h b/Engine/Source/Runtime/MovieScene/Public/MovieSceneSequenceID.h
--- a/Engine/Source/Runtime/MovieScene/Public/MovieSceneSequenceID.h
+++ b/Engine/Source/Runtime/MovieScene/Public/MovieSceneSequenceID.h
@@ -95,3 +95,8 @@ namespace MovieSceneSequenceID
 	static const FMovieSceneSequenceID Invalid(-1);
 	static const FMovieSceneSequenceID Root(0);
 }
+
+inline FString LexToString(const FMovieSceneSequenceID& SequenceID)
+{
+	return *FString::Printf(TEXT("SeqID(%d)"), SequenceID.GetInternalValue());
+}
