Sequencer: Move InvalidWarpCount to a cpp file for older compilers

diff --git a/Engine/Source/Runtime/MovieScene/Private/Evaluation/MovieSceneTimeWarping.cpp b/Engine/Source/Runtime/MovieScene/Private/Evaluation/MovieSceneTimeWarping.cpp
new file mode 100644
--- /dev/null
+++ b/Engine/Source/Runtime/MovieScene/Private/Evaluation/MovieSceneTimeWarping.cpp
@@ -0,0 +1,5 @@
+// Copyright 1998-2019 Epic Games, Inc. All Rights Reserved.
+
+#include "Evaluation/MovieSceneTimeWarping.h"
+
+const uint32 FMovieSceneTimeWarping::InvalidWarpCount = (uint32)(-1);
diff --git a/Engine/Source/Runtime/MovieScene/Public/Evaluation/MovieSceneTimeWarping.h b/Engine/Source/Runtime/MovieScene/Public/Evaluation/MovieSceneTimeWarping.h
--- a/Engine/Source/Runtime/MovieScene/Public/Evaluation/MovieSceneTimeWarping.h
+++ b/Engine/Source/Runtime/MovieScene/Public/Evaluation/MovieSceneTimeWarping.h
@@ -15,7 +15,7 @@
 USTRUCT()
 struct FMovieSceneTimeWarping
 {
-	static const uint32 InvalidWarpCount = (uint32)(-1);
+	MOVIESCENE_API static const uint32 InvalidWarpCount;
 
 	GENERATED_BODY()
 
