Sequencer: Fix signed/unsigned comparison.

diff --git a/Engine/Source/Runtime/MovieScene/Private/Tests/MovieSceneCompilerTests.cpp b/Engine/Source/Runtime/MovieScene/Private/Tests/MovieSceneCompilerTests.cpp
--- a/Engine/Source/Runtime/MovieScene/Private/Tests/MovieSceneCompilerTests.cpp
+++ b/Engine/Source/Runtime/MovieScene/Private/Tests/MovieSceneCompilerTests.cpp
@@ -366,7 +366,7 @@ bool FMovieSceneCompilerSubSequencesTest::RunTest(const FString& Parameters)
 
 		FMovieSceneEvaluationTemplate& AccessTemplate(UMovieSceneSequence& InSequence) override
 		{
-			for (size_t i = 0; i < ExpectedSequences.Num(); ++i)
+			for (int32 i = 0; i < ExpectedSequences.Num(); ++i)
 			{
 				if (ExpectedSequences[i] == &InSequence)
 				{
