Sequencer: const correctness.

diff --git a/Engine/Source/Runtime/MovieScene/Public/MovieSceneSequenceID.h b/Engine/Source/Runtime/MovieScene/Public/MovieSceneSequenceID.h
--- a/Engine/Source/Runtime/MovieScene/Public/MovieSceneSequenceID.h
+++ b/Engine/Source/Runtime/MovieScene/Public/MovieSceneSequenceID.h
@@ -45,7 +45,7 @@ struct FMovieSceneSequenceID
 		return GetTypeHash(In.Value);
 	}
 
-	FORCEINLINE FMovieSceneSequenceID AccumulateParentID(FMovieSceneSequenceID InParentID)
+	FORCEINLINE FMovieSceneSequenceID AccumulateParentID(FMovieSceneSequenceID InParentID) const
 	{
 		return Value == 0 ? InParentID : FMovieSceneSequenceID(HashCombine(Value, InParentID.Value));
 	}
