Sequencer: Recover code lost in merge error.

#rnx

diff --git a/Engine/Source/Runtime/MovieScene/Private/Compilation/MovieSceneCompiler.cpp b/Engine/Source/Runtime/MovieScene/Private/Compilation/MovieSceneCompiler.cpp
--- a/Engine/Source/Runtime/MovieScene/Private/Compilation/MovieSceneCompiler.cpp
+++ b/Engine/Source/Runtime/MovieScene/Private/Compilation/MovieSceneCompiler.cpp
@@ -67,6 +67,8 @@ struct FGatherParameters
 
 		SubParams.AccountForRounding();
 
+		SubParams.SequenceLoopCounter.AddWarpingLevel(LoopIndex);
+
 		return SubParams;
 	}
 
