Copy up from Dev-Editor [REDACTED_HANDLE]

diff --git a/Engine/Source/Runtime/ImageWriteQueue/Public/ImagePixelData.h b/Engine/Source/Runtime/ImageWriteQueue/Public/ImagePixelData.h
--- a/Engine/Source/Runtime/ImageWriteQueue/Public/ImagePixelData.h
+++ b/Engine/Source/Runtime/ImageWriteQueue/Public/ImagePixelData.h
@@ -18,6 +18,13 @@ enum class EImagePixelType
 	Float32,
 };
 
+struct IImagePixelDataPayload
+{
+	virtual ~IImagePixelDataPayload() {}
+};
+
+typedef TSharedPtr<IImagePixelDataPayload, ESPMode::ThreadSafe> FImagePixelPayloadPtr;
+
 struct FImagePixelData
 {
 	virtual ~FImagePixelData() {}
@@ -68,7 +75,7 @@ struct FImagePixelData
 	bool IsDataWellFormed() const
 	{
 		const void* RawPtr    = nullptr;
-		int32       SizeBytes = 0;
+		int64       SizeBytes = 0;
 
 		return GetRawData(RawPtr, SizeBytes);
 	}
@@ -76,14 +83,15 @@ struct FImagePixelData
 	/**
 	 * Get the data and its size only if it is well formed
 	 */
-	bool GetRawData(const void*& OutRawData, int32& OutSizeBytes) const
+	bool GetRawData(const void*& OutRawData, int64& OutSizeBytes) const
 	{
 		const void* RawPtr    = nullptr;
-		int32       SizeBytes = 0;
+		int64       SizeBytes = 0;
 
 		RetrieveData(RawPtr, SizeBytes);
 
-		if (RawPtr && SizeBytes == Size.X*Size.Y*BitDepth/8*NumChannels)
+		int64 FoundTotalSize = int64(Size.X)*int64(Size.Y)*int64(BitDepth / 8)*int64(NumChannels);
+		if (RawPtr && SizeBytes == FoundTotalSize)
 		{
 			OutRawData = RawPtr;
 			OutSizeBytes = SizeBytes;
@@ -108,14 +116,21 @@ struct FImagePixelData
 		return Move();
 	}
 
+	/**
+	* Return a pointer to the Payload stored in this data.
+	*/
+	template<typename T>
+	T* GetPayload() { return static_cast<T*>(Payload.Get()); }
+
 protected:
 
-	FImagePixelData(const FIntPoint& InSize, EImagePixelType InPixelType, ERGBFormat InPixelLayout, uint8 InBitDepth, uint8 InNumChannels)
+	FImagePixelData(const FIntPoint& InSize, EImagePixelType InPixelType, ERGBFormat InPixelLayout, uint8 InBitDepth, uint8 InNumChannels, FImagePixelPayloadPtr InPayload)
 		: Size(InSize)
 		, Type(InPixelType)
 		, PixelLayout(InPixelLayout)
 		, BitDepth(InBitDepth)
 		, NumChannels(InNumChannels)
+		, Payload(InPayload)
 	{}
 
 private:
@@ -123,7 +138,7 @@ private:
 	/**
 	 * Retrieve the raw pixel data
 	 */
-	virtual void RetrieveData(const void*& OutDataPtr, int32& OutSizeBytes) const = 0;
+	virtual void RetrieveData(const void*& OutDataPtr, int64& OutSizeBytes) const = 0;
 	virtual TUniquePtr<FImagePixelData> Copy() const = 0;
 	virtual TUniquePtr<FImagePixelData> Move() = 0;
 
@@ -141,6 +156,9 @@ private:
 
 	/** Number of channels in the data */
 	uint8 NumChannels;
+
+	/** Optional user-specified payload */
+	FImagePixelPayloadPtr Payload;
 };
 
 /**
@@ -149,14 +167,23 @@ private:
 template<typename PixelType>
 struct TImagePixelData : FImagePixelData
 {
-	TArray<PixelType> Pixels;
+	TArray64<PixelType> Pixels;
 
 	TImagePixelData(const FIntPoint& InSize)
-		: FImagePixelData(InSize, TImagePixelDataTraits<PixelType>::PixelType, TImagePixelDataTraits<PixelType>::PixelLayout, TImagePixelDataTraits<PixelType>::BitDepth, TImagePixelDataTraits<PixelType>::NumChannels)
+		: FImagePixelData(InSize, TImagePixelDataTraits<PixelType>::PixelType, TImagePixelDataTraits<PixelType>::PixelLayout, TImagePixelDataTraits<PixelType>::BitDepth, TImagePixelDataTraits<PixelType>::NumChannels, nullptr)
+	{}
+
+	TImagePixelData(const FIntPoint& InSize, TArray64<PixelType>&& InPixels)
+		: FImagePixelData(InSize, TImagePixelDataTraits<PixelType>::PixelType, TImagePixelDataTraits<PixelType>::PixelLayout, TImagePixelDataTraits<PixelType>::BitDepth, TImagePixelDataTraits<PixelType>::NumChannels, nullptr)
+		, Pixels(MoveTemp(InPixels))
+	{}
+
+	TImagePixelData(const FIntPoint& InSize, FImagePixelPayloadPtr InPayload)
+		: FImagePixelData(InSize, TImagePixelDataTraits<PixelType>::PixelType, TImagePixelDataTraits<PixelType>::PixelLayout, TImagePixelDataTraits<PixelType>::BitDepth, TImagePixelDataTraits<PixelType>::NumChannels, InPayload)
 	{}
 
-	TImagePixelData(const FIntPoint& InSize, TArray<PixelType>&& InPixels)
-		: FImagePixelData(InSize, TImagePixelDataTraits<PixelType>::PixelType, TImagePixelDataTraits<PixelType>::PixelLayout, TImagePixelDataTraits<PixelType>::BitDepth, TImagePixelDataTraits<PixelType>::NumChannels)
+	TImagePixelData(const FIntPoint& InSize, TArray64<PixelType>&& InPixels, FImagePixelPayloadPtr InPayload)
+		: FImagePixelData(InSize, TImagePixelDataTraits<PixelType>::PixelType, TImagePixelDataTraits<PixelType>::PixelLayout, TImagePixelDataTraits<PixelType>::BitDepth, TImagePixelDataTraits<PixelType>::NumChannels, InPayload)
 		, Pixels(MoveTemp(InPixels))
 	{}
 
@@ -170,7 +197,7 @@ struct TImagePixelData : FImagePixelData
 		return MakeUnique<TImagePixelData<PixelType>>(*this);
 	}
 
-	virtual void RetrieveData(const void*& OutDataPtr, int32& OutSizeBytes) const override
+	virtual void RetrieveData(const void*& OutDataPtr, int64& OutSizeBytes) const override
 	{
 		OutDataPtr = static_cast<const void*>(&Pixels[0]);
 		OutSizeBytes = Pixels.Num() * sizeof(PixelType);
diff --git a/Engine/Source/Runtime/MovieScene/Private/MovieSceneSequencePlayer.cpp b/Engine/Source/Runtime/MovieScene/Private/MovieSceneSequencePlayer.cpp
--- a/Engine/Source/Runtime/MovieScene/Private/MovieSceneSequencePlayer.cpp
+++ b/Engine/Source/Runtime/MovieScene/Private/MovieSceneSequencePlayer.cpp
@@ -995,7 +995,7 @@ void UMovieSceneSequencePlayer::RPC_ExplicitServerUpdateEvent_Implementation(EUp
 
 #if !NO_LOGGING
 	// Log the sync event if necessary
-	if (UE_LOG_ACTIVE(LogMovieSceneRepl, Verbose))
+	if (UE_LOG_ACTIVE(LogMovieScene, Verbose))
 	{
 		FFrameTime   CurrentTime     = PlayPosition.GetCurrentPosition();
 		FString      SequenceName    = RootTemplateInstance.GetSequence(MovieSceneSequenceID::Root)->GetName();
@@ -1006,7 +1006,7 @@ void UMovieSceneSequencePlayer::RPC_ExplicitServerUpdateEvent_Implementation(EUp
 			SequenceName += FString::Printf(TEXT(" (client %d)"), GPlayInEditorID - 1);
 		}
 
-		UE_LOG(LogMovieSceneRepl, Verbose, TEXT("Explicit update event for sequence %s %s @ frame %d, subframe %f. Server has moved to frame %d, subframe %f with EUpdatePositionMethod::%s."),
+		UE_LOG(LogMovieScene, Verbose, TEXT("Explicit update event for sequence %s %s @ frame %d, subframe %f. Server has moved to frame %d, subframe %f with EUpdatePositionMethod::%s."),
 			*SequenceName, *UEnum::GetValueAsString(TEXT("MovieScene.EMovieScenePlayerStatus"), Status.GetValue()), CurrentTime.FrameNumber.Value, CurrentTime.GetSubFrame(),
 			NetSyncProps.LastKnownPosition.FrameNumber.Value, NetSyncProps.LastKnownPosition.GetSubFrame(), *UEnum::GetValueAsString(TEXT("MovieScene.EUpdatePositionMethod"), NetSyncProps.LastKnownStatus.GetValue()));
 	}
@@ -1094,6 +1094,8 @@ void UMovieSceneSequencePlayer::PostNetReceive()
 	const bool bHasChangedTime    = NetSyncProps.LastKnownPosition != PlayPosition.GetCurrentPosition();
 
 	const FFrameTime PingLag      = (PingMs/1000.f) * PlayPosition.GetInputRate();
+	//const FFrameTime LagThreshold = 0.2f * PlayPosition.GetInputRate();
+	//const FFrameTime LagDisparity = FMath::Abs(PlayPosition.GetCurrentPosition() - NetSyncProps.LastKnownPosition);
 	const FFrameTime LagThreshold = (GSequencerNetSyncThresholdMS * 0.001f) * PlayPosition.GetInputRate();
 
 	if (!bHasChangedStatus && !bHasChangedTime)
diff --git a/Engine/Source/Runtime/MovieScene/Private/Tests/MovieSceneTransformTests.cpp b/Engine/Source/Runtime/MovieScene/Private/Tests/MovieSceneTransformTests.cpp
--- a/Engine/Source/Runtime/MovieScene/Private/Tests/MovieSceneTransformTests.cpp
+++ b/Engine/Source/Runtime/MovieScene/Private/Tests/MovieSceneTransformTests.cpp
@@ -9,6 +9,7 @@
 
 #define LOCTEXT_NAMESPACE "MovieSceneTransformTests"
 
+// Range equality.
 bool IsEqual(TRangeBound<FFrameNumber> A, TRangeBound<FFrameNumber> B)
 {
 	if (A.IsOpen() || B.IsOpen())
@@ -23,12 +24,66 @@ bool IsEqual(TRangeBound<FFrameNumber> A, TRangeBound<FFrameNumber> B)
 	return A.GetValue() == B.GetValue();
 }
 
+// Range equality.
 bool IsEqual(TRange<FFrameNumber> A, TRange<FFrameNumber> B)
 {
 	return IsEqual(A.GetLowerBound(), B.GetLowerBound()) && IsEqual(A.GetUpperBound(), B.GetUpperBound());
 }
 
-bool TestTransform(FAutomationTestBase& Test, FMovieSceneSequenceTransform Transform, TArrayView<TRange<FFrameNumber>> InSource, TArrayView<TRange<FFrameNumber>> InExpected, const TCHAR* TestName)
+// Frame number equality.
+bool IsEqual(FFrameNumber A, FFrameNumber B)
+{
+	return A.Value == B.Value;
+}
+
+// Frame time equality.
+bool IsEqual(FFrameTime A, FFrameTime B)
+{
+	return IsEqual(A.FrameNumber, B.FrameNumber) && FMath::IsNearlyEqual(A.GetSubFrame(), B.GetSubFrame());
+}
+
+// Most time transformations are not "round" so they return a frame time that must be rounded down to a frame number,
+// except for time warping which doesn't stretch anything and returns a frame number.
+template<typename TTransform>
+FFrameNumber TransformToFrameNumber(TTransform Transform, FFrameNumber Value)
+{
+	return (Value * Transform).FloorToFrame();
+}
+template<>
+FFrameNumber TransformToFrameNumber(FMovieSceneTimeWarping Transform, FFrameNumber Value)
+{
+	return Value * Transform;
+}
+
+// Generic method for testing the transform of frames and times.
+template<typename TTransform>
+bool TestTransform(FAutomationTestBase& Test, TTransform Transform, TArrayView<FFrameNumber> InSource, TArrayView<FFrameNumber> InExpected, const TCHAR* TestName)
+{
+	check(InSource.Num() == InExpected.Num());
+
+	bool bSuccess = true;
+	for (int32 Index = 0; Index < InSource.Num(); ++Index)
+	{
+		FFrameNumber Result = TransformToFrameNumber(Transform, InSource[Index]);
+		if (!IsEqual(Result, InExpected[Index]))
+		{
+			Test.AddError(FString::Printf(TEXT("Test '%s' failed (Index %d). Transform %s did not apply correctly (%s != %s)"),
+				TestName,
+				Index,
+				*LexToString(Transform),
+				*LexToString(Result),
+				*LexToString(InExpected[Index])));
+
+			bSuccess = false;
+		}
+	}
+
+	return bSuccess;
+}
+
+// A variant of the above method for testing the transform of ranges.
+template<typename TTransform>
+bool TestTransform(FAutomationTestBase& Test, TTransform Transform, TArrayView<TRange<FFrameNumber>> InSource, TArrayView<TRange<FFrameNumber>> InExpected, const TCHAR* TestName)
 {
 	check(InSource.Num() == InExpected.Num());
 
@@ -38,12 +93,10 @@ bool TestTransform(FAutomationTestBase& Test, FMovieSceneSequenceTransform Trans
 		TRange<FFrameNumber> Result = InSource[Index] * Transform;
 		if (!IsEqual(Result, InExpected[Index]))
 		{
-			Test.AddError(FString::Printf(TEXT("Test '%s' failed (Index %d). Transform (Scale %.3f, Offset %i+%.3f) did not apply correctly (%s != %s)"),
+			Test.AddError(FString::Printf(TEXT("Test '%s' failed (Index %d). Transform %s did not apply correctly (%s != %s)"),
 				TestName,
 				Index,
-				Transform.TimeScale,
-				Transform.Offset.FrameNumber.Value,
-				Transform.Offset.GetSubFrame(),
+				*LexToString(Transform),
 				*LexToString(Result),
 				*LexToString(InExpected[Index])));
 
@@ -61,8 +114,139 @@ FMovieSceneSequenceTransform TransformRange(FFrameNumber StartA, FFrameNumber En
 	return FMovieSceneSequenceTransform(StartB, Scale) * FMovieSceneSequenceTransform(-StartA);
 }
 
-IMPLEMENT_SIMPLE_AUTOMATION_TEST(FMovieSceneSubSectionCoreTransformsTest, "System.Engine.Sequencer.Core.Transforms", EAutomationTestFlags::EditorContext | EAutomationTestFlags::EngineFilter)
-bool FMovieSceneSubSectionCoreTransformsTest::RunTest(const FString& Parameters)
+// Linear transform tests
+IMPLEMENT_SIMPLE_AUTOMATION_TEST(FMovieSceneSubSectionCoreLinearTransformsTest, 
+		"System.Engine.Sequencer.Core.LinearTransforms", 
+		EAutomationTestFlags::EditorContext | EAutomationTestFlags::EngineFilter)
+bool FMovieSceneSubSectionCoreLinearTransformsTest::RunTest(const FString& Parameters)
+{
+	FFrameNumber SourceTimes[] = {
+		FFrameNumber(500),
+		FFrameNumber(525)
+	};
+
+	bool bSuccess = true;
+
+	{
+		FFrameNumber ExpectedTimes[] = {
+			FFrameNumber(500),
+			FFrameNumber(525)
+		};
+		FMovieSceneTimeTransform Transform(0);
+		bSuccess = TestTransform(*this, Transform, SourceTimes, ExpectedTimes, TEXT("IdentityTransform")) && bSuccess;
+
+		Transform = Transform.Inverse();
+		bSuccess = TestTransform(*this, Transform, ExpectedTimes, SourceTimes, TEXT("IdentityTransformInverse")) && bSuccess;
+	}
+
+	{
+		FFrameNumber ExpectedTimes[] = {
+			FFrameNumber(1000),
+			FFrameNumber(1050)
+		};
+		FMovieSceneTimeTransform Transform(0, 2.f);
+		bSuccess = TestTransform(*this, Transform, SourceTimes, ExpectedTimes, TEXT("OffsetTransform")) && bSuccess;
+
+		Transform = Transform.Inverse();
+		bSuccess = TestTransform(*this, Transform, ExpectedTimes, SourceTimes, TEXT("OffsetTransformInverse")) && bSuccess;
+	}
+
+	{
+		FFrameNumber ExpectedTimes[] = {
+			FFrameNumber(0),
+			FFrameNumber(50)
+		};
+		FMovieSceneTimeTransform Transform(-1000, 2.f);
+		bSuccess = TestTransform(*this, Transform, SourceTimes, ExpectedTimes, TEXT("OffsetAndScaleTransform")) && bSuccess;
+
+		Transform = Transform.Inverse();
+		bSuccess = TestTransform(*this, Transform, ExpectedTimes, SourceTimes, TEXT("OffsetAndScaleTransformInverse")) && bSuccess;
+	}
+
+	{
+		FFrameNumber ExpectedTimes[] = {
+			FFrameNumber(0),
+			FFrameNumber(50)
+		};
+		FMovieSceneTimeTransform Transform = FMovieSceneTimeTransform(0, 2.f) * FMovieSceneTimeTransform(-500);
+		bSuccess = TestTransform(*this, Transform, SourceTimes, ExpectedTimes, TEXT("OffsetAndScaleTransformObtainedFromMultiplication")) && bSuccess;
+	}
+
+	return bSuccess;
+}
+
+// Warping transform tests
+IMPLEMENT_SIMPLE_AUTOMATION_TEST(FMovieSceneSubSectionCoreWarpTransformsTest,
+		"System.Engine.Sequencer.Core.WarpTransforms",
+		EAutomationTestFlags::EditorContext | EAutomationTestFlags::EngineFilter)
+bool FMovieSceneSubSectionCoreWarpTransformsTest::RunTest(const FString& Parameters)
+{
+	FFrameNumber SourceTimes[] = {
+		FFrameNumber(0),
+		FFrameNumber(25),
+		FFrameNumber(50),
+		FFrameNumber(60),
+		FFrameNumber(120)
+	};
+
+	bool bSuccess = true;
+
+	{
+		FFrameNumber ExpectedTimes[] = {
+			FFrameNumber(0),
+			FFrameNumber(25),
+			FFrameNumber(0),
+			FFrameNumber(10),
+			FFrameNumber(20)
+		};
+		FMovieSceneTimeWarping Warping(0, 50);
+		bSuccess = TestTransform(*this, Warping, SourceTimes, ExpectedTimes, TEXT("SimpleWarping")) && bSuccess;
+
+		FMovieSceneTimeTransform Transform = Warping.InverseFromWarp(0);
+		bSuccess = TestTransform(*this, Transform, MakeArrayView(ExpectedTimes, 2), MakeArrayView(SourceTimes, 2), TEXT("SimpleWarpingInverseLoop0")) && bSuccess;
+
+		FMovieSceneTimeTransform Transform2 = Warping.InverseFromWarp(1);
+		bSuccess = TestTransform(*this, Transform2, MakeArrayView(ExpectedTimes + 2, 2), MakeArrayView(SourceTimes + 2, 2), TEXT("SimpleWarpingInverseLoop1")) && bSuccess;
+
+		FMovieSceneTimeTransform Transform3 = Warping.InverseFromWarp(2);
+		bSuccess = TestTransform(*this, Transform3, MakeArrayView(ExpectedTimes + 4, 1), MakeArrayView(SourceTimes + 4, 1), TEXT("SimpleWarpingInverseLoop2")) && bSuccess;
+	}
+
+	SourceTimes[0] = FFrameNumber(3);
+	SourceTimes[1] = FFrameNumber(28);
+	SourceTimes[2] = FFrameNumber(53);
+	SourceTimes[3] = FFrameNumber(63);
+	SourceTimes[4] = FFrameNumber(123);
+
+	{
+		FFrameNumber ExpectedTimes[] {
+			FFrameNumber(3),
+			FFrameNumber(28),
+			FFrameNumber(14),
+			FFrameNumber(24),
+			FFrameNumber(6)
+		};
+		FMovieSceneTimeWarping Warping(3, 42);
+		bSuccess = TestTransform(*this, Warping, SourceTimes, ExpectedTimes, TEXT("WarpingWithTrim")) && bSuccess;
+
+		FMovieSceneTimeTransform Transform = Warping.InverseFromWarp(0);
+		bSuccess = TestTransform(*this, Transform, MakeArrayView(ExpectedTimes, 2), MakeArrayView(SourceTimes, 2), TEXT("WarpingWithTrimInverseLoop0")) && bSuccess;
+
+		FMovieSceneTimeTransform Transform2 = Warping.InverseFromWarp(1);
+		bSuccess = TestTransform(*this, Transform2, MakeArrayView(ExpectedTimes + 2, 2), MakeArrayView(SourceTimes + 2, 2), TEXT("WarpingWithTrimInverseLoop1")) && bSuccess;
+
+		FMovieSceneTimeTransform Transform3 = Warping.InverseFromWarp(3);  // We lapsed one full loop
+		bSuccess = TestTransform(*this, Transform3, MakeArrayView(ExpectedTimes + 4, 1), MakeArrayView(SourceTimes + 4, 1), TEXT("WarpingWithTrimInverseLoop2")) && bSuccess;
+	}
+
+	return bSuccess;
+}
+
+// Sequence transform tests 
+IMPLEMENT_SIMPLE_AUTOMATION_TEST(FMovieSceneSubSectionCoreSequenceTransformsTest,
+		"System.Engine.Sequencer.Core.SequenceTransforms",
+		EAutomationTestFlags::EditorContext | EAutomationTestFlags::EngineFilter)
+bool FMovieSceneSubSectionCoreSequenceTransformsTest::RunTest(const FString& Parameters)
 {
 	// We test using ranges since that implicitly tests frame number transformation as well
 	static const TRangeBound<FFrameNumber> OpenBound;
@@ -86,7 +270,7 @@ bool FMovieSceneSubSectionCoreTransformsTest::RunTest(const FString& Parameters)
 			InfiniteRange, OpenLowerRange, OpenUpperRange, ClosedRange
 		};
 		
-		bSuccess = TestTransform(*this, IdentityTransform, SourceRanges, Expected, TEXT("IdentityTransform")) && bSuccess;
+		bSuccess = TestTransform(*this, IdentityTransform.LinearTransform, SourceRanges, Expected, TEXT("IdentityTransform")) && bSuccess;
 	}
 
 	{
@@ -97,7 +281,7 @@ bool FMovieSceneSubSectionCoreTransformsTest::RunTest(const FString& Parameters)
 			InfiniteRange, TRange<FFrameNumber>(OpenBound, FFrameNumber(300)), TRange<FFrameNumber>(FFrameNumber(200), OpenBound), TRange<FFrameNumber>(200, 300)
 		};
 
-		bSuccess = TestTransform(*this, Transform, SourceRanges, Expected, TEXT("Simple Translation")) && bSuccess;
+		bSuccess = TestTransform(*this, Transform.LinearTransform, SourceRanges, Expected, TEXT("Simple Translation")) && bSuccess;
 	}
 
 	{
@@ -110,7 +294,7 @@ bool FMovieSceneSubSectionCoreTransformsTest::RunTest(const FString& Parameters)
 			InfiniteRange, TRange<FFrameNumber>(OpenBound, FFrameNumber(1000)), TRange<FFrameNumber>(FFrameNumber(-200), OpenBound), TRange<FFrameNumber>(-200, 1000)
 		};
 
-		bSuccess = TestTransform(*this, Transform, SourceRanges, Expected, TEXT("Simple Translation + half speed")) && bSuccess;
+		bSuccess = TestTransform(*this, Transform.LinearTransform, SourceRanges, Expected, TEXT("Simple Translation + half speed")) && bSuccess;
 	}
 
 	{
@@ -134,7 +318,7 @@ bool FMovieSceneSubSectionCoreTransformsTest::RunTest(const FString& Parameters)
 			AddError(FString::Printf(TEXT("Accumulated transform does not have the same effect as separate transformations (%i+%.5f != %i+%.5f)"), AccumValue.FrameNumber.Value, AccumValue.GetSubFrame(), SeedValue.FrameNumber.Value, SeedValue.GetSubFrame()));
 		}
 
-		FMovieSceneSequenceTransform InverseTransform = AccumulatedTransform.Inverse();
+		FMovieSceneSequenceTransform InverseTransform = AccumulatedTransform.InverseLinearOnly();
 
 		FFrameTime InverseValue = AccumValue * InverseTransform;
 		if (InverseValue != 10)
@@ -146,5 +330,93 @@ bool FMovieSceneSubSectionCoreTransformsTest::RunTest(const FString& Parameters)
 	return true;
 }
 
+// Miscellaneous warping and scaling tests
+IMPLEMENT_SIMPLE_AUTOMATION_TEST(FMovieSceneSubSectionCoreWarpingAndScalingTransformsTest,
+		"System.Engine.Sequencer.Core.WarpingAndScalingTransforms",
+		EAutomationTestFlags::EditorContext | EAutomationTestFlags::EngineFilter)
+bool FMovieSceneSubSectionCoreWarpingAndScalingTransformsTest::RunTest(const FString& Parameters)
+{
+	{
+		// Sub-sequence at 0, playing at x2
+		FMovieSceneSequenceTransform Transform;
+		Transform.NestedTransforms.Emplace(FMovieSceneTimeTransform(0, 2.f), FMovieSceneTimeWarping(0, 30));
+		TestEqual("Transform time 1", FFrameNumber(10) * Transform, FFrameTime(20));
+		TestEqual("Transform time 2", FFrameNumber(18) * Transform, FFrameTime(6));
+
+		FMovieSceneTimeTransform Inv = Transform.InverseFromWarp(TArray<uint32> { 0 });
+		TestEqual("Inverse time 3", FFrameNumber(20) * Inv, FFrameTime(10));
+		Inv = Transform.InverseFromWarp(TArray<uint32> { 1 });
+		TestEqual("Inverse time 4", FFrameNumber(20) * Inv, FFrameTime(25));
+	}
+
+	{
+		// Sub-sequence at 0, playing at x2, with start offset 20
+		FMovieSceneSequenceTransform Transform;
+		Transform.NestedTransforms.Emplace(FMovieSceneTimeTransform(20, 2.f), FMovieSceneTimeWarping(20, 50));
+		TestEqual("Transform time 5", FFrameNumber(10) * Transform, FFrameTime(40));
+		TestEqual("Transform time 6", FFrameNumber(18) * Transform, FFrameTime(26));
+
+		FMovieSceneTimeTransform Inv = Transform.InverseFromWarp(TArray<uint32> { 0 });
+		TestEqual("Inverse time 7", FFrameNumber(40) * Inv, FFrameTime(10));
+		Inv = Transform.InverseFromWarp(TArray<uint32> { 1 });
+		TestEqual("Inverse time 8", FFrameNumber(40) * Inv, FFrameTime(25));
+	}
+
+	{
+		// Sub-sequence at 3, playing at x2
+		FMovieSceneSequenceTransform Transform;
+		Transform.NestedTransforms.Emplace(FMovieSceneTimeTransform(-6, 2.f), FMovieSceneTimeWarping(0, 30));
+		TestEqual("Transform time 9", FFrameNumber(13) * Transform, FFrameTime(20));
+		TestEqual("Transform time 10", FFrameNumber(21) * Transform, FFrameTime(6));
+
+		FMovieSceneTimeTransform Inv = Transform.InverseFromWarp(TArray<uint32> { 0 });
+		TestEqual("Inverse time 11", FFrameNumber(20) * Inv, FFrameTime(13));
+		Inv = Transform.InverseFromWarp(TArray<uint32> { 1 });
+		TestEqual("Inverse time 12", FFrameNumber(20) * Inv, FFrameTime(28));
+	}
+
+	{
+		// Sub-sequence at 3, playing at x2, with start offset 20
+		FMovieSceneSequenceTransform Transform;
+		Transform.NestedTransforms.Emplace(FMovieSceneTimeTransform(-6 + 20, 2.f), FMovieSceneTimeWarping(20, 50));
+		TestEqual("Transform time 13", FFrameNumber(13) * Transform, FFrameTime(40));
+		TestEqual("Transform time 14", FFrameNumber(21) * Transform, FFrameTime(26));
+
+		FMovieSceneTimeTransform Inv = Transform.InverseFromWarp(TArray<uint32> { 0 });
+		TestEqual("Inverse time 15", FFrameNumber(40) * Inv, FFrameTime(13));
+		Inv = Transform.InverseFromWarp(TArray<uint32> { 1 });
+		TestEqual("Inverse time 16", FFrameNumber(40) * Inv, FFrameTime(28));
+	}
+
+	{
+		// Two levels of sub-sequences: one placed at 10 and warping, the second placed at 6 with x2 scaling
+		FMovieSceneSequenceTransform Transform;
+		Transform.NestedTransforms.Emplace(FMovieSceneTimeTransform(-10, 1.f), FMovieSceneTimeWarping(0, 30));
+		Transform.NestedTransforms.Emplace(FMovieSceneTimeTransform(-12, 2.f));
+		TestEqual("Transform time 17", FFrameNumber(18) * Transform, FFrameTime(4));
+		TestEqual("Transform time 18", FFrameNumber(55) * Transform, FFrameTime(18));
+
+		FMovieSceneTimeTransform Inv = Transform.InverseFromWarp(TArray<uint32> { 0, FMovieSceneTimeWarping::InvalidWarpCount });
+		TestEqual("Inverse time 17", FFrameNumber(4) * Inv, FFrameTime(18));
+		Inv = Transform.InverseFromWarp(TArray<uint32> { 1, FMovieSceneTimeWarping::InvalidWarpCount });
+		TestEqual("Inverse time 18", FFrameNumber(18) * Inv, FFrameTime(55));
+	}
+
+	{
+		// Two levels of sub-sequences: one placed at 10, the second placed at 6 with x2 scaling and warping
+		FMovieSceneSequenceTransform Transform;
+		Transform.NestedTransforms.Emplace(FMovieSceneTimeTransform(-10, 1.f));
+		Transform.NestedTransforms.Emplace(FMovieSceneTimeTransform(-12, 2.f), FMovieSceneTimeWarping(0, 14));
+		TestEqual("Transform time 17", FFrameNumber(19) * Transform, FFrameTime(6));
+		TestEqual("Transform time 18", FFrameNumber(32) * Transform, FFrameTime(4));
+
+		FMovieSceneTimeTransform Inv = Transform.InverseFromWarp(TArray<uint32> { FMovieSceneTimeWarping::InvalidWarpCount, 0 });
+		TestEqual("Inverse time 17", FFrameNumber(6) * Inv, FFrameTime(19));
+		Inv = Transform.InverseFromWarp(TArray<uint32> { FMovieSceneTimeWarping::InvalidWarpCount, 2 });
+		TestEqual("Inverse time 18", FFrameNumber(4) * Inv, FFrameTime(32));
+	}
+
+	return true;
+}
 
 #undef LOCTEXT_NAMESPACE
diff --git a/Engine/Source/Runtime/Renderer/Private/SceneRendering.cpp b/Engine/Source/Runtime/Renderer/Private/SceneRendering.cpp
--- a/Engine/Source/Runtime/Renderer/Private/SceneRendering.cpp
+++ b/Engine/Source/Runtime/Renderer/Private/SceneRendering.cpp
@@ -889,7 +889,6 @@ void FViewInfo::Init()
 	TemporalJitterPixels = FVector2D::ZeroVector;
 
 	PreExposure = 1.0f;
-	MaterialTextureMipBias = 0.0f;
 
 	// Cache TEXTUREGROUP_World's for the render thread to create the material textures' shared sampler.
 	if (IsInGameThread())
@@ -1377,7 +1376,7 @@ void FViewInfo::SetupUniformBufferParameters(
 
 		float FinalMaterialTextureMipBias = GlobalMipBias;
 
-		if (bIsValidWorldTextureGroupSamplerFilter && PrimaryScreenPercentageMethod == EPrimaryScreenPercentageMethod::TemporalUpscale)
+		if (bIsValidWorldTextureGroupSamplerFilter && !FMath::IsNearlyZero(MaterialTextureMipBias))
 		{
 			ViewUniformShaderParameters.MaterialTextureMipBias = MaterialTextureMipBias;
 			ViewUniformShaderParameters.MaterialTextureDerivativeMultiply = FMath::Pow(2.0f, MaterialTextureMipBias);
diff --git a/Engine/Source/Runtime/Renderer/Private/SceneRendering.h b/Engine/Source/Runtime/Renderer/Private/SceneRendering.h
--- a/Engine/Source/Runtime/Renderer/Private/SceneRendering.h
+++ b/Engine/Source/Runtime/Renderer/Private/SceneRendering.h
@@ -1082,9 +1082,6 @@ public:
 	/** Frame's exposure. Always > 0. */
 	float PreExposure;
 
-	/** Mip bias to apply in material's samplers. */
-	float MaterialTextureMipBias;
-
 	/** Precomputed visibility data, the bits are indexed by VisibilityId of a primitive component. */
 	const uint8* PrecomputedVisibilityData;
 
diff --git a/Engine/Source/Runtime/Renderer/Private/SceneVisibility.cpp b/Engine/Source/Runtime/Renderer/Private/SceneVisibility.cpp
--- a/Engine/Source/Runtime/Renderer/Private/SceneVisibility.cpp
+++ b/Engine/Source/Runtime/Renderer/Private/SceneVisibility.cpp
@@ -3114,6 +3114,11 @@ void FSceneRenderer::PreVisibilityFrameSetup(FRHICommandListImmediate& RHICmdLis
 			{
 				ViewState->FrameIndex++;
 			}
+
+			if (View.OverrideFrameIndexValue.IsSet())
+			{
+				ViewState->FrameIndex = View.OverrideFrameIndexValue.GetValue();
+			}
 		}
 		
 		// Subpixel jitter for temporal AA
@@ -3293,7 +3298,7 @@ void FSceneRenderer::PreVisibilityFrameSetup(FRHICommandListImmediate& RHICmdLis
 				ViewState->PrevFrameViewInfo.ViewMatrices.GetViewMatrix(),
 				ViewState->PrevFrameViewInfo.ViewMatrices.GetViewOrigin(),
 				45.0f, 10000.0f);
-			const bool bResetCamera = (bFirstFrameOrTimeWasReset || View.bCameraCut || bIsLargeCameraMovement);
+			const bool bResetCamera = (bFirstFrameOrTimeWasReset || View.bCameraCut || bIsLargeCameraMovement || View.bForceCameraVisibilityReset);
 			
 #if RHI_RAYTRACING
 			// Note: 0.18 deg is the minimum angle for avoiding numerical precision issue (which would cause constant invalidation)
@@ -3355,6 +3360,7 @@ void FSceneRenderer::PreVisibilityFrameSetup(FRHICommandListImmediate& RHICmdLis
 			if (bFirstFrameOrTimeWasReset || 
 				ViewState->LastRenderTime + GEngine->PrimitiveProbablyVisibleTime < View.Family->CurrentRealTime ||
 				View.bCameraCut ||
+				View.bForceCameraVisibilityReset ||
 				IsLargeCameraMovement(
 					View, 
 				    ViewState->PrevViewMatrixForOcclusionQuery, 
@@ -3393,7 +3399,8 @@ void FSceneRenderer::PreVisibilityFrameSetup(FRHICommandListImmediate& RHICmdLis
 
 			// we don't use DeltaTime as it can be 0 (in editor) and is computed by subtracting floats (loses precision over time)
 			// Clamp DeltaWorldTime to reasonable values for the purposes of motion blur, things like TimeDilation can make it very small
-			if (!ViewFamily.bWorldIsPaused)
+			// Offline renders always control the timestep for the view and always need the timescales calculated.
+			if (!ViewFamily.bWorldIsPaused || View.bIsOfflineRender)
 			{
 				ViewState->UpdateMotionBlurTimeScale(View);
 			}
@@ -3432,8 +3439,8 @@ void FSceneViewState::UpdateMotionBlurTimeScale(const FViewInfo& View)
 {
 	const int32 MotionBlurTargetFPS = View.FinalPostProcessSettings.MotionBlurTargetFPS;
 
-	// Frame rates over 120 FPS are clamped to avoid creating huge motion vectors.
-	float DeltaWorldTime = FMath::Max(View.Family->DeltaWorldTime, 1.0f / 120.0f);
+	// Ensure we can divide by the Delta Time later without a divide by zero.
+	float DeltaWorldTime = FMath::Max(View.Family->DeltaWorldTime, SMALL_NUMBER);
 
 	// Track the current FPS by using an exponential moving average of the current delta time.
 	if (MotionBlurTargetFPS <= 0)
diff --git a/Engine/Source/Runtime/Slate/Private/Framework/Application/SlateApplication.cpp b/Engine/Source/Runtime/Slate/Private/Framework/Application/SlateApplication.cpp
--- a/Engine/Source/Runtime/Slate/Private/Framework/Application/SlateApplication.cpp
+++ b/Engine/Source/Runtime/Slate/Private/Framework/Application/SlateApplication.cpp
@@ -1345,9 +1345,20 @@ void FSlateApplication::Tick(ESlateTickType TickType)
 
 		if (TickType == ESlateTickType::All)
 		{
+#if WITH_ACCESSIBILITY
+			// We ensure to only call this in TickType::All to avoid the movie thread also calling this unnecessarily 
+			GetAccessibleMessageHandler()->ProcessAccessibleTasks();
+#endif
 			TickPlatform(DeltaTime);
 		}
 		TickApplication(TickType, DeltaTime);
+#if WITH_ACCESSIBILITY
+		if (TickType == ESlateTickType::All)
+		{
+			// we call this again to improve the responsiveness of accessibility navigation and announcements 
+			GetAccessibleMessageHandler()->ProcessAccessibleTasks();
+		}
+#endif
 	}
 }
 
@@ -1371,7 +1382,6 @@ void FSlateApplication::TickPlatform(float DeltaTime)
 		}
 
 		PlatformApplication->Tick(DeltaTime);
-
 		PlatformApplication->ProcessDeferredEvents(DeltaTime);
 	}
 }
@@ -1470,7 +1480,7 @@ void FSlateApplication::TickApplication(ESlateTickType TickType, float DeltaTime
 		// This relies on Widgets properly registering for Active timer when they need something to happen even
 		// when the user is not providing any input (ie, animations, viewport rendering, async polling, etc).
 		bIsSlateAsleep = true;
-		if (!AllowSlateToSleep.GetValueOnGameThread() || bAnyActiveTimersPending || !bIsUserIdle || bSynthesizedCursorMove || FApp::UseVRFocus())
+		if	(!AllowSlateToSleep.GetValueOnGameThread() || bAnyActiveTimersPending || !bIsUserIdle || bSynthesizedCursorMove || FApp::UseVRFocus())
 		{
 			bIsSlateAsleep = false; // if we get here, then Slate is not sleeping
 
diff --git a/Engine/Source/Runtime/Slate/Private/Widgets/Input/SSlider.cpp b/Engine/Source/Runtime/Slate/Private/Widgets/Input/SSlider.cpp
--- a/Engine/Source/Runtime/Slate/Private/Widgets/Input/SSlider.cpp
+++ b/Engine/Source/Runtime/Slate/Private/Widgets/Input/SSlider.cpp
@@ -91,13 +91,15 @@ int32 SSlider::OnPaint( const FPaintArgs& Args, const FGeometry& AllottedGeometr
 	// draw slider bar
 	auto BarTopLeft = FVector2D(SliderStartPoint.X, SliderStartPoint.Y - Style->BarThickness * 0.5f);
 	auto BarSize = FVector2D(SliderEndPoint.X - SliderStartPoint.X, Style->BarThickness);
+	auto BarImage = GetBarImage();
+	auto ThumbImage = GetThumbImage();
 	FSlateDrawElement::MakeBox(
 		OutDrawElements,
 		LayerId,
 		SliderGeometry.ToPaintGeometry(BarTopLeft, BarSize),
-		GetBarImage(),
+		BarImage,
 		DrawEffects,
-		SliderBarColor.Get().GetColor(InWidgetStyle) * InWidgetStyle.GetColorAndOpacityTint()
+		BarImage->GetTint(InWidgetStyle) * SliderBarColor.Get().GetColor(InWidgetStyle) * InWidgetStyle.GetColorAndOpacityTint()
 		);
 
 	++LayerId;
@@ -107,9 +109,9 @@ int32 SSlider::OnPaint( const FPaintArgs& Args, const FGeometry& AllottedGeometr
 		OutDrawElements,
 		LayerId,
 		SliderGeometry.ToPaintGeometry(HandleTopLeftPoint, GetThumbImage()->ImageSize),
-		GetThumbImage(),
+		ThumbImage,
 		DrawEffects,
-		SliderHandleColor.Get().GetColor(InWidgetStyle) * InWidgetStyle.GetColorAndOpacityTint()
+		ThumbImage->GetTint(InWidgetStyle) * SliderHandleColor.Get().GetColor(InWidgetStyle) * InWidgetStyle.GetColorAndOpacityTint()
 	);
 
 	return LayerId;
@@ -161,9 +163,10 @@ void SSlider::ResetControllerState()
 
 FNavigationReply SSlider::OnNavigation(const FGeometry& MyGeometry, const FNavigationEvent& InNavigationEvent)
 {
-	FNavigationReply Reply = FNavigationReply::Escape();
 	if (bControllerInputCaptured || !bRequiresControllerLock)
 	{
+		FNavigationReply Reply = FNavigationReply::Escape();
+
 		float NewValue = ValueAttribute.Get();
 		if (Orientation == EOrientation::Orient_Horizontal)
 		{
@@ -194,15 +197,11 @@ FNavigationReply SSlider::OnNavigation(const FGeometry& MyGeometry, const FNavig
 		if (ValueAttribute.Get() != NewValue)
 		{
 			CommitValue(FMath::Clamp(NewValue, MinValue, MaxValue));
+			return Reply;
 		}
 	}
 
-	if (Reply.GetBoundaryRule() != EUINavigationRule::Escape)
-	{
-		Reply = SLeafWidget::OnNavigation(MyGeometry, InNavigationEvent);
-	}
-
-	return Reply;
+	return SLeafWidget::OnNavigation(MyGeometry, InNavigationEvent);
 }
 
 FReply SSlider::OnKeyDown(const FGeometry& MyGeometry, const FKeyEvent& InKeyEvent)
diff --git a/Engine/Source/Runtime/Slate/Public/Framework/MultiBox/MultiBoxBuilder.h b/Engine/Source/Runtime/Slate/Public/Framework/MultiBox/MultiBoxBuilder.h
--- a/Engine/Source/Runtime/Slate/Public/Framework/MultiBox/MultiBoxBuilder.h
+++ b/Engine/Source/Runtime/Slate/Public/Framework/MultiBox/MultiBoxBuilder.h
@@ -239,7 +239,7 @@ public:
 	*
 	* @return  New widget object
 	*/
-	virtual TSharedRef< class SWidget > MakeWidget( FMultiBox::FOnMakeMultiBoxBuilderOverride* InMakeMultiBoxBuilderOverride = nullptr, uint32 MaxHeight = INT_MAX) override;
+	virtual TSharedRef< class SWidget > MakeWidget( FMultiBox::FOnMakeMultiBoxBuilderOverride* InMakeMultiBoxBuilderOverride = nullptr, uint32 MaxHeight = 1000) override;
 
 	/**
 	 * Adds a menu separator
diff --git a/Engine/Source/Runtime/Slate/Public/Widgets/Views/STableRow.h b/Engine/Source/Runtime/Slate/Public/Widgets/Views/STableRow.h
--- a/Engine/Source/Runtime/Slate/Public/Widgets/Views/STableRow.h
+++ b/Engine/Source/Runtime/Slate/Public/Widgets/Views/STableRow.h
@@ -424,7 +424,7 @@ public:
 	{
 		TSharedRef< ITypedTableView<ItemType> > OwnerTable = OwnerTablePtr.Pin().ToSharedRef();
 
-		// Requires #include "Widgets/Views/STableViewBase.h"
+		// Requires #include "Widgets/Views/SListView.h" in your header (not done in STableRow.h to avoid circular reference).
 		TSharedRef< STableViewBase > OwnerTableViewBase = StaticCastSharedRef< SListView<ItemType> >(OwnerTable);
 
 		if ( MouseEvent.GetEffectingButton() == EKeys::LeftMouseButton )
diff --git a/Engine/Source/Runtime/SlateCore/Private/Widgets/Accessibility/SlateAccessibleMessageHandler.cpp b/Engine/Source/Runtime/SlateCore/Private/Widgets/Accessibility/SlateAccessibleMessageHandler.cpp
--- a/Engine/Source/Runtime/SlateCore/Private/Widgets/Accessibility/SlateAccessibleMessageHandler.cpp
+++ b/Engine/Source/Runtime/SlateCore/Private/Widgets/Accessibility/SlateAccessibleMessageHandler.cpp
@@ -10,12 +10,20 @@
 #include "Widgets/SWindow.h"
 #include "Input/HittestGrid.h"
 #include "HAL/IConsoleManager.h"
+#include "GenericPlatform/GenericPlatformProcess.h"
+
 
 DECLARE_CYCLE_STAT(TEXT("Slate Accessibility: Tick"), STAT_AccessibilitySlateTick, STATGROUP_Accessibility);
 DECLARE_CYCLE_STAT(TEXT("Slate Accessibility: Event Raised"), STAT_AccessibilitySlateEventRaised, STATGROUP_Accessibility);
+DECLARE_CYCLE_STAT(TEXT("Slate Accessibility: Run In Thread"), STAT_AccessibilitySlateRunInThread, STATGROUP_Accessibility);
+DECLARE_CYCLE_STAT(TEXT("Slate Accessibility: Enqueue Accessible Task"), STAT_AccessibilitySlateEnqueueAccessibleTask, STATGROUP_Accessibility);
+DECLARE_CYCLE_STAT(TEXT("Slate Accessibility: Process Accessible Tasks"), STAT_AccessibilitySlateProcessAccessibleTasks, STATGROUP_Accessibility);
 
 FSlateAccessibleMessageHandler::FSlateAccessibleMessageHandler()
 	: FGenericAccessibleMessageHandler()
+#if ACCESSIBILITY_DEBUG_RESPONSIVENESS
+	, EnqueueEvent(FGenericPlatformProcess::GetSynchEventFromPool(false))
+#endif
 	, bDirty(false)
 {
 	bApplicationIsAccessible = true;
@@ -56,6 +64,109 @@ TSharedPtr<IAccessibleWidget> FSlateAccessibleMessageHandler::GetAccessibleWidge
 	return FSlateAccessibleWidgetCache::GetAccessibleWidgetFromId(Id);
 }
 
+void FSlateAccessibleMessageHandler::RunInThread(const TFunction<void()>& InFunction, bool bWaitForCompletion, ENamedThreads::Type InThread)
+{
+	if (IsActive())
+	{
+		SCOPE_CYCLE_COUNTER(STAT_AccessibilitySlateRunInThread);
+		// For Slate Accessible Message Handler, we HAVE to fulfill all these requests on the game thread 
+		check(InThread == ENamedThreads::GameThread);
+		check(InFunction);
+		
+		ENamedThreads::Type CurrentThreadID = FTaskGraphInterface::Get().GetCurrentThreadIfKnown();
+		// If we are already on the correct thread, just execute the task 
+		if (CurrentThreadID == ENamedThreads::GameThread)
+		{
+			InFunction();
+		}
+		else
+		{
+			// otherwise queue the event to be run in the game thread during FSlateApplication::Tick() 
+			if (bWaitForCompletion)
+			{
+				FEvent* CompletionEvent = FGenericPlatformProcess::GetSynchEventFromPool(false);
+				//we should always be able to get an event from the platform 
+				check(CompletionEvent);
+				EnqueueAccessibleTask(FSlateAccessibleTask(InFunction, CompletionEvent));
+				CompletionEvent->Wait();
+				FGenericPlatformProcess::ReturnSynchEventToPool(CompletionEvent);
+			}
+			else
+			{
+				EnqueueAccessibleTask(FSlateAccessibleTask(InFunction));
+			}
+		}
+	}
+}
+
+void FSlateAccessibleMessageHandler::EnqueueAccessibleTask(const FSlateAccessibleTask& InAccessibleTask)
+{
+	SCOPE_CYCLE_COUNTER(STAT_AccessibilitySlateEnqueueAccessibleTask);
+	FScopeLock Lock(&QueueCriticalSection);
+	AccessibleTaskStorageQueue.Push(InAccessibleTask);
+#if ACCESSIBILITY_DEBUG_RESPONSIVENESS
+	EnqueueEvent->Trigger();
+#endif 
+}
+
+void FSlateAccessibleMessageHandler::ProcessAccessibleTasks()
+{
+	if (!IsActive())
+	{
+		return;
+	}
+	SCOPE_CYCLE_COUNTER(STAT_AccessibilitySlateProcessAccessibleTasks);
+#if !ACCESSIBILITY_DEBUG_RESPONSIVENESS
+	// The processing queue should always be empty at this point as everything was processed last call 
+	check(AccessibleTaskProcessingQueue.Num() == 0);
+	// We swap the content of the queues to avoid constantly holding the lock while processing all tasks 
+	{
+		FScopeLock Lock(&QueueCriticalSection);
+		Swap(AccessibleTaskStorageQueue, AccessibleTaskProcessingQueue);
+	}
+	while (AccessibleTaskProcessingQueue.Num() > 0)
+	{
+		FSlateAccessibleTask TaskToRun = AccessibleTaskProcessingQueue.Pop();
+		TaskToRun.DoTask();
+	}
+#else
+	// We run all tasks in a time slice to improve responsiveness 
+	check(AccessibleTaskProcessingQueue.Num() == 0);
+	{
+		FScopeLock Lock(&QueueCriticalSection);
+		Swap(AccessibleTaskStorageQueue, AccessibleTaskProcessingQueue);
+	}
+	if (AccessibleTaskProcessingQueue.Num() == 0)
+	{
+		return; 
+	}
+	// Change this value to alter the time slice. The longer the time slice, the better the accessibility responsiveness, but the worse the performance for everything else  
+	static const double TimeSliceMilliseconds = 0.5;
+	const double TimeSliceSeconds = TimeSliceMilliseconds * 0.001;
+	const double StartTime = FPlatformTime::Seconds();
+	double RemainingTime = TimeSliceSeconds;
+	while (RemainingTime > 0)
+	{
+		while (AccessibleTaskProcessingQueue.Num() > 0)
+		{
+			FSlateAccessibleTask TaskToRun = AccessibleTaskProcessingQueue.Pop();
+			TaskToRun.DoTask();
+		}
+		double ElapsedTime = FPlatformTime::Seconds() - StartTime;
+		RemainingTime = TimeSliceSeconds - ElapsedTime;
+		if (RemainingTime > 0)
+		{
+
+#if ACCESSIBILITY_DEBUG_RESPONSIVENESS
+			EnqueueEvent->Wait(FTimespan::FromSeconds(RemainingTime));
+#endif
+			ElapsedTime = FPlatformTime::Seconds() - StartTime;
+			RemainingTime = TimeSliceSeconds - ElapsedTime;
+		}
+	}
+#endif 
+}
+
 void FSlateAccessibleMessageHandler::OnWidgetRemoved(SWidget* Widget)
 {
 	if (IsActive())
diff --git a/Engine/Source/Runtime/SlateCore/Public/Widgets/Accessibility/SlateAccessibleMessageHandler.h b/Engine/Source/Runtime/SlateCore/Public/Widgets/Accessibility/SlateAccessibleMessageHandler.h
--- a/Engine/Source/Runtime/SlateCore/Public/Widgets/Accessibility/SlateAccessibleMessageHandler.h
+++ b/Engine/Source/Runtime/SlateCore/Public/Widgets/Accessibility/SlateAccessibleMessageHandler.h
@@ -4,16 +4,19 @@
 
 #if WITH_ACCESSIBILITY
 
-#include "GenericPlatform/GenericAccessibleInterfaces.h"
+#include "GenericPlatform/Accessibility/GenericAccessibleInterfaces.h"
 #include "Misc/Variant.h"
 
 class FSlateAccessibleWidget;
 class SWidget;
 class SWindow;
 
+// Used to radically increase the responsiveness of accessibility navigation for debugging pruposes. The implementation is too slow to be feasible for shipping builds. 
+#define ACCESSIBILITY_DEBUG_RESPONSIVENESS 0 && !UE_BUILD_SHIPPING
+
 /**
  * Message handling system for Slate Accessibility API, dealing with both receiving events and pushing them back to the platform layer.
- * The message handler is also responsible for processing the Slate widget tree and 
+ * The message handler is also responsible for processing the Slate widget tree and queuing/processing accessibility requests for widget data from the platform.
  */
 class SLATECORE_API FSlateAccessibleMessageHandler : public FGenericAccessibleMessageHandler
 {
@@ -26,6 +29,7 @@ public:
 	virtual TSharedPtr<IAccessibleWidget> GetAccessibleWindow(const TSharedRef<FGenericWindow>& InWindow) const override;
 	virtual AccessibleWidgetId GetAccessibleWindowId(const TSharedRef<FGenericWindow>& InWindow) const override;
 	virtual TSharedPtr<IAccessibleWidget> GetAccessibleWidgetFromId(AccessibleWidgetId Id) const override;
+	virtual void RunInThread(const TFunction<void()>& InFunction, bool bWaitForCompletion = true, ENamedThreads::Type InThread = ENamedThreads::GameThread) override;
 	//~
 
 	/**
@@ -51,12 +55,54 @@ public:
 	 */
 	void MarkDirty() { bDirty = true; }
 
+	/**
+	* Processes all the queued tasks in the AccessibleTaskQueue
+	* Should only be called from FSlateApplication::TickPlatform() 
+	*/
+	void ProcessAccessibleTasks();
+
 	/**
 	 * Process any pending Slate widgets and update the accessible widget tree.
 	 */
 	void Tick();
 
 private:
+	/**
+	*  A helper class that wraps an accessibility task and the event to be triggered when the task finishes executing. 
+	*/
+	class FSlateAccessibleTask
+	{
+	public:
+		/** If completion event is nullptr, The task is treated as asynchronous */
+		FSlateAccessibleTask(const TFunction<void()>& InTask, FEvent* InCompletionEvent = nullptr)
+			: Task(InTask)
+			, CompletionEvent(InCompletionEvent)
+		{
+			check(Task);
+		}
+
+		/** Executes the task and triggers the completion event if available */
+		void DoTask()
+		{
+			// We should always have a valid task 
+			check(Task);
+			Task();
+			if (CompletionEvent)
+			{
+				CompletionEvent->Trigger();
+			}
+		}
+	private:
+		/** A lambda function that requests accessibility data from Slate in the Game Thread */
+		TFunction<void()> Task;
+
+		/** The event to be triggered when the task has completed execution. Can be null if the task should be executed asynchronously */
+		FEvent* CompletionEvent;
+	};
+
+	/** Queues an FSlateAccessibleTask to be processed in the Game Thread */
+	void EnqueueAccessibleTask(const FSlateAccessibleTask& InAccessibleTask);
+
 	struct FWidgetAndParent {
 		FWidgetAndParent(TWeakPtr<SWidget> InWidget, TSharedRef<FSlateAccessibleWidget> InParent)
 			: Widget(InWidget), Parent(InParent)
@@ -67,8 +113,30 @@ private:
 		TSharedRef<FSlateAccessibleWidget> Parent;
 	};
 
-	/** A list of widgets waiting to be processed in order to keep the accessible widget tree up to date. */
+	/** A list of widgets waiting to be processed in order to keep the accessible widget tree up to date */
 	TArray<FWidgetAndParent> ToProcess;
+
+	/** 
+	* A queue of FSlateAccessibleTasks to be processed in the Game Thread 
+	* This queue is only meant to store tasks, NOT process them . 
+	* We opt for a TArray as a TQueue has slower performance in this use case 
+	* @see ProcessAccessibleTasks() 
+	*/
+	TArray<FSlateAccessibleTask> AccessibleTaskStorageQueue;
+	/**
+	* The queue used to process all queued accessible tasks 
+	* The contents of AccessibleTaskStorageQueue is  swapped inot this queue as an optimization to avoid holding on to the lock 
+	* @see ProcessAccessibleTasks() 
+	*/
+	TArray<FSlateAccessibleTask> AccessibleTaskProcessingQueue;
+
+	/** Critical section used for synchronization of the AccessibleTaskQueue */
+	FCriticalSection QueueCriticalSection;
+
+#if ACCESSIBILITY_DEBUG_RESPONSIVENESS
+	FEvent* EnqueueEvent;
+#endif 
+
 	/** If true, Tick() will begin the update process to the accessible widget tree. Use MarkDirty() to set. */
 	bool bDirty;
 };
diff --git a/Engine/Source/Runtime/UMG/Public/Components/WidgetInteractionComponent.h b/Engine/Source/Runtime/UMG/Public/Components/WidgetInteractionComponent.h
--- a/Engine/Source/Runtime/UMG/Public/Components/WidgetInteractionComponent.h
+++ b/Engine/Source/Runtime/UMG/Public/Components/WidgetInteractionComponent.h
@@ -184,7 +184,7 @@ public:
 	UFUNCTION(BlueprintCallable, Category = "Interaction")
 	void SetFocus(UWidget* FocusWidget);
 
-private:
+protected:
 	/**
 	 * Represents the virtual user in slate.  When this component is registered, it gets a handle to the 
 	 * virtual slate user it will be, so virtual slate user 0, is probably real slate user 8, as that's the first
