Sequencer: Ensure that multiple ObjectToken containers do not exist for the same object when replacing restore state tokens.

#note this is a copy from 10774053

diff --git a/Engine/Source/Runtime/MovieScene/Private/Evaluation/MovieScenePreAnimatedState.cpp b/Engine/Source/Runtime/MovieScene/Private/Evaluation/MovieScenePreAnimatedState.cpp
--- a/Engine/Source/Runtime/MovieScene/Private/Evaluation/MovieScenePreAnimatedState.cpp
+++ b/Engine/Source/Runtime/MovieScene/Private/Evaluation/MovieScenePreAnimatedState.cpp
@@ -175,6 +175,35 @@ void TMovieSceneSavedTokens<TokenType>::OnPreAnimated(ECapturePreAnimatedState C
 	}
 }
 
+template<typename TokenType>
+void TMovieSceneSavedTokens<TokenType>::CopyFrom(TMovieSceneSavedTokens& OtherTokens)
+{
+	for (const FMovieSceneEntityAndAnimTypeID& Entity : OtherTokens.AnimatedEntities)
+	{
+		if (!AnimatedEntities.Contains(Entity))
+		{
+			AnimatedEntities.Add(Entity);
+		}
+	}
+
+	for (int32 OtherIndex = 0; OtherIndex < OtherTokens.AllAnimatedTypeIDs.Num(); ++OtherIndex)
+	{
+		FMovieSceneAnimTypeID OtherTypeID = OtherTokens.AllAnimatedTypeIDs[OtherIndex];
+
+		const int32 ExistingIndex = AllAnimatedTypeIDs.IndexOfByKey(OtherTypeID);
+		if (ExistingIndex != INDEX_NONE)
+		{
+			PreAnimatedTokens[ExistingIndex] = MoveTemp(OtherTokens.PreAnimatedTokens[OtherIndex]);
+		}
+		else
+		{
+			AllAnimatedTypeIDs.Add(OtherTypeID);
+			PreAnimatedTokens.Add(MoveTemp(OtherTokens.PreAnimatedTokens[OtherIndex]));
+		}
+	}
+}
+
+
 template<typename TokenType>
 void TMovieSceneSavedTokens<TokenType>::Restore(IMovieScenePlayer& Player)
 {
@@ -412,24 +441,34 @@ void FMovieScenePreAnimatedState::OnObjectsReplaced(const TMap<UObject*, UObject
 	{
 		UObject* OldObject = Iter->Key;
 		UObject* NewObject = Iter->Value;
+		if (!OldObject || !NewObject)
+		{
+			continue;
+		}
 
 		FObjectKey OldKey = FObjectKey(OldObject);
-		if (OldObject && NewObject && ObjectTokens.Contains(OldKey))
+		if (!ObjectTokens.Contains(OldKey))
 		{
-			FObjectKey NewKey = FObjectKey(NewObject);
+			continue;
+		}
 
-			ObjectTokens.Add(NewKey, MoveTemp(ObjectTokens[OldKey]));
-			ObjectTokens[NewKey].SetPayload(NewObject);
+		FObjectKey NewKey = FObjectKey(NewObject);
 
+		{
+			TMovieSceneSavedTokens<IMovieScenePreAnimatedTokenPtr>& NewTokens = ObjectTokens.FindOrAdd(NewKey, TMovieSceneSavedTokens<IMovieScenePreAnimatedTokenPtr>(NewObject));
+			TMovieSceneSavedTokens<IMovieScenePreAnimatedTokenPtr>& OldTokens = ObjectTokens.FindChecked(OldKey);
+
+			NewTokens.CopyFrom(OldTokens);
 			ObjectTokens.Remove(OldKey);
+			// NewTokens is not invalid
+		}
 
-			for (auto& Pair : EntityToAnimatedObjects)
+		for (auto& Pair : EntityToAnimatedObjects)
+		{
+			if (Pair.Value.Contains(OldKey))
 			{
-				if (Pair.Value.Contains(OldKey))
-				{
-					Pair.Value.Add(NewKey);
-					Pair.Value.Remove(OldKey);
-				}
+				Pair.Value.AddUnique(NewKey);
+				Pair.Value.Remove(OldKey);
 			}
 		}
 	}
diff --git a/Engine/Source/Runtime/MovieScene/Public/Evaluation/MovieScenePreAnimatedState.h b/Engine/Source/Runtime/MovieScene/Public/Evaluation/MovieScenePreAnimatedState.h
--- a/Engine/Source/Runtime/MovieScene/Public/Evaluation/MovieScenePreAnimatedState.h
+++ b/Engine/Source/Runtime/MovieScene/Public/Evaluation/MovieScenePreAnimatedState.h
@@ -147,7 +147,7 @@ struct MOVIESCENE_API TMovieSceneSavedTokens
 	 */
 	void Reset();
 
-	void SetPayload(PayloadType&& InPayload) { Payload = MoveTemp(InPayload); }
+	void CopyFrom(TMovieSceneSavedTokens& Other);
 
 private:
 
