Sequencer: Exclude some PostSerialize codepaths in non-editor builds. No longer call SetupOverrides for reference collector archives

diff --git a/Engine/Source/Runtime/MovieScene/Private/Evaluation/MovieSceneEvaluationTemplate.cpp b/Engine/Source/Runtime/MovieScene/Private/Evaluation/MovieSceneEvaluationTemplate.cpp
--- a/Engine/Source/Runtime/MovieScene/Private/Evaluation/MovieSceneEvaluationTemplate.cpp
+++ b/Engine/Source/Runtime/MovieScene/Private/Evaluation/MovieSceneEvaluationTemplate.cpp
@@ -19,6 +19,7 @@ void FMovieSceneTemplateGenerationLedger::AddTrack(const FGuid& InSignature, FMo
 	TrackSignatureToTrackIdentifier.Add(InSignature, Identifier);
 }
 
+#if WITH_EDITORONLY_DATA
 void FMovieSceneEvaluationTemplate::PostSerialize(const FArchive& Ar)
 {
 	if (Ar.IsLoading())
@@ -34,6 +35,7 @@ void FMovieSceneEvaluationTemplate::PostSerialize(const FArchive& Ar)
 		}
 	}
 }
+#endif
 
 void FMovieSceneEvaluationTemplate::ResetFieldData()
 {
diff --git a/Engine/Source/Runtime/MovieScene/Private/Evaluation/MovieSceneEvaluationTrack.cpp b/Engine/Source/Runtime/MovieScene/Private/Evaluation/MovieSceneEvaluationTrack.cpp
--- a/Engine/Source/Runtime/MovieScene/Private/Evaluation/MovieSceneEvaluationTrack.cpp
+++ b/Engine/Source/Runtime/MovieScene/Private/Evaluation/MovieSceneEvaluationTrack.cpp
@@ -132,8 +132,9 @@ FMovieSceneEvaluationTrack::FMovieSceneEvaluationTrack(const FGuid& InObjectBind
 
 void FMovieSceneEvaluationTrack::PostSerialize(const FArchive& Ar)
 {
-	if (Ar.IsLoading())
+	if (Ar.IsLoading() && !Ar.IsObjectReferenceCollector())
 	{
+#if WITH_EDITORONLY_DATA
 		// Guard against serialization mismatches where structs have been removed
 		TArray<int32, TInlineAllocator<2>> ImplsToRemove;
 		for (int32 Index = 0; Index < ChildTemplates.Num(); ++Index)
@@ -152,8 +153,10 @@ void FMovieSceneEvaluationTrack::PostSerialize(const FArchive& Ar)
 				Segment.Impls.RemoveAll([&](const FSectionEvaluationData& In) { return ImplsToRemove.Contains(In.ImplIndex); });
 			}
 		}
+#endif
+
+		SetupOverrides();
 	}
-	SetupOverrides();
 }
 
 void FMovieSceneEvaluationTrack::DefineAsSingleTemplate(FMovieSceneEvalTemplatePtr&& InTemplate)
diff --git a/Engine/Source/Runtime/MovieScene/Public/Evaluation/MovieSceneEvaluationTemplate.h b/Engine/Source/Runtime/MovieScene/Public/Evaluation/MovieSceneEvaluationTemplate.h
--- a/Engine/Source/Runtime/MovieScene/Public/Evaluation/MovieSceneEvaluationTemplate.h
+++ b/Engine/Source/Runtime/MovieScene/Public/Evaluation/MovieSceneEvaluationTemplate.h
@@ -282,7 +282,9 @@ public:
 	/**
 	 * Called after this template has been serialized in some way
 	 */
+#if WITH_EDITORONLY_DATA
 	MOVIESCENE_API void PostSerialize(const FArchive& Ar);
+#endif
 
 	/**
 	 * Purge any stale tracks we may have
@@ -360,4 +362,6 @@ private:
 	FMovieSceneSubSectionFieldData SubSectionFieldData;
 
 };
-template<> struct TStructOpsTypeTraits<FMovieSceneEvaluationTemplate> : public TStructOpsTypeTraitsBase2<FMovieSceneEvaluationTemplate> { enum { WithPostSerialize = true }; };
\ No newline at end of file
+#if WITH_EDITORONLY_DATA
+template<> struct TStructOpsTypeTraits<FMovieSceneEvaluationTemplate> : public TStructOpsTypeTraitsBase2<FMovieSceneEvaluationTemplate> { enum { WithPostSerialize = true }; };
+#endif
\ No newline at end of file
