Sequencer: Speculative cook non-determinism fix by changing a sort to a stable sort

diff --git a/Engine/Source/Runtime/MovieScene/Private/Compilation/MovieSceneCompiler.cpp b/Engine/Source/Runtime/MovieScene/Private/Compilation/MovieSceneCompiler.cpp
--- a/Engine/Source/Runtime/MovieScene/Private/Compilation/MovieSceneCompiler.cpp
+++ b/Engine/Source/Runtime/MovieScene/Private/Compilation/MovieSceneCompiler.cpp
@@ -252,7 +252,7 @@ void FMovieSceneCompiler::Compile(UMovieSceneSequence& InCompileSequence, IMovie
 		//  1. Group
 		//  2. Hierarchical bias
 		//  3. Evaluation priority
-		CompileData.Sort(SortPredicate);
+		CompileData.StableSort(SortPredicate);
 
 		// Compose the final result for the compiled range
 		FCompiledGroupResult Result(It.Range());
@@ -353,7 +353,7 @@ void FMovieSceneCompiler::CompileRange(TRange<FFrameNumber> InGlobalRange, UMovi
 		//  1. Group
 		//  2. Hierarchical bias
 		//  3. Evaluation priority
-		SortedCompileData.Sort(SortPredicate);
+		SortedCompileData.StableSort(SortPredicate);
 
 		// Compose the final result for the compiled range
 		FCompiledGroupResult Result(CompiledRange);
