Sequencer: Remove unused variables.

diff --git a/Engine/Source/Runtime/MovieScene/Private/Sections/MovieSceneSubSection.cpp b/Engine/Source/Runtime/MovieScene/Private/Sections/MovieSceneSubSection.cpp
--- a/Engine/Source/Runtime/MovieScene/Private/Sections/MovieSceneSubSection.cpp
+++ b/Engine/Source/Runtime/MovieScene/Private/Sections/MovieSceneSubSection.cpp
@@ -44,9 +44,6 @@ FMovieSceneSequenceTransform UMovieSceneSubSection::OuterToInnerTransform() cons
 	const FFrameRate   OuterFrameRate = GetTypedOuter<UMovieScene>()->GetTickResolution();
 	const float        FrameRateScale = (OuterFrameRate == InnerFrameRate) ? 1.f : (InnerFrameRate / OuterFrameRate).AsDecimal();
 
-	const FFrameNumber StartOffset = Parameters.StartFrameOffset;
-	const FFrameNumber EndOffset = Parameters.bCanLoop ? Parameters.EndFrameOffset : 0;
-
 	const TRange<FFrameNumber> MovieScenePlaybackRange = GetValidatedInnerPlaybackRange(Parameters, *MovieScenePtr);
 	const FFrameNumber InnerStartTime = MovieScene::DiscreteInclusiveLower(MovieScenePlaybackRange);
 	const FFrameNumber OuterStartTime = MovieScene::DiscreteInclusiveLower(SubRange);
