Sequencer: Fix for last tick game time seconds getting set to 0 which ends up stopping sequence playback immediately.

diff --git a/Engine/Source/Runtime/MovieScene/Private/MovieSceneSequencePlayer.cpp b/Engine/Source/Runtime/MovieScene/Private/MovieSceneSequencePlayer.cpp
--- a/Engine/Source/Runtime/MovieScene/Private/MovieSceneSequencePlayer.cpp
+++ b/Engine/Source/Runtime/MovieScene/Private/MovieSceneSequencePlayer.cpp
@@ -56,7 +56,6 @@ UMovieSceneSequencePlayer::UMovieSceneSequencePlayer(const FObjectInitializer& I
 	, StartTime(0)
 	, DurationFrames(0)
 	, CurrentNumLoops(0)
-	, LastTickGameTimeSeconds(-1.f)
 {
 	PlayPosition.Reset(FFrameTime(0));
 
@@ -233,7 +232,7 @@ void UMovieSceneSequencePlayer::Pause()
 		Status = EMovieScenePlayerStatus::Paused;
 		TimeController->StopPlaying(GetCurrentTime());
 
-		LastTickGameTimeSeconds = -1;
+		LastTickGameTimeSeconds.Reset();
 
 		// Evaluate the sequence at its current time, with a status of 'stopped' to ensure that animated state pauses correctly. (ie. audio sounds should stop/pause)
 		{
@@ -311,7 +310,7 @@ void UMovieSceneSequencePlayer::StopInternal(FFrameTime TimeToResetTo)
 		}
 
 		CurrentNumLoops = 0;
-		LastTickGameTimeSeconds = -1;
+		LastTickGameTimeSeconds.Reset();
 
 		// Reset loop count on stop so that it doesn't persist to the next call to play
 		PlaybackSettings.LoopCount.Value = 0;
@@ -727,9 +726,9 @@ void UMovieSceneSequencePlayer::Update(const float DeltaSeconds)
 
 		float DeltaTimeForFunction = DeltaSeconds;
 
-		if (LastTickGameTimeSeconds >= 0.f)
+		if (LastTickGameTimeSeconds.IsSet() && LastTickGameTimeSeconds.GetValue() >= 0.f)
 		{
-			DeltaTimeForFunction = CurrentWorldTime - LastTickGameTimeSeconds;
+			DeltaTimeForFunction = CurrentWorldTime - LastTickGameTimeSeconds.GetValue();
 		}
 
 		TimeController->Tick(DeltaTimeForFunction, PlayRate);
@@ -743,7 +742,10 @@ void UMovieSceneSequencePlayer::Update(const float DeltaSeconds)
 		UpdateTimeCursorPosition(NewTime, EUpdatePositionMethod::Play);
 	}
 
-	LastTickGameTimeSeconds = CurrentWorldTime;
+	if (World)
+	{
+		LastTickGameTimeSeconds = CurrentWorldTime;
+	}
 }
 
 void UMovieSceneSequencePlayer::UpdateTimeCursorPosition(FFrameTime NewPosition, EUpdatePositionMethod Method)
diff --git a/Engine/Source/Runtime/MovieScene/Public/MovieSceneSequencePlayer.h b/Engine/Source/Runtime/MovieScene/Public/MovieSceneSequencePlayer.h
--- a/Engine/Source/Runtime/MovieScene/Public/MovieSceneSequencePlayer.h
+++ b/Engine/Source/Runtime/MovieScene/Public/MovieSceneSequencePlayer.h
@@ -687,9 +687,7 @@ private:
 
 	/**
 	* The last world game time at which we were ticked. Game time used is dependent on bTickEvenWhenPaused
-	* Valid only if we've been ticked at least once since having a tick interval; otherwise set to -1.f
+	* Valid only if we've been ticked at least once since having a tick interval
 	*/
-	float LastTickGameTimeSeconds;
-
-
+	TOptional<float> LastTickGameTimeSeconds;
 };
