Sequencer: Added a Timecode clock which sets the sequencer clock to the timecode provider. Renamed the existing Timecode clock to Relative Timecode.

diff --git a/Engine/Plugins/VirtualProduction/Takes/Source/TakeRecorder/Private/Recorder/TakeRecorder.cpp b/Engine/Plugins/VirtualProduction/Takes/Source/TakeRecorder/Private/Recorder/TakeRecorder.cpp
--- a/Engine/Plugins/VirtualProduction/Takes/Source/TakeRecorder/Private/Recorder/TakeRecorder.cpp
+++ b/Engine/Plugins/VirtualProduction/Takes/Source/TakeRecorder/Private/Recorder/TakeRecorder.cpp
@@ -749,7 +749,7 @@ void UTakeRecorder::Start(const FTimecode& InTimecodeSource)
 		UMovieScene*   MovieScene = SequenceAsset->GetMovieScene();
 		if (MovieScene)
 		{
-			MovieScene->SetClockSource(EUpdateClockSource::Timecode);
+			MovieScene->SetClockSource(EUpdateClockSource::RelativeTimecode);
 			Sequencer->ResetTimeController();
 
 			FFrameRate FrameRate = MovieScene->GetDisplayRate();
diff --git a/Engine/Source/Editor/Sequencer/Private/Sequencer.cpp b/Engine/Source/Editor/Sequencer/Private/Sequencer.cpp
--- a/Engine/Source/Editor/Sequencer/Private/Sequencer.cpp
+++ b/Engine/Source/Editor/Sequencer/Private/Sequencer.cpp
@@ -11517,6 +11517,7 @@ void FSequencer::ResetTimeController()
 	{
 	case EUpdateClockSource::Audio:    TimeController = MakeShared<FMovieSceneTimeController_AudioClock>();         break;
 	case EUpdateClockSource::Platform: TimeController = MakeShared<FMovieSceneTimeController_PlatformClock>();      break;
+	case EUpdateClockSource::RelativeTimecode: TimeController = MakeShared<FMovieSceneTimeController_RelativeTimecodeClock>();      break;
 	case EUpdateClockSource::Timecode: TimeController = MakeShared<FMovieSceneTimeController_TimecodeClock>();      break;
 	case EUpdateClockSource::Custom:   TimeController = MovieScene->MakeCustomTimeController(GetPlaybackContext()); break;
 	default:                           TimeController = MakeShared<FMovieSceneTimeController_Tick>();               break;
diff --git a/Engine/Source/Runtime/MovieScene/Private/MovieSceneSequencePlayer.cpp b/Engine/Source/Runtime/MovieScene/Private/MovieSceneSequencePlayer.cpp
--- a/Engine/Source/Runtime/MovieScene/Private/MovieSceneSequencePlayer.cpp
+++ b/Engine/Source/Runtime/MovieScene/Private/MovieSceneSequencePlayer.cpp
@@ -693,6 +693,7 @@ void UMovieSceneSequencePlayer::Initialize(UMovieSceneSequence* InSequence, cons
 		{
 		case EUpdateClockSource::Audio:    TimeController = MakeShared<FMovieSceneTimeController_AudioClock>();    break;
 		case EUpdateClockSource::Platform: TimeController = MakeShared<FMovieSceneTimeController_PlatformClock>(); break;
+		case EUpdateClockSource::RelativeTimecode: TimeController = MakeShared<FMovieSceneTimeController_RelativeTimecodeClock>(); break;
 		case EUpdateClockSource::Timecode: TimeController = MakeShared<FMovieSceneTimeController_TimecodeClock>(); break;
 		default:                           TimeController = MakeShared<FMovieSceneTimeController_Tick>();          break;
 		}
diff --git a/Engine/Source/Runtime/MovieScene/Private/MovieSceneTimeController.cpp b/Engine/Source/Runtime/MovieScene/Private/MovieSceneTimeController.cpp
--- a/Engine/Source/Runtime/MovieScene/Private/MovieSceneTimeController.cpp
+++ b/Engine/Source/Runtime/MovieScene/Private/MovieSceneTimeController.cpp
@@ -108,7 +108,7 @@ double FMovieSceneTimeController_AudioClock::GetCurrentTime() const
 	return AudioDevice ? AudioDevice->GetAudioClock() : FPlatformTime::Seconds();
 }
 
-double FMovieSceneTimeController_TimecodeClock::GetCurrentTime() const
+double FMovieSceneTimeController_RelativeTimecodeClock::GetCurrentTime() const
 {
 	const TOptional<FQualifiedFrameTime> CurrentFrameTime = FApp::GetCurrentFrameTime();
 	if (CurrentFrameTime.IsSet())
@@ -121,6 +121,21 @@ double FMovieSceneTimeController_TimecodeClock::GetCurrentTime() const
 	}
 }
 
+FFrameTime FMovieSceneTimeController_TimecodeClock::OnRequestCurrentTime(const FQualifiedFrameTime& InCurrentTime, float InPlayRate)
+{
+	if (GEngine && GEngine->GetTimecodeProvider() && GEngine->GetTimecodeProvider()->GetSynchronizationState() == ETimecodeProviderSynchronizationState::Synchronized)
+	{
+		FTimecode Timecode = FApp::GetTimecode();
+		FFrameRate FrameRate = FApp::GetTimecodeFrameRate();
+
+		// Convert timecode to raw number of frames at the timecode's framerate.
+		FFrameNumber FrameNumber = Timecode.ToFrameNumber(FrameRate);
+		return FFrameRate::TransformTime(FFrameTime(FrameNumber), FrameRate, InCurrentTime.Rate);
+	}
+
+	return FFrameTime(0);
+}
+
 void FMovieSceneTimeController_Tick::OnStartPlaying(const FQualifiedFrameTime& InStartTime)
 {
 	CurrentOffsetSeconds = 0.0;
diff --git a/Engine/Source/Runtime/MovieScene/Public/MovieSceneFwd.h b/Engine/Source/Runtime/MovieScene/Public/MovieSceneFwd.h
--- a/Engine/Source/Runtime/MovieScene/Public/MovieSceneFwd.h
+++ b/Engine/Source/Runtime/MovieScene/Public/MovieSceneFwd.h
@@ -48,6 +48,9 @@ enum class EUpdateClockSource : uint8
 	/** Use the audio clock for timing. Does not honor world or actor pause state. */
 	Audio,
 
+	/** Time relative to the timecode provider for timing. Does not honor world or actor pause state. */
+	RelativeTimecode,
+
 	/** Use current timecode provider for timing. Does not honor world or actor pause state. */
 	Timecode,
 
diff --git a/Engine/Source/Runtime/MovieScene/Public/MovieSceneTimeController.h b/Engine/Source/Runtime/MovieScene/Public/MovieSceneTimeController.h
--- a/Engine/Source/Runtime/MovieScene/Public/MovieSceneTimeController.h
+++ b/Engine/Source/Runtime/MovieScene/Public/MovieSceneTimeController.h
@@ -117,14 +117,23 @@ struct MOVIESCENE_API FMovieSceneTimeController_AudioClock : FMovieSceneTimeCont
 
 
 /**
-* A timing manager that retrieves its time from the Timecode clock
+* A timing manager that retrieves its time relative to the Timecode clock
 */
-struct MOVIESCENE_API FMovieSceneTimeController_TimecodeClock : FMovieSceneTimeController_ExternalClock
+struct MOVIESCENE_API FMovieSceneTimeController_RelativeTimecodeClock : FMovieSceneTimeController_ExternalClock
 {
 	virtual double GetCurrentTime() const override;
 };
 
 
+/**
+* A timing manager that retrieves its time from the Timecode clock
+*/
+struct MOVIESCENE_API FMovieSceneTimeController_TimecodeClock : FMovieSceneTimeController
+{
+	virtual FFrameTime OnRequestCurrentTime(const FQualifiedFrameTime& InCurrentTime, float InPlayRate) override;
+};
+
+
 /**
  * A timing manager that accumulates delta times from a world
  */
