Sequencer: Fixes to prevent reentrant evaluation during stop.

Added bIsEvaluating in StopInternal()
Added latent play action

diff --git a/Engine/Source/Runtime/MovieScene/Private/MovieSceneSequencePlayer.cpp b/Engine/Source/Runtime/MovieScene/Private/MovieSceneSequencePlayer.cpp
--- a/Engine/Source/Runtime/MovieScene/Private/MovieSceneSequencePlayer.cpp
+++ b/Engine/Source/Runtime/MovieScene/Private/MovieSceneSequencePlayer.cpp
@@ -138,6 +138,12 @@ void UMovieSceneSequencePlayer::PlayLooping(int32 NumLoops)
 
 void UMovieSceneSequencePlayer::PlayInternal()
 {
+	if (bIsEvaluating)
+	{
+		LatentActions.Emplace(FLatentAction::EType::Play);
+		return;
+	}
+
 	if (!IsPlaying() && Sequence && CanPlay())
 	{
 		float PlayRate = bReversePlayback ? -PlaybackSettings.PlayRate : PlaybackSettings.PlayRate;
@@ -221,14 +227,14 @@ void UMovieSceneSequencePlayer::PlayInternal()
 
 void UMovieSceneSequencePlayer::Pause()
 {
-	if (IsPlaying())
+	if (bIsEvaluating)
 	{
-		if (bIsEvaluating)
-		{
-			LatentActions.Emplace(FLatentAction::EType::Pause);
-			return;
-		}
+		LatentActions.Emplace(FLatentAction::EType::Pause);
+		return;
+	}
 
+	if (IsPlaying())
+	{
 		Status = EMovieScenePlayerStatus::Paused;
 		TimeController->StopPlaying(GetCurrentTime());
 
@@ -292,14 +298,15 @@ void UMovieSceneSequencePlayer::StopAtCurrentTime()
 
 void UMovieSceneSequencePlayer::StopInternal(FFrameTime TimeToResetTo)
 {
-	if (IsPlaying() || IsPaused())
+	if (bIsEvaluating)
 	{
-		if (bIsEvaluating)
-		{
-			LatentActions.Emplace(FLatentAction::EType::Stop, TimeToResetTo);
-			return;
-		}
+		LatentActions.Emplace(FLatentAction::EType::Stop, TimeToResetTo);
+		return;
+	}
 
+	if (IsPlaying() || IsPaused())
+	{
+		bIsEvaluating = true;
 		Status = EMovieScenePlayerStatus::Stopped;
 
 		// Put the cursor at the specified position
@@ -353,6 +360,9 @@ void UMovieSceneSequencePlayer::StopInternal(FFrameTime TimeToResetTo)
 		{
 			OnStop.Broadcast();
 		}
+		bIsEvaluating = false;
+
+		ApplyLatentActions();
 	}
 }
 
@@ -942,6 +952,7 @@ void UMovieSceneSequencePlayer::ApplyLatentActions()
 		{
 		case FLatentAction::EType::Stop:   StopInternal(LatentAction.Position); continue;
 		case FLatentAction::EType::Pause:  Pause();                             continue;
+		case FLatentAction::EType::Play:   PlayInternal();                      continue;
 		}
 
 		check(LatentAction.Type == FLatentAction::EType::Update);
diff --git a/Engine/Source/Runtime/MovieScene/Public/MovieSceneSequencePlayer.h b/Engine/Source/Runtime/MovieScene/Public/MovieSceneSequencePlayer.h
--- a/Engine/Source/Runtime/MovieScene/Public/MovieSceneSequencePlayer.h
+++ b/Engine/Source/Runtime/MovieScene/Public/MovieSceneSequencePlayer.h
@@ -635,7 +635,7 @@ protected:
 
 	struct FLatentAction
 	{
-		enum class EType : uint8 { Stop, Pause, Update };
+		enum class EType : uint8 { Stop, Pause, Update, Play };
 
 		FLatentAction(EType InType, FFrameTime DesiredTime = 0)
 			: Type(InType), Position(DesiredTime)
