Sequencer: Improve editing of blending camera cuts.

Split helper method to fixup consecutive sections, making a blending and non-blending one.
Don't edit easing curves because they're already computed in tracks' UpdateEasing method (and better).

diff --git a/Engine/Source/Runtime/MovieScene/Private/MovieSceneCommonHelpers.cpp b/Engine/Source/Runtime/MovieScene/Private/MovieSceneCommonHelpers.cpp
--- a/Engine/Source/Runtime/MovieScene/Private/MovieSceneCommonHelpers.cpp
+++ b/Engine/Source/Runtime/MovieScene/Private/MovieSceneCommonHelpers.cpp
@@ -88,7 +88,44 @@ void MovieSceneHelpers::SortConsecutiveSections(TArray<UMovieSceneSection*>& Sec
 	);
 }
 
-void MovieSceneHelpers::FixupConsecutiveSections(TArray<UMovieSceneSection*>& Sections, UMovieSceneSection& Section, bool bDelete, bool bOnlyOnSameRow, bool bAllowOverlapBlending)
+void MovieSceneHelpers::FixupConsecutiveSections(TArray<UMovieSceneSection*>& Sections, UMovieSceneSection& Section, bool bDelete)
+{
+	// Find the previous section and extend it to take the place of the section being deleted
+	int32 SectionIndex = INDEX_NONE;
+
+	TRange<FFrameNumber> SectionRange = Section.GetRange();
+
+	if (Sections.Find(&Section, SectionIndex))
+	{
+		int32 PrevSectionIndex = SectionIndex - 1;
+		if( Sections.IsValidIndex( PrevSectionIndex ) )
+		{
+			// Extend the previous section
+			if (bDelete)
+			{
+				Sections[PrevSectionIndex]->SetEndFrame(SectionRange.GetUpperBound());
+			}
+			else
+			{
+				Sections[PrevSectionIndex]->SetEndFrame(TRangeBound<FFrameNumber>::FlipInclusion(SectionRange.GetLowerBound()));
+			}
+		}
+
+		if( !bDelete )
+		{
+			int32 NextSectionIndex = SectionIndex + 1;
+			if(Sections.IsValidIndex(NextSectionIndex))
+			{
+				// Shift the next CameraCut's start time so that it starts when the new CameraCut ends
+				Sections[NextSectionIndex]->SetStartFrame(TRangeBound<FFrameNumber>::FlipInclusion(SectionRange.GetUpperBound()));
+			}
+		}
+	}
+
+	SortConsecutiveSections(Sections);
+}
+
+void MovieSceneHelpers::FixupConsecutiveBlendingSections(TArray<UMovieSceneSection*>& Sections, UMovieSceneSection& Section, bool bDelete)
 {
 	int32 SectionIndex = INDEX_NONE;
 
@@ -101,7 +138,7 @@ void MovieSceneHelpers::FixupConsecutiveSections(TArray<UMovieSceneSection*>& Se
 		if (Sections.IsValidIndex(PrevSectionIndex))
 		{
 			UMovieSceneSection* PrevSection = Sections[PrevSectionIndex];
-			if (!bOnlyOnSameRow || PrevSection->GetRowIndex() == Section.GetRowIndex())
+			if (PrevSection->GetRowIndex() == Section.GetRowIndex())
 			{
 				PrevSection->Modify();
 
@@ -110,29 +147,27 @@ void MovieSceneHelpers::FixupConsecutiveSections(TArray<UMovieSceneSection*>& Se
 					// The current section was deleted... extend the previous section to fill the gap.
 					PrevSection->SetEndFrame(SectionRange.GetUpperBound());
 				}
-				else if (!bAllowOverlapBlending)
-				{
-					// Keep the previous section snug against the current section's start.
-					PrevSection->SetEndFrame(TRangeBound<FFrameNumber>::FlipInclusion(SectionRange.GetLowerBound()));
-				}
 				else
 				{
 					// If we made a gap: adjust the previous section's end time so that it ends wherever the current section's ease-in ends.
-					// If we created an overlap: adjust the current section's ease-in so it ends where the previous section ends.
+					// If we created an overlap: calls to UMovieSceneTrack::UpdateEasing have already set the easing curves correctly based on overlaps.
 					const FFrameNumber GapOrOverlap = SectionRange.GetLowerBoundValue() - PrevSection->GetRange().GetUpperBoundValue();
 					if (GapOrOverlap > 0)
 					{
 						// It's a gap!
 						PrevSection->SetEndFrame(TRangeBound<FFrameNumber>::Exclusive(SectionRange.GetLowerBoundValue() + Section.Easing.GetEaseInDuration()));
 					}
-					else
-					{
-						// It's an overlap!
-						Section.Easing.AutoEaseInDuration = -GapOrOverlap.Value;
-					}
 				}
 			}
 		}
+		else
+		{
+			if (!bDelete)
+			{
+				// The given section is the first section. Let's clear its auto ease-in since there's no overlap anymore with a previous section.
+				Section.Easing.AutoEaseInDuration = 0;
+			}
+		}
 
 		// Find the next section and adjust its start time to match the moved/resized section's new end time.
 		if (!bDelete)
@@ -141,33 +176,25 @@ void MovieSceneHelpers::FixupConsecutiveSections(TArray<UMovieSceneSection*>& Se
 			if (Sections.IsValidIndex(NextSectionIndex))
 			{
 				UMovieSceneSection* NextSection = Sections[NextSectionIndex];
-				if (!bOnlyOnSameRow || NextSection->GetRowIndex() == Section.GetRowIndex())
+				if (NextSection->GetRowIndex() == Section.GetRowIndex())
 				{
 					NextSection->Modify();
 
-					if (!bAllowOverlapBlending)
-					{
-						// Keep the next section snug against the current section's end.
-						NextSection->SetStartFrame(TRangeBound<FFrameNumber>::FlipInclusion(SectionRange.GetUpperBound()));
-					}
-					else
+					// If we made a gap: adjust the next section's start time so that it lines up with the current section's end.
+					// If we created an overlap: adjust the next section's ease-in so it ends where the current section ends.
+					const FFrameNumber GapOrOverlap = NextSection->GetRange().GetLowerBoundValue() - SectionRange.GetUpperBoundValue();
+					if (GapOrOverlap > 0)
 					{
-						// If we made a gap: adjust the next section's start time so that it lines up with the current section's end.
-						// If we created an overlap: adjust the next section's ease-in so it ends where the current section ends.
-						const FFrameNumber GapOrOverlap = NextSection->GetRange().GetLowerBoundValue() - SectionRange.GetUpperBoundValue();
-						if (GapOrOverlap > 0)
-						{
-							// It's a gap!
-							NextSection->SetStartFrame(TRangeBound<FFrameNumber>::Inclusive(SectionRange.GetUpperBoundValue() - NextSection->Easing.GetEaseInDuration()));
-						}
-						else
-						{
-							// It's an overlap!
-							NextSection->Easing.AutoEaseInDuration = -GapOrOverlap.Value;
-						}
+						// It's a gap!
+						NextSection->SetStartFrame(TRangeBound<FFrameNumber>::Inclusive(SectionRange.GetUpperBoundValue() - NextSection->Easing.GetEaseInDuration()));
 					}
 				}
 			}
+			else
+			{
+				// The given section is the last section. Let's clear its auto ease-out since there's no overlap anymore with a next section.
+				Section.Easing.AutoEaseOutDuration = 0;
+			}
 		}
 	}
 
diff --git a/Engine/Source/Runtime/MovieScene/Public/MovieSceneCommonHelpers.h b/Engine/Source/Runtime/MovieScene/Public/MovieSceneCommonHelpers.h
--- a/Engine/Source/Runtime/MovieScene/Public/MovieSceneCommonHelpers.h
+++ b/Engine/Source/Runtime/MovieScene/Public/MovieSceneCommonHelpers.h
@@ -47,7 +47,9 @@ public:
 	 * @param Section The section that was modified 
 	 * @param bDelete Was this a deletion?
 	 */
-	static void FixupConsecutiveSections(TArray<UMovieSceneSection*>& Sections, UMovieSceneSection& Section, bool bDelete, bool bOnlyOnSameRow = false, bool bAllowOverlapBlending = false);
+	static void FixupConsecutiveSections(TArray<UMovieSceneSection*>& Sections, UMovieSceneSection& Section, bool bDelete);
+
+	static void FixupConsecutiveBlendingSections(TArray<UMovieSceneSection*>& Sections, UMovieSceneSection& Section, bool bDelete);
 
 	/*
  	 * Sort consecutive sections so that they are in order based on start time
diff --git a/Engine/Source/Runtime/MovieSceneTracks/Private/Tracks/MovieSceneCameraCutTrack.cpp b/Engine/Source/Runtime/MovieSceneTracks/Private/Tracks/MovieSceneCameraCutTrack.cpp
--- a/Engine/Source/Runtime/MovieSceneTracks/Private/Tracks/MovieSceneCameraCutTrack.cpp
+++ b/Engine/Source/Runtime/MovieSceneTracks/Private/Tracks/MovieSceneCameraCutTrack.cpp
@@ -70,7 +70,14 @@ UMovieSceneCameraCutSection* UMovieSceneCameraCutTrack::AddNewCameraCut(const FM
 	MovieSceneHelpers::SortConsecutiveSections(Sections);
 
 	// Once CameraCuts are sorted fixup the surrounding CameraCuts to fix any gaps
-	MovieSceneHelpers::FixupConsecutiveSections(Sections, *NewSection, false, true, bCanBlend);
+	if (bCanBlend)
+	{
+		MovieSceneHelpers::FixupConsecutiveBlendingSections(Sections, *NewSection, false);
+	}
+	else
+	{
+		MovieSceneHelpers::FixupConsecutiveSections(Sections, *NewSection, false);
+	}
 
 	return NewSection;
 }
@@ -131,8 +138,14 @@ void UMovieSceneCameraCutTrack::RemoveSection(UMovieSceneSection& Section)
 {
 	Sections.Remove(&Section);
 
-	MovieSceneHelpers::FixupConsecutiveSections(Sections, Section, true, true, bCanBlend);
-	MovieSceneHelpers::SortConsecutiveSections(Sections);
+	if (bCanBlend)
+	{
+		MovieSceneHelpers::FixupConsecutiveSections(Sections, Section, true);
+	}
+	else
+	{
+		MovieSceneHelpers::FixupConsecutiveBlendingSections(Sections, Section, true);
+	}
 
 	// @todo Sequencer: The movie scene owned by the section is now abandoned.  Should we offer to delete it?  
 }
@@ -140,7 +153,14 @@ void UMovieSceneCameraCutTrack::RemoveSection(UMovieSceneSection& Section)
 void UMovieSceneCameraCutTrack::RemoveSectionAt(int32 SectionIndex)
 {
 	UMovieSceneSection* SectionToDelete = Sections[SectionIndex];
-	MovieSceneHelpers::FixupConsecutiveSections(Sections, *SectionToDelete, true, true, bCanBlend);
+	if (bCanBlend)
+	{
+		MovieSceneHelpers::FixupConsecutiveSections(Sections, *SectionToDelete, true);
+	}
+	else
+	{
+		MovieSceneHelpers::FixupConsecutiveBlendingSections(Sections, *SectionToDelete, true);
+	}
 
 	Sections.RemoveAt(SectionIndex);
 	MovieSceneHelpers::SortConsecutiveSections(Sections);
@@ -187,7 +207,14 @@ FText UMovieSceneCameraCutTrack::GetDefaultDisplayName() const
 #if WITH_EDITOR
 void UMovieSceneCameraCutTrack::OnSectionMoved(UMovieSceneSection& Section, const FMovieSceneSectionMovedParams& Params)
 {
-	MovieSceneHelpers::FixupConsecutiveSections(Sections, Section, false, true, bCanBlend);
+	if (bCanBlend)
+	{
+		MovieSceneHelpers::FixupConsecutiveBlendingSections(Sections, Section, false);
+	}
+	else
+	{
+		MovieSceneHelpers::FixupConsecutiveSections(Sections, Section, false);
+	}
 }
 #endif
 
