Sequencer: Remove cook optimization from movie scene

diff --git a/Engine/Source/Runtime/MovieScene/Private/MovieScene.cpp b/Engine/Source/Runtime/MovieScene/Private/MovieScene.cpp
--- a/Engine/Source/Runtime/MovieScene/Private/MovieScene.cpp
+++ b/Engine/Source/Runtime/MovieScene/Private/MovieScene.cpp
@@ -76,20 +76,6 @@ void UMovieScene::Serialize( FArchive& Ar )
 	Ar.UsingCustomVersion(FMovieSceneEvaluationCustomVersion::GUID);
 	Ar.UsingCustomVersion(FSequencerObjectVersion::GUID);
 
-#if WITH_EDITOR
-
-	// Perform optimizations for cooking
-	if (Ar.IsCooking())
-	{
-		// @todo: Optimize master tracks?
-
-		// Optimize object bindings
-		OptimizeObjectArray(Spawnables);
-		OptimizeObjectArray(Possessables);
-	}
-
-#endif // WITH_EDITOR
-
 	Super::Serialize(Ar);
 
 #if WITH_EDITORONLY_DATA
@@ -143,41 +129,6 @@ void UMovieScene::Serialize( FArchive& Ar )
 
 #if WITH_EDITOR
 
-template<typename T>
-void UMovieScene::OptimizeObjectArray(TArray<T>& ObjectArray)
-{
-	for (int32 ObjectIndex = ObjectArray.Num() - 1; ObjectIndex >= 0; --ObjectIndex)
-	{
-		FGuid ObjectGuid = ObjectArray[ObjectIndex].GetGuid();
-
-		// Find the binding relating to this object, and optimize its tracks
-		// @todo: ObjectBindings mapped by ID to avoid linear search
-		for (int32 BindingIndex = 0; BindingIndex < ObjectBindings.Num(); ++BindingIndex)
-		{
-			FMovieSceneBinding& Binding = ObjectBindings[BindingIndex];
-			if (Binding.GetObjectGuid() != ObjectGuid)
-			{
-				continue;
-			}
-			
-			bool bShouldRemoveObject = false;
-
-			// Optimize any tracks
-			Binding.PerformCookOptimization(bShouldRemoveObject);
-
-			// Remove the object if it's completely redundant
-			if (bShouldRemoveObject)
-			{
-				ObjectBindings.RemoveAtSwap(BindingIndex, 1, false);
-				ObjectArray.RemoveAtSwap(ObjectIndex, 1, false);
-			}
-
-			// Process next object
-			break;
-		}
-	}
-}
-
 // @todo sequencer: Some of these methods should only be used by tools, and should probably move out of MovieScene!
 FGuid UMovieScene::AddSpawnable( const FString& Name, UObject& ObjectTemplate )
 {
diff --git a/Engine/Source/Runtime/MovieScene/Private/MovieSceneBinding.cpp b/Engine/Source/Runtime/MovieScene/Private/MovieSceneBinding.cpp
--- a/Engine/Source/Runtime/MovieScene/Private/MovieSceneBinding.cpp
+++ b/Engine/Source/Runtime/MovieScene/Private/MovieSceneBinding.cpp
@@ -16,23 +16,3 @@ bool FMovieSceneBinding::RemoveTrack(UMovieSceneTrack& Track)
 {
 	return (Tracks.RemoveSingle(&Track) != 0);
 }
-
-#if WITH_EDITOR
-void FMovieSceneBinding::PerformCookOptimization(bool& bShouldRemoveObject)
-{
-	for (int32 Index = Tracks.Num() - 1; Index >= 0; --Index)
-	{
-		ECookOptimizationFlags Flags = Tracks[Index]->GetCookOptimizationFlags();
-
-		if ((Flags & ECookOptimizationFlags::RemoveObject) != ECookOptimizationFlags::None)
-		{
-			bShouldRemoveObject = true;
-			return;
-		}
-		else if ((Flags & ECookOptimizationFlags::RemoveTrack) != ECookOptimizationFlags::None)
-		{
-			Tracks.RemoveAtSwap(Index, 1, false);
-		}
-	}
-}
-#endif
diff --git a/Engine/Source/Runtime/MovieScene/Public/MovieScene.h b/Engine/Source/Runtime/MovieScene/Public/MovieScene.h
--- a/Engine/Source/Runtime/MovieScene/Public/MovieScene.h
+++ b/Engine/Source/Runtime/MovieScene/Public/MovieScene.h
@@ -965,12 +965,6 @@ protected:
 	 */
 	void RemoveBinding(const FGuid& Guid);
 
-#if WITH_EDITOR
-	/** Templated helper for optimizing lists of possessables and spawnables for cook */
-	template<typename T>
-	void OptimizeObjectArray(TArray<T>& ObjectArray);
-#endif
-	
 protected:
 
 	/** Called after this object has been deserialized */
diff --git a/Engine/Source/Runtime/MovieScene/Public/MovieSceneBinding.h b/Engine/Source/Runtime/MovieScene/Public/MovieSceneBinding.h
--- a/Engine/Source/Runtime/MovieScene/Public/MovieSceneBinding.h
+++ b/Engine/Source/Runtime/MovieScene/Public/MovieSceneBinding.h
@@ -137,15 +137,6 @@ struct FMovieSceneBinding
 		Tracks = InTracks;
 	}
 
-#if WITH_EDITOR
-	/**
-	 * Perform cook-time optimization on this object binding
-	 * @param bShouldRemoveObject 		(Out) Boolean that will set to true if this whole binding should be considered redundant.
-	 */
-	void PerformCookOptimization(bool& bShouldRemoveObject);
-
-#endif
-
 #if WITH_EDITORONLY_DATA
 	/**
 	* Get this folder's desired sorting order
diff --git a/Engine/Source/Runtime/MovieScene/Public/MovieSceneTrack.h b/Engine/Source/Runtime/MovieScene/Public/MovieSceneTrack.h
--- a/Engine/Source/Runtime/MovieScene/Public/MovieSceneTrack.h
+++ b/Engine/Source/Runtime/MovieScene/Public/MovieSceneTrack.h
@@ -18,18 +18,6 @@ struct FMovieSceneTrackSegmentBlender;
 struct FMovieSceneTrackRowSegmentBlender;
 struct IMovieSceneTemplateGenerator;
 
-/** Flags used to perform cook-time optimization of movie scene data */
-enum class ECookOptimizationFlags
-{
-	/** Perform no cook optimization */
-	None 			= 0,
-	/** Remove this track since its of no consequence to runtime */
-	RemoveTrack		= 1 << 0,
-	/** Remove this track's object since its of no consequence to runtime */
-	RemoveObject	= 1 << 1,
-};
-ENUM_CLASS_FLAGS(ECookOptimizationFlags)
-
 /** Track compiler arguments */
 struct FMovieSceneTrackCompilerArgs
 {
@@ -371,16 +359,6 @@ public:
 	 */
 	virtual void RemoveSectionAt(int32 SectionIndex) PURE_VIRTUAL(UMovieSceneSection::RemoveSectionAt, );
 
-#if WITH_EDITOR
-
-	/**
-	 * Called when this track's movie scene is being cooked to determine if/how this track should be cooked.
-	 * @return ECookOptimizationFlags detailing how to optimize this track
-	 */
-	virtual ECookOptimizationFlags GetCookOptimizationFlags() const { return ECookOptimizationFlags::None; }
-
-#endif
-
 #if WITH_EDITORONLY_DATA
 
 	/**
