Sequencer: Remove stray pragma disable/enable

#rnx

diff --git a/Engine/Source/Runtime/MovieScene/Private/MovieSceneSequencePlayer.cpp b/Engine/Source/Runtime/MovieScene/Private/MovieSceneSequencePlayer.cpp
--- a/Engine/Source/Runtime/MovieScene/Private/MovieSceneSequencePlayer.cpp
+++ b/Engine/Source/Runtime/MovieScene/Private/MovieSceneSequencePlayer.cpp
@@ -44,8 +44,6 @@ bool FMovieSceneSequencePlaybackSettings::SerializeFromMismatchedTag( const FPro
 	return false;
 }
 
-PRAGMA_DISABLE_OPTIMIZATION
-
 UMovieSceneSequencePlayer::UMovieSceneSequencePlayer(const FObjectInitializer& Init)
 	: Super(Init)
 	, Status(EMovieScenePlayerStatus::Stopped)
@@ -1283,5 +1281,3 @@ bool UMovieSceneSequencePlayer::CallRemoteFunction(UFunction* Function, void* Pa
 	return false;
 }
 
-
-PRAGMA_ENABLE_OPTIMIZATION
\ No newline at end of file
