Sequencer: Disallow looping in RPC_OnStopEvent because calls to PlayToFrame could loop to the start.

diff --git a/Engine/Source/Runtime/MovieScene/Private/MovieSceneSequencePlayer.cpp b/Engine/Source/Runtime/MovieScene/Private/MovieSceneSequencePlayer.cpp
--- a/Engine/Source/Runtime/MovieScene/Private/MovieSceneSequencePlayer.cpp
+++ b/Engine/Source/Runtime/MovieScene/Private/MovieSceneSequencePlayer.cpp
@@ -1092,11 +1092,17 @@ void UMovieSceneSequencePlayer::RPC_OnStopEvent_Implementation(FFrameTime Stoppe
 	}
 #endif
 
-	switch (Status.GetValue())
+	// If we're behind the target time to stop at, advance to that target time.
+	const bool bBehindTime = StoppedTime < PlayPosition.GetCurrentPosition();
+
+	if (bBehindTime)
 	{
-	case EMovieScenePlayerStatus::Playing:   PlayToFrame( StoppedTime);  break;
-	case EMovieScenePlayerStatus::Stopped:   JumpToFrame( StoppedTime);  break;
-	case EMovieScenePlayerStatus::Scrubbing: ScrubToFrame(StoppedTime);  break;
+		switch (Status.GetValue())
+		{
+		case EMovieScenePlayerStatus::Playing:   PlayToFrame(StoppedTime);  break;
+		case EMovieScenePlayerStatus::Stopped:   JumpToFrame(StoppedTime);  break;
+		case EMovieScenePlayerStatus::Scrubbing: ScrubToFrame(StoppedTime);  break;
+		}
 	}
 
 	StopInternal(StoppedTime);
