Sequencer: Fix client stop so that if the client is very close to the end, don't play to the end because it will stop soon anyway. Fix the test so that this only happens if the client is behind the target time to stop at.

#rnx

diff --git a/Engine/Source/Runtime/MovieScene/Private/MovieSceneSequencePlayer.cpp b/Engine/Source/Runtime/MovieScene/Private/MovieSceneSequencePlayer.cpp
--- a/Engine/Source/Runtime/MovieScene/Private/MovieSceneSequencePlayer.cpp
+++ b/Engine/Source/Runtime/MovieScene/Private/MovieSceneSequencePlayer.cpp
@@ -1092,16 +1092,40 @@ void UMovieSceneSequencePlayer::RPC_OnStopEvent_Implementation(FFrameTime Stoppe
 	}
 #endif
 
-	// If we're behind the target time to stop at, advance to that target time.
-	const bool bBehindTime = StoppedTime < PlayPosition.GetCurrentPosition();
+	float PingMs = 0.f;
 
-	if (bBehindTime)
+	UWorld* PlayWorld = GetPlaybackWorld();
+	if (PlayWorld)
 	{
-		switch (Status.GetValue())
+		UNetDriver* NetDriver = PlayWorld->GetNetDriver();
+		if (NetDriver && NetDriver->ServerConnection && NetDriver->ServerConnection->PlayerController && NetDriver->ServerConnection->PlayerController->PlayerState)
 		{
-		case EMovieScenePlayerStatus::Playing:   PlayToFrame(StoppedTime);  break;
-		case EMovieScenePlayerStatus::Stopped:   JumpToFrame(StoppedTime);  break;
-		case EMovieScenePlayerStatus::Scrubbing: ScrubToFrame(StoppedTime);  break;
+			PingMs = NetDriver->ServerConnection->PlayerController->PlayerState->ExactPing * (bReversePlayback ? -1.f : 1.f);
+		}
+	}
+
+	const FFrameTime PingLag = (PingMs / 1000.f) * PlayPosition.GetInputRate();
+	const FFrameTime LagThreshold = (GSequencerNetSyncThresholdMS * 0.001f) * PlayPosition.GetInputRate();
+
+	// When the server has stopped and a client is near the end (and is thus about to loop), we don't want to forcibly synchronize the time unless
+	// the *real* difference in time is above the threshold. We compute the real-time difference by adding SequenceDuration*LoopCountDifference to the server position:
+	const int32        LoopOffset = (NetSyncProps.LastKnownNumLoops - CurrentNumLoops) * (bReversePlayback ? -1 : 1);
+	const FFrameTime   OffsetServerTime = (NetSyncProps.LastKnownPosition + PingLag) + GetFrameDuration() * LoopOffset;
+	const FFrameTime   Difference = FMath::Abs(PlayPosition.GetCurrentPosition() - OffsetServerTime);
+
+	// If the difference is large enough and the client is behind the target time to stop at, advance to the target time.
+	if (Difference > LagThreshold + PingLag)
+	{
+		const bool bBehindTime = PlayPosition.GetCurrentPosition() < StoppedTime;
+
+		if (bBehindTime)
+		{
+			switch (Status.GetValue())
+			{
+			case EMovieScenePlayerStatus::Playing:   PlayToFrame(StoppedTime);  break;
+			case EMovieScenePlayerStatus::Stopped:   JumpToFrame(StoppedTime);  break;
+			case EMovieScenePlayerStatus::Scrubbing: ScrubToFrame(StoppedTime);  break;
+			}
 		}
 	}
 
