Sequencer: Ensure that all segments for tracks are compiled as part of a full compile.

  - This hopefully fixes an issue of cook non-determinism where a parent sequence would generate segments for an inner sequence, but would only generate the segments it needed. This becomes an issue if the inner sequence is cooked afterwards (causing its template to be re-compiled) and generates all its segments. Although this wouldn't create a crash (due to the indirection of the segment ID) it _would_ create a binary difference in the cooked data.

  - I have fixed this by always compiling the full segment list when doing a full compile. This should leave the compile-on-the-fly path unchanged (which only compiles subsets of the full sequence range)

diff --git a/Engine/Source/Runtime/MovieScene/Private/Compilation/MovieSceneCompiler.cpp b/Engine/Source/Runtime/MovieScene/Private/Compilation/MovieSceneCompiler.cpp
--- a/Engine/Source/Runtime/MovieScene/Private/Compilation/MovieSceneCompiler.cpp
+++ b/Engine/Source/Runtime/MovieScene/Private/Compilation/MovieSceneCompiler.cpp
@@ -688,6 +688,13 @@ void FMovieSceneCompiler::GatherCompileDataForTrack(FMovieSceneEvaluationTrack&
 		return Track.HasChildTemplate(EvalData.ImplIndex) && Track.GetChildTemplate(EvalData.ImplIndex).RequiresInitialization();
 	};
 
+	// When doing a full compile, ensure that we always compile all segments for tracks. This guarantees that we generate the same
+	// segments for tracks regardless of whether they are being compiled as a constrained sub sequence or not.
+	if (Params.RootCompileRange == TRange<FFrameNumber>::All())
+	{
+		Track.GetSegmentsInRange(TRange<FFrameNumber>::All());
+	}
+
 	FMovieSceneSequenceTransform SequenceToRootTransform  = Params.RootToSequenceTransform.InverseFromWarp(Params.SequenceLoopCounter);
 	FMovieSceneSequenceID        CurrentSequenceID        = Params.RootPath.Remap(MovieSceneSequenceID::Root);
 	TRange<FFrameNumber>         CompileClampIntersection = TRange<FFrameNumber>::Intersection(Params.LocalCompileRange, Params.LocalClampRange);
