Sequencer: Fix NodeGroupCollection nativization - IsEditorOnly

diff --git a/Engine/Source/Runtime/MovieScene/Public/MovieScene.h b/Engine/Source/Runtime/MovieScene/Public/MovieScene.h
--- a/Engine/Source/Runtime/MovieScene/Public/MovieScene.h
+++ b/Engine/Source/Runtime/MovieScene/Public/MovieScene.h
@@ -256,6 +256,8 @@ class MOVIESCENE_API UMovieSceneNodeGroup : public UObject
 {
 	GENERATED_BODY()
 
+	virtual bool IsEditorOnly() const override { return true; }
+
 #if WITH_EDITORONLY_DATA
 public:
 	const FName GetName() const { return Name; }
@@ -313,6 +315,8 @@ class MOVIESCENE_API UMovieSceneNodeGroupCollection : public UObject
 {
 	GENERATED_BODY()
 
+	virtual bool IsEditorOnly() const override { return true; }
+
 #if WITH_EDITORONLY_DATA
 public:
 	/** Called after this object has been deserialized */
@@ -1209,11 +1213,9 @@ private:
 	UPROPERTY()
 	TArray<FMovieSceneSectionGroup> SectionGroups;
 
-#if WITH_EDITORONLY_DATA
 	/** Collection of user-defined groups */
 	UPROPERTY()
 	UMovieSceneNodeGroupCollection* NodeGroupCollection;
-#endif
 
 	/** Whether this scene's marked frames should be shown globally */
 	bool bGloballyShowMarkedFrames;
