Sequencer: Don't need to regenerate signature if RF_NeedLoad | RF_LoadCompleted

diff --git a/Engine/Source/Runtime/MovieScene/Private/MovieSceneSignedObject.cpp b/Engine/Source/Runtime/MovieScene/Private/MovieSceneSignedObject.cpp
--- a/Engine/Source/Runtime/MovieScene/Private/MovieSceneSignedObject.cpp
+++ b/Engine/Source/Runtime/MovieScene/Private/MovieSceneSignedObject.cpp
@@ -8,9 +8,6 @@
 UMovieSceneSignedObject::UMovieSceneSignedObject(const FObjectInitializer& Init)
 	: Super(Init)
 {
-#if WITH_EDITOR
-	PreLoadSignature.A = PreLoadSignature.B = PreLoadSignature.C = PreLoadSignature.D = 0xFFFFFFFF;
-#endif
 }
 
 void UMovieSceneSignedObject::PostInitProperties()
@@ -19,28 +16,15 @@ void UMovieSceneSignedObject::PostInitProperties()
 
 	// Always seed newly created objects with a new signature
 	// (CDO and archetypes always have a zero GUID)
-	if (!HasAnyFlags(RF_ClassDefaultObject | RF_ArchetypeObject) && Signature == GetDefault<UMovieSceneSignedObject>()->Signature)
+	if (!HasAnyFlags(RF_ClassDefaultObject | RF_ArchetypeObject | RF_NeedLoad | RF_LoadCompleted) && Signature == GetDefault<UMovieSceneSignedObject>()->Signature)
 	{
 		Signature = FGuid::NewGuid();
-#if WITH_EDITOR
-		PreLoadSignature = Signature;
-#endif
 	}
 }
 
 void UMovieSceneSignedObject::PostLoad()
 {
 	Super::PostLoad();
-#if WITH_EDITOR
-	if (!HasAnyFlags(RF_ClassDefaultObject | RF_ArchetypeObject) && Signature == PreLoadSignature)
-	{
-		UPackage* Package = GetOutermost();
-		UMovieSceneSequence* Sequence = GetTypedOuter<UMovieSceneSequence>();
-		FString PackageName = Package ? Package->GetName() : TEXT("Unknown package");
-		FString SequenceName = Sequence ? Sequence->GetName() : TEXT("Unknown sequence");
-		UE_LOG(LogMovieScene, Warning, TEXT("Legacy data detected in sequence '%s (%s)'. This will cause deterministic cooking issues. Please resave the package."), *PackageName, *SequenceName);
-	}
-#endif
 }
 
 void UMovieSceneSignedObject::MarkAsChanged()
diff --git a/Engine/Source/Runtime/MovieScene/Public/MovieSceneSignedObject.h b/Engine/Source/Runtime/MovieScene/Public/MovieSceneSignedObject.h
--- a/Engine/Source/Runtime/MovieScene/Public/MovieSceneSignedObject.h
+++ b/Engine/Source/Runtime/MovieScene/Public/MovieSceneSignedObject.h
@@ -53,11 +53,6 @@ private:
 	UPROPERTY()
 	FGuid Signature;
 
-#if WITH_EDITOR
-	/** Keep track of the above signature before and after post load to ensure that it got deserialized. If it didn't this will create deterministic cooking issues. */
-	FGuid PreLoadSignature;
-#endif
-
 	/** Event that is triggered whenever this object's signature has changed */
 	FOnSignatureChanged OnSignatureChangedEvent;
 };
