Sequencer: Playback controls

Deprecated PlayTo, ScrubTo, JumpTo for Frames, Seconds, and MarkedFrames. Replaced with SetPlaybackPosition() which takes PlaybackParams to specify the position (ie, frame, seconds, or marked frame) and an update method (play, scrub, jump).

Also, adding a new feature/function:

PlayTo() - play from the current time to a specified time

diff --git a/Engine/Plugins/Developer/Concert/ConcertSync/ConcertSyncClient/Source/ConcertSyncClient/Private/ConcertClientSequencerManager.cpp b/Engine/Plugins/Developer/Concert/ConcertSync/ConcertSyncClient/Source/ConcertSyncClient/Private/ConcertClientSequencerManager.cpp
--- a/Engine/Plugins/Developer/Concert/ConcertSync/ConcertSyncClient/Source/ConcertSyncClient/Private/ConcertClientSequencerManager.cpp
+++ b/Engine/Plugins/Developer/Concert/ConcertSync/ConcertSyncClient/Source/ConcertSyncClient/Private/ConcertClientSequencerManager.cpp
@@ -524,7 +524,7 @@ void FConcertClientSequencerManager::ApplyEventToPlayers(const FConcertSequencer
 					CompensatedTime.FrameNumber.Value, CompensatedTime.GetSubFrame(), CompensatedTime / SequenceRate
 				);
 
-				Player->PlayToFrame(CompensatedTime);
+				Player->SetPlaybackPosition(FMovieSceneSequencePlaybackParams(CompensatedTime, EUpdatePositionMethod::Play));
 				Player->SetPlayRate(EventState.PlaybackSpeed);
 			}
 			else
@@ -552,20 +552,20 @@ void FConcertClientSequencerManager::ApplyEventToPlayers(const FConcertSequencer
 			case EConcertMovieScenePlayerStatus::Stepping:
 				// fallthrough, handles as scrub
 			case EConcertMovieScenePlayerStatus::Scrubbing:
-				Player->ScrubToFrame(IncomingTime);
+				Player->SetPlaybackPosition(FMovieSceneSequencePlaybackParams(IncomingTime, EUpdatePositionMethod::Scrub));
 				break;
 			case EConcertMovieScenePlayerStatus::Paused:
-				Player->JumpToFrame(IncomingTime);
+				Player->SetPlaybackPosition(FMovieSceneSequencePlaybackParams(IncomingTime, EUpdatePositionMethod::Jump));
 				Player->Pause();
 				break;
 			case EConcertMovieScenePlayerStatus::Stopped:
-				Player->JumpToFrame(IncomingTime);
+				Player->SetPlaybackPosition(FMovieSceneSequencePlaybackParams(IncomingTime, EUpdatePositionMethod::Jump));
 				Player->Stop();
 				break;
 			case EConcertMovieScenePlayerStatus::Jumping:
 				// fallthrough, handles as stop
 			default:
-				Player->JumpToFrame(IncomingTime);
+				Player->SetPlaybackPosition(FMovieSceneSequencePlaybackParams(IncomingTime, EUpdatePositionMethod::Jump));
 			}
 
 			Player->SetPlayRate(EventState.PlaybackSpeed);
diff --git a/Engine/Plugins/Experimental/ControlRig/Source/ControlRigEditor/Private/Sequencer/ControlRigSequenceExporter.cpp b/Engine/Plugins/Experimental/ControlRig/Source/ControlRigEditor/Private/Sequencer/ControlRigSequenceExporter.cpp
--- a/Engine/Plugins/Experimental/ControlRig/Source/ControlRigEditor/Private/Sequencer/ControlRigSequenceExporter.cpp
+++ b/Engine/Plugins/Experimental/ControlRig/Source/ControlRigEditor/Private/Sequencer/ControlRigSequenceExporter.cpp
@@ -256,7 +256,7 @@ void Convert(UControlRigSequence* Sequence, UAnimSequence* AnimSequence, USkelet
 				float CurrentTime = (float)(StartTime + FMath::Clamp(((double)Frame / FrameCountDouble) * DurationSeconds, 0.0, DurationSeconds));
 
 				// Tick Sequence
-				LevelSequenceActor->SequencePlayer->JumpToSeconds(CurrentTime);
+				LevelSequenceActor->SequencePlayer->SetPlaybackPosition(FMovieSceneSequencePlaybackParams(CurrentTime, EUpdatePositionMethod::Jump));
 
 				// Tick skeletal mesh component
 				SkeletalMeshComponent->TickAnimation(FrameLength, false);
diff --git a/Engine/Plugins/MovieScene/MovieRenderPipeline/Source/MovieRenderPipelineCore/Private/MoviePipelineTiming.cpp b/Engine/Plugins/MovieScene/MovieRenderPipeline/Source/MovieRenderPipelineCore/Private/MoviePipelineTiming.cpp
--- a/Engine/Plugins/MovieScene/MovieRenderPipeline/Source/MovieRenderPipelineCore/Private/MoviePipelineTiming.cpp
+++ b/Engine/Plugins/MovieScene/MovieRenderPipeline/Source/MovieRenderPipelineCore/Private/MoviePipelineTiming.cpp
@@ -93,7 +93,7 @@ void UMoviePipeline::TickProducingFrames()
 		FFrameTime TimeInPlayRate = FFrameRate::TransformTime(CurrentCameraCut.CurrentMasterSeqTick, TargetSequence->GetMovieScene()->GetTickResolution(), TargetSequence->GetMovieScene()->GetDisplayRate());
 		
 		// We tell it to jump so that things responding to different scrub-types work correctly.
-		LevelSequenceActor->GetSequencePlayer()->JumpToFrame(TimeInPlayRate);
+		LevelSequenceActor->GetSequencePlayer()->SetPlaybackPosition(FMovieSceneSequencePlaybackParams(TimeInPlayRate, EUpdatePositionMethod::Jump));
 		CustomSequenceTimeController->SetCachedFrameTiming(FQualifiedFrameTime(CurrentCameraCut.CurrentMasterSeqTick, TargetSequence->GetMovieScene()->GetTickResolution()));
 
 		// Ensure we don't try to evaluate as we want to sit and wait during warm up and motion blur frames.
@@ -212,7 +212,7 @@ void UMoviePipeline::TickProducingFrames()
 
 		// Jump to the motion blur frame
 		FFrameTime TimeInPlayRate = FFrameRate::TransformTime(FinalEvalTime, FrameMetrics.TickResolution, TargetSequence->GetMovieScene()->GetDisplayRate());
-		LevelSequenceActor->GetSequencePlayer()->JumpToFrame(TimeInPlayRate);
+		LevelSequenceActor->GetSequencePlayer()->SetPlaybackPosition(FMovieSceneSequencePlaybackParams(TimeInPlayRate, EUpdatePositionMethod::Jump));
 		CustomSequenceTimeController->SetCachedFrameTiming(FQualifiedFrameTime(FinalEvalTime, FrameMetrics.TickResolution));
 		
 		// We early out on this loop so that at least one tick passes for motion blur. We need to leave our state in
@@ -484,7 +484,7 @@ void UMoviePipeline::TickProducingFrames()
 			// important because we're going backwards in time from the motion blur frame, and some tracks need to be notified
 			// of a jump when time goes backwards.
 			FFrameTime TimeInPlayRate = FFrameRate::TransformTime(FinalEvalTime, FrameMetrics.TickResolution, TargetSequence->GetMovieScene()->GetDisplayRate());
-			LevelSequenceActor->GetSequencePlayer()->JumpToFrame(TimeInPlayRate);
+			LevelSequenceActor->GetSequencePlayer()->SetPlaybackPosition(FMovieSceneSequencePlaybackParams(TimeInPlayRate, EUpdatePositionMethod::Jump));
 			LevelSequenceActor->GetSequencePlayer()->Play();
 		}
 
diff --git a/Engine/Plugins/MovieScene/SequencerScripting/Source/SequencerScriptingEditor/Private/SequencerTools.cpp b/Engine/Plugins/MovieScene/SequencerScripting/Source/SequencerScriptingEditor/Private/SequencerTools.cpp
--- a/Engine/Plugins/MovieScene/SequencerScripting/Source/SequencerScriptingEditor/Private/SequencerTools.cpp
+++ b/Engine/Plugins/MovieScene/SequencerScripting/Source/SequencerScriptingEditor/Private/SequencerTools.cpp
@@ -119,7 +119,7 @@ TArray<FSequencerBoundObjects> USequencerToolsFunctionLibrary::GetBoundObjects(U
 	FFrameRate Resolution = InSequence->GetMovieScene()->GetTickResolution();
 	TRange<FFrameNumber> SpecifiedRange = InRange.ToNative(Resolution);
 	Player->Play();
-	Player->PlayToFrame(SpecifiedRange.GetLowerBoundValue());
+	Player->SetPlaybackPosition(FMovieSceneSequencePlaybackParams(SpecifiedRange.GetLowerBoundValue().Value, EUpdatePositionMethod::Play));
 
 	FMovieSceneSequenceID SequenceId = Player->State.FindSequenceId(InSequence);
 
@@ -150,7 +150,7 @@ TArray<FSequencerBoundObjects> USequencerToolsFunctionLibrary::GetObjectBindings
 	FFrameRate Resolution = InSequence->GetMovieScene()->GetTickResolution();
 	TRange<FFrameNumber> SpecifiedRange = InRange.ToNative(Resolution);
 	Player->Play();
-	Player->PlayToFrame(SpecifiedRange.GetLowerBoundValue());
+	Player->SetPlaybackPosition(FMovieSceneSequencePlaybackParams(SpecifiedRange.GetLowerBoundValue().Value, EUpdatePositionMethod::Play));
 
 	TArray<FSequencerBoundObjects> BoundObjects;
 
@@ -204,7 +204,7 @@ bool USequencerToolsFunctionLibrary::ExportFBX(UWorld* World, ULevelSequence* Se
 		{
 			// Evaluate at the beginning of the subscene time to ensure that spawnables are created before export
 			Player->Play();
-			Player->PlayToFrame(MovieScene::DiscreteInclusiveLower(MovieScene->GetPlaybackRange()));
+			Player->SetPlaybackPosition(FMovieSceneSequencePlaybackParams(MovieScene::DiscreteInclusiveLower(MovieScene->GetPlaybackRange()).Value, EUpdatePositionMethod::Play));
 		}
 
 		bDidExport = MovieSceneToolHelpers::ExportFBX(World, MovieScene, Player, Bindings, NodeNameAdapter, Template, InFBXFileName, RootToLocalTransform);
diff --git a/Engine/Plugins/VirtualProduction/Takes/Source/TakeRecorderSources/Private/TakeRecorderLevelSequenceSource.cpp b/Engine/Plugins/VirtualProduction/Takes/Source/TakeRecorderSources/Private/TakeRecorderLevelSequenceSource.cpp
--- a/Engine/Plugins/VirtualProduction/Takes/Source/TakeRecorderSources/Private/TakeRecorderLevelSequenceSource.cpp
+++ b/Engine/Plugins/VirtualProduction/Takes/Source/TakeRecorderSources/Private/TakeRecorderLevelSequenceSource.cpp
@@ -58,7 +58,7 @@ void UTakeRecorderLevelSequenceSource::StartRecording(const FTimecode& InSection
 		if (ActorToTrigger->SequencePlayer)
 		{
 			ActorToTrigger->SequencePlayer->SetDisableCameraCuts(true);
-			ActorToTrigger->SequencePlayer->JumpToFrame(ActorToTrigger->SequencePlayer->GetStartTime().Time.FrameNumber);
+			ActorToTrigger->SequencePlayer->SetPlaybackPosition(FMovieSceneSequencePlaybackParams(ActorToTrigger->SequencePlayer->GetStartTime().Time, EUpdatePositionMethod::Jump));
 			ActorToTrigger->SequencePlayer->Play();
 		}
 		else
diff --git a/Engine/Source/Editor/MovieSceneTools/Private/AutomatedLevelSequenceCapture.cpp b/Engine/Source/Editor/MovieSceneTools/Private/AutomatedLevelSequenceCapture.cpp
--- a/Engine/Source/Editor/MovieSceneTools/Private/AutomatedLevelSequenceCapture.cpp
+++ b/Engine/Source/Editor/MovieSceneTools/Private/AutomatedLevelSequenceCapture.cpp
@@ -601,7 +601,7 @@ void UAutomatedLevelSequenceCapture::SetupFrameRange()
 				// Override the movie scene's playback range
 				Actor->SequencePlayer->SetFrameRate(Settings.GetFrameRate());
 				Actor->SequencePlayer->SetFrameRange(PlaybackStartFrame.Value, (PlaybackEndFrame - PlaybackStartFrame).Value);
-				Actor->SequencePlayer->JumpToFrame(PlaybackStartFrame.Value);
+				Actor->SequencePlayer->SetPlaybackPosition(FMovieSceneSequencePlaybackParams(PlaybackStartFrame.Value, EUpdatePositionMethod::Jump));
 
 				Actor->SequencePlayer->SetSnapshotOffsetFrames(WarmUpFrameCount);
 			}
@@ -741,7 +741,7 @@ void UAutomatedLevelSequenceCapture::OnTick(float DeltaSeconds)
 			FFrameNumber EndTimePlayRateSpace   = ConvertFrameTime(EndTime,   MovieScene->GetTickResolution(), Settings.GetFrameRate()).CeilToFrame();
 
 			Actor->SequencePlayer->SetFrameRange(StartTimePlayRateSpace.Value, (EndTimePlayRateSpace - StartTimePlayRateSpace).Value);
-			Actor->SequencePlayer->JumpToFrame(StartTimePlayRateSpace.Value);
+			Actor->SequencePlayer->SetPlaybackPosition(FMovieSceneSequencePlaybackParams(StartTimePlayRateSpace.Value, EUpdatePositionMethod::Jump));
 			Actor->SequencePlayer->Play();
 
 			// We need to re-register to the binding when we start each shot. When a shot reaches the last frame it unregisters the binding so that
diff --git a/Engine/Source/Editor/SequenceRecorder/Private/SequenceRecorder.cpp b/Engine/Source/Editor/SequenceRecorder/Private/SequenceRecorder.cpp
--- a/Engine/Source/Editor/SequenceRecorder/Private/SequenceRecorder.cpp
+++ b/Engine/Source/Editor/SequenceRecorder/Private/SequenceRecorder.cpp
@@ -709,7 +709,7 @@ bool FSequenceRecorder::StartRecordingInternal(UWorld* World)
 				if (DupActorToTrigger->SequencePlayer)
 				{
 					DupActorToTrigger->SequencePlayer->SetDisableCameraCuts(true);
-					DupActorToTrigger->SequencePlayer->JumpToFrame(0);
+					DupActorToTrigger->SequencePlayer->SetPlaybackPosition(FMovieSceneSequencePlaybackParams(0, EUpdatePositionMethod::Jump));
 					DupActorToTrigger->SequencePlayer->Play();
 				}
 				else
diff --git a/Engine/Source/Runtime/LevelSequence/Private/SequenceMediaController.cpp b/Engine/Source/Runtime/LevelSequence/Private/SequenceMediaController.cpp
--- a/Engine/Source/Runtime/LevelSequence/Private/SequenceMediaController.cpp
+++ b/Engine/Source/Runtime/LevelSequence/Private/SequenceMediaController.cpp
@@ -98,7 +98,7 @@ void ALevelSequenceMediaController::Client_Play()
 
 	UE_LOG(LogLevelSequence, Log, TEXT("Initiating playback of sequence '%s' starting at time %fs"), *Sequence->GetName(), StartTimeSeconds);
 
-	Sequence->GetSequencePlayer()->JumpToSeconds(StartTimeSeconds);
+	Sequence->GetSequencePlayer()->SetPlaybackPosition(FMovieSceneSequencePlaybackParams(StartTimeSeconds, EUpdatePositionMethod::Jump));
 	Sequence->GetSequencePlayer()->Play();
 }
 
@@ -126,7 +126,7 @@ void ALevelSequenceMediaController::Client_ConditionallyForceTime(float DesyncTh
 	if (Difference > DesyncThresholdSeconds)
 	{
 		const double NewTimeSeconds = Sequence->GetSequencePlayer()->GetStartTime().AsSeconds() + (CurrentServerTime - ServerStartTimeSeconds);
-		Sequence->GetSequencePlayer()->JumpToSeconds(NewTimeSeconds);
+		Sequence->GetSequencePlayer()->SetPlaybackPosition(FMovieSceneSequencePlaybackParams(NewTimeSeconds, EUpdatePositionMethod::Jump));
 
 		UE_LOG(LogLevelSequence, Warning, TEXT("Forcibly synchronizing sequence '%s' to time %f to server time (it is out by %fs)."), *Sequence->GetName(), NewTimeSeconds, Difference);
 	}
diff --git a/Engine/Source/Runtime/MovieScene/Private/MovieSceneSequencePlayer.cpp b/Engine/Source/Runtime/MovieScene/Private/MovieSceneSequencePlayer.cpp
--- a/Engine/Source/Runtime/MovieScene/Private/MovieSceneSequencePlayer.cpp
+++ b/Engine/Source/Runtime/MovieScene/Private/MovieSceneSequencePlayer.cpp
@@ -44,6 +44,36 @@ bool FMovieSceneSequencePlaybackSettings::SerializeFromMismatchedTag( const FPro
 	return false;
 }
 
+FFrameTime FMovieSceneSequencePlaybackParams::GetPlaybackPosition(UMovieSceneSequencePlayer* Player) const
+{
+	FFrameTime PlaybackPosition = Player->GetCurrentTime().Time;
+
+	if (PositionType == EMovieScenePositionType::Frame)
+	{
+		PlaybackPosition = Frame;
+	}
+	else if (PositionType == EMovieScenePositionType::Time)
+	{
+		PlaybackPosition = Time * Player->GetFrameRate();
+	}
+	else if (PositionType == EMovieScenePositionType::MarkedFrame)
+	{
+		UMovieScene* MovieScene = Player->GetSequence() ? Player->GetSequence()->GetMovieScene() : nullptr;
+
+		if (MovieScene)
+		{
+			int32 MarkedIndex = MovieScene->FindMarkedFrameByLabel(MarkedFrame);
+
+			if (MarkedIndex != INDEX_NONE)
+			{
+				PlaybackPosition = ConvertFrameTime(MovieScene->GetMarkedFrames()[MarkedIndex].FrameNumber, MovieScene->GetTickResolution(), MovieScene->GetDisplayRate());
+			}
+		}
+	}
+
+	return PlaybackPosition;
+}
+
 UMovieSceneSequencePlayer::UMovieSceneSequencePlayer(const FObjectInitializer& Init)
 	: Super(Init)
 	, Status(EMovieScenePlayerStatus::Stopped)
@@ -151,14 +181,14 @@ void UMovieSceneSequencePlayer::PlayInternal()
 		{
 			if (PlayRate > 0.f)
 			{
-				JumpToFrame(FFrameTime(StartTime));
+				SetPlaybackPosition(FMovieSceneSequencePlaybackParams(FFrameTime(StartTime), EUpdatePositionMethod::Jump));
 			}
 		}
 		else if (GetCurrentTime().Time == FFrameTime(StartTime))
 		{
 			if (PlayRate < 0.f)
 			{
-				JumpToFrame(GetLastValidTime());
+				SetPlaybackPosition(FMovieSceneSequencePlaybackParams(GetLastValidTime(), EUpdatePositionMethod::Jump));
 			}
 		}
 
@@ -236,6 +266,7 @@ void UMovieSceneSequencePlayer::Pause()
 		Status = EMovieScenePlayerStatus::Paused;
 		TimeController->StopPlaying(GetCurrentTime());
 
+		PauseOnFrame.Reset();
 		LastTickGameTimeSeconds.Reset();
 
 		// Evaluate the sequence at its current time, with a status of 'stopped' to ensure that animated state pauses correctly. (ie. audio sounds should stop/pause)
@@ -315,6 +346,7 @@ void UMovieSceneSequencePlayer::StopInternal(FFrameTime TimeToResetTo)
 		}
 
 		CurrentNumLoops = 0;
+		PauseOnFrame.Reset();
 		LastTickGameTimeSeconds.Reset();
 
 		// Reset loop count on stop so that it doesn't persist to the next call to play
@@ -374,7 +406,7 @@ void UMovieSceneSequencePlayer::GoToEndAndStop()
 	}
 
 	Status = EMovieScenePlayerStatus::Playing;
-	JumpToFrame(LastValidTime);
+	SetPlaybackPosition(FMovieSceneSequencePlaybackParams(LastValidTime, EUpdatePositionMethod::Jump));
 	StopInternal(LastValidTime);
 }
 
@@ -453,114 +485,33 @@ void UMovieSceneSequencePlayer::SetTimeRange( float StartTimeSeconds, float Dura
 	SetFrameRange(StartFrame.Value, Duration.Value);
 }
 
-void UMovieSceneSequencePlayer::PlayToFrame(FFrameTime NewPosition)
+void UMovieSceneSequencePlayer::PlayTo(FMovieSceneSequencePlaybackParams InPlaybackParams)
 {
-	UpdateTimeCursorPosition(NewPosition, EUpdatePositionMethod::Play);
+	PauseOnFrame = InPlaybackParams.GetPlaybackPosition(this);
 
-	TimeController->Reset(GetCurrentTime());
-
-	if (HasAuthority())
+	if (GetCurrentTime().Time < PauseOnFrame.GetValue())
 	{
-		RPC_ExplicitServerUpdateEvent(EUpdatePositionMethod::Play, NewPosition);
+		Play();
 	}
-}
-
-void UMovieSceneSequencePlayer::ScrubToFrame(FFrameTime NewPosition)
-{
-	UpdateTimeCursorPosition(NewPosition, EUpdatePositionMethod::Scrub);
-
-	TimeController->Reset(GetCurrentTime());
-
-	if (HasAuthority())
-	{
-		RPC_ExplicitServerUpdateEvent(EUpdatePositionMethod::Scrub, NewPosition);
-	}
-}
-
-void UMovieSceneSequencePlayer::JumpToFrame(FFrameTime NewPosition)
-{
-	UpdateTimeCursorPosition(NewPosition, EUpdatePositionMethod::Jump);
-
-	TimeController->Reset(GetCurrentTime());
-
-	if (HasAuthority())
-	{
-		RPC_ExplicitServerUpdateEvent(EUpdatePositionMethod::Jump, NewPosition);
-	}
-}
-
-void UMovieSceneSequencePlayer::PlayToSeconds(float TimeInSeconds)
-{
-	PlayToFrame(TimeInSeconds * PlayPosition.GetInputRate());
-}
-
-void UMovieSceneSequencePlayer::ScrubToSeconds(float TimeInSeconds)
-{
-	ScrubToFrame(TimeInSeconds * PlayPosition.GetInputRate());
-}
-
-void UMovieSceneSequencePlayer::JumpToSeconds(float TimeInSeconds)
-{
-	JumpToFrame(TimeInSeconds * PlayPosition.GetInputRate());
-}
-
-
-int32 UMovieSceneSequencePlayer::FindMarkedFrameByLabel(const FString& InLabel) const
-{
-	if (!Sequence)
+	else
 	{
-		return INDEX_NONE;
+		PlayReverse();
 	}
-
-	UMovieScene* MovieScene = Sequence ? Sequence->GetMovieScene() : nullptr;
-
-	int32 MarkedIndex = MovieScene->FindMarkedFrameByLabel(InLabel);
-	return MarkedIndex;
 }
 
-bool UMovieSceneSequencePlayer::PlayToMarkedFrame(const FString& InLabel)
-{
-	int32 MarkedIndex = FindMarkedFrameByLabel(InLabel);
-
-	if (MarkedIndex != INDEX_NONE)
-	{
-		UMovieScene* MovieScene = Sequence->GetMovieScene();
-		PlayToFrame(ConvertFrameTime(MovieScene->GetMarkedFrames()[MarkedIndex].FrameNumber, MovieScene->GetTickResolution(), MovieScene->GetDisplayRate()));
-
-		return true;
-	}
-
-	return false;
-}
 
-bool UMovieSceneSequencePlayer::ScrubToMarkedFrame(const FString& InLabel)
+void UMovieSceneSequencePlayer::SetPlaybackPosition(FMovieSceneSequencePlaybackParams InPlaybackParams)
 {
-	int32 MarkedIndex = FindMarkedFrameByLabel(InLabel);
+	FFrameTime NewPosition = InPlaybackParams.GetPlaybackPosition(this);
 
-	if (MarkedIndex != INDEX_NONE)
-	{
-		UMovieScene* MovieScene = Sequence->GetMovieScene();
-		ScrubToFrame(ConvertFrameTime(MovieScene->GetMarkedFrames()[MarkedIndex].FrameNumber, MovieScene->GetTickResolution(), MovieScene->GetDisplayRate()));
+	UpdateTimeCursorPosition(NewPosition, InPlaybackParams.UpdateMethod);
 
-		return true;
-	}
-
-	return false;
-}
-
-bool UMovieSceneSequencePlayer::JumpToMarkedFrame(const FString& InLabel)
-{
-	int32 MarkedIndex = FindMarkedFrameByLabel(InLabel);
+	TimeController->Reset(GetCurrentTime());
 
-	if (MarkedIndex != INDEX_NONE)
+	if (HasAuthority())
 	{
-		UMovieScene* MovieScene = Sequence->GetMovieScene();
-		JumpToFrame(ConvertFrameTime(MovieScene->GetMarkedFrames()[MarkedIndex].FrameNumber, MovieScene->GetTickResolution(), MovieScene->GetDisplayRate()));
-
-		return true;
+		RPC_ExplicitServerUpdateEvent(InPlaybackParams.UpdateMethod, NewPosition);
 	}
-
-	return false;
 }
 
 bool UMovieSceneSequencePlayer::IsPlaying() const
@@ -578,11 +529,6 @@ bool UMovieSceneSequencePlayer::IsReversed() const
 	return bReversePlayback;
 }
 
-float UMovieSceneSequencePlayer::GetLength() const
-{
-	return GetDuration().AsSeconds();
-}
-
 float UMovieSceneSequencePlayer::GetPlayRate() const
 {
 	return PlaybackSettings.PlayRate;
@@ -593,11 +539,6 @@ void UMovieSceneSequencePlayer::SetPlayRate(float PlayRate)
 	PlaybackSettings.PlayRate = PlayRate;
 }
 
-void UMovieSceneSequencePlayer::SetPlaybackRange( float NewStartTime, float NewEndTime )
-{
-	SetTimeRange(NewStartTime, NewEndTime - NewStartTime);
-}
-
 FFrameTime UMovieSceneSequencePlayer::GetLastValidTime() const
 {
 	UMovieScene* MovieScene = Sequence->GetMovieScene();
@@ -629,6 +570,24 @@ bool UMovieSceneSequencePlayer::ShouldStopOrLoop(FFrameTime NewPosition) const
 	return bShouldStopOrLoop;
 }
 
+bool UMovieSceneSequencePlayer::ShouldPause(FFrameTime NewPosition) const
+{
+	bool bShouldPause = false;
+	if (IsPlaying() && PauseOnFrame.IsSet())
+	{
+		if (!bReversePlayback)
+		{
+			bShouldPause = PauseOnFrame.GetValue() <= PlayPosition.GetCurrentPosition();
+		}
+		else
+		{
+			bShouldPause = PauseOnFrame.GetValue() >= PlayPosition.GetCurrentPosition();
+		}
+	}
+
+	return bShouldPause;
+}
+
 void UMovieSceneSequencePlayer::Initialize(UMovieSceneSequence* InSequence, const FMovieSceneSequencePlaybackSettings& InSettings)
 {
 	check(InSequence);
@@ -807,7 +766,15 @@ void UMovieSceneSequencePlayer::UpdateTimeCursorPosition_Internal(FFrameTime New
 		bPendingOnStartedPlaying = false;
 	}
 
-	if (Method == EUpdatePositionMethod::Play && ShouldStopOrLoop(NewPosition))
+	if (Method == EUpdatePositionMethod::Play && ShouldPause(NewPosition))
+	{
+		if (PauseOnFrame.GetValue() != PlayPosition.GetCurrentPosition())
+		{
+			UpdateTimeCursorPosition(PauseOnFrame.GetValue(), EUpdatePositionMethod::Jump);
+		}
+		Pause();
+	}
+	else if (Method == EUpdatePositionMethod::Play && ShouldStopOrLoop(NewPosition))
 	{
 		// The actual start time taking into account reverse playback
 		FFrameNumber StartTimeWithReversed = bReversePlayback ? GetLastValidTime().FrameNumber : StartTime;
@@ -954,12 +921,7 @@ void UMovieSceneSequencePlayer::ApplyLatentActions()
 		}
 
 		check(LatentAction.Type == FLatentAction::EType::Update);
-		switch (LatentAction.UpdateMethod)
-		{
-		case EUpdatePositionMethod::Play:  PlayToFrame( LatentAction.Position); continue;
-		case EUpdatePositionMethod::Jump:  JumpToFrame( LatentAction.Position); continue;
-		case EUpdatePositionMethod::Scrub: ScrubToFrame(LatentAction.Position); continue;
-		}
+		SetPlaybackPosition(FMovieSceneSequencePlaybackParams(LatentAction.Position, LatentAction.UpdateMethod));
 	}
 }
 
@@ -1055,12 +1017,7 @@ void UMovieSceneSequencePlayer::RPC_ExplicitServerUpdateEvent_Implementation(EUp
 	// Note: in the case of PlayToFrame this will not necessarily sweep the exact same range as the server did
 	// because this client player is unlikely to be at exactly the same time that the server was at when it performed the operation.
 	// This is irrelevant for jumps and scrubs as only the new time is meaningful.
-	switch (EventMethod)
-	{
-	case EUpdatePositionMethod::Play:  PlayToFrame( MarkerTime);  break;
-	case EUpdatePositionMethod::Jump:  JumpToFrame( MarkerTime);  break;
-	case EUpdatePositionMethod::Scrub: ScrubToFrame(MarkerTime);  break;
-	}
+	SetPlaybackPosition(FMovieSceneSequencePlaybackParams(MarkerTime, EventMethod));
 }
 
 void UMovieSceneSequencePlayer::RPC_OnStopEvent_Implementation(FFrameTime StoppedTime)
@@ -1120,12 +1077,7 @@ void UMovieSceneSequencePlayer::RPC_OnStopEvent_Implementation(FFrameTime Stoppe
 
 		if (bBehindTime)
 		{
-			switch (Status.GetValue())
-			{
-			case EMovieScenePlayerStatus::Playing:   PlayToFrame(StoppedTime);  break;
-			case EMovieScenePlayerStatus::Stopped:   JumpToFrame(StoppedTime);  break;
-			case EMovieScenePlayerStatus::Scrubbing: ScrubToFrame(StoppedTime);  break;
-			}
+			SetPlaybackPosition(FMovieSceneSequencePlaybackParams(StoppedTime, (EUpdatePositionMethod)Status.GetValue()));
 		}
 	}
 
@@ -1202,7 +1154,7 @@ void UMovieSceneSequencePlayer::PostNetReceive()
 		if (LagDisparity > LagThreshold)
 		{
 			// Synchronize to the server time as best we can if there is a large disparity
-			PlayToFrame(NetSyncProps.LastKnownPosition + PingLag);
+			SetPlaybackPosition(FMovieSceneSequencePlaybackParams(NetSyncProps.LastKnownPosition + PingLag, EUpdatePositionMethod::Play));
 		}
 	}
 	else
@@ -1229,7 +1181,7 @@ void UMovieSceneSequencePlayer::PostNetReceive()
 				if (bHasChangedStatus)
 				{
 					// If the status has changed forcibly play to the server position before setting the new status
-					PlayToFrame(NetSyncProps.LastKnownPosition + PingLag);
+					SetPlaybackPosition(FMovieSceneSequencePlaybackParams(NetSyncProps.LastKnownPosition + PingLag, EUpdatePositionMethod::Play));
 				}
 				else if (Difference > LagThreshold + PingLag)
 				{
@@ -1255,21 +1207,21 @@ void UMovieSceneSequencePlayer::PostNetReceive()
 					const bool bPlayToFrame = bReversePlayback ? NetSyncProps.LastKnownPosition < PlayPosition.GetCurrentPosition() : NetSyncProps.LastKnownPosition > PlayPosition.GetCurrentPosition();
 					if (bPlayToFrame)
 					{
-						PlayToFrame(NetSyncProps.LastKnownPosition + PingLag);
+						SetPlaybackPosition(FMovieSceneSequencePlaybackParams(NetSyncProps.LastKnownPosition + PingLag, EUpdatePositionMethod::Play));
 					}
 					else
 					{
-						JumpToFrame(NetSyncProps.LastKnownPosition + PingLag);
+						SetPlaybackPosition(FMovieSceneSequencePlaybackParams(NetSyncProps.LastKnownPosition + PingLag, EUpdatePositionMethod::Jump));
 					}
 				}
 			}
 			else if (Status == EMovieScenePlayerStatus::Stopped)
 			{
-				JumpToFrame(NetSyncProps.LastKnownPosition);
+				SetPlaybackPosition(FMovieSceneSequencePlaybackParams(NetSyncProps.LastKnownPosition, EUpdatePositionMethod::Jump));
 			}
 			else if (Status == EMovieScenePlayerStatus::Scrubbing)
 			{
-				ScrubToFrame(NetSyncProps.LastKnownPosition);
+				SetPlaybackPosition(FMovieSceneSequencePlaybackParams(NetSyncProps.LastKnownPosition, EUpdatePositionMethod::Scrub));
 			}
 		}
 
diff --git a/Engine/Source/Runtime/MovieScene/Public/MovieSceneSequencePlayer.h b/Engine/Source/Runtime/MovieScene/Public/MovieSceneSequencePlayer.h
--- a/Engine/Source/Runtime/MovieScene/Public/MovieSceneSequencePlayer.h
+++ b/Engine/Source/Runtime/MovieScene/Public/MovieSceneSequencePlayer.h
@@ -18,7 +18,7 @@ DECLARE_DYNAMIC_MULTICAST_DELEGATE(FOnMovieSceneSequencePlayerEvent);
 /**
  * Enum used to define how to update to a particular time
  */
-UENUM()
+UENUM(BlueprintType)
 enum class EUpdatePositionMethod : uint8
 {
 	/** Update from the current position to a specified position (including triggering events), using the current player status */
@@ -158,6 +158,59 @@ struct FMovieSceneSequencePlaybackSettings
 	MOVIESCENE_API bool SerializeFromMismatchedTag(const FPropertyTag& Tag, FStructuredArchive::FSlot Slot);
 };
 
+UENUM(BlueprintType)
+enum class EMovieScenePositionType : uint8
+{
+	Frame,
+	Time,
+	MarkedFrame,
+};
+
+USTRUCT(BlueprintType)
+struct FMovieSceneSequencePlaybackParams
+{
+	GENERATED_BODY()
+
+	FMovieSceneSequencePlaybackParams()
+		: Time(0.f)
+		, PositionType(EMovieScenePositionType::Frame)
+		, UpdateMethod(EUpdatePositionMethod::Play) {}
+
+	FMovieSceneSequencePlaybackParams(FFrameTime InFrame, EUpdatePositionMethod InUpdateMethod)
+		: Frame(InFrame)
+		, Time(0.f)
+		, PositionType(EMovieScenePositionType::Frame)
+		, UpdateMethod(InUpdateMethod) {}
+
+	FMovieSceneSequencePlaybackParams(float InTime, EUpdatePositionMethod InUpdateMethod)
+		: Time(InTime)
+		, PositionType(EMovieScenePositionType::Time)
+		, UpdateMethod(InUpdateMethod) {}
+
+	FMovieSceneSequencePlaybackParams(const FString& InMarkedFrame, EUpdatePositionMethod InUpdateMethod)
+		: Time(0.f)
+		, MarkedFrame(InMarkedFrame)
+		, PositionType(EMovieScenePositionType::MarkedFrame)
+		, UpdateMethod(InUpdateMethod) {}
+
+	FFrameTime GetPlaybackPosition(UMovieSceneSequencePlayer* Player) const;
+
+	UPROPERTY(EditAnywhere, BlueprintReadWrite, Category="Cinematic", meta=(EditCondition="PositionType == EMovieScenePositionType::Frame"))
+	FFrameTime Frame;
+	
+	UPROPERTY(EditAnywhere, BlueprintReadWrite, Category="Cinematic", meta=(EditCondition="PositionType == EMovieScenePositionType::Time", unit=s))
+	float Time;
+	
+	UPROPERTY(EditAnywhere, BlueprintReadWrite, Category="Cinematic", meta=(EditCondition="PositionType == EMovieScenePositionType::MarkedFrame"))
+	FString MarkedFrame;
+
+	UPROPERTY(EditAnywhere, BlueprintReadWrite, Category="Cinematic")
+	EMovieScenePositionType PositionType;
+
+	UPROPERTY(EditAnywhere, BlueprintReadWrite, Category="Cinematic")
+	EUpdatePositionMethod UpdateMethod;
+};
+
 template<> struct TStructOpsTypeTraits<FMovieSceneSequencePlaybackSettings> : public TStructOpsTypeTraitsBase2<FMovieSceneSequencePlaybackSettings>
 {
 	enum { WithCopy = true, WithStructuredSerializeFromMismatchedTag = true };
@@ -216,65 +269,6 @@ public:
 	UFUNCTION(BlueprintCallable, Category="Game|Cinematic", meta = (ToolTip = "Go to end of the sequence and stop. Adheres to 'When Finished' section rules."))
 	void GoToEndAndStop();
 
-public:
-
-	/**
-	 * Get the current playback position
-	 * @return The current playback position
-	 */
-	UE_DEPRECATED(4.20, "Please use GetCurrentTime instead")
-	UFUNCTION(BlueprintCallable, Category="Game|Cinematic", meta = (DeprecatedFunction, DeprecationMessage = "Please use GetCurrentTime instead"))
-	float GetPlaybackPosition() const { return GetCurrentTime().AsSeconds() - StartTime / PlayPosition.GetInputRate(); }
-
-	/**
-	 * Get the playback length of the sequence
-	 */
-	UE_DEPRECATED(4.20, "Please use GetDuration instead")
-	UFUNCTION(BlueprintCallable, Category="Game|Cinematic", meta = (DeprecatedFunction, DeprecationMessage = "Please use GetDuration instead"))
-	float GetLength() const;
-
-	/**
-	 * Get the offset within the level sequence to start playing
-	 */
-	UE_DEPRECATED(4.20, "Please use GetStartTime instead")
-	UFUNCTION(BlueprintCallable, Category="Game|Cinematic", meta = (DeprecatedFunction, DeprecationMessage = "Please use GetStartTime instead"))
-	float GetPlaybackStart() const { return StartTime / PlayPosition.GetInputRate(); }
-
-	/**
-	 * Get the offset within the level sequence to finish playing
-	 */
-	UE_DEPRECATED(4.20, "Please use GetEndTime instead")
-	UFUNCTION(BlueprintCallable, Category="Game|Cinematic", meta = (DeprecatedFunction, DeprecationMessage = "Please use GetEndTime instead"))
-	float GetPlaybackEnd() const { return (StartTime + DurationFrames) / PlayPosition.GetInputRate(); }
-
-	/**
-	 * Set the current playback position
-	 * @param NewPlaybackPosition - The new playback position to set.
-	 * If the animation is currently playing, it will continue to do so from the new position
-	 */
-	UE_DEPRECATED(4.20, "Please use PlayToFrame instead")
-	UFUNCTION(BlueprintCallable, Category="Game|Cinematic", meta = (DeprecatedFunction, DeprecationMessage = "Please use PlayToFrame instead"))
-	void SetPlaybackPosition(float NewPlaybackPosition) { Status == EMovieScenePlayerStatus::Playing ? PlayToSeconds(NewPlaybackPosition + StartTime / PlayPosition.GetInputRate()) : JumpToSeconds(NewPlaybackPosition + StartTime / PlayPosition.GetInputRate()); }
-
-	/**
-	 * Sets the range in time to be played back by this player, overriding the default range stored in the asset
-	 *
-	 * @param	NewStartTime	The new starting time for playback
-	 * @param	NewEndTime		The new ending time for playback.  Must be larger than the start time.
-	 */
-	UE_DEPRECATED(4.20, "Please use SetFrameRange or SetTimeRange instead")
-	UFUNCTION(BlueprintCallable, Category="Game|Cinematic", meta = (DeprecatedFunction, DeprecationMessage = "Please use SetFrameRange or SetTimeRange instead"))
-	void SetPlaybackRange( const float NewStartTime, const float NewEndTime );
-
-	/**
-	 * Jump to new playback position
-	 * @param NewPlaybackPosition - The new playback position to set.
-	 * This can be used to update sequencer repeatedly, as if in a scrubbing state
-	 */
-	UE_DEPRECATED(4.20, "Please use ScrubToTime instead")
-	UFUNCTION(BlueprintCallable, Category="Game|Cinematic", meta = (DeprecatedFunction, DeprecationMessage = "Please use ScrubToTime instead"))
-	void JumpToPosition(float NewPlaybackPosition) { ScrubToSeconds(NewPlaybackPosition); }
-
 public:
 
 	/**
@@ -343,90 +337,61 @@ public:
 public:
 
 	/**
-	 * Low-level call to set the current time of the player by evaluating from the current time to the specified time, as if the sequence is playing. 
-	 * Triggers events that lie within the evaluated range. Does not alter the persistent playback status of the player (IsPlaying).
+	 * Play from the current position to the requested position and pause. If requested position is before the current position, 
+	 * playback will be reversed. Playback to the requested position will be cancelled if Stop() or Pause() is invoked during this 
+	 * playback.
 	 *
-	 * @param NewPosition     The new frame time to play to
+	 * @param PlaybackParams The position settings (ie. the position to play to)
 	 */
-	UFUNCTION(BlueprintCallable, Category="Game|Cinematic", DisplayName="Play To (Frames)")
-	void PlayToFrame(FFrameTime NewPosition);
+	UFUNCTION(BlueprintCallable, Category="Game|Cinematic")
+	void PlayTo(FMovieSceneSequencePlaybackParams PlaybackParams);
 
 	/**
-	 * Low-level call to set the current time of the player by evaluating only the specified time. Will not trigger any events. 
-	 * Does not alter the persistent playback status of the player (IsPlaying).
+	 * Set the current time of the player by evaluating from the current time to the specified time, as if the sequence is playing. 
+	 * Triggers events that lie within the evaluated range. Does not alter the persistent playback status of the player (IsPlaying).
 	 *
-	 * @param NewPosition     The new frame time to scrub to
+	 * @param PlaybackParams The position settings (ie. the position to set playback to)
 	 */
-	UFUNCTION(BlueprintCallable, Category="Game|Cinematic", DisplayName="Scrub To (Frames)")
-	void ScrubToFrame(FFrameTime NewPosition);
+	UFUNCTION(BlueprintCallable, Category="Game|Cinematic")
+	void SetPlaybackPosition(FMovieSceneSequencePlaybackParams PlaybackParams);
 
-	/**
-	 * Low-level call to set the current time of the player by evaluating only the specified time, as if scrubbing the timeline. Will trigger only events that exist at the specified time. 
-	 * Does not alter the persistent playback status of the player (IsPlaying).
-	 *
-	 * @param NewPosition     The new frame time to jump to
-	 */
-	UFUNCTION(BlueprintCallable, Category="Game|Cinematic", DisplayName="Jump To (Frames)")
-	void JumpToFrame(FFrameTime NewPosition);
+public:
 
+	UE_DEPRECATED(4.26, "PlayToFrame is deprecated, use SetPlaybackPosition.")
+	UFUNCTION(BlueprintCallable, Category = "Game|Cinematic", DisplayName = "Play To (Frames)", meta=(DeprecatedFunction, DeprecationMessage="PlayToFrame is deprecated, use SetPlaybackPosition."))
+	void PlayToFrame(FFrameTime NewPosition) { SetPlaybackPosition(FMovieSceneSequencePlaybackParams(NewPosition, EUpdatePositionMethod::Play)); }
 
-	/**
-	 * Low-level call to set the current time of the player by evaluating only the specified time. Will not trigger any events.
-	 * Does not alter the persistent playback status of the player (IsPlaying).
-	 *
-	 * @param TimeInSeconds   The desired time in seconds
-	 */
-	UFUNCTION(BlueprintCallable, Category="Game|Cinematic", DisplayName="Play To (Seconds)")
-	void PlayToSeconds(float TimeInSeconds);
+	UE_DEPRECATED(4.26, "ScrubToFrame is deprecated, use SetPlaybackPosition.")
+	UFUNCTION(BlueprintCallable, Category = "Game|Cinematic", DisplayName = "Scrub To (Frames)", meta=(DeprecatedFunction, DeprecationMessage="ScrubToFrame is deprecated, use SetPlaybackPosition."))
+	void ScrubToFrame(FFrameTime NewPosition) { SetPlaybackPosition(FMovieSceneSequencePlaybackParams(NewPosition, EUpdatePositionMethod::Scrub)); }
 
-	/**
-	 * Low-level call to set the current time of the player by evaluating only the specified time. Will not trigger any events.
-	 * Does not alter the persistent playback status of the player (IsPlaying).
-	 *
-	 * @param TimeInSeconds   The desired time in seconds
-	 */
-	UFUNCTION(BlueprintCallable, Category="Game|Cinematic", DisplayName="Scrub To (Seconds)")
-	void ScrubToSeconds(float TimeInSeconds);
+	UE_DEPRECATED(4.26, "JumpToFrame is deprecated, use SetPlaybackPosition.")
+	UFUNCTION(BlueprintCallable, Category="Game|Cinematic", DisplayName="Jump To (Frames)", meta=(DeprecatedFunction, DeprecationMessage="JumpToFrame is deprecated, use SetPlaybackPosition."))
+	void JumpToFrame(FFrameTime NewPosition) { SetPlaybackPosition(FMovieSceneSequencePlaybackParams(NewPosition, EUpdatePositionMethod::Jump)); }
 
-	/**
-	 * Low-level call to set the current time of the player by evaluating only the specified time, as if scrubbing the timeline. Will trigger only events that exist at the specified time.
-	 * Does not alter the persistent playback status of the player (IsPlaying).
-	 *
-	 * @param TimeInSeconds   The desired time in seconds
-	 */
-	UFUNCTION(BlueprintCallable, Category="Game|Cinematic", DisplayName="Jump To (Seconds)")
-	void JumpToSeconds(float TimeInSeconds);
+	UE_DEPRECATED(4.26, "PlayToSeconds is deprecated, use SetPlaybackPosition.")
+	UFUNCTION(BlueprintCallable, Category = "Game|Cinematic", DisplayName = "Play To (Seconds)", meta=(DeprecatedFunction, DeprecationMessage="PlayToSeconds is deprecated, use SetPlaybackPosition."))
+	void PlayToSeconds(float TimeInSeconds) { SetPlaybackPosition(FMovieSceneSequencePlaybackParams(TimeInSeconds, EUpdatePositionMethod::Play)); }
 
+	UE_DEPRECATED(4.26, "ScrubToSeconds is deprecated, use SetPlaybackPosition.")
+	UFUNCTION(BlueprintCallable, Category = "Game|Cinematic", DisplayName = "Scrub To (Seconds)", meta=(DeprecatedFunction, DeprecationMessage="ScrubToSeconds is deprecated, use SetPlaybackPosition."))
+	void ScrubToSeconds(float TimeInSeconds) { SetPlaybackPosition(FMovieSceneSequencePlaybackParams(TimeInSeconds, EUpdatePositionMethod::Scrub)); }
 
-	/**
-	 * Low-level call to set the current time of the player to the marked frame by label by evaluating only the specified time. Will not trigger any events.
-	 * Does not alter the persistent playback status of the player (IsPlaying).
-	 *
-	 * @param InLabel   The desired marked frame label to play to
-	 * @return Whether the marked frame was found
-	 */
-	UFUNCTION(BlueprintCallable, Category="Game|Cinematic")
-	bool PlayToMarkedFrame(const FString& InLabel);
+	UE_DEPRECATED(4.26, "JumpToSeconds is deprecated, use SetPlaybackPosition.")
+	UFUNCTION(BlueprintCallable, Category = "Game|Cinematic", DisplayName = "Jump To (Seconds)", meta=(DeprecatedFunction, DeprecationMessage="JumpToSeconds is deprecated, use SetPlaybackPosition."))
+	void JumpToSeconds(float TimeInSeconds) { SetPlaybackPosition(FMovieSceneSequencePlaybackParams(TimeInSeconds, EUpdatePositionMethod::Jump)); }
 
-	/**
-	 * Low-level call to set the current time of the player to the marked frame by label by evaluating only the specified time. Will not trigger any events.
-	 * Does not alter the persistent playback status of the player (IsPlaying).
-	 *
-	 * @param InLabel   The desired marked frame label to scrub to
-	 * @return Whether the marked frame was found
-	 */
-	UFUNCTION(BlueprintCallable, Category="Game|Cinematic")
-	bool ScrubToMarkedFrame(const FString& InLabel);
+	UE_DEPRECATED(4.26, "PlayToMarkedFrame is deprecated, use SetPlaybackPosition.")
+	UFUNCTION(BlueprintCallable, Category = "Game|Cinematic", meta=(DeprecatedFunction, DeprecationMessage="PlayToMarkedFrame is deprecated, use SetPlaybackPosition."))
+	bool PlayToMarkedFrame(const FString& InLabel) { SetPlaybackPosition(FMovieSceneSequencePlaybackParams(InLabel, EUpdatePositionMethod::Play)); return true; }
 
-	/**
-	 * Low-level call to set the current time of the player to the marked frame by label by evaluating only the specified time, as if scrubbing the timeline. Will trigger only events that exist at the specified time.
-	 * Does not alter the persistent playback status of the player (IsPlaying).
-	 *
-	 * @param InLabel   The desired marked frame label to jump to
-	 * @return Whether the marked frame was found
-	 */
-	UFUNCTION(BlueprintCallable, Category = "Game|Cinematic")
-	bool JumpToMarkedFrame(const FString& InLabel);
+	UE_DEPRECATED(4.26, "ScrubToMarkedFrame is deprecated, use SetPlaybackPosition.")
+	UFUNCTION(BlueprintCallable, Category = "Game|Cinematic", meta=(DeprecatedFunction, DeprecationMessage="ScrubToMarkedFrame is deprecated, use SetPlaybackPosition."))
+	bool ScrubToMarkedFrame(const FString& InLabel) { SetPlaybackPosition(FMovieSceneSequencePlaybackParams(InLabel, EUpdatePositionMethod::Scrub)); return true; }
+
+	UE_DEPRECATED(4.26, "JumpToMarkedFrame is deprecated, use SetPlaybackPosition.")
+	UFUNCTION(BlueprintCallable, Category = "Game|Cinematic", meta=(DeprecatedFunction, DeprecationMessage="JumpToMarkedFrame is deprecated, use SetPlaybackPosition."))
+	bool JumpToMarkedFrame(const FString& InLabel) { SetPlaybackPosition(FMovieSceneSequencePlaybackParams(InLabel, EUpdatePositionMethod::Jump)); return true; }
 
 public:
 
@@ -531,13 +496,12 @@ protected:
 
 	void UpdateTimeCursorPosition(FFrameTime NewPosition, EUpdatePositionMethod Method);
 	bool ShouldStopOrLoop(FFrameTime NewPosition) const;
+	bool ShouldPause(FFrameTime NewPosition) const;
 
 	UWorld* GetPlaybackWorld() const;
 
 	FFrameTime GetLastValidTime() const;
 
-	int32 FindMarkedFrameByLabel(const FString& InLabel) const;
-
 protected:
 
 	//~ IMovieScenePlayer interface
@@ -690,4 +654,7 @@ private:
 	* Valid only if we've been ticked at least once since having a tick interval
 	*/
 	TOptional<float> LastTickGameTimeSeconds;
+
+	/** If set, pause playback on this frame */
+	TOptional<FFrameTime> PauseOnFrame;
 };
