Fix compiling with MOVIESCENE_DETAILED_STATS 1

diff --git a/Engine/Source/Runtime/MovieScene/Private/Evaluation/MovieSceneEvaluationState.cpp b/Engine/Source/Runtime/MovieScene/Private/Evaluation/MovieSceneEvaluationState.cpp
--- a/Engine/Source/Runtime/MovieScene/Private/Evaluation/MovieSceneEvaluationState.cpp
+++ b/Engine/Source/Runtime/MovieScene/Private/Evaluation/MovieSceneEvaluationState.cpp
@@ -7,6 +7,7 @@
 #include "IMovieScenePlaybackClient.h"
 
 DECLARE_CYCLE_STAT(TEXT("Find Bound Objects"), MovieSceneEval_FindBoundObjects, STATGROUP_MovieSceneEval);
+DECLARE_CYCLE_STAT(TEXT("Iterate Bound Objects"), MovieSceneEval_IterateBoundObjects, STATGROUP_MovieSceneEval);
 
 FMovieSceneSharedDataId FMovieSceneSharedDataId::Allocate()
 {
