Sequencer: Ensure that bindings are correctly invalidated in the instance registry when FMovieSceneEvaluationState::ClearObjectCaches is called

diff --git a/Engine/Source/Runtime/MovieScene/Private/Evaluation/MovieSceneEvaluationState.cpp b/Engine/Source/Runtime/MovieScene/Private/Evaluation/MovieSceneEvaluationState.cpp
--- a/Engine/Source/Runtime/MovieScene/Private/Evaluation/MovieSceneEvaluationState.cpp
+++ b/Engine/Source/Runtime/MovieScene/Private/Evaluation/MovieSceneEvaluationState.cpp
@@ -236,6 +236,7 @@ void FMovieSceneObjectCache::Clear(IMovieScenePlayer& Player)
 	ChildBindings.Reset();
 
 	Player.NotifyBindingsChanged();
+	OnBindingInvalidated.Broadcast(FGuid());
 }
 
 
diff --git a/Engine/Source/Runtime/MovieScene/Public/EntitySystem/MovieSceneInstanceRegistry.h b/Engine/Source/Runtime/MovieScene/Public/EntitySystem/MovieSceneInstanceRegistry.h
--- a/Engine/Source/Runtime/MovieScene/Public/EntitySystem/MovieSceneInstanceRegistry.h
+++ b/Engine/Source/Runtime/MovieScene/Public/EntitySystem/MovieSceneInstanceRegistry.h
@@ -84,7 +84,8 @@ struct FInstanceRegistry
 
 	bool IsBindingInvalidated(const FGuid& ObjectBindingID, FInstanceHandle InstanceHandle) const
 	{
-		return InvalidatedObjectBindings.Contains(MakeTuple(ObjectBindingID, InstanceHandle));
+		// The binding is invalidated if it is contained within the invalid set, or if an empty GUID with the same instance handle exists (implying _all_ bindings are invalidated for that instance handle)
+		return InvalidatedObjectBindings.Contains(MakeTuple(ObjectBindingID, InstanceHandle)) || InvalidatedObjectBindings.Contains(MakeTuple(FGuid(), InstanceHandle));
 	}
 
 	bool HasInvalidatedBindings() const
@@ -109,6 +110,7 @@ private:
 	TSparseArray<FSequenceInstance> Instances;
 	uint16 InstanceSerialNumber;
 
+	/** Set of invalidated object bindings by their instance handle. Empty guids indicate that all bindings for that instance handle are invalid */
 	TSet<TTuple<FGuid, FInstanceHandle>> InvalidatedObjectBindings;
 };
 
