Sequencer: Fixed attempting to process nearest sections for empty trees

diff --git a/Engine/Source/Runtime/MovieScene/Private/Compilation/MovieSceneEvaluationTreePopulationRules.cpp b/Engine/Source/Runtime/MovieScene/Private/Compilation/MovieSceneEvaluationTreePopulationRules.cpp
--- a/Engine/Source/Runtime/MovieScene/Private/Compilation/MovieSceneEvaluationTreePopulationRules.cpp
+++ b/Engine/Source/Runtime/MovieScene/Private/Compilation/MovieSceneEvaluationTreePopulationRules.cpp
@@ -132,6 +132,11 @@ void FEvaluationTreePopulationRules::HighPassCustomPerRow(TArrayView<UMovieScene
 
 void FEvaluationTreePopulationRules::PopulateNearestSection(TArrayView<UMovieSceneSection* const> Sections, TMovieSceneEvaluationTree<FMovieSceneTrackEvaluationData>& OutTree)
 {
+	if (OutTree.IsEmpty())
+	{
+		return;
+	}
+
 	// Fill in gaps
 	TArray<TTuple<TRange<FFrameNumber>, FMovieSceneTrackEvaluationData>> RangesToInsert;
 
