Sequencer: Upgraded object binding IDs to be more portable when nested

  - This change introduces two separate types for dealing with binding IDs in different spaces explicitly: FFixedObjectBindingID and FRelativeObjectBindingID. The latter is portable if the sequence it was created within is nested inside another foreign hierarchy.
  - FMovieSceneObjectBindingID remains, and is able to represent both of the new types, but should only be used for serialization to ensure compatability and consistency with editor tools that operate in different spaces.
  - Any manual ResolveLocalToRoot or GetSequenceID processing has now been replaced with the new ResolveBoundObjects API, which provides a less error-prone way of resolving bindings that live in different sequences

diff --git a/Engine/Plugins/Animation/LiveLink/Source/LiveLinkSequencer/Private/TakeRecorderSource/TakeRecorderLiveLinkSource.cpp b/Engine/Plugins/Animation/LiveLink/Source/LiveLinkSequencer/Private/TakeRecorderSource/TakeRecorderLiveLinkSource.cpp
--- a/Engine/Plugins/Animation/LiveLink/Source/LiveLinkSequencer/Private/TakeRecorderSource/TakeRecorderLiveLinkSource.cpp
+++ b/Engine/Plugins/Animation/LiveLink/Source/LiveLinkSequencer/Private/TakeRecorderSource/TakeRecorderLiveLinkSource.cpp
@@ -29,7 +29,7 @@ UTakeRecorderLiveLinkSource::UTakeRecorderLiveLinkSource(const FObjectInitialize
 	TrackTint = FColor(74, 108, 164);
 }
 
-TArray<UTakeRecorderSource*> UTakeRecorderLiveLinkSource::PreRecording(class ULevelSequence* InSequence, class ULevelSequence* InMasterSequence, FManifestSerializer* InManifestSerializer) 
+TArray<UTakeRecorderSource*> UTakeRecorderLiveLinkSource::PreRecording(ULevelSequence* InSequence, FMovieSceneSequenceID InSequenceID, ULevelSequence* InMasterSequence, FManifestSerializer* InManifestSerializer) 
 {
 	UMovieScene* MovieScene = InSequence->GetMovieScene();
 	TrackRecorder = NewObject<UMovieSceneLiveLinkTrackRecorder>();
diff --git a/Engine/Plugins/Animation/LiveLink/Source/LiveLinkSequencer/Private/TakeRecorderSource/TakeRecorderLiveLinkSource.h b/Engine/Plugins/Animation/LiveLink/Source/LiveLinkSequencer/Private/TakeRecorderSource/TakeRecorderLiveLinkSource.h
--- a/Engine/Plugins/Animation/LiveLink/Source/LiveLinkSequencer/Private/TakeRecorderSource/TakeRecorderLiveLinkSource.h
+++ b/Engine/Plugins/Animation/LiveLink/Source/LiveLinkSequencer/Private/TakeRecorderSource/TakeRecorderLiveLinkSource.h
@@ -84,7 +84,7 @@ public:
 private:
 
 	// UTakeRecorderSource
-	virtual TArray<UTakeRecorderSource*> PreRecording(class ULevelSequence* InSequence, class ULevelSequence* InMasterSequence, FManifestSerializer* InManifestSerializer) override;
+	virtual TArray<UTakeRecorderSource*> PreRecording(ULevelSequence* InSequence, FMovieSceneSequenceID InSequenceID, ULevelSequence* InMasterSequence, FManifestSerializer* InManifestSerializer) override;
 	virtual void StartRecording(const FTimecode& InSectionStartTimecode, const FFrameNumber& InSectionFirstFrame, class ULevelSequence* InSequence) override;
 	virtual void TickRecording(const FQualifiedFrameTime& CurrentSequenceTime) override;
 	virtual void StopRecording(class ULevelSequence* InSequence) override;
diff --git a/Engine/Plugins/MovieScene/LevelSequenceEditor/Source/LevelSequenceEditor/Private/LevelSequenceEditorBlueprintLibrary.cpp b/Engine/Plugins/MovieScene/LevelSequenceEditor/Source/LevelSequenceEditor/Private/LevelSequenceEditorBlueprintLibrary.cpp
--- a/Engine/Plugins/MovieScene/LevelSequenceEditor/Source/LevelSequenceEditor/Private/LevelSequenceEditorBlueprintLibrary.cpp
+++ b/Engine/Plugins/MovieScene/LevelSequenceEditor/Source/LevelSequenceEditor/Private/LevelSequenceEditorBlueprintLibrary.cpp
@@ -252,7 +252,7 @@ TArray<UObject*> ULevelSequenceEditorBlueprintLibrary::GetBoundObjects(FMovieSce
 	TArray<UObject*> BoundObjects;
 	if (CurrentSequencer.IsValid())
 	{
-		for (TWeakObjectPtr<> WeakObject : CurrentSequencer.Pin()->FindBoundObjects(ObjectBinding.GetGuid(), ObjectBinding.GetSequenceID()))
+		for (TWeakObjectPtr<> WeakObject : ObjectBinding.ResolveBoundObjects(CurrentSequencer.Pin()->GetFocusedTemplateID(), *CurrentSequencer.Pin()))
 		{
 			if (WeakObject.IsValid())
 			{
diff --git a/Engine/Plugins/MovieScene/SequencerScripting/Source/SequencerScripting/Private/ExtensionLibraries/MovieSceneSequenceExtensions.cpp b/Engine/Plugins/MovieScene/SequencerScripting/Source/SequencerScripting/Private/ExtensionLibraries/MovieSceneSequenceExtensions.cpp
--- a/Engine/Plugins/MovieScene/SequencerScripting/Source/SequencerScripting/Private/ExtensionLibraries/MovieSceneSequenceExtensions.cpp
+++ b/Engine/Plugins/MovieScene/SequencerScripting/Source/SequencerScripting/Private/ExtensionLibraries/MovieSceneSequenceExtensions.cpp
@@ -565,28 +565,98 @@ TArray<UObject*> UMovieSceneSequenceExtensions::LocateBoundObjects(UMovieSceneSe
 
 FMovieSceneObjectBindingID UMovieSceneSequenceExtensions::MakeBindingID(UMovieSceneSequence* MasterSequence, const FSequencerBindingProxy& InBinding, EMovieSceneObjectBindingSpace Space)
 {
-	FMovieSceneSequenceID SequenceID = MovieSceneSequenceID::Root;
+	// This function was kinda flawed before - when ::Local was passed for the Space parameter,
+	// and the sub sequence ID could not be found it would always fall back to a binding for ::Root without any Sequence ID
+	FMovieSceneObjectBindingID BindingID = GetPortableBindingID(MasterSequence, MasterSequence, InBinding);
+	if (Space == EMovieSceneObjectBindingSpace::Root)
+	{
+		BindingID.ReinterpretAsFixed();
+	}
+	return BindingID;
+}
 
-	FMovieSceneCompiledDataID DataID = UMovieSceneCompiledDataManager::GetPrecompiledData()->Compile(MasterSequence);
+FMovieSceneObjectBindingID UMovieSceneSequenceExtensions::GetBindingID(const FSequencerBindingProxy& InBinding)
+{
+	return UE::MovieScene::FRelativeObjectBindingID(InBinding.BindingID);
+}
+
+FMovieSceneObjectBindingID UMovieSceneSequenceExtensions::GetPortableBindingID(UMovieSceneSequence* MasterSequence, UMovieSceneSequence* DestinationSequence, const FSequencerBindingProxy& InBinding)
+{
+	if (!MasterSequence || !DestinationSequence || !InBinding.Sequence)
+	{
+		FFrame::KismetExecutionMessage(TEXT("Invalid sequence sepcified."), ELogVerbosity::Error);
+		return FMovieSceneObjectBindingID();
+	}
 
+	// If they are all the same sequence, we're dealing with a local binding - this requires no computation
+	if (MasterSequence == DestinationSequence && MasterSequence == InBinding.Sequence)
+	{
+		return UE::MovieScene::FRelativeObjectBindingID(InBinding.BindingID);
+	}
 
-	const FMovieSceneSequenceHierarchy* Hierarchy = UMovieSceneCompiledDataManager::GetPrecompiledData()->FindHierarchy(DataID);
-	if (Hierarchy)
+	// Destination is the destination for the BindingID to be serialized / resolved within
+	// Target is the target sequence that contains the actual binding
+
+	TOptional<FMovieSceneSequenceID> DestinationSequenceID;
+	TOptional<FMovieSceneSequenceID> TargetSequenceID;
+
+	if (MasterSequence == DestinationSequence)
+	{
+		DestinationSequenceID = MovieSceneSequenceID::Root;
+	}
+	if (MasterSequence == InBinding.Sequence)
 	{
-		for (const TTuple<FMovieSceneSequenceID, FMovieSceneSubSequenceData>& Pair : Hierarchy->AllSubSequenceData())
+		TargetSequenceID = MovieSceneSequenceID::Root;
+	}
+
+	// We know that we have at least one sequence ID to find, otherwise we would have entered the ::Local branch above
+	UMovieSceneCompiledDataManager*     CompiledDataManager = UMovieSceneCompiledDataManager::GetPrecompiledData();
+	const FMovieSceneCompiledDataID     DataID              = CompiledDataManager->Compile(MasterSequence);
+	const FMovieSceneSequenceHierarchy* Hierarchy           = CompiledDataManager->FindHierarchy(DataID);
+
+	// If we have no hierarchy, the supplied MasterSequence does not have any sub-sequences so the callee has given us bogus parameters
+	if (!Hierarchy)
+	{
+		FFrame::KismetExecutionMessage(*FString::Printf(TEXT("Master Sequence ('%s') does not have any sub-sequences."), *MasterSequence->GetPathName()), ELogVerbosity::Error);
+		return FMovieSceneObjectBindingID();
+	}
+
+	// Find the destination and/or target sequence IDs as required.
+	// This method is flawed if there is more than one instance of the sequence within the hierarchy
+	// In this case we just pick the first one we find
+	for (const TTuple<FMovieSceneSequenceID, FMovieSceneSubSequenceData>& Pair : Hierarchy->AllSubSequenceData())
+	{
+		if (UMovieSceneSequence* SubSequence = Pair.Value.GetSequence())
 		{
-			if (UMovieSceneSequence* SubSequence = Pair.Value.GetSequence())
+			if (!TargetSequenceID.IsSet() && SubSequence == InBinding.Sequence)
+			{
+				TargetSequenceID = Pair.Key;
+			}
+			if (!DestinationSequenceID.IsSet() && SubSequence == DestinationSequence)
 			{
-				if (SubSequence == InBinding.Sequence)
-				{
-					SequenceID = Pair.Key;
-					break;
-				}
+				DestinationSequenceID = Pair.Key;
+			}
+
+			if (DestinationSequenceID.IsSet() && TargetSequenceID.IsSet())
+			{
+				break;
 			}
 		}
 	}
 
-	return FMovieSceneObjectBindingID(InBinding.BindingID, SequenceID, Space);
+	if (!DestinationSequenceID.IsSet())
+	{
+		FFrame::KismetExecutionMessage(*FString::Printf(TEXT("Unable to locate DestinationSequence ('%s') within Master Sequence hierarchy ('%s')."), *DestinationSequence->GetPathName(), *MasterSequence->GetPathName()), ELogVerbosity::Error);
+		return FMovieSceneObjectBindingID();
+	}
+
+	if (!TargetSequenceID.IsSet())
+	{
+		FFrame::KismetExecutionMessage(*FString::Printf(TEXT("Unable to locate Sequence for InBinding ('%s') within Master Sequence hierarchy ('%s')."), *InBinding.Sequence->GetPathName(), *MasterSequence->GetPathName()), ELogVerbosity::Error);
+		return FMovieSceneObjectBindingID();
+	}
+
+	return UE::MovieScene::FRelativeObjectBindingID(DestinationSequenceID.GetValue(), TargetSequenceID.GetValue(), InBinding.BindingID, Hierarchy);
 }
 
 FSequencerBindingProxy UMovieSceneSequenceExtensions::ResolveBindingID(UMovieSceneSequence* MasterSequence, FMovieSceneObjectBindingID InObjectBindingID)
@@ -598,7 +668,7 @@ FSequencerBindingProxy UMovieSceneSequenceExtensions::ResolveBindingID(UMovieSce
 	const FMovieSceneSequenceHierarchy* Hierarchy = UMovieSceneCompiledDataManager::GetPrecompiledData()->FindHierarchy(DataID);
 	if (Hierarchy)
 	{
-		if (UMovieSceneSequence* SubSequence = Hierarchy->FindSubSequence(InObjectBindingID.GetSequenceID()))
+		if (UMovieSceneSequence* SubSequence = Hierarchy->FindSubSequence(InObjectBindingID.ResolveSequenceID(MovieSceneSequenceID::Root, Hierarchy)))
 		{
 			Sequence = SubSequence;
 		}
diff --git a/Engine/Plugins/MovieScene/SequencerScripting/Source/SequencerScripting/Public/ExtensionLibraries/MovieSceneSequenceExtensions.h b/Engine/Plugins/MovieScene/SequencerScripting/Source/SequencerScripting/Public/ExtensionLibraries/MovieSceneSequenceExtensions.h
--- a/Engine/Plugins/MovieScene/SequencerScripting/Source/SequencerScripting/Public/ExtensionLibraries/MovieSceneSequenceExtensions.h
+++ b/Engine/Plugins/MovieScene/SequencerScripting/Source/SequencerScripting/Public/ExtensionLibraries/MovieSceneSequenceExtensions.h
@@ -441,10 +441,35 @@ public:
 	 * @param Space The object binding space to resolve from (Root or Local)
 	 * @return The new object binding id
 	 */
-	UFUNCTION(BlueprintCallable, Category = "Sequence", meta = (ScriptMethod))
+	UE_DEPRECATED(5.0, "Please migrate to GetBindingID or GetPortableBindingID depending on use-case.")
+	UFUNCTION(BlueprintCallable, Category = "Sequence", meta = (ScriptMethod, DeprecatedFunction, DeprecationMessage="Please migrate to GetBindingID or GetPortableBindingID depending on use-case."))
 	static FMovieSceneObjectBindingID MakeBindingID(UMovieSceneSequence* MasterSequence, const FSequencerBindingProxy& InBinding, EMovieSceneObjectBindingSpace Space = EMovieSceneObjectBindingSpace::Root);
 
 
+	/**
+	 * Get the binding ID for a binding within a sequence.
+	 * @note: The resulting binding is only valid when applied to properties within the same sequence as this binding. Use GetPortableBindingID for bindings which live in different sub-sequences.
+	 *
+	 * @param Binding The binding proxy to generate the binding id from
+	 * @return The binding's id
+	 */
+	UFUNCTION(BlueprintCallable, Category = "Sequence", meta = (ScriptMethod))
+	static FMovieSceneObjectBindingID GetBindingID(const FSequencerBindingProxy& InBinding);
+
+
+	/**
+	 * Get a portable binding ID for a binding that resides in a different sequence to the one where this binding will be resolved.
+	 * @note: This function must be used over GetBindingID when the target binding resides in different shots or sub-sequences.
+	 * @note: Only unique instances of sequences within a master sequences are supported
+	 *
+	 * @param MasterSequence The master sequence that contains both the destination sequence (that will resolve the binding ID) and the target sequence that contains the actual binding
+	 * @param DestinationSequence The sequence that will own or resolve the resulting binding ID. For example, if the binding ID will be applied to a camera cut section pass the sequence that contains the camera cut track to this parameter.
+	 * @param Binding The target binding to create the ID from
+	 * @return The binding's id
+	 */
+	UFUNCTION(BlueprintCallable, Category = "Sequence", meta = (ScriptMethod))
+	static FMovieSceneObjectBindingID GetPortableBindingID(UMovieSceneSequence* MasterSequence, UMovieSceneSequence* DestinationSequence, const FSequencerBindingProxy& InBinding);
+
 	/**
 	 * Make a binding for the given binding ID
 	 *
diff --git a/Engine/Plugins/MovieScene/SequencerScripting/Source/SequencerScriptingEditor/Private/SequencerTools.cpp b/Engine/Plugins/MovieScene/SequencerScripting/Source/SequencerScriptingEditor/Private/SequencerTools.cpp
--- a/Engine/Plugins/MovieScene/SequencerScripting/Source/SequencerScriptingEditor/Private/SequencerTools.cpp
+++ b/Engine/Plugins/MovieScene/SequencerScripting/Source/SequencerScriptingEditor/Private/SequencerTools.cpp
@@ -135,7 +135,7 @@ TArray<FSequencerBoundObjects> USequencerToolsFunctionLibrary::GetBoundObjects(U
 	TArray<FSequencerBoundObjects> BoundObjects;
 	for (FSequencerBindingProxy Binding : InBindings)
 	{
-		FMovieSceneObjectBindingID ObjectBinding(Binding.BindingID, SequenceId);
+		FMovieSceneObjectBindingID ObjectBinding = UE::MovieScene::FFixedObjectBindingID(Binding.BindingID, SequenceId);
 		BoundObjects.Add(FSequencerBoundObjects(Binding, Player->GetBoundObjects(ObjectBinding)));
 	}
 
@@ -168,7 +168,8 @@ TArray<FSequencerBoundObjects> USequencerToolsFunctionLibrary::GetObjectBindings
 		TArray<FMovieSceneObjectBindingID> ObjectBindings = Player->GetObjectBindings(Object);
 		for (FMovieSceneObjectBindingID ObjectBinding : ObjectBindings)
 		{
-			FSequencerBindingProxy Binding(ObjectBinding.GetGuid(), Player->State.FindSequence(ObjectBinding.GetSequenceID()));
+			FMovieSceneSequenceID SequenceID = ObjectBinding.ResolveSequenceID(MovieSceneSequenceID::Root, *Player);
+			FSequencerBindingProxy Binding(ObjectBinding.GetGuid(), Player->State.FindSequence(SequenceID));
 			BoundObjects.Add(FSequencerBoundObjects(Binding, TArray<UObject*>({ Object })));
 		}
 	}
diff --git a/Engine/Plugins/MovieScene/TemplateSequence/Source/TemplateSequence/Private/Section/TemplateSequenceSection.cpp b/Engine/Plugins/MovieScene/TemplateSequence/Source/TemplateSequence/Private/Section/TemplateSequenceSection.cpp
--- a/Engine/Plugins/MovieScene/TemplateSequence/Source/TemplateSequence/Private/Section/TemplateSequenceSection.cpp
+++ b/Engine/Plugins/MovieScene/TemplateSequence/Source/TemplateSequence/Private/Section/TemplateSequenceSection.cpp
@@ -3,6 +3,7 @@
 #include "Sections/TemplateSequenceSection.h"
 #include "TemplateSequence.h"
 #include "Systems/TemplateSequenceSystem.h"
+#include "Evaluation/MovieSceneRootOverridePath.h"
 
 UTemplateSequenceSection::UTemplateSequenceSection(const FObjectInitializer& ObjInitializer)
 	: Super(ObjInitializer)
@@ -24,15 +25,10 @@ void UTemplateSequenceSection::ImportEntityImpl(UMovieSceneEntitySystemLinker* E
 
 	if (UTemplateSequence* TemplateSubSequence = Cast<UTemplateSequence>(GetSequence()))
 	{
-		const FMovieSceneObjectBindingID InnerObjectBindingID(
-				TemplateSubSequence->GetRootObjectBindingID(), GetSequenceID(), EMovieSceneObjectBindingSpace::Local);
+		const FSubSequencePath      PathToRoot = EntityLinker->GetInstanceRegistry()->GetInstance(Params.Sequence.InstanceHandle).GetSubSequencePath();
+		const FMovieSceneSequenceID ResolvedSequenceID = PathToRoot.ResolveChildSequenceID(this->GetSequenceID());
 
-		const FSequenceInstance& SequenceInstance = EntityLinker->GetInstanceRegistry()->GetInstance(Params.Sequence.InstanceHandle);
-		const FMovieSceneObjectBindingID AbsoluteInnerObjectBindingID(
-				InnerObjectBindingID.ResolveLocalToRoot(SequenceInstance.GetSequenceID(), *SequenceInstance.GetPlayer()));
-
-		ComponentData.InnerOperand = FMovieSceneEvaluationOperand(
-				AbsoluteInnerObjectBindingID.GetSequenceID(), AbsoluteInnerObjectBindingID.GetGuid());
+		ComponentData.InnerOperand = FMovieSceneEvaluationOperand(ResolvedSequenceID, TemplateSubSequence->GetRootObjectBindingID());
 	}
 
 	FGuid ObjectBindingID = Params.GetObjectBindingID();
diff --git a/Engine/Plugins/Runtime/GameplayAbilities/Source/GameplayAbilities/Private/Sequencer/MovieSceneGameplayCueSections.cpp b/Engine/Plugins/Runtime/GameplayAbilities/Source/GameplayAbilities/Private/Sequencer/MovieSceneGameplayCueSections.cpp
--- a/Engine/Plugins/Runtime/GameplayAbilities/Source/GameplayAbilities/Private/Sequencer/MovieSceneGameplayCueSections.cpp
+++ b/Engine/Plugins/Runtime/GameplayAbilities/Source/GameplayAbilities/Private/Sequencer/MovieSceneGameplayCueSections.cpp
@@ -65,8 +65,7 @@ void ExecuteGameplayCueEvent(IMovieScenePlayer* Player, const UE::MovieScene::FE
 
 	if (GameplayCueKey.Instigator.IsValid())
 	{
-		FMovieSceneObjectBindingID Resolved = GameplayCueKey.Instigator.ResolveLocalToRoot(Params.SequenceID, *Player);
-		for (TWeakObjectPtr<> WeakObject : Player->FindBoundObjects(Resolved.GetGuid(), Resolved.GetSequenceID()))
+		for (TWeakObjectPtr<> WeakObject : GameplayCueKey.Instigator.ResolveBoundObjects(Params.SequenceID, *Player))
 		{
 			if (AActor* Actor = ActorFromResolvedObject(WeakObject.Get()))
 			{
@@ -78,8 +77,7 @@ void ExecuteGameplayCueEvent(IMovieScenePlayer* Player, const UE::MovieScene::FE
 
 	if (GameplayCueKey.EffectCauser.IsValid())
 	{
-		FMovieSceneObjectBindingID Resolved = GameplayCueKey.EffectCauser.ResolveLocalToRoot(Params.SequenceID, *Player);
-		for (TWeakObjectPtr<> WeakObject : Player->FindBoundObjects(Resolved.GetGuid(), Resolved.GetSequenceID()))
+		for (TWeakObjectPtr<> WeakObject : GameplayCueKey.EffectCauser.ResolveBoundObjects(Params.SequenceID, *Player))
 		{
 			if (AActor* Actor = ActorFromResolvedObject(WeakObject.Get()))
 			{
diff --git a/Engine/Plugins/VirtualProduction/DMX/DMXEngine/Source/DMXEditor/Private/Sequencer/TakeRecorderDMXLibrarySource.cpp b/Engine/Plugins/VirtualProduction/DMX/DMXEngine/Source/DMXEditor/Private/Sequencer/TakeRecorderDMXLibrarySource.cpp
--- a/Engine/Plugins/VirtualProduction/DMX/DMXEngine/Source/DMXEditor/Private/Sequencer/TakeRecorderDMXLibrarySource.cpp
+++ b/Engine/Plugins/VirtualProduction/DMX/DMXEngine/Source/DMXEditor/Private/Sequencer/TakeRecorderDMXLibrarySource.cpp
@@ -64,7 +64,7 @@ void UTakeRecorderDMXLibrarySource::OnEntitiesUpdated(UDMXLibrary* UpdatedLibrar
 	}
 }
 
-TArray<UTakeRecorderSource*> UTakeRecorderDMXLibrarySource::PreRecording(ULevelSequence* InSequence, ULevelSequence* InMasterSequence, FManifestSerializer* InManifestSerializer)
+TArray<UTakeRecorderSource*> UTakeRecorderDMXLibrarySource::PreRecording(ULevelSequence* InSequence, FMovieSceneSequenceID InSequenceID, ULevelSequence* InMasterSequence, FManifestSerializer* InManifestSerializer)
 {
 	if (DMXLibrary)
 	{
diff --git a/Engine/Plugins/VirtualProduction/DMX/DMXEngine/Source/DMXEditor/Private/Sequencer/TakeRecorderDMXLibrarySource.h b/Engine/Plugins/VirtualProduction/DMX/DMXEngine/Source/DMXEditor/Private/Sequencer/TakeRecorderDMXLibrarySource.h
--- a/Engine/Plugins/VirtualProduction/DMX/DMXEngine/Source/DMXEditor/Private/Sequencer/TakeRecorderDMXLibrarySource.h
+++ b/Engine/Plugins/VirtualProduction/DMX/DMXEngine/Source/DMXEditor/Private/Sequencer/TakeRecorderDMXLibrarySource.h
@@ -66,7 +66,7 @@ private:
 	void OnEntitiesUpdated(UDMXLibrary* UpdatedLibrary);
 
 	//~ UTakeRecorderSource
-	virtual TArray<UTakeRecorderSource*> PreRecording(ULevelSequence* InSequence, ULevelSequence* InMasterSequence, FManifestSerializer* InManifestSerializer) override;
+	virtual TArray<UTakeRecorderSource*> PreRecording(ULevelSequence* InSequence, FMovieSceneSequenceID InSequenceID, ULevelSequence* InMasterSequence, FManifestSerializer* InManifestSerializer) override;
 	virtual void StartRecording(const FTimecode& InSectionStartTimecode, const FFrameNumber& InSectionFirstFrame, class ULevelSequence* InSequence) override;
 	virtual void StopRecording(class ULevelSequence* InSequence) override;
 	virtual void TickRecording(const FQualifiedFrameTime& CurrentTime) override;
diff --git a/Engine/Plugins/VirtualProduction/Takes/Source/TakeRecorderSources/Private/TakeRecorderActorSource.cpp b/Engine/Plugins/VirtualProduction/Takes/Source/TakeRecorderSources/Private/TakeRecorderActorSource.cpp
--- a/Engine/Plugins/VirtualProduction/Takes/Source/TakeRecorderSources/Private/TakeRecorderActorSource.cpp
+++ b/Engine/Plugins/VirtualProduction/Takes/Source/TakeRecorderSources/Private/TakeRecorderActorSource.cpp
@@ -178,6 +178,8 @@ UTakeRecorderActorSource::UTakeRecorderActorSource(const FObjectInitializer& Obj
 	bReduceKeys = true;
 	bRecordParentHierarchy = true;
 
+	TargetSequenceID = MovieSceneSequenceID::Invalid;
+
 	// Build the property map on initialization so that sources created at runtime have a default map
 	RebuildRecordedPropertyMap();
 }
@@ -188,7 +190,7 @@ bool UTakeRecorderActorSource::IsValid() const
 }
 
 
-TArray<UTakeRecorderSource*> UTakeRecorderActorSource::PreRecording(class ULevelSequence* InSequence, class ULevelSequence* InMasterSequence, FManifestSerializer* InManifestSerializer)
+TArray<UTakeRecorderSource*> UTakeRecorderActorSource::PreRecording(ULevelSequence* InSequence, FMovieSceneSequenceID InSequenceID, ULevelSequence* InMasterSequence, FManifestSerializer* InManifestSerializer)
 {
 	// Don't bother doing anything if we don't have a valid actor to record.
 	if (!Target.IsValid())
@@ -202,7 +204,7 @@ TArray<UTakeRecorderSource*> UTakeRecorderActorSource::PreRecording(class ULevel
 	AActor* ActorToRecord = Target.Get();
 	TargetLevelSequence = InSequence;
 	MasterLevelSequence = InMasterSequence;
-	SequenceID.Reset(); 
+	TargetSequenceID = InSequenceID;
 
 	FString ObjectBindingName = ActorToRecord->GetName();
 
diff --git a/Engine/Plugins/VirtualProduction/Takes/Source/TakeRecorderSources/Private/TakeRecorderActorSource.h b/Engine/Plugins/VirtualProduction/Takes/Source/TakeRecorderSources/Private/TakeRecorderActorSource.h
--- a/Engine/Plugins/VirtualProduction/Takes/Source/TakeRecorderSources/Private/TakeRecorderActorSource.h
+++ b/Engine/Plugins/VirtualProduction/Takes/Source/TakeRecorderSources/Private/TakeRecorderActorSource.h
@@ -121,7 +121,7 @@ public:
 	
 	// UTakeRecorderSource Interface
 	virtual bool IsValid()const override;
-	virtual TArray<UTakeRecorderSource*> PreRecording(class ULevelSequence* InSequence, class ULevelSequence* InMasterSequence, FManifestSerializer* InManifestSerializer) override;
+	virtual TArray<UTakeRecorderSource*> PreRecording(ULevelSequence* InSequence, FMovieSceneSequenceID InSequenceID, ULevelSequence* InMasterSequence, FManifestSerializer* InManifestSerializer) override;
 	virtual void StartRecording(const FTimecode& InSectionStartTimecode, const FFrameNumber& InSectionFirstFrame, class ULevelSequence* InSequence) override;
 	virtual void TickRecording(const FQualifiedFrameTime& CurrentSequenceTime) override;
 	virtual void StopRecording(class ULevelSequence* InSequence) override;
@@ -130,6 +130,8 @@ public:
 	virtual FString GetSubsceneTrackName(ULevelSequence* InSequence) const override;
 	virtual FString GetSubsceneAssetName(ULevelSequence* InSequence) const override;
 	virtual void AddContentsToFolder(class UMovieSceneFolder* InFolder) override;
+	virtual FMovieSceneSequenceID GetSequenceID() const override { return TargetSequenceID; }
+
 	// ~UTakeRecorderSource Interface
 
 	/** Set the Target actor that we are going to record. Will reset the Recorded Property Map to defaults. */
@@ -148,13 +150,6 @@ public:
 	/** Get the record type. If set to project default, gets the type from the project settings */
 	bool GetRecordToPossessable() const;
 
-	TOptional<FMovieSceneSequenceID> GetSequenceID() const
-	{
-		return SequenceID;
-	}
-
-	void SetSequenceID(FMovieSceneSequenceID InSequenceID) { SequenceID = InSequenceID; }
-
 protected:
 
 	/**
@@ -310,9 +305,9 @@ private:
 	TSet<class UActorComponent*> CachedComponentList;
 
 	/**
-	* Optional ID of the TargetLevelSequence. Cached since calculating this can be heavy. 
+	* ID of the TargetLevelSequence.
 	*/
-	TOptional<FMovieSceneSequenceID> SequenceID;
+	FMovieSceneSequenceID TargetSequenceID;
 
 	/**
 	*  Serializer
diff --git a/Engine/Plugins/VirtualProduction/Takes/Source/TakeRecorderSources/Private/TakeRecorderCameraCutSource.cpp b/Engine/Plugins/VirtualProduction/Takes/Source/TakeRecorderSources/Private/TakeRecorderCameraCutSource.cpp
--- a/Engine/Plugins/VirtualProduction/Takes/Source/TakeRecorderSources/Private/TakeRecorderCameraCutSource.cpp
+++ b/Engine/Plugins/VirtualProduction/Takes/Source/TakeRecorderSources/Private/TakeRecorderCameraCutSource.cpp
@@ -26,7 +26,7 @@ UTakeRecorderCameraCutSource::UTakeRecorderCameraCutSource(const FObjectInitiali
 	TrackTint = FColor(160, 160, 160);
 }
 
-TArray<UTakeRecorderSource*> UTakeRecorderCameraCutSource::PreRecording(class ULevelSequence* InSequence, class ULevelSequence* InMasterSequence, FManifestSerializer* InManifestSerializer)
+TArray<UTakeRecorderSource*> UTakeRecorderCameraCutSource::PreRecording(ULevelSequence* InSequence, FMovieSceneSequenceID InSequenceID, ULevelSequence* InMasterSequence, FManifestSerializer* InManifestSerializer)
 {
 	World = TakeRecorderSourcesUtils::GetSourceWorld(InSequence);
 	MasterLevelSequence = InMasterSequence;
@@ -141,7 +141,8 @@ TArray<UTakeRecorderSource*> UTakeRecorderCameraCutSource::PostRecording(class U
 			}
 
 			UMovieSceneCameraCutSection* CameraCutSection = Cast<UMovieSceneCameraCutSection>(CameraCutTrack->CreateNewSection());
-			CameraCutSection->SetCameraBindingID(FMovieSceneObjectBindingID(CameraCutData[CameraCutIndex].Guid, CameraCutData[CameraCutIndex].SequenceID, EMovieSceneObjectBindingSpace::Local));
+			FMovieSceneObjectBindingID CameraCutBinding = UE::MovieScene::FRelativeObjectBindingID(CameraCutData[CameraCutIndex].Guid, CameraCutData[CameraCutIndex].SequenceID);
+			CameraCutSection->SetCameraBindingID(CameraCutBinding);
 			CameraCutSection->SetRange(Range);
 			CameraCutTrack->AddSection(*CameraCutSection);
 		}
diff --git a/Engine/Plugins/VirtualProduction/Takes/Source/TakeRecorderSources/Private/TakeRecorderCameraCutSource.h b/Engine/Plugins/VirtualProduction/Takes/Source/TakeRecorderSources/Private/TakeRecorderCameraCutSource.h
--- a/Engine/Plugins/VirtualProduction/Takes/Source/TakeRecorderSources/Private/TakeRecorderCameraCutSource.h
+++ b/Engine/Plugins/VirtualProduction/Takes/Source/TakeRecorderSources/Private/TakeRecorderCameraCutSource.h
@@ -25,7 +25,7 @@ public:
 private:
 
 	// UTakeRecorderSource
-	virtual TArray<UTakeRecorderSource*> PreRecording(class ULevelSequence* InSequence, class ULevelSequence* InMasterSequence, FManifestSerializer* InManifestSerializer) override;
+	virtual TArray<UTakeRecorderSource*> PreRecording(ULevelSequence* InSequence, FMovieSceneSequenceID InSequenceID, ULevelSequence* InMasterSequence, FManifestSerializer* InManifestSerializer) override;
 	virtual void TickRecording(const FQualifiedFrameTime& CurrentTime) override;
 	virtual TArray<UTakeRecorderSource*> PostRecording(class ULevelSequence* InSequence, class ULevelSequence* InMasterSequence) override;
 	virtual FText GetDisplayTextImpl() const override;
diff --git a/Engine/Plugins/VirtualProduction/Takes/Source/TakeRecorderSources/Private/TakeRecorderLevelSequenceSource.cpp b/Engine/Plugins/VirtualProduction/Takes/Source/TakeRecorderSources/Private/TakeRecorderLevelSequenceSource.cpp
--- a/Engine/Plugins/VirtualProduction/Takes/Source/TakeRecorderSources/Private/TakeRecorderLevelSequenceSource.cpp
+++ b/Engine/Plugins/VirtualProduction/Takes/Source/TakeRecorderSources/Private/TakeRecorderLevelSequenceSource.cpp
@@ -15,7 +15,7 @@ UTakeRecorderLevelSequenceSource::UTakeRecorderLevelSequenceSource(const FObject
 	TrackTint = FColor(164, 74, 74);
 }
 
-TArray<UTakeRecorderSource*> UTakeRecorderLevelSequenceSource::PreRecording(class ULevelSequence* InSequence, class ULevelSequence* InMasterSequence, FManifestSerializer* InManifestSerializer)
+TArray<UTakeRecorderSource*> UTakeRecorderLevelSequenceSource::PreRecording(ULevelSequence* InSequence, FMovieSceneSequenceID InSequenceID, ULevelSequence* InMasterSequence, FManifestSerializer* InManifestSerializer)
 {
 	UWorld* World = TakesUtils::GetFirstPIEWorld();
 	if (!World)
diff --git a/Engine/Plugins/VirtualProduction/Takes/Source/TakeRecorderSources/Private/TakeRecorderLevelSequenceSource.h b/Engine/Plugins/VirtualProduction/Takes/Source/TakeRecorderSources/Private/TakeRecorderLevelSequenceSource.h
--- a/Engine/Plugins/VirtualProduction/Takes/Source/TakeRecorderSources/Private/TakeRecorderLevelSequenceSource.h
+++ b/Engine/Plugins/VirtualProduction/Takes/Source/TakeRecorderSources/Private/TakeRecorderLevelSequenceSource.h
@@ -25,7 +25,7 @@ public:
 private:
 
 	// UTakeRecorderSource
-	virtual TArray<UTakeRecorderSource*> PreRecording(class ULevelSequence* InSequence, class ULevelSequence* InMasterSequence, FManifestSerializer* InManifestSerializer) override;
+	virtual TArray<UTakeRecorderSource*> PreRecording(ULevelSequence* InSequence, FMovieSceneSequenceID InSequenceID, ULevelSequence* InMasterSequence, FManifestSerializer* InManifestSerializer) override;
 	virtual void StartRecording(const FTimecode& InSectionStartTimecode, const FFrameNumber& InSectionFirstFrame, class ULevelSequence* InSequence) override;
 	virtual void StopRecording(class ULevelSequence* InSequence) override;
 	virtual FText GetDisplayTextImpl() const override;
diff --git a/Engine/Plugins/VirtualProduction/Takes/Source/TakeRecorderSources/Private/TakeRecorderLevelVisibilitySource.cpp b/Engine/Plugins/VirtualProduction/Takes/Source/TakeRecorderSources/Private/TakeRecorderLevelVisibilitySource.cpp
--- a/Engine/Plugins/VirtualProduction/Takes/Source/TakeRecorderSources/Private/TakeRecorderLevelVisibilitySource.cpp
+++ b/Engine/Plugins/VirtualProduction/Takes/Source/TakeRecorderSources/Private/TakeRecorderLevelVisibilitySource.cpp
@@ -57,7 +57,7 @@ UTakeRecorderLevelVisibilitySource::UTakeRecorderLevelVisibilitySource(const FOb
 	: Super(ObjInit)
 {}
 
-TArray<UTakeRecorderSource*> UTakeRecorderLevelVisibilitySource::PreRecording(class ULevelSequence* InSequence, class ULevelSequence* InMasterSequence, FManifestSerializer* InManifestSerializer)
+TArray<UTakeRecorderSource*> UTakeRecorderLevelVisibilitySource::PreRecording(ULevelSequence* InSequence, FMovieSceneSequenceID InSequenceID, ULevelSequence* InMasterSequence, FManifestSerializer* InManifestSerializer)
 {
 	UWorld* World = TakeRecorderSourcesUtils::GetSourceWorld(InSequence);
 
diff --git a/Engine/Plugins/VirtualProduction/Takes/Source/TakeRecorderSources/Private/TakeRecorderLevelVisibilitySource.h b/Engine/Plugins/VirtualProduction/Takes/Source/TakeRecorderSources/Private/TakeRecorderLevelVisibilitySource.h
--- a/Engine/Plugins/VirtualProduction/Takes/Source/TakeRecorderSources/Private/TakeRecorderLevelVisibilitySource.h
+++ b/Engine/Plugins/VirtualProduction/Takes/Source/TakeRecorderSources/Private/TakeRecorderLevelVisibilitySource.h
@@ -40,7 +40,7 @@ public:
 private:
 
 	// UTakeRecorderSource
-	virtual TArray<UTakeRecorderSource*> PreRecording(class ULevelSequence* InSequence, class ULevelSequence* InMasterSequence, FManifestSerializer* InManifestSerializer) override;
+	virtual TArray<UTakeRecorderSource*> PreRecording(ULevelSequence* InSequence, FMovieSceneSequenceID InSequenceID, ULevelSequence* InMasterSequence, FManifestSerializer* InManifestSerializer) override;
 	virtual void TickRecording(const FQualifiedFrameTime& CurrentTime) override;
 	virtual void AddContentsToFolder(class UMovieSceneFolder* InFolder) override;
 	virtual FText GetDisplayTextImpl() const override;
diff --git a/Engine/Plugins/VirtualProduction/Takes/Source/TakeRecorderSources/Private/TakeRecorderMicrophoneAudioSource.cpp b/Engine/Plugins/VirtualProduction/Takes/Source/TakeRecorderSources/Private/TakeRecorderMicrophoneAudioSource.cpp
--- a/Engine/Plugins/VirtualProduction/Takes/Source/TakeRecorderSources/Private/TakeRecorderMicrophoneAudioSource.cpp
+++ b/Engine/Plugins/VirtualProduction/Takes/Source/TakeRecorderSources/Private/TakeRecorderMicrophoneAudioSource.cpp
@@ -93,7 +93,7 @@ static FString MakeNewAssetName(const FString& BaseAssetPath, const FString& Bas
 	return AssetName;
 }
 
-TArray<UTakeRecorderSource*> UTakeRecorderMicrophoneAudioSource::PreRecording(ULevelSequence* InSequence, ULevelSequence* InMasterSequence, FManifestSerializer* InManifestSerializer)
+TArray<UTakeRecorderSource*> UTakeRecorderMicrophoneAudioSource::PreRecording(ULevelSequence* InSequence, FMovieSceneSequenceID InSequenceID, ULevelSequence* InMasterSequence, FManifestSerializer* InManifestSerializer)
 {
 	UMovieScene* MovieScene = InSequence->GetMovieScene();
 	for (auto MasterTrack : MovieScene->GetMasterTracks())
diff --git a/Engine/Plugins/VirtualProduction/Takes/Source/TakeRecorderSources/Private/TakeRecorderMicrophoneAudioSource.h b/Engine/Plugins/VirtualProduction/Takes/Source/TakeRecorderSources/Private/TakeRecorderMicrophoneAudioSource.h
--- a/Engine/Plugins/VirtualProduction/Takes/Source/TakeRecorderSources/Private/TakeRecorderMicrophoneAudioSource.h
+++ b/Engine/Plugins/VirtualProduction/Takes/Source/TakeRecorderSources/Private/TakeRecorderMicrophoneAudioSource.h
@@ -59,7 +59,7 @@ public:
 private:
 
 	// UTakeRecorderSource
-	virtual TArray<UTakeRecorderSource*> PreRecording(class ULevelSequence* InSequence, class ULevelSequence* InMasterSequence, FManifestSerializer* InManifestSerializer) override;
+	virtual TArray<UTakeRecorderSource*> PreRecording(ULevelSequence* InSequence, FMovieSceneSequenceID InSequenceID, ULevelSequence* InMasterSequence, FManifestSerializer* InManifestSerializer) override;
 	virtual void StartRecording(const FTimecode& InSectionStartTimecode, const FFrameNumber& InSectionFirstFrame, class ULevelSequence* InSequence) override;
 	virtual void StopRecording(class ULevelSequence* InSequence) override;
 	virtual FText GetDisplayTextImpl() const override;
diff --git a/Engine/Plugins/VirtualProduction/Takes/Source/TakeRecorderSources/Private/TakeRecorderPlayerSource.cpp b/Engine/Plugins/VirtualProduction/Takes/Source/TakeRecorderSources/Private/TakeRecorderPlayerSource.cpp
--- a/Engine/Plugins/VirtualProduction/Takes/Source/TakeRecorderSources/Private/TakeRecorderPlayerSource.cpp
+++ b/Engine/Plugins/VirtualProduction/Takes/Source/TakeRecorderSources/Private/TakeRecorderPlayerSource.cpp
@@ -19,7 +19,7 @@ UTakeRecorderPlayerSource::UTakeRecorderPlayerSource(const FObjectInitializer& O
 }
 
 
-TArray<UTakeRecorderSource*> UTakeRecorderPlayerSource::PreRecording(class ULevelSequence* InSequence, class ULevelSequence* InMasterSequence, FManifestSerializer* InManifestSerializer)
+TArray<UTakeRecorderSource*> UTakeRecorderPlayerSource::PreRecording(ULevelSequence* InSequence, FMovieSceneSequenceID InSequenceID, ULevelSequence* InMasterSequence, FManifestSerializer* InManifestSerializer)
 {
 	TArray<UTakeRecorderSource*> NewSources;
 
diff --git a/Engine/Plugins/VirtualProduction/Takes/Source/TakeRecorderSources/Private/TakeRecorderPlayerSource.h b/Engine/Plugins/VirtualProduction/Takes/Source/TakeRecorderSources/Private/TakeRecorderPlayerSource.h
--- a/Engine/Plugins/VirtualProduction/Takes/Source/TakeRecorderSources/Private/TakeRecorderPlayerSource.h
+++ b/Engine/Plugins/VirtualProduction/Takes/Source/TakeRecorderSources/Private/TakeRecorderPlayerSource.h
@@ -21,7 +21,7 @@ public:
 private:
 
 	// UTakeRecorderSource
-	virtual TArray<UTakeRecorderSource*> PreRecording(class ULevelSequence* InSequence, class ULevelSequence* InMasterSequence, FManifestSerializer* InManifestSerializer) override;
+	virtual TArray<UTakeRecorderSource*> PreRecording(ULevelSequence* InSequence, FMovieSceneSequenceID InSequenceID, ULevelSequence* InMasterSequence, FManifestSerializer* InManifestSerializer) override;
 	virtual TArray<UTakeRecorderSource*> PostRecording(class ULevelSequence* InSequence, class ULevelSequence* InMasterSequence) override;
 	virtual FText GetDisplayTextImpl() const override;
 	virtual bool CanAddSource(UTakeRecorderSources* InSources) const override;
diff --git a/Engine/Plugins/VirtualProduction/Takes/Source/TakeRecorderSources/Private/TakeRecorderSourcesUtils.h b/Engine/Plugins/VirtualProduction/Takes/Source/TakeRecorderSources/Private/TakeRecorderSourcesUtils.h
--- a/Engine/Plugins/VirtualProduction/Takes/Source/TakeRecorderSources/Private/TakeRecorderSourcesUtils.h
+++ b/Engine/Plugins/VirtualProduction/Takes/Source/TakeRecorderSources/Private/TakeRecorderSourcesUtils.h
@@ -158,8 +158,8 @@ namespace TakeRecorderSourcesUtils
 	*/
 	static FMovieSceneSequenceID GetLevelSequenceID(const UTakeRecorderSource* InSource, class AActor* Actor, ULevelSequence* MasterLevelSequence)
 	{
-		FMovieSceneSequenceID OutSequenceID = MovieSceneSequenceID::Root;
 		UTakeRecorderSources* OwningSources = CastChecked<UTakeRecorderSources>(InSource->GetOuter());
+
 		for (UTakeRecorderSource* Source : OwningSources->GetSources())
 		{
 			if (UTakeRecorderActorSource* ActorSource = Cast<UTakeRecorderActorSource>(Source))
@@ -169,34 +169,12 @@ namespace TakeRecorderSourcesUtils
 				{
 					if (ActorSource->TargetLevelSequence != MasterLevelSequence)
 					{
-						if (!ActorSource->GetSequenceID().IsSet()) // only compile if it's not set yet this rcording.
-						{
-							UMovieSceneCompiledDataManager::GetPrecompiledData()->Compile(MasterLevelSequence);
-
-							FMovieSceneCompiledDataID DataID = UMovieSceneCompiledDataManager::GetPrecompiledData()->GetDataID(MasterLevelSequence);
-							const FMovieSceneSequenceHierarchy* Hierarchy = UMovieSceneCompiledDataManager::GetPrecompiledData()->FindHierarchy(DataID);
-							if (Hierarchy)
-							{
-								for (const TTuple<FMovieSceneSequenceID, FMovieSceneSubSequenceData>& Pair : Hierarchy->AllSubSequenceData())
-								{
-									if (Pair.Value.Sequence == ActorSource->TargetLevelSequence)
-									{
-										OutSequenceID = Pair.Key;
-										ActorSource->SetSequenceID(OutSequenceID);
-										break;
-
-									}
-								}
-							}
-						}
-						else
-						{
-							OutSequenceID = ActorSource->GetSequenceID().GetValue();
-						}
+						return ActorSource->GetSequenceID();
 					}
 				}
 			}
 		}
-		return OutSequenceID;
+
+		return MovieSceneSequenceID::Root;
 	}
 }
\ No newline at end of file
diff --git a/Engine/Plugins/VirtualProduction/Takes/Source/TakeRecorderSources/Private/TakeRecorderWorldSource.cpp b/Engine/Plugins/VirtualProduction/Takes/Source/TakeRecorderSources/Private/TakeRecorderWorldSource.cpp
--- a/Engine/Plugins/VirtualProduction/Takes/Source/TakeRecorderSources/Private/TakeRecorderWorldSource.cpp
+++ b/Engine/Plugins/VirtualProduction/Takes/Source/TakeRecorderSources/Private/TakeRecorderWorldSource.cpp
@@ -40,7 +40,7 @@ UTakeRecorderWorldSource::UTakeRecorderWorldSource(const FObjectInitializer& Obj
 	TrackTint = FColor(129, 129, 129);
 }
 
-TArray<UTakeRecorderSource*> UTakeRecorderWorldSource::PreRecording(class ULevelSequence* InSequence, class ULevelSequence* InMasterSequence, FManifestSerializer* InManifestSerializer)
+TArray<UTakeRecorderSource*> UTakeRecorderWorldSource::PreRecording(ULevelSequence* InSequence, FMovieSceneSequenceID InSequenceID, ULevelSequence* InMasterSequence, FManifestSerializer* InManifestSerializer)
 {
 	TArray<UTakeRecorderSource*> NewSources;
 
diff --git a/Engine/Plugins/VirtualProduction/Takes/Source/TakeRecorderSources/Private/TakeRecorderWorldSource.h b/Engine/Plugins/VirtualProduction/Takes/Source/TakeRecorderSources/Private/TakeRecorderWorldSource.h
--- a/Engine/Plugins/VirtualProduction/Takes/Source/TakeRecorderSources/Private/TakeRecorderWorldSource.h
+++ b/Engine/Plugins/VirtualProduction/Takes/Source/TakeRecorderSources/Private/TakeRecorderWorldSource.h
@@ -45,7 +45,7 @@ public:
 private:
 
 	// UTakeRecorderSource
-	virtual TArray<UTakeRecorderSource*> PreRecording(class ULevelSequence* InSequence, class ULevelSequence* InMasterSequence, FManifestSerializer* InManifestSerializer) override;
+	virtual TArray<UTakeRecorderSource*> PreRecording(ULevelSequence* InSequence, FMovieSceneSequenceID InSequenceID, ULevelSequence* InMasterSequence, FManifestSerializer* InManifestSerializer) override;
 	virtual TArray<UTakeRecorderSource*> PostRecording(class ULevelSequence* InSequence, class ULevelSequence* InMasterSequence) override;
 	virtual bool SupportsTakeNumber() const override { return false; }
 	virtual FText GetDisplayTextImpl() const override;
diff --git a/Engine/Plugins/VirtualProduction/Takes/Source/TakeTrackRecorders/Private/MovieScene3DAttachTrackRecorder.cpp b/Engine/Plugins/VirtualProduction/Takes/Source/TakeTrackRecorders/Private/MovieScene3DAttachTrackRecorder.cpp
--- a/Engine/Plugins/VirtualProduction/Takes/Source/TakeTrackRecorders/Private/MovieScene3DAttachTrackRecorder.cpp
+++ b/Engine/Plugins/VirtualProduction/Takes/Source/TakeTrackRecorders/Private/MovieScene3DAttachTrackRecorder.cpp
@@ -6,6 +6,7 @@
 #include "Modules/ModuleManager.h"
 #include "SequenceRecorderUtils.h"
 #include "MovieScene.h"
+#include "LevelSequence.h"
 
 bool FMovieScene3DAttachTrackRecorderFactory::CanRecordObject(UObject* InObjectToRecord) const
 {
@@ -55,9 +56,11 @@ void UMovieScene3DAttachTrackRecorder::RecordSampleImpl(const FQualifiedFrameTim
 				MovieSceneSection->TimecodeSource = MovieScene->TimecodeSource;
 				MovieSceneSection->SetRange(TRange<FFrameNumber>(CurrentFrame, CurrentFrame));
 
-				FMovieSceneSequenceID SequenceID;
-				SequenceID = OwningTakeRecorderSource->GetLevelSequenceID(AttachedToActor);
-				MovieSceneSection->SetConstraintId(Guid, SequenceID);
+				FMovieSceneSequenceID TargetSequenceID = OwningTakeRecorderSource->GetLevelSequenceID(AttachedToActor);
+				FMovieSceneSequenceID ThisSequenceID   = OwningTakeRecorderSource->GetSequenceID();
+
+				FMovieSceneObjectBindingID NewBinding = UE::MovieScene::FRelativeObjectBindingID(ThisSequenceID, TargetSequenceID, Guid, OwningTakeRecorderSource->GetMasterLevelSequence());
+				MovieSceneSection->SetConstraintBindingID(NewBinding);
 			}
 
 			ActorAttachedTo = AttachedToActor;
@@ -88,9 +91,11 @@ void UMovieScene3DAttachTrackRecorder::FinalizeTrackImpl()
 				UE_LOG(LogTemp, Warning, TEXT("Could not find binding to attach (%s) to its parent (%s), perhaps (%s) was not recorded?"), *ActorToRecord->GetActorLabel(), *AttachedToActor->GetActorLabel(), *AttachedToActor->GetActorLabel());
 			}
 
-			FMovieSceneSequenceID SequenceID;
-			SequenceID = OwningTakeRecorderSource->GetLevelSequenceID(AttachedToActor);
-			MovieSceneSection->SetConstraintId(Guid, SequenceID);
+			FMovieSceneSequenceID TargetSequenceID = OwningTakeRecorderSource->GetLevelSequenceID(AttachedToActor);
+			FMovieSceneSequenceID ThisSequenceID   = OwningTakeRecorderSource->GetSequenceID();
+
+			FMovieSceneObjectBindingID NewBinding = UE::MovieScene::FRelativeObjectBindingID(ThisSequenceID, TargetSequenceID, Guid, OwningTakeRecorderSource->GetMasterLevelSequence());
+			MovieSceneSection->SetConstraintBindingID(NewBinding);
 		}
 	}
 }
diff --git a/Engine/Plugins/VirtualProduction/Takes/Source/TakeTrackRecorders/Public/TrackRecorders/IMovieSceneTrackRecorderHost.h b/Engine/Plugins/VirtualProduction/Takes/Source/TakeTrackRecorders/Public/TrackRecorders/IMovieSceneTrackRecorderHost.h
--- a/Engine/Plugins/VirtualProduction/Takes/Source/TakeTrackRecorders/Public/TrackRecorders/IMovieSceneTrackRecorderHost.h
+++ b/Engine/Plugins/VirtualProduction/Takes/Source/TakeTrackRecorders/Public/TrackRecorders/IMovieSceneTrackRecorderHost.h
@@ -187,4 +187,5 @@ public:
 	virtual ULevelSequence* GetMasterLevelSequence() const = 0;
 
 	/**/
+	virtual FMovieSceneSequenceID GetSequenceID() const = 0;
 };
diff --git a/Engine/Plugins/VirtualProduction/Takes/Source/TakesCore/Private/TakeRecorderSource.cpp b/Engine/Plugins/VirtualProduction/Takes/Source/TakesCore/Private/TakeRecorderSource.cpp
--- a/Engine/Plugins/VirtualProduction/Takes/Source/TakesCore/Private/TakeRecorderSource.cpp
+++ b/Engine/Plugins/VirtualProduction/Takes/Source/TakesCore/Private/TakeRecorderSource.cpp
@@ -2,6 +2,7 @@
 
 #include "TakeRecorderSource.h"
 #include "Styling/SlateIconFinder.h"
+#include "MovieSceneSequenceID.h"
 
 UTakeRecorderSource::UTakeRecorderSource(const FObjectInitializer& ObjInit)
 	: Super(ObjInit)
@@ -14,4 +15,9 @@ UTakeRecorderSource::UTakeRecorderSource(const FObjectInitializer& ObjInit)
 const FSlateBrush* UTakeRecorderSource::GetDisplayIconImpl() const
 {
 	return FSlateIconFinder::FindCustomIconBrushForClass(GetClass(), TEXT("ClassThumbnail"));
+}
+
+TArray<UTakeRecorderSource*> UTakeRecorderSource::PreRecording(ULevelSequence* InSequence, FMovieSceneSequenceID InSequenceID, ULevelSequence* InMasterSequence, FManifestSerializer* InManifestSerializer)
+{
+	return TArray<UTakeRecorderSource*>();
 }
\ No newline at end of file
diff --git a/Engine/Plugins/VirtualProduction/Takes/Source/TakesCore/Private/TakeRecorderSources.cpp b/Engine/Plugins/VirtualProduction/Takes/Source/TakesCore/Private/TakeRecorderSources.cpp
--- a/Engine/Plugins/VirtualProduction/Takes/Source/TakesCore/Private/TakeRecorderSources.cpp
+++ b/Engine/Plugins/VirtualProduction/Takes/Source/TakesCore/Private/TakeRecorderSources.cpp
@@ -94,6 +94,7 @@ void UTakeRecorderSources::StartRecordingRecursive(TArray<UTakeRecorderSource*>
 		if (Source->bEnabled)
 		{
 			ULevelSequence* TargetSequence = InMasterSequence;
+			FMovieSceneSequenceID TargetSequenceID = MovieSceneSequenceID::Root;
 
 			// The Sequencer Take system is built around swapping out sub-sequences. If they want to use this system, we create a sub-sequence
 			// for the Source and tell it to write into this sub-sequence instead of the master sequence. We then keep track of which Source
@@ -157,6 +158,8 @@ void UTakeRecorderSources::StartRecordingRecursive(TArray<UTakeRecorderSource*>
 				NewSubSection->SetRowIndex(SubsceneTrack->GetMaxRowIndex() + 1);
 				SubsceneTrack->FixRowIndices();
 
+				TargetSequenceID = NewSubSection->GetSequenceID();
+
 				ActiveSubSections.Add(NewSubSection);
 				if (InManifestSerializer)
 				{
@@ -198,7 +201,7 @@ void UTakeRecorderSources::StartRecordingRecursive(TArray<UTakeRecorderSource*>
 			// Update our mappings of which sources use which sub-sequence.
 			SourceSubSequenceMap.FindOrAdd(Source) = TargetSequence;
 			Source->TimecodeSource = Timecode;
-			for (UTakeRecorderSource* NewlyAddedSource : Source->PreRecording(TargetSequence, InMasterSequence, InManifestSerializer))
+			for (UTakeRecorderSource* NewlyAddedSource : Source->PreRecording(TargetSequence, TargetSequenceID, InMasterSequence, InManifestSerializer))
 			{
 				// Add it to our classes list of sources 
 				Sources.Add(NewlyAddedSource);
@@ -257,6 +260,7 @@ void UTakeRecorderSources::PreRecordingRecursive(TArray<UTakeRecorderSource*> In
 		if (Source->bEnabled)
 		{
 			ULevelSequence* TargetSequence = InMasterSequence;
+			FMovieSceneSequenceID TargetSequenceID = MovieSceneSequenceID::Root;
 
 			// The Sequencer Take system is built around swapping out sub-sequences. If they want to use this system, we create a sub-sequence
 			// for the Source and tell it to write into this sub-sequence instead of the master sequence. We then keep track of which Source
@@ -318,6 +322,8 @@ void UTakeRecorderSources::PreRecordingRecursive(TArray<UTakeRecorderSource*> In
 				NewSubSection->SetRowIndex(SubsceneTrack->GetMaxRowIndex() + 1);
 				SubsceneTrack->FixRowIndices();
 
+				TargetSequenceID = NewSubSection->GetSequenceID();
+
 				ActiveSubSections.Add(NewSubSection);
 				if (InManifestSerializer)
 				{
@@ -359,7 +365,7 @@ void UTakeRecorderSources::PreRecordingRecursive(TArray<UTakeRecorderSource*> In
 			// Update our mappings of which sources use which sub-sequence.
 			SourceSubSequenceMap.FindOrAdd(Source) = TargetSequence;
 
-			for (UTakeRecorderSource* NewlyAddedSource : Source->PreRecording(TargetSequence, InMasterSequence, InManifestSerializer))
+			for (UTakeRecorderSource* NewlyAddedSource : Source->PreRecording(TargetSequence, TargetSequenceID, InMasterSequence, InManifestSerializer))
 			{
 				// Add it to our classes list of sources 
 				Sources.Add(NewlyAddedSource);
diff --git a/Engine/Plugins/VirtualProduction/Takes/Source/TakesCore/Private/TakesUtils.cpp b/Engine/Plugins/VirtualProduction/Takes/Source/TakesCore/Private/TakesUtils.cpp
--- a/Engine/Plugins/VirtualProduction/Takes/Source/TakesCore/Private/TakesUtils.cpp
+++ b/Engine/Plugins/VirtualProduction/Takes/Source/TakesCore/Private/TakesUtils.cpp
@@ -125,7 +125,7 @@ void CreateCameraCutTrack(ULevelSequence* LevelSequence, const FGuid& RecordedCa
 	}
 
 	UMovieSceneCameraCutSection* CameraCutSection = Cast<UMovieSceneCameraCutSection>(CameraCutTrack->CreateNewSection());
-	CameraCutSection->SetCameraBindingID(FMovieSceneObjectBindingID(RecordedCameraGuid, SequenceID, EMovieSceneObjectBindingSpace::Local));
+	CameraCutSection->SetCameraBindingID(UE::MovieScene::FRelativeObjectBindingID(RecordedCameraGuid));
 	CameraCutSection->SetRange(InRange);
 	CameraCutTrack->AddSection(*CameraCutSection);
 }
diff --git a/Engine/Plugins/VirtualProduction/Takes/Source/TakesCore/Public/TakeRecorderSource.h b/Engine/Plugins/VirtualProduction/Takes/Source/TakesCore/Public/TakeRecorderSource.h
--- a/Engine/Plugins/VirtualProduction/Takes/Source/TakesCore/Public/TakeRecorderSource.h
+++ b/Engine/Plugins/VirtualProduction/Takes/Source/TakesCore/Public/TakeRecorderSource.h
@@ -9,6 +9,7 @@
 #include "Misc/QualifiedFrameTime.h"
 #include "TakeRecorderSource.generated.h"
 
+struct FMovieSceneSequenceID;
 class ULevelSequence;
 
 /**
@@ -64,12 +65,13 @@ public:
 	* Will not be called if this recording source is not enabled.
 	*
 	* @param InSequence - The Level Sequence the take is being recorded into
+	* @param InSequenceID - The sequence ID for InSequence, or MovieSceneSequenceID::Root where InSequence == InMasterSequence (ie, not sub-sequence recording)
 	* @param InMasterSequence - The Master Level Sequence that may contain the InSequence as a child or if no subsequences, is the same as InSequence.
 	* @param InManifestSerializer - Manifest Serializer that we may write into.
 	* @return An array of newly created take recorder sources. Can be an empty list if no additional sources needed to be
 	* created by this source.
 	*/
-	virtual TArray<UTakeRecorderSource*> PreRecording(class ULevelSequence* InSequence, class ULevelSequence* InMasterSequence, FManifestSerializer* InManifestSerializer) { return TArray<UTakeRecorderSource*>(); }
+	virtual TArray<UTakeRecorderSource*> PreRecording(ULevelSequence* InSequence, FMovieSceneSequenceID InSequenceID, ULevelSequence* InMasterSequence, FManifestSerializer* InManifestSerializer);
 
 	/**
 	* This is called when the UTakeRecorderSources starts a recording, after all sources have had PreRecording called on them.
diff --git a/Engine/Source/Editor/MovieSceneTools/Private/Channels/BuiltInChannelEditors.cpp b/Engine/Source/Editor/MovieSceneTools/Private/Channels/BuiltInChannelEditors.cpp
--- a/Engine/Source/Editor/MovieSceneTools/Private/Channels/BuiltInChannelEditors.cpp
+++ b/Engine/Source/Editor/MovieSceneTools/Private/Channels/BuiltInChannelEditors.cpp
@@ -37,6 +37,7 @@
 #include "EntitySystem/Interrogation/MovieSceneInterrogationLinker.h"
 #include "EntitySystem/Interrogation/MovieSceneInterrogatedPropertyInstantiator.h"
 #include "Systems/MovieScenePropertyInstantiator.h"
+#include "MovieSceneSpawnableAnnotation.h"
 #include "ISequencerModule.h"
 #include "MovieSceneTracksComponentTypes.h"
 
@@ -182,11 +183,21 @@ FKeyHandle AddOrUpdateKey(FMovieSceneActorReferenceData* Channel, UMovieSceneSec
 				AActor* CurrentActor = Cast<AActor>(PropertyBindings->GetCurrentValue<UObject*>(*Object));
 				if (CurrentActor)
 				{
-					FGuid ThisGuid = Sequencer.FindObjectId(*CurrentActor, Sequencer.GetFocusedTemplateID());
+					FMovieSceneObjectBindingID Binding;
 
-					FMovieSceneObjectBindingID NewValue(ThisGuid, MovieSceneSequenceID::Root, EMovieSceneObjectBindingSpace::Local);
+					TOptional<FMovieSceneSpawnableAnnotation> Spawnable = FMovieSceneSpawnableAnnotation::Find(CurrentActor);
+					if (Spawnable.IsSet())
+					{
+						// Check whether the spawnable is underneath the current sequence, if so, we can remap it to a local sequence ID
+						Binding = UE::MovieScene::FRelativeObjectBindingID(Sequencer.GetFocusedTemplateID(), Spawnable->SequenceID, Spawnable->ObjectBindingID, Sequencer);
+					}
+					else
+					{
+						FGuid ThisGuid = Sequencer.GetHandleToObject(CurrentActor);
+						Binding = UE::MovieScene::FRelativeObjectBindingID(ThisGuid);
+					}
 
-					int32 NewIndex = Channel->GetData().AddKey(InTime, NewValue);
+					int32 NewIndex = Channel->GetData().AddKey(InTime, Binding);
 
 					return Channel->GetData().GetHandle(NewIndex);
 				}
@@ -383,18 +394,30 @@ public:
 
 		ChildSlot
 		[
-			SNew(SComboButton)
-			.OnGetMenuContent(this, &SActorReferenceBox::GetPickerMenu)
-			.ContentPadding(FMargin(0.0, 0.0))
-			.ButtonStyle(FEditorStyle::Get(), "PropertyEditor.AssetComboStyle")
-			.ForegroundColor(FEditorStyle::GetColor("PropertyEditor.AssetName.ColorAndOpacity"))
-			.ButtonContent()
+			SNew(SHorizontalBox)
+
+			+ SHorizontalBox::Slot()
+			[
+				SNew(SComboButton)
+				.OnGetMenuContent(this, &SActorReferenceBox::GetPickerMenu)
+				.ContentPadding(FMargin(0.0, 0.0))
+				.ButtonStyle(FEditorStyle::Get(), "PropertyEditor.AssetComboStyle")
+				.ForegroundColor(FEditorStyle::GetColor("PropertyEditor.AssetName.ColorAndOpacity"))
+				.ButtonContent()
+				[
+					GetCurrentItemWidget(
+						SNew(STextBlock)
+						.TextStyle(FEditorStyle::Get(), "PropertyEditor.AssetClass")
+						.Font(FCoreStyle::GetDefaultFontStyle("Regular", 8))
+					)
+				]
+			]
+
+			+ SHorizontalBox::Slot()
+			.AutoWidth()
+			.Padding(FMargin(4.f, 0.f, 0.f, 0.f))
 			[
-				GetCurrentItemWidget(
-					SNew(STextBlock)
-					.TextStyle(FEditorStyle::Get(), "PropertyEditor.AssetClass")
-					.Font(FCoreStyle::GetDefaultFontStyle("Regular", 8))
-				)
+				GetWarningWidget()
 			]
 		];
 
@@ -478,10 +501,9 @@ TSharedRef<SWidget> CreateKeyEditor(const TMovieSceneChannelHandle<FMovieSceneAc
 
 		// Look for components to choose
 		ISequencer* Sequencer = KeyEditor.GetSequencer();
-		FMovieSceneEvaluationOperand ObjectOperand(ActorKey.Object.GetSequenceID(), ActorKey.Object.GetGuid());
 		TArray<USceneComponent*> ComponentsWithSockets;
 		AActor* Actor = nullptr;
-		for (TWeakObjectPtr<> WeakObject : Sequencer->FindBoundObjects(ObjectOperand))
+		for (TWeakObjectPtr<> WeakObject : ActorKey.Object.ResolveBoundObjects(MovieSceneSequenceID::Root, *Sequencer))
 		{
 			Actor = Cast<AActor>(WeakObject.Get());
 			if (Actor)
diff --git a/Engine/Source/Editor/MovieSceneTools/Private/K2Node_GetSequenceBinding.cpp b/Engine/Source/Editor/MovieSceneTools/Private/K2Node_GetSequenceBinding.cpp
--- a/Engine/Source/Editor/MovieSceneTools/Private/K2Node_GetSequenceBinding.cpp
+++ b/Engine/Source/Editor/MovieSceneTools/Private/K2Node_GetSequenceBinding.cpp
@@ -136,7 +136,7 @@ UMovieScene* UK2Node_GetSequenceBinding::GetObjectMovieScene() const
 		// Ensure that the sequence data is as loaded as it can be - we many only be able to partially load the structural information as part of a blueprint compile as that may happen at Preload time
 		EnsureFullyLoaded(Sequence);
 
-		FMovieSceneSequenceID SequenceID = Binding.GetSequenceID();
+		FMovieSceneSequenceID SequenceID = Binding.GetRelativeSequenceID();
 		if (SequenceID == MovieSceneSequenceID::Root)
 		{
 			// Look it up in the moviescene itself
diff --git a/Engine/Source/Editor/MovieSceneTools/Private/MatineeImportTools.cpp b/Engine/Source/Editor/MovieSceneTools/Private/MatineeImportTools.cpp
--- a/Engine/Source/Editor/MovieSceneTools/Private/MatineeImportTools.cpp
+++ b/Engine/Source/Editor/MovieSceneTools/Private/MatineeImportTools.cpp
@@ -860,8 +860,7 @@ bool FMatineeImportTools::CopyInterpDirectorTrack( UInterpTrackDirector* Directo
 				FGuid CameraHandle = Player.FindObjectId(*CameraActor, MovieSceneSequenceID::Root);
 				if (CameraHandle.IsValid())
 				{
-					FMovieSceneObjectBindingID CameraBindingID(CameraHandle, MovieSceneSequenceID::Root);
-					CameraCutTrack->AddNewCameraCut(CameraBindingID, (TrackCut.Time * FrameRate).RoundToFrame());
+					CameraCutTrack->AddNewCameraCut(UE::MovieScene::FRelativeObjectBindingID(CameraHandle), (TrackCut.Time * FrameRate).RoundToFrame());
 					bCutsAdded = true;
 				}
 			}
diff --git a/Engine/Source/Editor/MovieSceneTools/Private/MovieSceneObjectBindingIDCustomization.cpp b/Engine/Source/Editor/MovieSceneTools/Private/MovieSceneObjectBindingIDCustomization.cpp
--- a/Engine/Source/Editor/MovieSceneTools/Private/MovieSceneObjectBindingIDCustomization.cpp
+++ b/Engine/Source/Editor/MovieSceneTools/Private/MovieSceneObjectBindingIDCustomization.cpp
@@ -56,23 +56,35 @@ void FMovieSceneObjectBindingIDCustomization::CustomizeHeader(TSharedRef<IProper
 	]
 	.ValueContent()
 	[
-		SNew(SDropTarget)
-		.OnDrop(this, &FMovieSceneObjectBindingIDCustomization::OnDrop)
-		.OnAllowDrop_Static(IsAcceptable)
-		.OnIsRecognized_Static(IsAcceptable)
+		SNew(SHorizontalBox)
+
+		+ SHorizontalBox::Slot()
 		[
-			SNew(SComboButton)
-			.ToolTipText(this, &FMovieSceneObjectBindingIDCustomization::GetToolTipText)
-			.OnGetMenuContent(this, &FMovieSceneObjectBindingIDCustomization::GetPickerMenu)
-			.ContentPadding(FMargin(4.0, 2.0))
-			.ButtonContent()
+			SNew(SDropTarget)
+			.OnDrop(this, &FMovieSceneObjectBindingIDCustomization::OnDrop)
+			.OnAllowDrop_Static(IsAcceptable)
+			.OnIsRecognized_Static(IsAcceptable)
 			[
-				GetCurrentItemWidget(
-					SNew(STextBlock)
-					.Font(CustomizationUtils.GetRegularFont())
-				)
+				SNew(SComboButton)
+				.ToolTipText(this, &FMovieSceneObjectBindingIDCustomization::GetToolTipText)
+				.OnGetMenuContent(this, &FMovieSceneObjectBindingIDCustomization::GetPickerMenu)
+				.ContentPadding(FMargin(4.0, 2.0))
+				.ButtonContent()
+				[
+					GetCurrentItemWidget(
+						SNew(STextBlock)
+						.Font(CustomizationUtils.GetRegularFont())
+					)
+				]
 			]
 		]
+
+		+ SHorizontalBox::Slot()
+		.AutoWidth()
+		.Padding(FMargin(4.f, 0.f, 0.f, 0.f))
+		[
+			GetWarningWidget()
+		]
 	];
 }
 
@@ -81,7 +93,7 @@ FReply FMovieSceneObjectBindingIDCustomization::OnDrop(TSharedPtr<FDragDropOpera
 	FSequencerObjectBindingDragDropOp* SequencerOp = InOperation->IsOfType<FSequencerObjectBindingDragDropOp>() ? static_cast<FSequencerObjectBindingDragDropOp*>(InOperation.Get()) : nullptr;
 	if (SequencerOp)
 	{
-		TArray<FMovieSceneObjectBindingID> Bindings = SequencerOp->GetDraggedBindings();
+		TArray<UE::MovieScene::FFixedObjectBindingID> Bindings = SequencerOp->GetDraggedBindings();
 		if (Bindings.Num() == 1)
 		{
 			SetBindingId(Bindings[0]);
diff --git a/Engine/Source/Editor/MovieSceneTools/Private/MovieSceneObjectBindingIDPicker.cpp b/Engine/Source/Editor/MovieSceneTools/Private/MovieSceneObjectBindingIDPicker.cpp
--- a/Engine/Source/Editor/MovieSceneTools/Private/MovieSceneObjectBindingIDPicker.cpp
+++ b/Engine/Source/Editor/MovieSceneTools/Private/MovieSceneObjectBindingIDPicker.cpp
@@ -12,12 +12,15 @@
 #include "Textures/SlateIcon.h"
 #include "EditorStyleSet.h"
 #include "Widgets/Images/SImage.h"
+#include "EditorFontGlyphs.h"
+#include "Widgets/Input/SButton.h"
 #include "Widgets/SOverlay.h"
 #include "ISequencer.h"
 #include "Evaluation/MovieSceneEvaluationTemplateInstance.h"
 #include "Evaluation/MovieSceneSequenceHierarchy.h"
 #include "Compilation/MovieSceneCompiledDataManager.h"
 #include "Framework/Application/SlateApplication.h"
+#include "Styling/StarshipCoreStyle.h"
 
 #define LOCTEXT_NAMESPACE "MovieSceneObjectBindingIDPicker"
 
@@ -50,7 +53,7 @@ void FMovieSceneObjectBindingIDPicker::OnGetMenuContent(FMenuBuilder& MenuBuilde
 
 	bool bHadAnyEntries = false;
 
-	if (Node->BindingID.GetGuid().IsValid())
+	if (Node->BindingID.Guid.IsValid())
 	{
 		bHadAnyEntries = true;
 		MenuBuilder.AddMenuEntry(
@@ -67,7 +70,7 @@ void FMovieSceneObjectBindingIDPicker::OnGetMenuContent(FMenuBuilder& MenuBuilde
 	{
 		check(Child.IsValid())
 
-		if (!Child->BindingID.GetGuid().IsValid())
+		if (!Child->BindingID.Guid.IsValid())
 		{
 			if (Child->Children.Num())
 			{
@@ -156,9 +159,39 @@ TSharedRef<SWidget> FMovieSceneObjectBindingIDPicker::GetCurrentItemWidget(TShar
 		];
 }
 
-void FMovieSceneObjectBindingIDPicker::SetBindingId(FMovieSceneObjectBindingID InBindingId)
+TSharedRef<SWidget> FMovieSceneObjectBindingIDPicker::GetWarningWidget()
 {
-	SetRemappedCurrentValue(InBindingId);
+	return SNew(SButton)
+		.HAlign(HAlign_Center)
+		.VAlign(VAlign_Center)
+		.ContentPadding(FMargin(0))
+		.ButtonStyle(FEditorStyle::Get(), "HoverHintOnly")
+		.ToolTipText(LOCTEXT("FixedBindingWarningText", "This binding is fixed to the current Master Sequence hierarchy, so will break if evaluated in a different hierarchy.\nClick here to fix this problem."))
+		.Visibility_Raw(this, &FMovieSceneObjectBindingIDPicker::GetFixedWarningVisibility)
+		.OnClicked_Raw(this, &FMovieSceneObjectBindingIDPicker::AttemptBindingFixup)
+		[
+			SNew(SImage)
+			.Image(FStarshipCoreStyle::GetCoreStyle().GetBrush("Icons.Warning"))
+		];
+}
+
+EVisibility FMovieSceneObjectBindingIDPicker::GetFixedWarningVisibility() const
+{
+	FMovieSceneObjectBindingID CurrentValue = GetCurrentValue();
+
+	const bool bShowError = CurrentValue.IsFixedBinding() && LocalSequenceID != MovieSceneSequenceID::Invalid;
+	return bShowError ? EVisibility::Visible : EVisibility::Collapsed;
+}
+
+FReply FMovieSceneObjectBindingIDPicker::AttemptBindingFixup()
+{
+	SetCurrentValueFromFixed(GetCurrentValueAsFixed());
+	return FReply::Handled();
+}
+
+void FMovieSceneObjectBindingIDPicker::SetBindingId(UE::MovieScene::FFixedObjectBindingID InBindingId)
+{
+	SetCurrentValueFromFixed(InBindingId);
 	UpdateCachedData();
 
 	TSharedPtr<SWidget> MenuWidget = DismissWidget.Pin();
@@ -170,8 +203,11 @@ void FMovieSceneObjectBindingIDPicker::SetBindingId(FMovieSceneObjectBindingID I
 
 void FMovieSceneObjectBindingIDPicker::UpdateCachedData()
 {
-	FMovieSceneObjectBindingID CurrentValue = GetRemappedCurrentValue();
-	TSharedPtr<FSequenceBindingNode> Object = CurrentValue.IsValid() ? DataTree->FindNode(CurrentValue) : nullptr;
+	using namespace UE::MovieScene;
+
+	FFixedObjectBindingID CurrentValue = GetCurrentValueAsFixed();
+
+	TSharedPtr<FSequenceBindingNode> Object = CurrentValue.Guid.IsValid() ? DataTree->FindNode(CurrentValue) : nullptr;
 
 	if (!Object.IsValid())
 	{
@@ -187,7 +223,7 @@ void FMovieSceneObjectBindingIDPicker::UpdateCachedData()
 		bIsCurrentItemSpawnable = Object->bIsSpawnable;
 
 		ToolTipText = FText();
-		while (Object.IsValid() && Object->BindingID != FMovieSceneObjectBindingID())
+		while (Object.IsValid() && Object->BindingID.SequenceID != MovieSceneSequenceID::Invalid)
 		{
 			ToolTipText = ToolTipText.IsEmpty() ? Object->DisplayString : FText::Format(LOCTEXT("ToolTipFormat", "{0} -> {1}"), Object->DisplayString, ToolTipText);
 			Object = DataTree->FindNode(Object->ParentID);
@@ -220,63 +256,37 @@ EVisibility FMovieSceneObjectBindingIDPicker::GetSpawnableIconOverlayVisibility(
 	return bIsCurrentItemSpawnable ? EVisibility::Visible : EVisibility::Collapsed;
 }
 
-FMovieSceneObjectBindingID FMovieSceneObjectBindingIDPicker::GetRemappedCurrentValue() const
+UE::MovieScene::FFixedObjectBindingID FMovieSceneObjectBindingIDPicker::GetCurrentValueAsFixed() const
 {
 	FMovieSceneObjectBindingID ID = GetCurrentValue();
 
-	TSharedPtr<ISequencer> Sequencer = WeakSequencer.Pin();
-
 	// If the ID is in local space, remap it to the root space as according to the LocalSequenceID we were created with
-	if (Sequencer.IsValid() && LocalSequenceID != MovieSceneSequenceID::Root && ID.IsValid() && ID.GetBindingSpace() == EMovieSceneObjectBindingSpace::Local)
+	if (TSharedPtr<ISequencer> Sequencer = WeakSequencer.Pin())
 	{
-		ID = ID.ResolveLocalToRoot(LocalSequenceID, *Sequencer);
+		ID = ID.ResolveToFixed(LocalSequenceID, *Sequencer);
 	}
 
-	return ID;
+	return ID.ReinterpretAsFixed();
 }
 
-void FMovieSceneObjectBindingIDPicker::SetRemappedCurrentValue(FMovieSceneObjectBindingID InValue)
+void FMovieSceneObjectBindingIDPicker::SetCurrentValueFromFixed(UE::MovieScene::FFixedObjectBindingID InValue)
 {
-	TSharedPtr<ISequencer> Sequencer = WeakSequencer.Pin();
+	TSharedPtr<ISequencer>              Sequencer = WeakSequencer.Pin();
+	const FMovieSceneSequenceHierarchy* Hierarchy = Sequencer.IsValid() ? Sequencer->GetEvaluationTemplate().GetHierarchy() : nullptr;
 
-	// If we have a local sequence ID set, and the supplied binding is in root space, we attempt to remap it into the local sequence ID's space, and use a sequence ID
-	// that will resolve from LocalSequenceID instead of from the root. This ensures that you can work on sub sequences on their own, or within a master sequence
-	// and the binding will resolve correctly.
-	if (LocalSequenceID.IsValid() && Sequencer.IsValid() && InValue.GetGuid().IsValid() && InValue.GetBindingSpace() == EMovieSceneObjectBindingSpace::Root)
+	// If we don't know the local sequence ID, or we have no hierarchy, or we're resetting the binding; just set the ID directly
+	if (LocalSequenceID == MovieSceneSequenceID::Invalid || !InValue.Guid.IsValid())
 	{
-		FMovieSceneRootEvaluationTemplateInstance& Instance = Sequencer->GetEvaluationTemplate();
-
-		const FMovieSceneSequenceHierarchy* Hierarchy = Instance.GetCompiledDataManager()->FindHierarchy(Instance.GetCompiledDataID());
-		if (Hierarchy)
-		{
-			FMovieSceneSequenceID NewLocalSequenceID = MovieSceneSequenceID::Root;
-			FMovieSceneSequenceID CurrentSequenceID = InValue.GetSequenceID();
-
-			while (CurrentSequenceID.IsValid())
-			{
-				if (LocalSequenceID == CurrentSequenceID)
-				{
-					// Found it
-					InValue = FMovieSceneObjectBindingID(InValue.GetGuid(), NewLocalSequenceID, EMovieSceneObjectBindingSpace::Local);
-					break;
-				}
-
-				const FMovieSceneSequenceHierarchyNode* CurrentNode = Hierarchy->FindNode(CurrentSequenceID);
-				if (!ensureAlwaysMsgf(CurrentNode, TEXT("Malformed sequence hierarchy")))
-				{
-					break;
-				}
-				else if (const FMovieSceneSubSequenceData* SubData = Hierarchy->FindSubData(CurrentSequenceID))
-				{
-					NewLocalSequenceID = NewLocalSequenceID.AccumulateParentID(SubData->DeterministicSequenceID);
-				}
-				
-				CurrentSequenceID = CurrentNode->ParentID;
-			}
-		}
+		SetCurrentValue(InValue);
+	}
+	else
+	{
+		// Attempt to remap the desired binding to the current local sequence by either making it local to this sequence
+		// or specifying a parent index so that this binding is still able to resolve correctly if the master sequence is added
+		// as a subsequence elsewhere
+		// This ensures that you can work on sub sequences on their own, or within a master sequence and the binding will resolve correctly.
+		SetCurrentValue(InValue.ConvertToRelative(LocalSequenceID, Hierarchy));
 	}
-
-	SetCurrentValue(InValue);
 }
 
 #undef LOCTEXT_NAMESPACE
diff --git a/Engine/Source/Editor/MovieSceneTools/Private/MovieSceneToolHelpers.cpp b/Engine/Source/Editor/MovieSceneTools/Private/MovieSceneToolHelpers.cpp
--- a/Engine/Source/Editor/MovieSceneTools/Private/MovieSceneToolHelpers.cpp
+++ b/Engine/Source/Editor/MovieSceneTools/Private/MovieSceneToolHelpers.cpp
@@ -2858,7 +2858,7 @@ void ImportCameraCut(UnFbx::FFbxImporter* FbxImporter, UMovieScene* InMovieScene
 				FGuid CameraGuid = FindCameraGuid(AllCameras[value], InObjectBindingMap);
 				if (CameraGuid != FGuid())
 				{
-					CameraCutTrack->AddNewCameraCut(FMovieSceneObjectBindingID(CameraGuid, MovieSceneSequenceID::Root), (key.GetTime().GetSecondDouble() * FrameRate).RoundToFrame());
+					CameraCutTrack->AddNewCameraCut(UE::MovieScene::FRelativeObjectBindingID(CameraGuid), (key.GetTime().GetSecondDouble() * FrameRate).RoundToFrame());
 				}
 			}
 		}
diff --git a/Engine/Source/Editor/MovieSceneTools/Private/ObjectBindingTagCache.cpp b/Engine/Source/Editor/MovieSceneTools/Private/ObjectBindingTagCache.cpp
--- a/Engine/Source/Editor/MovieSceneTools/Private/ObjectBindingTagCache.cpp
+++ b/Engine/Source/Editor/MovieSceneTools/Private/ObjectBindingTagCache.cpp
@@ -70,7 +70,7 @@ void FObjectBindingTagCache::ConditionalUpdate(UMovieSceneSequence* RootSequence
 
 		for (FMovieSceneObjectBindingID ID : Pair.Value.IDs)
 		{
-			ExposedNameReverseLUT.Add(ID, Pair.Key);
+			ExposedNameReverseLUT.Add(ID.ReinterpretAsFixed(), Pair.Key);
 		}
 	}
 
diff --git a/Engine/Source/Editor/MovieSceneTools/Private/SObjectBindingTag.cpp b/Engine/Source/Editor/MovieSceneTools/Private/SObjectBindingTag.cpp
--- a/Engine/Source/Editor/MovieSceneTools/Private/SObjectBindingTag.cpp
+++ b/Engine/Source/Editor/MovieSceneTools/Private/SObjectBindingTag.cpp
@@ -20,7 +20,7 @@
 
 #define LOCTEXT_NAMESPACE "SObjectBindingTag"
 
-void SObjectBindingTags::Construct(const FArguments& InArgs, const FMovieSceneObjectBindingID& InBindingID, FObjectBindingTagCache* BindingCache)
+void SObjectBindingTags::Construct(const FArguments& InArgs, const UE::MovieScene::FFixedObjectBindingID& InBindingID, FObjectBindingTagCache* BindingCache)
 {
 	BindingID = InBindingID;
 	OnTagDeletedEvent = InArgs._OnTagDeleted;
diff --git a/Engine/Source/Editor/MovieSceneTools/Private/SequenceBindingTree.cpp b/Engine/Source/Editor/MovieSceneTools/Private/SequenceBindingTree.cpp
--- a/Engine/Source/Editor/MovieSceneTools/Private/SequenceBindingTree.cpp
+++ b/Engine/Source/Editor/MovieSceneTools/Private/SequenceBindingTree.cpp
@@ -84,6 +84,8 @@ bool FSequenceBindingTree::ConditionalRebuild(UMovieSceneSequence* InSequence, F
 
 void FSequenceBindingTree::ForceRebuild(UMovieSceneSequence* InSequence, FObjectKey InActiveSequence, FMovieSceneSequenceID InActiveSequenceID)
 {
+	using namespace UE::MovieScene;
+
 	CachedSequenceSignatures.Empty();
 	bIsEmpty = true;
 
@@ -95,7 +97,7 @@ void FSequenceBindingTree::ForceRebuild(UMovieSceneSequence* InSequence, FObject
 	ActiveSequenceNode = nullptr;
 
 	// Create a node for the root sequence
-	FMovieSceneObjectBindingID RootSequenceID;
+	FFixedObjectBindingID RootSequenceID = FFixedObjectBindingID(FGuid(), MovieSceneSequenceID::Root);
 	TSharedRef<FSequenceBindingNode> RootSequenceNode = MakeShared<FSequenceBindingNode>(FText(), RootSequenceID, FSlateIcon());
 	Hierarchy.Add(RootSequenceID, RootSequenceNode);
 
@@ -121,7 +123,7 @@ void FSequenceBindingTree::ForceRebuild(UMovieSceneSequence* InSequence, FObject
 			ActiveParent->Children.Remove(ActiveSequenceNode.ToSharedRef());
 
 			// Make a new top level node (with an invalid ID)
-			FMovieSceneObjectBindingID TopLevelID = FMovieSceneObjectBindingID(FGuid(), FMovieSceneSequenceID());
+			FFixedObjectBindingID TopLevelID = FFixedObjectBindingID(FGuid(), MovieSceneSequenceID::Invalid);
 			TopLevelNode = MakeShared<FSequenceBindingNode>(FText(), TopLevelID, FSlateIcon());
 
 			// Override the display string and icon
@@ -140,9 +142,9 @@ void FSequenceBindingTree::Sort(TSharedRef<FSequenceBindingNode> Node)
 		[](TSharedRef<FSequenceBindingNode> A, TSharedRef<FSequenceBindingNode> B)
 		{
 			// Sort shots first
-			if (A->BindingID.IsValid() != B->BindingID.IsValid())
+			if (A->BindingID.Guid.IsValid() != B->BindingID.Guid.IsValid())
 			{
-				return !A->BindingID.IsValid();
+				return !A->BindingID.Guid.IsValid();
 			}
 			return A->DisplayString.CompareToCaseIgnored(B->DisplayString) < 0;
 		}
@@ -156,6 +158,8 @@ void FSequenceBindingTree::Sort(TSharedRef<FSequenceBindingNode> Node)
 
 void FSequenceBindingTree::Build(UMovieSceneSequence* InSequence, FSequenceIDStack& SequenceIDStack)
 {
+	using namespace UE::MovieScene;
+
 	check(InSequence);
 
 	UMovieScene* MovieScene = InSequence->GetMovieScene();
@@ -173,7 +177,7 @@ void FSequenceBindingTree::Build(UMovieSceneSequence* InSequence, FSequenceIDSta
 		}
 
 		// Keep track of the active sequence node
-		ActiveSequenceNode = Hierarchy.FindChecked(FMovieSceneObjectBindingID(FGuid(), SequenceIDStack.GetCurrent()));
+		ActiveSequenceNode = Hierarchy.FindChecked(FFixedObjectBindingID(FGuid(), SequenceIDStack.GetCurrent()));
 	}
 
 	// Iterate all sub sections
@@ -192,7 +196,7 @@ void FSequenceBindingTree::Build(UMovieSceneSequence* InSequence, FSequenceIDSta
 					FMovieSceneSequenceID ParentID = SequenceIDStack.GetCurrent();
 					SequenceIDStack.Push(SubSection->GetSequenceID());
 					
-					FMovieSceneObjectBindingID CurrentID(FGuid(), SequenceIDStack.GetCurrent());
+					FFixedObjectBindingID CurrentID = FFixedObjectBindingID(FGuid(), SequenceIDStack.GetCurrent());
 
 					UMovieSceneCinematicShotSection* ShotSection = Cast<UMovieSceneCinematicShotSection>(Section);
 					FText DisplayString = ShotSection ? FText::FromString(ShotSection->GetShotDisplayName()) : FText::FromName(SubSequence->GetFName());
@@ -222,7 +226,7 @@ void FSequenceBindingTree::Build(UMovieSceneSequence* InSequence, FSequenceIDSta
 	{
 		const FMovieSceneSpawnable& Spawnable = MovieScene->GetSpawnable(Index);
 		
-		FMovieSceneObjectBindingID ID(Spawnable.GetGuid(), CurrentSequenceID);
+		FFixedObjectBindingID ID = FFixedObjectBindingID(Spawnable.GetGuid(), CurrentSequenceID);
 
 		FSlateIcon Icon = FSlateIconFinder::FindIconForClass(Spawnable.GetObjectTemplate()->GetClass());
 		TSharedRef<FSequenceBindingNode> NewNode = MakeShared<FSequenceBindingNode>(MovieScene->GetObjectDisplayName(Spawnable.GetGuid()), ID, Icon);
@@ -242,7 +246,7 @@ void FSequenceBindingTree::Build(UMovieSceneSequence* InSequence, FSequenceIDSta
 		const FMovieScenePossessable& Possessable = MovieScene->GetPossessable(Index);
 		if (InSequence->CanRebindPossessable(Possessable))
 		{
-			FMovieSceneObjectBindingID ID(Possessable.GetGuid(), CurrentSequenceID);
+			FFixedObjectBindingID ID = FFixedObjectBindingID(Possessable.GetGuid(), CurrentSequenceID);
 
 			FSlateIcon Icon = FSlateIconFinder::FindIconForClass(Possessable.GetPossessedObjectClass());
 			TSharedRef<FSequenceBindingNode> NewNode = MakeShared<FSequenceBindingNode>(MovieScene->GetObjectDisplayName(Possessable.GetGuid()), ID, Icon);
@@ -258,7 +262,9 @@ void FSequenceBindingTree::Build(UMovieSceneSequence* InSequence, FSequenceIDSta
 
 TSharedRef<FSequenceBindingNode> FSequenceBindingTree::EnsureParent(const FGuid& InParentGuid, UMovieScene* InMovieScene, FMovieSceneSequenceID SequenceID)
 {
-	FMovieSceneObjectBindingID ParentPtr(InParentGuid, SequenceID);
+	using namespace UE::MovieScene;
+
+	FFixedObjectBindingID ParentPtr = FFixedObjectBindingID(InParentGuid, SequenceID);
 
 	// If the node already exists
 	TSharedPtr<FSequenceBindingNode> Parent = Hierarchy.FindRef(ParentPtr);
diff --git a/Engine/Source/Editor/MovieSceneTools/Private/TrackEditors/ActorPickerTrackEditor.cpp b/Engine/Source/Editor/MovieSceneTools/Private/TrackEditors/ActorPickerTrackEditor.cpp
--- a/Engine/Source/Editor/MovieSceneTools/Private/TrackEditors/ActorPickerTrackEditor.cpp
+++ b/Engine/Source/Editor/MovieSceneTools/Private/TrackEditors/ActorPickerTrackEditor.cpp
@@ -161,7 +161,7 @@ void FActorPickerTrackEditor::ShowActorSubMenu(FMenuBuilder& MenuBuilder, TArray
 
 void FActorPickerTrackEditor::ActorPicked(AActor* ParentActor, TArray<FGuid> ObjectGuids, UMovieSceneSection* Section)
 {
-	ActorPickerIDPicked(FActorPickerID(ParentActor, FMovieSceneObjectBindingID()), ObjectGuids, Section);
+	ActorPickerIDPicked(FActorPickerID(ParentActor, UE::MovieScene::FFixedObjectBindingID()), ObjectGuids, Section);
 }
 
 
@@ -174,21 +174,15 @@ void FActorPickerTrackEditor::ExistingBindingPicked(FMovieSceneObjectBindingID E
 	if (ExistingBindingID.IsValid())
 	{
 		// Ensure that this ID is resolvable from the root, based on the current local sequence ID
-		FMovieSceneObjectBindingID RootBindingID = ExistingBindingID.ResolveLocalToRoot(SequenceID, *SequencerPtr);
-		SequenceID = RootBindingID.GetSequenceID();
+		SequenceID = ExistingBindingID.ResolveSequenceID(SequenceID, *SequencerPtr);
 	}
 
-	TArrayView<TWeakObjectPtr<UObject>> RuntimeObjects = SequencerPtr->FindBoundObjects(ExistingBindingID.GetGuid(), SequenceID);
-	for (auto RuntimeObject : RuntimeObjects)
+	for (TWeakObjectPtr<> RuntimeObject : SequencerPtr->FindBoundObjects(ExistingBindingID.GetGuid(), SequenceID))
 	{
-		if (RuntimeObject.IsValid())
+		if (AActor* Actor = Cast<AActor>(RuntimeObject.Get()))
 		{
-			AActor* Actor = Cast<AActor>(RuntimeObject.Get());
-			if (Actor)
-			{
-				ActorPickerIDPicked(FActorPickerID(Actor, ExistingBindingID), ObjectBindings, Section);
-				return;
-			}
+			ActorPickerIDPicked(FActorPickerID(Actor, ExistingBindingID), ObjectBindings, Section);
+			return;
 		}
 	}
 
diff --git a/Engine/Source/Editor/MovieSceneTools/Private/TrackEditors/ActorPickerTrackEditor.h b/Engine/Source/Editor/MovieSceneTools/Private/TrackEditors/ActorPickerTrackEditor.h
--- a/Engine/Source/Editor/MovieSceneTools/Private/TrackEditors/ActorPickerTrackEditor.h
+++ b/Engine/Source/Editor/MovieSceneTools/Private/TrackEditors/ActorPickerTrackEditor.h
@@ -11,6 +11,7 @@ class AActor;
 class FMenuBuilder;
 class USceneComponent;
 class FTrackEditorBindingIDPicker;
+
 struct FMovieSceneObjectBindingID;
 
 /**
@@ -22,7 +23,7 @@ public:
 
 	struct FActorPickerID
 	{
-		FActorPickerID(AActor* InActorPicked, FMovieSceneObjectBindingID InExistingBindingID) : ActorPicked(InActorPicked), ExistingBindingID(InExistingBindingID) {}
+		FActorPickerID(AActor* InActorPicked, const FMovieSceneObjectBindingID& InExistingBindingID) : ActorPicked(InActorPicked), ExistingBindingID(InExistingBindingID) {}
 
 		TWeakObjectPtr<AActor> ActorPicked;
 		FMovieSceneObjectBindingID ExistingBindingID;
diff --git a/Engine/Source/Editor/MovieSceneTools/Private/TrackEditors/AttachTrackEditor.cpp b/Engine/Source/Editor/MovieSceneTools/Private/TrackEditors/AttachTrackEditor.cpp
--- a/Engine/Source/Editor/MovieSceneTools/Private/TrackEditors/AttachTrackEditor.cpp
+++ b/Engine/Source/Editor/MovieSceneTools/Private/TrackEditors/AttachTrackEditor.cpp
@@ -53,15 +53,7 @@ public:
 			TSharedPtr<ISequencer> Sequencer = AttachTrackEditor->GetSequencer();
 			if (Sequencer.IsValid())
 			{
-				FMovieSceneSequenceID SequenceID = Sequencer->GetFocusedTemplateID();
-				if (AttachSection->GetConstraintBindingID().GetSequenceID().IsValid())
-				{
-					// Ensure that this ID is resolvable from the root, based on the current local sequence ID
-					FMovieSceneObjectBindingID RootBindingID = AttachSection->GetConstraintBindingID().ResolveLocalToRoot(SequenceID, *Sequencer);
-					SequenceID = RootBindingID.GetSequenceID();
-				}
-
-				TArrayView<TWeakObjectPtr<UObject>> RuntimeObjects = Sequencer->FindBoundObjects(AttachSection->GetConstraintBindingID().GetGuid(), SequenceID);
+				TArrayView<TWeakObjectPtr<UObject>> RuntimeObjects = AttachSection->GetConstraintBindingID().ResolveBoundObjects(Sequencer->GetFocusedTemplateID(), *Sequencer);
 				if (RuntimeObjects.Num() == 1 && RuntimeObjects[0].IsValid())
 				{
 					if (AActor* Actor = Cast<AActor>(RuntimeObjects[0].Get()))
@@ -298,13 +290,12 @@ void F3DAttachTrackEditor::ActorSocketPicked(const FName SocketName, USceneCompo
 			if (Spawnable.IsSet())
 			{
 				// Check whether the spawnable is underneath the current sequence, if so, we can remap it to a local sequence ID
-				ConstraintBindingID = FMovieSceneObjectBindingID(Spawnable->ObjectBindingID, Spawnable->SequenceID);
-				ConstraintBindingID.TryAndMakeLocalTo(SequencerPtr->GetFocusedTemplateID(), *SequencerPtr);
+				ConstraintBindingID = UE::MovieScene::FRelativeObjectBindingID(SequencerPtr->GetFocusedTemplateID(), Spawnable->SequenceID, Spawnable->ObjectBindingID, *SequencerPtr);
 			}
 			else
 			{
 				FGuid ParentActorId = FindOrCreateHandleToObject(Actor).Handle;
-				ConstraintBindingID = FMovieSceneObjectBindingID(ParentActorId, MovieSceneSequenceID::Root, EMovieSceneObjectBindingSpace::Local);
+				ConstraintBindingID = UE::MovieScene::FRelativeObjectBindingID(ParentActorId);
 			}
 		}
 
@@ -491,15 +482,7 @@ FTransform GetLocationAtTime(UMovieScene3DTransformTrack* TransformTrack, FFrame
 
 AActor* GetConstraintActor(TSharedPtr<ISequencer> InSequencer, const FMovieSceneObjectBindingID& InConstraintBindingID)
 {
-	FMovieSceneSequenceID SequenceID = InSequencer->GetFocusedTemplateID();
-	if (InConstraintBindingID.GetSequenceID().IsValid())
-	{
-		// Ensure that this ID is resolvable from the root, based on the current local sequence ID
-		FMovieSceneObjectBindingID RootBindingID = InConstraintBindingID.ResolveLocalToRoot(SequenceID, *InSequencer);
-		SequenceID = RootBindingID.GetSequenceID();
-	}
-
-	TArrayView<TWeakObjectPtr<UObject>> RuntimeObjects = InSequencer->FindBoundObjects(InConstraintBindingID.GetGuid(), SequenceID);
+	TArrayView<TWeakObjectPtr<UObject>> RuntimeObjects = InConstraintBindingID.ResolveBoundObjects(InSequencer->GetFocusedTemplateID(), *InSequencer);
 
 	if (RuntimeObjects.Num() >= 1 && RuntimeObjects[0].IsValid())
 	{
@@ -1142,14 +1125,13 @@ FKeyPropertyResult F3DAttachTrackEditor::AddKeyInternal( FFrameNumber KeyTime, c
 		if (Spawnable.IsSet())
 		{
 			// Check whether the spawnable is underneath the current sequence, if so, we can remap it to a local sequence ID
-			ConstraintBindingID = FMovieSceneObjectBindingID(Spawnable->ObjectBindingID, Spawnable->SequenceID);
-			ConstraintBindingID.TryAndMakeLocalTo(SequencerPtr->GetFocusedTemplateID(), *SequencerPtr);
+			ConstraintBindingID = UE::MovieScene::FRelativeObjectBindingID(SequencerPtr->GetFocusedTemplateID(), Spawnable->SequenceID, Spawnable->ObjectBindingID, *SequencerPtr);
 		}
 		else
 		{
 			FFindOrCreateHandleResult HandleResult = FindOrCreateHandleToObject(ActorPickerID.ActorPicked.Get());
 			KeyPropertyResult.bHandleCreated |= HandleResult.bWasCreated;
-			ConstraintBindingID = FMovieSceneObjectBindingID(HandleResult.Handle, MovieSceneSequenceID::Root, EMovieSceneObjectBindingSpace::Local);
+			ConstraintBindingID = UE::MovieScene::FRelativeObjectBindingID(HandleResult.Handle);
 		}
 	}
 
@@ -1169,7 +1151,7 @@ FKeyPropertyResult F3DAttachTrackEditor::AddKeyInternal( FFrameNumber KeyTime, c
 	TOptional<TArrayView<FMovieSceneFloatChannel*>> ParentChannels;
 
 	// If the constraint exists within this sequence, we can perform transform compensation
-	if (ConstraintBindingID.ResolveLocalToRoot(CurrentSequenceID,*GetSequencer()).GetSequenceID() == CurrentSequenceID)
+	if (ConstraintBindingID.ResolveSequenceID(CurrentSequenceID,*GetSequencer()) == CurrentSequenceID)
 	{
 		UMovieScene3DTransformTrack* ParentTransformTrack = MovieScene->FindTrack<UMovieScene3DTransformTrack>(ConstraintBindingID.GetGuid());
 		if (ParentTransformTrack && ParentTransformTrack->GetAllSections().Num() == 1)
diff --git a/Engine/Source/Editor/MovieSceneTools/Private/TrackEditors/CameraCutTrackEditor.cpp b/Engine/Source/Editor/MovieSceneTools/Private/TrackEditors/CameraCutTrackEditor.cpp
--- a/Engine/Source/Editor/MovieSceneTools/Private/TrackEditors/CameraCutTrackEditor.cpp
+++ b/Engine/Source/Editor/MovieSceneTools/Private/TrackEditors/CameraCutTrackEditor.cpp
@@ -358,7 +358,7 @@ FKeyPropertyResult FCameraCutTrackEditor::AddKeyInternal( FFrameNumber KeyTime,
 	UMovieSceneCameraCutTrack* CameraCutTrack = FindOrCreateCameraCutTrack();
 	const TArray<UMovieSceneSection*>& AllSections = CameraCutTrack->GetAllSections();
 
-	UMovieSceneCameraCutSection* NewSection = CameraCutTrack->AddNewCameraCut(FMovieSceneObjectBindingID(ObjectGuid, MovieSceneSequenceID::Root, EMovieSceneObjectBindingSpace::Local), KeyTime);
+	UMovieSceneCameraCutSection* NewSection = CameraCutTrack->AddNewCameraCut(UE::MovieScene::FRelativeObjectBindingID(ObjectGuid), KeyTime);
 	KeyPropertyResult.bTrackModified = true;
 	KeyPropertyResult.SectionsCreated.Add(NewSection);
 
diff --git a/Engine/Source/Editor/MovieSceneTools/Private/TrackEditors/PathTrackEditor.cpp b/Engine/Source/Editor/MovieSceneTools/Private/TrackEditors/PathTrackEditor.cpp
--- a/Engine/Source/Editor/MovieSceneTools/Private/TrackEditors/PathTrackEditor.cpp
+++ b/Engine/Source/Editor/MovieSceneTools/Private/TrackEditors/PathTrackEditor.cpp
@@ -12,6 +12,7 @@
 #include "Components/SplineComponent.h"
 #include "MovieSceneToolHelpers.h"
 #include "Evaluation/MovieSceneEvaluationTemplateInstance.h"
+#include "MovieSceneSpawnableAnnotation.h"
 
 
 #define LOCTEXT_NAMESPACE "FPathTrackEditor"
@@ -42,15 +43,7 @@ public:
 			TSharedPtr<ISequencer> Sequencer = PathTrackEditor->GetSequencer();
 			if (Sequencer.IsValid())
 			{
-				FMovieSceneSequenceID SequenceID = Sequencer->GetFocusedTemplateID();
-				if (PathSection->GetConstraintBindingID().GetSequenceID().IsValid())
-				{
-					// Ensure that this ID is resolvable from the root, based on the current local sequence ID
-					FMovieSceneObjectBindingID RootBindingID = PathSection->GetConstraintBindingID().ResolveLocalToRoot(SequenceID, *Sequencer);
-					SequenceID = RootBindingID.GetSequenceID();
-				}
-
-				TArrayView<TWeakObjectPtr<UObject>> RuntimeObjects = Sequencer->FindBoundObjects(PathSection->GetConstraintBindingID().GetGuid(), SequenceID);
+				TArrayView<TWeakObjectPtr<UObject>> RuntimeObjects = PathSection->GetConstraintBindingID().ResolveBoundObjects(Sequencer->GetFocusedTemplateID(), *Sequencer);
 				if (RuntimeObjects.Num() == 1 && RuntimeObjects[0].IsValid())
 				{
 					if (AActor* Actor = Cast<AActor>(RuntimeObjects[0].Get()))
@@ -140,16 +133,8 @@ bool F3DPathTrackEditor::IsActorPickable(const AActor* const ParentActor, FGuid
 	UMovieScene3DPathSection* PathSection = Cast<UMovieScene3DPathSection>(InSection);
 	if (PathSection != nullptr)
 	{
-		FMovieSceneSequenceID SequenceID = GetSequencer()->GetFocusedTemplateID();
-		if (PathSection->GetConstraintBindingID().GetSequenceID().IsValid())
-		{
-			// Ensure that this ID is resolvable from the root, based on the current local sequence ID
-			FMovieSceneObjectBindingID RootBindingID = PathSection->GetConstraintBindingID().ResolveLocalToRoot(SequenceID, *GetSequencer());
-			SequenceID = RootBindingID.GetSequenceID();
-		}
-
-		TArrayView<TWeakObjectPtr<UObject>> RuntimeObjects = GetSequencer()->FindBoundObjects(PathSection->GetConstraintBindingID().GetGuid(), SequenceID);
-		
+		TSharedPtr<ISequencer> SequencerPtr = GetSequencer();
+		TArrayView<TWeakObjectPtr<UObject>> RuntimeObjects = PathSection->GetConstraintBindingID().ResolveBoundObjects(SequencerPtr->GetFocusedTemplateID(), *SequencerPtr);
 		if (RuntimeObjects.Contains(ParentActor))
 		{
 			return false;
@@ -187,10 +172,21 @@ void F3DPathTrackEditor::ActorSocketPicked(const FName SocketName, USceneCompone
 		{
 			ConstraintBindingID = ActorPickerID.ExistingBindingID;
 		}
-		else if (ActorPickerID.ActorPicked.IsValid())
+		else if (AActor* Actor = ActorPickerID.ActorPicked.Get())
 		{
-			FGuid ParentActorId = FindOrCreateHandleToObject(ActorPickerID.ActorPicked.Get()).Handle;
-			ConstraintBindingID = FMovieSceneObjectBindingID(ParentActorId, MovieSceneSequenceID::Root, EMovieSceneObjectBindingSpace::Local);
+			TSharedPtr<ISequencer> SequencerPtr = GetSequencer();
+
+			TOptional<FMovieSceneSpawnableAnnotation> Spawnable = FMovieSceneSpawnableAnnotation::Find(Actor);
+			if (Spawnable.IsSet())
+			{
+				// Check whether the spawnable is underneath the current sequence, if so, we can remap it to a local sequence ID
+				ConstraintBindingID = UE::MovieScene::FRelativeObjectBindingID(SequencerPtr->GetFocusedTemplateID(), Spawnable->SequenceID, Spawnable->ObjectBindingID, *SequencerPtr);
+			}
+			else
+			{
+				FGuid ParentActorId = FindOrCreateHandleToObject(Actor).Handle;
+				ConstraintBindingID = UE::MovieScene::FRelativeObjectBindingID(ParentActorId);
+			}
 		}
 
 		if (ConstraintBindingID.IsValid())
@@ -224,12 +220,24 @@ FKeyPropertyResult F3DPathTrackEditor::AddKeyInternal( FFrameNumber KeyTime, con
 	{
 		ConstraintBindingID = ActorPickerID.ExistingBindingID;
 	}
-	else if (ActorPickerID.ActorPicked.IsValid())
+	else if (AActor* Actor = ActorPickerID.ActorPicked.Get())
 	{
-		FFindOrCreateHandleResult HandleResult = FindOrCreateHandleToObject(ActorPickerID.ActorPicked.Get());
-		FGuid ParentActorId = HandleResult.Handle;
-		KeyPropertyResult.bHandleCreated |= HandleResult.bWasCreated;
-		ConstraintBindingID = FMovieSceneObjectBindingID(ParentActorId, MovieSceneSequenceID::Root, EMovieSceneObjectBindingSpace::Local);
+		TSharedPtr<ISequencer> SequencerPtr = GetSequencer();
+
+		TOptional<FMovieSceneSpawnableAnnotation> Spawnable = FMovieSceneSpawnableAnnotation::Find(Actor);
+		if (Spawnable.IsSet())
+		{
+			// Check whether the spawnable is underneath the current sequence, if so, we can remap it to a local sequence ID
+			ConstraintBindingID = UE::MovieScene::FRelativeObjectBindingID(SequencerPtr->GetFocusedTemplateID(), Spawnable->SequenceID, Spawnable->ObjectBindingID, *SequencerPtr);
+		}
+		else
+		{
+			FFindOrCreateHandleResult HandleResult = FindOrCreateHandleToObject(Actor);
+			FGuid ParentActorId = HandleResult.Handle;
+			KeyPropertyResult.bHandleCreated |= HandleResult.bWasCreated;
+
+			ConstraintBindingID = UE::MovieScene::FRelativeObjectBindingID(ParentActorId);
+		}
 	}
 
 	if (!ConstraintBindingID.IsValid())
diff --git a/Engine/Source/Editor/MovieSceneTools/Private/TrackEditors/PropertyTrackEditors/ActorReferencePropertyTrackEditor.cpp b/Engine/Source/Editor/MovieSceneTools/Private/TrackEditors/PropertyTrackEditors/ActorReferencePropertyTrackEditor.cpp
--- a/Engine/Source/Editor/MovieSceneTools/Private/TrackEditors/PropertyTrackEditors/ActorReferencePropertyTrackEditor.cpp
+++ b/Engine/Source/Editor/MovieSceneTools/Private/TrackEditors/PropertyTrackEditors/ActorReferencePropertyTrackEditor.cpp
@@ -2,6 +2,7 @@
 
 #include "TrackEditors/PropertyTrackEditors/ActorReferencePropertyTrackEditor.h"
 #include "GameFramework/Actor.h"
+#include "MovieSceneSpawnableAnnotation.h"
 
 
 TSharedRef<ISequencerTrackEditor> FActorReferencePropertyTrackEditor::CreateTrackEditor( TSharedRef<ISequencer> OwningSequencer )
@@ -15,11 +16,26 @@ void FActorReferencePropertyTrackEditor::GenerateKeysFromPropertyChanged( const
 	AActor* NewReferencedActor = Cast<AActor>(PropertyChangedParams.GetPropertyValue<UObject*>());
 	if ( NewReferencedActor != nullptr )
 	{
-		FGuid ActorGuid = GetSequencer()->GetHandleToObject( NewReferencedActor );
-		if ( ActorGuid.IsValid() )
+		TSharedPtr<ISequencer> SequencerPtr = GetSequencer();
+
+		FMovieSceneObjectBindingID Binding;
+
+		TOptional<FMovieSceneSpawnableAnnotation> Spawnable = FMovieSceneSpawnableAnnotation::Find(NewReferencedActor);
+		if (Spawnable.IsSet())
+		{
+			// Check whether the spawnable is underneath the current sequence, if so, we can remap it to a local sequence ID
+			Binding = UE::MovieScene::FRelativeObjectBindingID(SequencerPtr->GetFocusedTemplateID(), Spawnable->SequenceID, Spawnable->ObjectBindingID, *SequencerPtr);
+		}
+		else
+		{
+			FGuid ParentActorId = FindOrCreateHandleToObject(NewReferencedActor).Handle;
+			Binding = UE::MovieScene::FRelativeObjectBindingID(ParentActorId);
+		}
+
+		if (Binding.IsValid())
 		{
 			FMovieSceneActorReferenceKey NewKey;
-			NewKey.Object = FMovieSceneObjectBindingID(ActorGuid, MovieSceneSequenceID::Root, EMovieSceneObjectBindingSpace::Local);
+			NewKey.Object = Binding;
 			OutGeneratedKeys.Add(FMovieSceneChannelValueSetter::Create<FMovieSceneActorReferenceData>(0, NewKey, true));
 		}
 	}
diff --git a/Engine/Source/Editor/MovieSceneTools/Public/MovieSceneObjectBindingIDCustomization.h b/Engine/Source/Editor/MovieSceneTools/Public/MovieSceneObjectBindingIDCustomization.h
--- a/Engine/Source/Editor/MovieSceneTools/Public/MovieSceneObjectBindingIDCustomization.h
+++ b/Engine/Source/Editor/MovieSceneTools/Public/MovieSceneObjectBindingIDCustomization.h
@@ -11,7 +11,15 @@ class IPropertyHandle;
 class UMovieSceneSequence;
 class FDragDropOperation;
 class ISequencer;
-struct FMovieSceneObjectBindingID;
+
+namespace UE
+{
+namespace MovieScene
+{
+	struct FFixedObjectBindingID;
+}
+}
+
 
 
 class MOVIESCENETOOLS_API FMovieSceneObjectBindingIDCustomization
diff --git a/Engine/Source/Editor/MovieSceneTools/Public/MovieSceneObjectBindingIDPicker.h b/Engine/Source/Editor/MovieSceneTools/Public/MovieSceneObjectBindingIDPicker.h
--- a/Engine/Source/Editor/MovieSceneTools/Public/MovieSceneObjectBindingIDPicker.h
+++ b/Engine/Source/Editor/MovieSceneTools/Public/MovieSceneObjectBindingIDPicker.h
@@ -6,6 +6,7 @@
 #include "Textures/SlateIcon.h"
 #include "MovieSceneSequenceID.h"
 
+class FReply;
 class SWidget;
 class UMovieSceneSequence;
 class FMenuBuilder;
@@ -17,6 +18,17 @@ struct FSequenceBindingTree;
 struct FSequenceBindingNode;
 struct FMovieSceneObjectBindingID;
 
+namespace UE
+{
+namespace MovieScene
+{
+
+struct FRelativeObjectBindingID;
+struct FFixedObjectBindingID;
+
+} // namespace MovieScene
+} // namespace UE
+
 /**
  * Helper class that is used to pick object bindings for movie scene data
  */
@@ -76,7 +88,7 @@ protected:
 	EVisibility GetSpawnableIconOverlayVisibility() const;
 
 	/** Assign a new binding ID in response to user-input */
-	void SetBindingId(FMovieSceneObjectBindingID InBindingId);
+	void SetBindingId(UE::MovieScene::FFixedObjectBindingID InBindingId);
 
 	/** Build menu content that allows the user to choose a binding from inside the source sequence */
 	TSharedRef<SWidget> GetPickerMenu();
@@ -87,6 +99,9 @@ protected:
 	/** Get a widget that represents the currently chosen item */
 	TSharedRef<SWidget> GetCurrentItemWidget(TSharedRef<STextBlock> TextContent);
 
+	/** Get a widget that represents a warning/fixup button for this binding */
+	TSharedRef<SWidget> GetWarningWidget();
+
 	/** Optional sequencer ptr */
 	TWeakPtr<ISequencer> WeakSequencer;
 
@@ -98,15 +113,21 @@ protected:
 
 private:
 
-	/** Get the currently set binding ID, remapped to the root sequence if necessary */
-	FMovieSceneObjectBindingID GetRemappedCurrentValue() const;
+	/** Get the currently set binding ID as a fixed binding ID from the root sequence */
+	UE::MovieScene::FFixedObjectBindingID GetCurrentValueAsFixed() const;
 
-	/** Set the binding ID, remapped to the local sequence if possible */
-	void SetRemappedCurrentValue(FMovieSceneObjectBindingID InValue);
+	/** Set the binding ID */
+	void SetCurrentValueFromFixed(UE::MovieScene::FFixedObjectBindingID InValue);
 
 	/** Called when the combo box has been clicked to populate its menu content */
 	void OnGetMenuContent(FMenuBuilder& MenuBuilder, TSharedPtr<FSequenceBindingNode> Node);
 
+	/** Get the visibility for the warning/fixup button */
+	EVisibility GetFixedWarningVisibility() const;
+
+	/** Get the visibility for the warning/fixup button */
+	FReply AttemptBindingFixup();
+
 	/** Cached current text and tooltips */
 	FText CurrentText, ToolTipText;
 
diff --git a/Engine/Source/Editor/MovieSceneTools/Public/ObjectBindingTagCache.h b/Engine/Source/Editor/MovieSceneTools/Public/ObjectBindingTagCache.h
--- a/Engine/Source/Editor/MovieSceneTools/Public/ObjectBindingTagCache.h
+++ b/Engine/Source/Editor/MovieSceneTools/Public/ObjectBindingTagCache.h
@@ -20,7 +20,7 @@ class FObjectBindingTagCache
 {
 public:
 
-	using TagIterator = TMultiMap<FMovieSceneObjectBindingID, FName>::TConstKeyIterator;
+	using TagIterator = TMultiMap<UE::MovieScene::FFixedObjectBindingID, FName>::TConstKeyIterator;
 
 	DECLARE_MULTICAST_DELEGATE_OneParam(FOnUpdated, const FObjectBindingTagCache*);
 
@@ -46,7 +46,7 @@ public:
 	/*
 	 * Check whether the specified binding ID has the specified tag
 	 */
-	bool HasTag(const FMovieSceneObjectBindingID& BindingID, const FName& TagName) const
+	bool HasTag(const UE::MovieScene::FFixedObjectBindingID& BindingID, const FName& TagName) const
 	{
 		return ExposedNameReverseLUT.FindPair(BindingID, TagName) != nullptr;
 	}
@@ -55,7 +55,7 @@ public:
 	/*
 	 * Iterate all the tags for the specified binding ID
 	 */
-	TagIterator IterateTags(const FMovieSceneObjectBindingID& BindingID) const
+	TagIterator IterateTags(const UE::MovieScene::FFixedObjectBindingID& BindingID) const
 	{
 		return ExposedNameReverseLUT.CreateConstKeyIterator(BindingID);
 	}
@@ -69,5 +69,5 @@ private:
 	FGuid RootSequenceSignature;
 
 	/** Multi map (1->many) of a binding ID to any tags that are associated with it */
-	TMultiMap<FMovieSceneObjectBindingID, FName> ExposedNameReverseLUT;
+	TMultiMap<UE::MovieScene::FFixedObjectBindingID, FName> ExposedNameReverseLUT;
 };
\ No newline at end of file
diff --git a/Engine/Source/Editor/MovieSceneTools/Public/SObjectBindingTag.h b/Engine/Source/Editor/MovieSceneTools/Public/SObjectBindingTag.h
--- a/Engine/Source/Editor/MovieSceneTools/Public/SObjectBindingTag.h
+++ b/Engine/Source/Editor/MovieSceneTools/Public/SObjectBindingTag.h
@@ -27,7 +27,7 @@ public:
 	SLATE_END_ARGS()
 
 	/** Construct this widget. Binding cache can be retrieved through FSequencer::GetObjectTagCache()*/
-	MOVIESCENETOOLS_API void Construct(const FArguments& InArgs, const FMovieSceneObjectBindingID& InBindingID, FObjectBindingTagCache* BindingCache);
+	MOVIESCENETOOLS_API void Construct(const FArguments& InArgs, const UE::MovieScene::FFixedObjectBindingID& InBindingID, FObjectBindingTagCache* BindingCache);
 
 private:
 
@@ -40,7 +40,7 @@ private:
 private:
 
 	FTagEvent OnTagDeletedEvent;
-	FMovieSceneObjectBindingID BindingID;
+	UE::MovieScene::FFixedObjectBindingID BindingID;
 };
 
 /** A single named tag widget for an object binding within a sequence, represented as a rounded 'pill' */
diff --git a/Engine/Source/Editor/MovieSceneTools/Public/SequenceBindingTree.h b/Engine/Source/Editor/MovieSceneTools/Public/SequenceBindingTree.h
--- a/Engine/Source/Editor/MovieSceneTools/Public/SequenceBindingTree.h
+++ b/Engine/Source/Editor/MovieSceneTools/Public/SequenceBindingTree.h
@@ -19,8 +19,9 @@ class UMovieSceneSequence;
 /** Node that represents an object binding, or a sub sequence (where the guid is zero) */
 struct FSequenceBindingNode
 {
-	FSequenceBindingNode(FText InDisplayString, FMovieSceneObjectBindingID InBindingID, FSlateIcon InIcon)
+	FSequenceBindingNode(FText InDisplayString, const UE::MovieScene::FFixedObjectBindingID& InBindingID, FSlateIcon InIcon)
 		: BindingID(InBindingID)
+		, ParentID(FGuid(), MovieSceneSequenceID::Invalid)
 		, DisplayString(InDisplayString)
 		, Icon(InIcon)
 		, bIsSpawnable(false)
@@ -34,7 +35,7 @@ struct FSequenceBindingNode
 	}
 
 	/** This object's ID, and its parent's */
-	FMovieSceneObjectBindingID BindingID, ParentID;
+	UE::MovieScene::FFixedObjectBindingID BindingID, ParentID;
 	/** The display string that represents this node */
 	FText DisplayString;
 	/** A representative icon for the node */
@@ -74,7 +75,7 @@ struct MOVIESCENETOOLS_API FSequenceBindingTree
 	}
 
 	/** Find a node in the tree */
-	TSharedPtr<FSequenceBindingNode> FindNode(FMovieSceneObjectBindingID BindingID) const
+	TSharedPtr<FSequenceBindingNode> FindNode(UE::MovieScene::FFixedObjectBindingID BindingID) const
 	{
 		return Hierarchy.FindRef(BindingID);
 	}
@@ -108,7 +109,7 @@ private:
 	/** The top level (root) node in the tree */
 	TSharedPtr<FSequenceBindingNode> TopLevelNode;
 	/** Map of hierarchical information */
-	TMap<FMovieSceneObjectBindingID, TSharedPtr<FSequenceBindingNode>> Hierarchy;
+	TMap<UE::MovieScene::FFixedObjectBindingID, TSharedPtr<FSequenceBindingNode>> Hierarchy;
 	/** Map of sequence to its signature the last time we were built */
 	TMap<FObjectKey, FGuid> CachedSequenceSignatures;
 	/** Whether the tree is considered empty */
diff --git a/Engine/Source/Editor/SequenceRecorder/Private/SequenceRecorderUtils.cpp b/Engine/Source/Editor/SequenceRecorder/Private/SequenceRecorderUtils.cpp
--- a/Engine/Source/Editor/SequenceRecorder/Private/SequenceRecorderUtils.cpp
+++ b/Engine/Source/Editor/SequenceRecorder/Private/SequenceRecorderUtils.cpp
@@ -218,7 +218,7 @@ void CreateCameraCutTrack(ULevelSequence* LevelSequence, const FGuid& RecordedCa
 	}
 
 	UMovieSceneCameraCutSection* CameraCutSection = Cast<UMovieSceneCameraCutSection>(CameraCutTrack->CreateNewSection());
-	CameraCutSection->SetCameraBindingID(FMovieSceneObjectBindingID(RecordedCameraGuid, SequenceID, EMovieSceneObjectBindingSpace::Local));
+	CameraCutSection->SetCameraBindingID(UE::MovieScene::FRelativeObjectBindingID(RecordedCameraGuid, SequenceID));
 	CameraCutSection->SetRange(LevelSequence->GetMovieScene()->GetPlaybackRange());
 	CameraCutTrack->AddSection(*CameraCutSection);
 }
diff --git a/Engine/Source/Editor/SequenceRecorderSections/Private/MovieScene3DAttachSectionRecorder.cpp b/Engine/Source/Editor/SequenceRecorderSections/Private/MovieScene3DAttachSectionRecorder.cpp
--- a/Engine/Source/Editor/SequenceRecorderSections/Private/MovieScene3DAttachSectionRecorder.cpp
+++ b/Engine/Source/Editor/SequenceRecorderSections/Private/MovieScene3DAttachSectionRecorder.cpp
@@ -65,7 +65,7 @@ void FMovieScene3DAttachSectionRecorder::Record(float CurrentTime)
 				FFrameNumber CurrentFrame = (CurrentTime * TickResolution).FloorToFrame();
 
 				MovieSceneSection->SetRange(TRange<FFrameNumber>::Inclusive(CurrentFrame, CurrentFrame));
-				MovieSceneSection->SetConstraintId(Guid);
+				MovieSceneSection->SetConstraintBindingID(UE::MovieScene::FRelativeObjectBindingID(Guid));
 				MovieSceneSection->AttachSocketName = SocketName;
 				MovieSceneSection->AttachComponentName = ComponentName;
 
diff --git a/Engine/Source/Editor/Sequencer/Private/DisplayNodes/SequencerObjectBindingNode.cpp b/Engine/Source/Editor/Sequencer/Private/DisplayNodes/SequencerObjectBindingNode.cpp
--- a/Engine/Source/Editor/Sequencer/Private/DisplayNodes/SequencerObjectBindingNode.cpp
+++ b/Engine/Source/Editor/Sequencer/Private/DisplayNodes/SequencerObjectBindingNode.cpp
@@ -545,7 +545,7 @@ void FSequencerObjectBindingNode::AddTagMenu(FMenuBuilder& MenuBuilder)
 			{
 				const FGuid& ObjectID = StaticCastSharedRef<FSequencerObjectBindingNode>(Node)->GetObjectBinding();
 
-				FMovieSceneObjectBindingID BindingID(ObjectID, SequenceID);
+				UE::MovieScene::FFixedObjectBindingID BindingID(ObjectID, SequenceID);
 				for (auto It = Sequencer->GetObjectBindingTagCache()->IterateTags(BindingID); It; ++It)
 				{
 					AllTags.Add(It.Value());
@@ -600,7 +600,7 @@ ECheckBoxState FSequencerObjectBindingNode::GetTagCheckState(FName TagName)
 		{
 			const FGuid& ObjectID = StaticCastSharedRef<FSequencerObjectBindingNode>(Node)->GetObjectBinding();
 
-			FMovieSceneObjectBindingID BindingID(ObjectID, SequenceID);
+			UE::MovieScene::FFixedObjectBindingID BindingID(ObjectID, SequenceID);
 			ECheckBoxState ThisCheckState = Sequencer.GetObjectBindingTagCache()->HasTag(BindingID, TagName)
 				? ECheckBoxState::Checked
 				: ECheckBoxState::Unchecked;
@@ -631,7 +631,7 @@ void FSequencerObjectBindingNode::ToggleTag(FName TagName)
 		{
 			const FGuid& ObjectID = StaticCastSharedRef<FSequencerObjectBindingNode>(Node)->GetObjectBinding();
 
-			FMovieSceneObjectBindingID BindingID(ObjectID, SequenceID);
+			UE::MovieScene::FFixedObjectBindingID BindingID(ObjectID, SequenceID);
 			if (!Sequencer.GetObjectBindingTagCache()->HasTag(BindingID, TagName))
 			{
 				HandleAddTag(TagName);
@@ -657,7 +657,7 @@ void FSequencerObjectBindingNode::HandleDeleteTag(FName TagName)
 		{
 			const FGuid& ObjectID = StaticCastSharedRef<FSequencerObjectBindingNode>(Node)->GetObjectBinding();
 
-			MovieScene->UntagBinding(TagName, FMovieSceneObjectBindingID(ObjectID, SequenceID));
+			MovieScene->UntagBinding(TagName, UE::MovieScene::FFixedObjectBindingID(ObjectID, SequenceID));
 		}
 	}
 }
@@ -676,7 +676,7 @@ void FSequencerObjectBindingNode::HandleAddTag(FName TagName)
 		{
 			const FGuid& ObjectID = StaticCastSharedRef<FSequencerObjectBindingNode>(Node)->GetObjectBinding();
 
-			MovieScene->TagBinding(TagName, FMovieSceneObjectBindingID(ObjectID, SequenceID));
+			MovieScene->TagBinding(TagName, UE::MovieScene::FFixedObjectBindingID(ObjectID, SequenceID));
 		}
 	}
 }
@@ -715,7 +715,7 @@ TSharedRef<SWidget> FSequencerObjectBindingNode::GetCustomOutlinerContent()
 TSharedPtr<SWidget> FSequencerObjectBindingNode::GetAdditionalOutlinerLabel()
 {
 	FSequencer& Sequencer = GetSequencer();
-	FMovieSceneObjectBindingID BindingID(ObjectBinding, Sequencer.GetFocusedTemplateID());
+	UE::MovieScene::FFixedObjectBindingID BindingID(ObjectBinding, Sequencer.GetFocusedTemplateID());
 
 	return SNew(SObjectBindingTags, BindingID, Sequencer.GetObjectBindingTagCache());
 }
diff --git a/Engine/Source/Editor/Sequencer/Private/SObjectBindingTagManager.cpp b/Engine/Source/Editor/Sequencer/Private/SObjectBindingTagManager.cpp
--- a/Engine/Source/Editor/Sequencer/Private/SObjectBindingTagManager.cpp
+++ b/Engine/Source/Editor/Sequencer/Private/SObjectBindingTagManager.cpp
@@ -71,7 +71,7 @@ public:
 		else if (ColumnName == Column_Tags)
 		{
 			TSharedRef<SWidget> TagContent = SNullWidget::NullWidget;
-			if (BindingNode->BindingID.IsValid())
+			if (BindingNode->BindingID.Guid.IsValid())
 			{
 				TagContent = GenerateTagsContent(BindingNode);
 			}
@@ -387,7 +387,7 @@ void SObjectBindingTagManager::TagSelectionAs(FName NewName)
 	MovieScene->Modify();
 	for (TSharedRef<FSequenceBindingNode> SelectedItem : TreeView->GetSelectedItems())
 	{
-		if (SelectedItem->BindingID.IsValid())
+		if (SelectedItem->BindingID.Guid.IsValid())
 		{
 			bMadeAnyChanges = true;
 			MovieScene->TagBinding(NewName, SelectedItem->BindingID);
@@ -425,7 +425,7 @@ void SObjectBindingTagManager::UntagSelection(FName TagName, TSharedPtr<FSequenc
 	MovieScene->Modify();
 	for (TSharedRef<FSequenceBindingNode> SelectedItem : TreeView->GetSelectedItems())
 	{
-		if (SelectedItem->BindingID.IsValid())
+		if (SelectedItem->BindingID.Guid.IsValid())
 		{
 			bMadeAnyChanges = true;
 			MovieScene->UntagBinding(TagName, SelectedItem->BindingID);
@@ -433,7 +433,7 @@ void SObjectBindingTagManager::UntagSelection(FName TagName, TSharedPtr<FSequenc
 	}
 	if (InstigatorNode && !TreeView->IsItemSelected(InstigatorNode.ToSharedRef()))
 	{
-		if (InstigatorNode->BindingID.IsValid())
+		if (InstigatorNode->BindingID.Guid.IsValid())
 		{
 			bMadeAnyChanges = true;
 			MovieScene->UntagBinding(TagName, InstigatorNode->BindingID);
diff --git a/Engine/Source/Editor/Sequencer/Private/Sequencer.cpp b/Engine/Source/Editor/Sequencer/Private/Sequencer.cpp
--- a/Engine/Source/Editor/Sequencer/Private/Sequencer.cpp
+++ b/Engine/Source/Editor/Sequencer/Private/Sequencer.cpp
@@ -945,16 +945,16 @@ void FSequencer::FocusSequenceInstance(UMovieSceneSubSection& InSubSection)
 	TemplateIDBackwardStack.Push(ActiveTemplateIDs.Top());
 	TemplateIDForwardStack.Reset();
 
-	FMovieSceneRootOverridePath Path;
+	UE::MovieScene::FSubSequencePath Path;
 
 	// Ensure the hierarchy is up to date
 	FMovieSceneCompiledDataID DataID = CompiledDataManager->Compile(RootSequence.Get());
 	const FMovieSceneSequenceHierarchy& Hierarchy = CompiledDataManager->GetHierarchyChecked(DataID);
 
-	Path.Set(ActiveTemplateIDs.Last(), &Hierarchy);
+	Path.Reset(ActiveTemplateIDs.Last(), &Hierarchy);
 
 	// Root out the SequenceID for the sub section
-	FMovieSceneSequenceID SequenceID = Path.Remap(InSubSection.GetSequenceID());
+	FMovieSceneSequenceID SequenceID = Path.ResolveChildSequenceID(InSubSection.GetSequenceID());
 
 	// If the sequence isn't found, reset to the root and dive in from there
 	if (!Hierarchy.FindSubData(SequenceID))
@@ -1003,10 +1003,10 @@ void FSequencer::FocusSequenceInstance(UMovieSceneSubSection& InSubSection)
 			}
 		}
 
-		Path.Set(ActiveTemplateIDs.Last(), &Hierarchy);
+		Path.Reset(ActiveTemplateIDs.Last(), &Hierarchy);
 
 		// Root out the SequenceID for the sub section
-		SequenceID = Path.Remap(InSubSection.GetSequenceID());
+		SequenceID = Path.ResolveChildSequenceID(InSubSection.GetSequenceID());
 	}
 
 	if (!ensure(Hierarchy.FindSubData(SequenceID)))
@@ -2186,7 +2186,7 @@ void FSequencer::BakeTransform()
 			for (auto AttachSection : AttachTrack->GetAllSections())
 			{
 				FMovieSceneObjectBindingID ConstraintBindingID = (Cast<UMovieScene3DAttachSection>(AttachSection))->GetConstraintBindingID();
-				for (auto ParentObject : FindBoundObjects(ConstraintBindingID.GetGuid(), ConstraintBindingID.GetSequenceID()) )
+				for (auto ParentObject : ConstraintBindingID.ResolveBoundObjects(GetFocusedTemplateID(), *this))
 				{
 					AttachParentActor = Cast<AActor>(ParentObject.Get());
 					break;
@@ -5857,7 +5857,7 @@ void FSequencer::OnNewActorsDropped(const TArray<UObject*>& DroppedObjects, cons
 					// Create an attach track
 					UMovieScene3DAttachTrack* AttachTrack = Cast<UMovieScene3DAttachTrack>(OwnerMovieScene->AddTrack(UMovieScene3DAttachTrack::StaticClass(), NewCameraGuid));
 
-					FMovieSceneObjectBindingID AttachBindingID(NewGuid, MovieSceneSequenceID::Root);
+					FMovieSceneObjectBindingID AttachBindingID = UE::MovieScene::FRelativeObjectBindingID(NewGuid);
 					FFrameNumber StartTime = UE::MovieScene::DiscreteInclusiveLower(GetPlaybackRange());
 					FFrameNumber Duration  = UE::MovieScene::DiscreteSize(GetPlaybackRange());
 
diff --git a/Engine/Source/Editor/Sequencer/Private/SequencerDisplayNodeDragDropOp.cpp b/Engine/Source/Editor/Sequencer/Private/SequencerDisplayNodeDragDropOp.cpp
--- a/Engine/Source/Editor/Sequencer/Private/SequencerDisplayNodeDragDropOp.cpp
+++ b/Engine/Source/Editor/Sequencer/Private/SequencerDisplayNodeDragDropOp.cpp
@@ -92,9 +92,9 @@ FSequencer* FSequencerDisplayNodeDragDropOp::GetSequencer() const
 	return nullptr;
 }
 
-TArray<FMovieSceneObjectBindingID> FSequencerDisplayNodeDragDropOp::GetDraggedBindings() const
+TArray<UE::MovieScene::FFixedObjectBindingID> FSequencerDisplayNodeDragDropOp::GetDraggedBindings() const
 {
-	TArray<FMovieSceneObjectBindingID> Bindings;
+	TArray<UE::MovieScene::FFixedObjectBindingID> Bindings;
 
 	FSequencer* Sequencer = GetSequencer();
 	if (!Sequencer)
@@ -112,10 +112,10 @@ TArray<FMovieSceneObjectBindingID> FSequencerDisplayNodeDragDropOp::GetDraggedBi
 		if (!ObjectBindingID.IsValid())
 		{
 			// To avoid confusion over what is overridable, if anything is invalid, the entire drag is invalid
-			return TArray<FMovieSceneObjectBindingID>();
+			return TArray<UE::MovieScene::FFixedObjectBindingID>();
 		}
 
-		Bindings.Emplace(ObjectBindingID, SequenceID);
+		Bindings.Emplace(UE::MovieScene::FFixedObjectBindingID(ObjectBindingID, SequenceID));
 	}
 
 	return Bindings;
@@ -161,7 +161,7 @@ FReply FSequencerDisplayNodeDragDropOp::DroppedOnPanel( const TSharedRef< class
 			FGuid ObjectBindingID = GetBindingID(Node, Sequence);
 			if (ObjectBindingID.IsValid())
 			{
-				Template->Binding = FMovieSceneObjectBindingID(ObjectBindingID, SequenceID);
+				Template->Binding = UE::MovieScene::FFixedObjectBindingID(ObjectBindingID, SequenceID);
 				UEdGraphNode* NewNode = Action.PerformAction(&Graph, GetHoveredPin(), GraphPosition, false);
 
 				int32 Offset = FMath::Max(NewNode->NodeHeight, 100);
diff --git a/Engine/Source/Editor/Sequencer/Private/SequencerDisplayNodeDragDropOp.h b/Engine/Source/Editor/Sequencer/Private/SequencerDisplayNodeDragDropOp.h
--- a/Engine/Source/Editor/Sequencer/Private/SequencerDisplayNodeDragDropOp.h
+++ b/Engine/Source/Editor/Sequencer/Private/SequencerDisplayNodeDragDropOp.h
@@ -7,7 +7,13 @@
 #include "DisplayNodes/SequencerDisplayNode.h"
 #include "SequencerObjectBindingDragDropOp.h"
 
-struct FMovieSceneObjectBindingID;
+namespace UE
+{
+namespace MovieScene
+{
+	struct FFixedObjectBindingID;
+}
+}
 
 /** A decorated drag drop operation object for dragging sequencer display nodes. */
 class FSequencerDisplayNodeDragDropOp : public FSequencerObjectBindingDragDropOp
@@ -23,7 +29,7 @@ public:
 public:
 
 	//~ FSequencerObjectBindingDragDropOp interface
-	virtual TArray<FMovieSceneObjectBindingID> GetDraggedBindings() const override;
+	virtual TArray<UE::MovieScene::FFixedObjectBindingID> GetDraggedBindings() const override;
 
 	//~ FGraphEditorDragDropAction interface
 	virtual void HoverTargetChanged() override;
diff --git a/Engine/Source/Editor/Sequencer/Private/SequencerEdMode.cpp b/Engine/Source/Editor/Sequencer/Private/SequencerEdMode.cpp
--- a/Engine/Source/Editor/Sequencer/Private/SequencerEdMode.cpp
+++ b/Engine/Source/Editor/Sequencer/Private/SequencerEdMode.cpp
@@ -659,18 +659,7 @@ void FSequencerEdMode::DrawAudioTracks(FPrimitiveDrawInterface* PDI)
 					FMovieSceneObjectBindingID AttachBindingID = Values[Index].Object;
 					FName AttachSocketName = Values[Index].SocketName;
 
-					FMovieSceneSequenceID SequenceID = Sequencer->GetFocusedTemplateID();
-					if (AttachBindingID.GetSequenceID().IsValid())
-					{
-						// Ensure that this ID is resolvable from the root, based on the current local sequence ID
-						FMovieSceneObjectBindingID RootBindingID = AttachBindingID.ResolveLocalToRoot(SequenceID, *Sequencer);
-						SequenceID = RootBindingID.GetSequenceID();
-					}
-
-					// If the transform is set, otherwise use the bound actor's transform
-					FMovieSceneEvaluationOperand ObjectOperand(SequenceID, AttachBindingID.GetGuid());
-
-					for (TWeakObjectPtr<> WeakObject : Sequencer->FindBoundObjects(ObjectOperand))
+					for (TWeakObjectPtr<> WeakObject : AttachBindingID.ResolveBoundObjects(Sequencer->GetFocusedTemplateID(), *Sequencer))
 					{
 						AActor* AttachActor = Cast<AActor>(WeakObject.Get());
 						if (AttachActor)
diff --git a/Engine/Source/Editor/Sequencer/Public/SequencerObjectBindingDragDropOp.h b/Engine/Source/Editor/Sequencer/Public/SequencerObjectBindingDragDropOp.h
--- a/Engine/Source/Editor/Sequencer/Public/SequencerObjectBindingDragDropOp.h
+++ b/Engine/Source/Editor/Sequencer/Public/SequencerObjectBindingDragDropOp.h
@@ -6,12 +6,20 @@
 #include "Input/DragAndDrop.h"
 #include "GraphEditorDragDropAction.h"
 
-struct FMovieSceneObjectBindingID;
+namespace UE
+{
+namespace MovieScene
+{
+
+struct FFixedObjectBindingID;
+
+} // namespace MovieScene
+} // namespace UE
 
 class FSequencerObjectBindingDragDropOp : public FGraphEditorDragDropAction
 {
 public:
 	DRAG_DROP_OPERATOR_TYPE( FSequencerDisplayNodeDragDropOpBase, FGraphEditorDragDropAction )
 
-	virtual TArray<FMovieSceneObjectBindingID> GetDraggedBindings() const = 0;
+	virtual TArray<UE::MovieScene::FFixedObjectBindingID> GetDraggedBindings() const = 0;
 };
\ No newline at end of file
diff --git a/Engine/Source/Runtime/LevelSequence/Private/LevelSequenceActor.cpp b/Engine/Source/Runtime/LevelSequence/Private/LevelSequenceActor.cpp
--- a/Engine/Source/Runtime/LevelSequence/Private/LevelSequenceActor.cpp
+++ b/Engine/Source/Runtime/LevelSequence/Private/LevelSequenceActor.cpp
@@ -368,7 +368,8 @@ void ALevelSequenceActor::SetBinding(FMovieSceneObjectBindingID Binding, const T
 		BindingOverrides->SetBinding(Binding, TArray<UObject*>(Actors), bAllowBindingsFromAsset);
 		if (SequencePlayer)
 		{
-			SequencePlayer->State.Invalidate(Binding.GetGuid(), Binding.GetSequenceID());
+			FMovieSceneSequenceID SequenceID = Binding.ResolveSequenceID(MovieSceneSequenceID::Root, *SequencePlayer);
+			SequencePlayer->State.Invalidate(Binding.GetGuid(), SequenceID);
 		}
 	}
 }
@@ -405,7 +406,8 @@ void ALevelSequenceActor::AddBinding(FMovieSceneObjectBindingID Binding, AActor*
 		BindingOverrides->AddBinding(Binding, Actor, bAllowBindingsFromAsset);
 		if (SequencePlayer)
 		{
-			SequencePlayer->State.Invalidate(Binding.GetGuid(), Binding.GetSequenceID());
+			FMovieSceneSequenceID SequenceID = Binding.ResolveSequenceID(MovieSceneSequenceID::Root, *SequencePlayer);
+			SequencePlayer->State.Invalidate(Binding.GetGuid(), SequenceID);
 		}
 	}
 }
@@ -442,7 +444,8 @@ void ALevelSequenceActor::RemoveBinding(FMovieSceneObjectBindingID Binding, AAct
 		BindingOverrides->RemoveBinding(Binding, Actor);
 		if (SequencePlayer)
 		{
-			SequencePlayer->State.Invalidate(Binding.GetGuid(), Binding.GetSequenceID());
+			FMovieSceneSequenceID SequenceID = Binding.ResolveSequenceID(MovieSceneSequenceID::Root, *SequencePlayer);
+			SequencePlayer->State.Invalidate(Binding.GetGuid(), SequenceID);
 		}
 	}
 }
@@ -479,7 +482,8 @@ void ALevelSequenceActor::ResetBinding(FMovieSceneObjectBindingID Binding)
 		BindingOverrides->ResetBinding(Binding);
 		if (SequencePlayer)
 		{
-			SequencePlayer->State.Invalidate(Binding.GetGuid(), Binding.GetSequenceID());
+			FMovieSceneSequenceID SequenceID = Binding.ResolveSequenceID(MovieSceneSequenceID::Root, *SequencePlayer);
+			SequencePlayer->State.Invalidate(Binding.GetGuid(), SequenceID);
 		}
 	}
 }
diff --git a/Engine/Source/Runtime/MovieScene/Private/Compilation/MovieSceneCompiledDataManager.cpp b/Engine/Source/Runtime/MovieScene/Private/Compilation/MovieSceneCompiledDataManager.cpp
--- a/Engine/Source/Runtime/MovieScene/Private/Compilation/MovieSceneCompiledDataManager.cpp
+++ b/Engine/Source/Runtime/MovieScene/Private/Compilation/MovieSceneCompiledDataManager.cpp
@@ -1135,11 +1135,11 @@ bool UMovieSceneCompiledDataManager::CompileHierarchy(UMovieSceneSequence* Seque
 
 bool UMovieSceneCompiledDataManager::CompileHierarchy(UMovieSceneSequence* Sequence, const FGatherParameters& Params, FMovieSceneSequenceHierarchy* InOutHierarchy)
 {
-	FMovieSceneRootOverridePath RootPath;
+	UE::MovieScene::FSubSequencePath RootPath;
 	return CompileHierarchyImpl(Sequence, Params, FMovieSceneEvaluationOperand(), &RootPath, InOutHierarchy);
 }
 
-bool UMovieSceneCompiledDataManager::CompileHierarchyImpl(UMovieSceneSequence* Sequence, const FGatherParameters& Params, const FMovieSceneEvaluationOperand& Operand, FMovieSceneRootOverridePath* RootPath, FMovieSceneSequenceHierarchy* InOutHierarchy)
+bool UMovieSceneCompiledDataManager::CompileHierarchyImpl(UMovieSceneSequence* Sequence, const FGatherParameters& Params, const FMovieSceneEvaluationOperand& Operand, UE::MovieScene::FSubSequencePath* RootPath, FMovieSceneSequenceHierarchy* InOutHierarchy)
 {
 	UMovieScene* MovieScene = Sequence->GetMovieScene();
 	check(MovieScene && RootPath && InOutHierarchy);
@@ -1170,7 +1170,7 @@ bool UMovieSceneCompiledDataManager::CompileHierarchyImpl(UMovieSceneSequence* S
 	return bContainsSubSequences;
 }
 
-bool UMovieSceneCompiledDataManager::CompileSubTrackHierarchy(UMovieSceneSubTrack* SubTrack, const FGatherParameters& Params, const FMovieSceneEvaluationOperand& Operand, FMovieSceneRootOverridePath* RootPath, FMovieSceneSequenceHierarchy* InOutHierarchy)
+bool UMovieSceneCompiledDataManager::CompileSubTrackHierarchy(UMovieSceneSubTrack* SubTrack, const FGatherParameters& Params, const FMovieSceneEvaluationOperand& Operand, UE::MovieScene::FSubSequencePath* RootPath, FMovieSceneSequenceHierarchy* InOutHierarchy)
 {
 	bool bContainsSubSequences = false;
 
@@ -1199,7 +1199,7 @@ bool UMovieSceneCompiledDataManager::CompileSubTrackHierarchy(UMovieSceneSubTrac
 			continue;
 		}
 
-		const FMovieSceneSequenceID InnerSequenceID = RootPath->Remap(SubSection->GetSequenceID());
+		const FMovieSceneSequenceID InnerSequenceID = RootPath->ResolveChildSequenceID(SubSection->GetSequenceID());
 
 		SectionToID.Add(SubSection, InnerSequenceID);
 
@@ -1263,9 +1263,9 @@ bool UMovieSceneCompiledDataManager::CompileSubTrackHierarchy(UMovieSceneSubTrac
 			SubParams.Flags |= Entry.Flags;
 			SubParams.NetworkMask = NewMask;
 
-			RootPath->Push(SubData->DeterministicSequenceID);
+			RootPath->PushGeneration(SubSequenceID, SubData->DeterministicSequenceID);
 			CompileHierarchyImpl(SubData->GetSequence(), SubParams, Operand, RootPath, InOutHierarchy);
-			RootPath->Pop();
+			RootPath->PopGenerations(1);
 		}
 	}
 
diff --git a/Engine/Source/Runtime/MovieScene/Private/EntitySystem/MovieSceneEntityInstantiatorSystem.cpp b/Engine/Source/Runtime/MovieScene/Private/EntitySystem/MovieSceneEntityInstantiatorSystem.cpp
--- a/Engine/Source/Runtime/MovieScene/Private/EntitySystem/MovieSceneEntityInstantiatorSystem.cpp
+++ b/Engine/Source/Runtime/MovieScene/Private/EntitySystem/MovieSceneEntityInstantiatorSystem.cpp
@@ -76,11 +76,10 @@ void UMovieSceneEntityInstantiatorSystem::UnlinkStaleObjectBindings(UE::MovieSce
 	{
 		const FSequenceInstance* TargetInstance = &InstanceRegistry->GetInstance(InstanceHandle);
 
-		// Is the binding in the same instance as the entity?
-		bool bBindingInThisSequence = InObjectBindingID.GetBindingSpace() == EMovieSceneObjectBindingSpace::Local && InObjectBindingID.GetSequenceID() == MovieSceneSequenceID::Root;
-		bBindingInThisSequence |= InObjectBindingID.GetBindingSpace() == EMovieSceneObjectBindingSpace::Root && TargetInstance->GetSequenceID() == InObjectBindingID.GetSequenceID();
+		FMovieSceneSequenceID ThisSequenceID     = TargetInstance->GetSequenceID();
+		FMovieSceneSequenceID RemappedSequenceID = InObjectBindingID.ResolveSequenceID(ThisSequenceID, *TargetInstance->GetPlayer());
 
-		if (bBindingInThisSequence)
+		if (RemappedSequenceID == ThisSequenceID)
 		{
 			if (InstanceRegistry->IsBindingInvalidated(InObjectBindingID.GetGuid(), InstanceHandle))
 			{
@@ -91,16 +90,10 @@ void UMovieSceneEntityInstantiatorSystem::UnlinkStaleObjectBindings(UE::MovieSce
 		{
 			if (!TargetInstance->IsRootSequence())
 			{
-				if (InObjectBindingID.GetBindingSpace() == EMovieSceneObjectBindingSpace::Local)
-				{
-					InObjectBindingID = InObjectBindingID.ResolveLocalToRoot(TargetInstance->GetSequenceID(), *TargetInstance->GetPlayer());
-				}
-
 				TargetInstance = &InstanceRegistry->GetInstance(TargetInstance->GetRootInstanceHandle());
 			}
 
-
-			FInstanceHandle SubInstance = TargetInstance->FindSubInstance(InObjectBindingID.GetSequenceID());
+			FInstanceHandle SubInstance = TargetInstance->FindSubInstance(RemappedSequenceID);
 			if (!InstanceRegistry->IsHandleValid(SubInstance) || InstanceRegistry->IsBindingInvalidated(InObjectBindingID.GetGuid(), SubInstance))
 			{
 				StaleEntities.Add(EntityID);
diff --git a/Engine/Source/Runtime/MovieScene/Private/EntitySystem/MovieSceneSequenceInstance.cpp b/Engine/Source/Runtime/MovieScene/Private/EntitySystem/MovieSceneSequenceInstance.cpp
--- a/Engine/Source/Runtime/MovieScene/Private/EntitySystem/MovieSceneSequenceInstance.cpp
+++ b/Engine/Source/Runtime/MovieScene/Private/EntitySystem/MovieSceneSequenceInstance.cpp
@@ -10,6 +10,7 @@
 
 #include "Evaluation/MovieSceneEvaluationTemplateInstance.h"
 #include "Evaluation/Instances/MovieSceneTrackEvaluator.h"
+#include "Evaluation/MovieSceneRootOverridePath.h"
 
 #include "IMovieScenePlayer.h"
 #include "MovieSceneSequencePlayer.h"
@@ -311,5 +312,10 @@ FMovieSceneEntityID FSequenceInstance::FindEntity(UObject* Owner, uint32 EntityI
 	return Ledger.FindImportedEntity(FMovieSceneEvaluationFieldEntityKey{ decltype(FMovieSceneEvaluationFieldEntityKey::EntityOwner)(Owner), EntityID });
 }
 
+FSubSequencePath FSequenceInstance::GetSubSequencePath() const
+{
+	return FSubSequencePath(SequenceID, *GetPlayer());
+}
+
 } // namespace MovieScene
 } // namespace UE
diff --git a/Engine/Source/Runtime/MovieScene/Private/EntitySystem/MovieSceneSequenceUpdaters.cpp b/Engine/Source/Runtime/MovieScene/Private/EntitySystem/MovieSceneSequenceUpdaters.cpp
--- a/Engine/Source/Runtime/MovieScene/Private/EntitySystem/MovieSceneSequenceUpdaters.cpp
+++ b/Engine/Source/Runtime/MovieScene/Private/EntitySystem/MovieSceneSequenceUpdaters.cpp
@@ -419,7 +419,7 @@ void FSequenceUpdater_Hierarchical::Update(UMovieSceneEntitySystemLinker* Linker
 
 	FInstanceHandle             RootInstanceHandle = InstanceHandle;
 	FMovieSceneCompiledDataID   RootCompiledDataID = CompiledDataID;
-	FMovieSceneRootOverridePath RootOverridePath;
+	FSubSequencePath            RootOverridePath;
 	FMovieSceneContext          RootContext = Context;
 
 	TArray<FMovieSceneSequenceID, TInlineAllocator<16>> ActiveSequences;
@@ -435,7 +435,7 @@ void FSequenceUpdater_Hierarchical::Update(UMovieSceneEntitySystemLinker* Linker
 			RootCompiledDataID = CompiledDataManager->GetDataID(RootSequence);
 			RootContext        = Context.Transform(SubData->RootToSequenceTransform, SubData->TickResolution);
 
-			RootOverridePath.Set(RootOverrideSequenceID, MasterHierarchy);
+			RootOverridePath.Reset(RootOverrideSequenceID, MasterHierarchy);
 
 			ActiveSequences.Add(RootOverrideSequenceID);
 		}
@@ -507,7 +507,7 @@ void FSequenceUpdater_Hierarchical::Update(UMovieSceneEntitySystemLinker* Linker
 		for (FMovieSceneSubSequenceTreeEntry Entry : RootOverrideHierarchy->GetTree().GetAllData(SubSequenceIt.Node()))
 		{
 			// When a root override path is specified, we always remap the 'local' sequence IDs to their equivalents from the master sequence.
-			FMovieSceneSequenceID SequenceIDFromMaster = RootOverridePath.Remap(Entry.SequenceID);
+			FMovieSceneSequenceID SequenceIDFromMaster = RootOverridePath.ResolveChildSequenceID(Entry.SequenceID);
 
 			ActiveSequences.Add(SequenceIDFromMaster);
 
diff --git a/Engine/Source/Runtime/MovieScene/Private/Evaluation/Instances/MovieSceneTrackEvaluator.cpp b/Engine/Source/Runtime/MovieScene/Private/Evaluation/Instances/MovieSceneTrackEvaluator.cpp
--- a/Engine/Source/Runtime/MovieScene/Private/Evaluation/Instances/MovieSceneTrackEvaluator.cpp
+++ b/Engine/Source/Runtime/MovieScene/Private/Evaluation/Instances/MovieSceneTrackEvaluator.cpp
@@ -171,7 +171,7 @@ const FMovieSceneEvaluationGroup* FMovieSceneTrackEvaluator::SetupFrame(UMovieSc
 	check(OverrideRootSequence);
 
 	RootID = InOverrideRootID;
-	RootOverridePath.Set(InOverrideRootID, RootHierarchy);
+	RootOverridePath.Reset(InOverrideRootID, RootHierarchy);
 
 	const FMovieSceneEvaluationField* OverrideRootField = nullptr;
 	FMovieSceneSequenceTransform RootOverrideTransform;
@@ -236,7 +236,7 @@ void FMovieSceneTrackEvaluator::EvaluateGroup(const FMovieSceneEvaluationGroup&
 			FMovieSceneEvaluationFieldTrackPtr    TrackPtr   = TrackEntry.TrackPtr;
 
 			// Ensure we're able to find the sequence instance in our root if we've overridden
-			TrackPtr.SequenceID = RootOverridePath.Remap(TrackPtr.SequenceID);
+			TrackPtr.SequenceID = RootOverridePath.ResolveChildSequenceID(TrackPtr.SequenceID);
 
 			const FCachedPtrs&                EvalPtrs = CachedPtrs.FindChecked(TrackPtr.SequenceID);
 			const FMovieSceneEvaluationTrack* Track    = EvalPtrs.Template->FindTrack(TrackPtr.TrackIdentifier);
@@ -276,7 +276,7 @@ void FMovieSceneTrackEvaluator::EvaluateGroup(const FMovieSceneEvaluationGroup&
 			FMovieSceneEvaluationFieldTrackPtr    TrackPtr   = TrackEntry.TrackPtr;
 
 			// Ensure we're able to find the sequence instance in our root if we've overridden
-			TrackPtr.SequenceID = RootOverridePath.Remap(TrackPtr.SequenceID);
+			TrackPtr.SequenceID = RootOverridePath.ResolveChildSequenceID(TrackPtr.SequenceID);
 
 			const FCachedPtrs&                EvalPtrs = CachedPtrs.FindChecked(TrackPtr.SequenceID);
 			const FMovieSceneEvaluationTrack* Track    = EvalPtrs.Template->FindTrack(TrackPtr.TrackIdentifier);
@@ -334,7 +334,7 @@ void FMovieSceneTrackEvaluator::CallSetupTearDown(IMovieScenePlayer& Player, FDe
 		FMovieSceneEvaluationKey Key = OrderedKey.Key;
 
 		// Ensure we're able to find the sequence instance in our root if we've overridden
-		Key.SequenceID = RootOverridePath.Remap(Key.SequenceID);
+		Key.SequenceID = RootOverridePath.ResolveChildSequenceID(Key.SequenceID);
 
 		const FCachedPtrs* EvalPtrs = CachedPtrs.Find(Key.SequenceID);
 		if (EvalPtrs)
@@ -378,7 +378,7 @@ void FMovieSceneTrackEvaluator::CallSetupTearDown(IMovieScenePlayer& Player, FDe
 		FMovieSceneEvaluationKey Key = OrderedKey.Key;
 
 		// Ensure we're able to find the sequence instance in our root if we've overridden
-		Key.SequenceID = RootOverridePath.Remap(Key.SequenceID);
+		Key.SequenceID = RootOverridePath.ResolveChildSequenceID(Key.SequenceID);
 
 		const FCachedPtrs&                EvalPtrs = CachedPtrs.FindChecked(Key.SequenceID);
 		const FMovieSceneEvaluationTrack* Track    = EvalPtrs.Template->FindTrack(Key.TrackIdentifier);
diff --git a/Engine/Source/Runtime/MovieScene/Private/Evaluation/Instances/MovieSceneTrackEvaluator.h b/Engine/Source/Runtime/MovieScene/Private/Evaluation/Instances/MovieSceneTrackEvaluator.h
--- a/Engine/Source/Runtime/MovieScene/Private/Evaluation/Instances/MovieSceneTrackEvaluator.h
+++ b/Engine/Source/Runtime/MovieScene/Private/Evaluation/Instances/MovieSceneTrackEvaluator.h
@@ -134,7 +134,7 @@ private:
 	UMovieSceneCompiledDataManager* CompiledDataManager;
 
 	/** Override path that is used to remap inner sequence IDs to the root space when evaluating with a root override */
-	FMovieSceneRootOverridePath RootOverridePath;
+	UE::MovieScene::FSubSequencePath RootOverridePath;
 
 	/** Execution tokens that are used to apply animated state */
 	FMovieSceneExecutionTokens ExecutionTokens;
diff --git a/Engine/Source/Runtime/MovieScene/Private/Evaluation/MovieSceneEvaluationState.cpp b/Engine/Source/Runtime/MovieScene/Private/Evaluation/MovieSceneEvaluationState.cpp
--- a/Engine/Source/Runtime/MovieScene/Private/Evaluation/MovieSceneEvaluationState.cpp
+++ b/Engine/Source/Runtime/MovieScene/Private/Evaluation/MovieSceneEvaluationState.cpp
@@ -128,7 +128,7 @@ void FMovieSceneObjectCache::FilterObjectBindings(UObject* PredicateObject, IMov
 				UObject* Object = WeakObject.Get();
 				if (Object && Object == PredicateObject)
 				{
-					OutBindings->Add(FMovieSceneObjectBindingID(Pair.Key, SequenceID));
+					OutBindings->Add(UE::MovieScene::FFixedObjectBindingID(Pair.Key, SequenceID));
 					break;
 				}
 			}
@@ -149,7 +149,7 @@ void FMovieSceneObjectCache::FilterObjectBindings(UObject* PredicateObject, IMov
 			UObject* Object = WeakObject.Get();
 			if (Object && Object == PredicateObject)
 			{
-				OutBindings->Add(FMovieSceneObjectBindingID(DirtyBinding, SequenceID));
+				OutBindings->Add(UE::MovieScene::FFixedObjectBindingID(DirtyBinding, SequenceID));
 				break;
 			}
 		}
diff --git a/Engine/Source/Runtime/MovieScene/Private/Evaluation/MovieSceneRootOverridePath.cpp b/Engine/Source/Runtime/MovieScene/Private/Evaluation/MovieSceneRootOverridePath.cpp
--- a/Engine/Source/Runtime/MovieScene/Private/Evaluation/MovieSceneRootOverridePath.cpp
+++ b/Engine/Source/Runtime/MovieScene/Private/Evaluation/MovieSceneRootOverridePath.cpp
@@ -2,30 +2,194 @@
 
 #include "Evaluation/MovieSceneRootOverridePath.h"
 #include "Evaluation/MovieSceneSequenceHierarchy.h"
+#include "Evaluation/MovieSceneEvaluationTemplateInstance.h"
+#include "IMovieScenePlayer.h"
 
-void FMovieSceneRootOverridePath::Reset()
+#include "Algo/Find.h"
+
+namespace UE
+{
+namespace MovieScene
+{
+
+FSubSequencePath::FSubSequencePath()
+{}
+
+FSubSequencePath::FSubSequencePath(FMovieSceneSequenceID LeafID, IMovieScenePlayer& Player)
+{
+	Reset(LeafID, Player.GetEvaluationTemplate().GetHierarchy());
+}
+
+FSubSequencePath::FSubSequencePath(FMovieSceneSequenceID LeafID, const FMovieSceneSequenceHierarchy* RootHierarchy)
+{
+	Reset(LeafID, RootHierarchy);
+}
+
+FMovieSceneSequenceID FSubSequencePath::FindCommonParent(const FSubSequencePath& A, const FSubSequencePath& B)
+{
+	if (A.PathToRoot.Num() == 0 || B.PathToRoot.Num() == 0)
+	{
+		return MovieSceneSequenceID::Root;
+	}
+
+	FMovieSceneSequenceID CommonParent = MovieSceneSequenceID::Root;
+
+	int32 IndexA = A.PathToRoot.Num()-1;
+	int32 IndexB = B.PathToRoot.Num()-1;
+
+	// Keep walking into the path until we find a disprity in sequences
+	for (; IndexA >= 0 && IndexB >= 0; --IndexA, --IndexB)
+	{
+		FMovieSceneSequenceID CurrentSequenceID = A.PathToRoot[IndexA].Accumulated;
+		if (B.PathToRoot[IndexB].Accumulated == CurrentSequenceID)
+		{
+			CommonParent = CurrentSequenceID;
+		}
+		else
+		{
+			break;
+		}
+	}
+
+	return CommonParent;
+}
+
+void FSubSequencePath::Reset()
 {
-	ReverseOverrideRootPath.Reset();
+	PathToRoot.Reset();
 }
 
-void FMovieSceneRootOverridePath::Set(FMovieSceneSequenceID OverrideRootID, const FMovieSceneSequenceHierarchy* RootHierarchy)
+void FSubSequencePath::Reset(FMovieSceneSequenceID LeafID, const FMovieSceneSequenceHierarchy* RootHierarchy)
 {
-	ReverseOverrideRootPath.Reset();
+	check(LeafID != MovieSceneSequenceID::Invalid);
+
+	PathToRoot.Reset();
 
-	FMovieSceneSequenceID CurrentSequenceID = OverrideRootID;
+	FMovieSceneSequenceID CurrentSequenceID = LeafID;
 
-	check(OverrideRootID == MovieSceneSequenceID::Root || RootHierarchy != nullptr);
+	check(LeafID == MovieSceneSequenceID::Root || RootHierarchy != nullptr);
 
 	while (CurrentSequenceID != MovieSceneSequenceID::Root)
 	{
-		const FMovieSceneSequenceHierarchyNode* CurrentNode = RootHierarchy->FindNode(CurrentSequenceID);
-		const FMovieSceneSubSequenceData* OuterSubData = RootHierarchy->FindSubData(CurrentSequenceID);
+		const FMovieSceneSequenceHierarchyNode* CurrentNode  = RootHierarchy->FindNode(CurrentSequenceID);
+		const FMovieSceneSubSequenceData*       OuterSubData = RootHierarchy->FindSubData(CurrentSequenceID);
 		if (!ensureAlwaysMsgf(CurrentNode && OuterSubData, TEXT("Malformed sequence hierarchy")))
 		{
 			return;
 		}
 
-		ReverseOverrideRootPath.Add(OuterSubData->DeterministicSequenceID);
+		PathToRoot.Add(FSequenceIDPair{ OuterSubData->DeterministicSequenceID, CurrentSequenceID });
 		CurrentSequenceID = CurrentNode->ParentID;
 	}
 }
+
+bool FSubSequencePath::Contains(FMovieSceneSequenceID SequenceID) const
+{
+	return SequenceID == MovieSceneSequenceID::Root || Algo::FindBy(PathToRoot, SequenceID, &FSequenceIDPair::Accumulated) != nullptr;
+}
+
+int32 FSubSequencePath::NumGenerationsFromLeaf(FMovieSceneSequenceID SequenceID) const
+{
+	int32 Count = 0;
+
+	// Count sequence IDs until we find the sequence ID
+	// We must return 0 where SequenceID == Leaf, 1 for Immediate parents, 2 for grandparents etc
+	for ( ; Count < PathToRoot.Num(); ++Count)
+	{
+		if (PathToRoot[Count].Accumulated == SequenceID)
+		{
+			return Count;
+		}
+	}
+
+	ensureAlwaysMsgf(SequenceID == MovieSceneSequenceID::Root, TEXT("Specified SequenceID does not exist in this path"));
+	return Count;
+}
+
+int32 FSubSequencePath::NumGenerationsFromRoot(FMovieSceneSequenceID SequenceID) const
+{
+	const int32 Num = PathToRoot.Num();
+	int32 Count = 0;
+
+	// Count sequence IDs from the root (ie, the tail of the path) until we find the sequence ID
+	// We must return 0 where SequenceID == Root, 1 for children, 2 for grandchildren etc
+	for ( ; Count < Num; ++Count)
+	{
+		const int32 Index = Num - Count - 1;
+		if (PathToRoot[Index].Accumulated == SequenceID)
+		{
+			return Count;
+		}
+	}
+
+	FMovieSceneSequenceID LeafID = PathToRoot.Num() != 0 ? PathToRoot.Last().Accumulated : MovieSceneSequenceID::Root;
+	ensureAlwaysMsgf(LeafID == SequenceID, TEXT("Specified SequenceID does not exist in this path"));
+	return Count;
+}
+
+FMovieSceneSequenceID FSubSequencePath::MakeLocalSequenceID(FMovieSceneSequenceID ParentSequenceID) const
+{
+	FMovieSceneSequenceID AccumulatedID = MovieSceneSequenceID::Root;
+	for (FSequenceIDPair Pair : PathToRoot)
+	{
+		if (Pair.Accumulated == ParentSequenceID)
+		{
+			break;
+		}
+
+		AccumulatedID = AccumulatedID.AccumulateParentID(Pair.Unaccumulated);
+	}
+
+	return AccumulatedID;
+}
+
+FMovieSceneSequenceID FSubSequencePath::MakeLocalSequenceID(FMovieSceneSequenceID ParentSequenceID, FMovieSceneSequenceID TargetSequenceID) const
+{
+	FMovieSceneSequenceID AccumulatedID = MovieSceneSequenceID::Root;
+
+	int32 Index = 0;
+
+	// Skip over any that are not the target sequence ID
+	for ( ; Index < PathToRoot.Num() && TargetSequenceID != PathToRoot[Index].Accumulated; ++Index)
+	{}
+
+	for ( ; Index < PathToRoot.Num(); ++Index)
+	{
+		FSequenceIDPair Pair = PathToRoot[Index];
+
+		AccumulatedID = AccumulatedID.AccumulateParentID(Pair.Unaccumulated);
+
+		if (Pair.Accumulated == ParentSequenceID)
+		{
+			break;
+		}
+	}
+
+	return AccumulatedID;
+}
+
+void FSubSequencePath::PushGeneration(FMovieSceneSequenceID AccumulatedSequenceID, FMovieSceneSequenceID UnaccumulatedSequenceID)
+{
+	PathToRoot.Insert(FSequenceIDPair{ UnaccumulatedSequenceID, AccumulatedSequenceID }, 0);
+}
+
+void FSubSequencePath::PopTo(FMovieSceneSequenceID ParentSequenceID)
+{
+	const int32 NumToRemove = NumGenerationsFromLeaf(ParentSequenceID);
+	PopGenerations(NumToRemove);
+}
+
+void FSubSequencePath::PopGenerations(int32 NumGenerations)
+{
+	if (NumGenerations != 0)
+	{
+		check(NumGenerations <= PathToRoot.Num());
+
+		// Remove children from the head of the array
+		PathToRoot.RemoveAt(0, NumGenerations, false);
+	}
+}
+
+
+} // namespace MovieScene
+} // namespace UE
diff --git a/Engine/Source/Runtime/MovieScene/Private/MovieScene.cpp b/Engine/Source/Runtime/MovieScene/Private/MovieScene.cpp
--- a/Engine/Source/Runtime/MovieScene/Private/MovieScene.cpp
+++ b/Engine/Source/Runtime/MovieScene/Private/MovieScene.cpp
@@ -414,16 +414,18 @@ void UMovieScene::AddNewBindingTag(const FName& NewTag)
 	BindingGroups.Add(NewTag);
 }
 
-void UMovieScene::TagBinding(const FName& NewTag, FMovieSceneObjectBindingID BindingToTag)
+void UMovieScene::TagBinding(const FName& NewTag, const UE::MovieScene::FFixedObjectBindingID& BindingToTag)
 {
-	BindingGroups.FindOrAdd(NewTag).IDs.AddUnique(BindingToTag);
+	FMovieSceneObjectBindingID SerializedID = BindingToTag;
+	BindingGroups.FindOrAdd(NewTag).IDs.AddUnique(SerializedID);
 }
 
-void UMovieScene::UntagBinding(const FName& Tag, FMovieSceneObjectBindingID Binding)
+void UMovieScene::UntagBinding(const FName& Tag, const UE::MovieScene::FFixedObjectBindingID& Binding)
 {
+	FMovieSceneObjectBindingID PredicateID = Binding;
 	if (FMovieSceneObjectBindingIDs* Array = BindingGroups.Find(Tag))
 	{
-		Array->IDs.Remove(Binding);
+		Array->IDs.Remove(PredicateID);
 		if (Array->IDs.Num() == 0)
 		{
 			BindingGroups.Remove(Tag);
diff --git a/Engine/Source/Runtime/MovieScene/Private/MovieSceneBindingOverrides.cpp b/Engine/Source/Runtime/MovieScene/Private/MovieSceneBindingOverrides.cpp
--- a/Engine/Source/Runtime/MovieScene/Private/MovieSceneBindingOverrides.cpp
+++ b/Engine/Source/Runtime/MovieScene/Private/MovieSceneBindingOverrides.cpp
@@ -22,7 +22,7 @@ bool UMovieSceneBindingOverrides::LocateBoundObjects(const FGuid& InBindingId, F
 		const FMovieSceneBindingOverrideData& Data = BindingData[It.Value()];
 
 		// We have fast lookup only on GUID, so be sure to check the sequence ID before allowing overrides
-		if (Data.ObjectBindingId.GetSequenceID() != InSequenceID)
+		if (Data.ObjectBindingId.ReinterpretAsFixed().SequenceID != InSequenceID)
 		{
 			continue;
 		}
diff --git a/Engine/Source/Runtime/MovieScene/Private/MovieSceneObjectBindingID.cpp b/Engine/Source/Runtime/MovieScene/Private/MovieSceneObjectBindingID.cpp
--- a/Engine/Source/Runtime/MovieScene/Private/MovieSceneObjectBindingID.cpp
+++ b/Engine/Source/Runtime/MovieScene/Private/MovieSceneObjectBindingID.cpp
@@ -4,68 +4,185 @@
 #include "Evaluation/MovieSceneSequenceHierarchy.h"
 #include "Evaluation/MovieSceneEvaluationTemplateInstance.h"
 #include "Compilation/MovieSceneCompiledDataManager.h"
+#include "EntitySystem/MovieSceneSequenceInstance.h"
 #include "IMovieScenePlayer.h"
 
-FMovieSceneObjectBindingID FMovieSceneObjectBindingID::ResolveLocalToRoot(FMovieSceneSequenceID LocalSequenceID, IMovieScenePlayer& Player) const
+
+
+namespace UE
+{
+namespace MovieScene
+{
+
+
+FMovieSceneSequenceID ResolveExternalSequenceID(FMovieSceneSequenceID SourceSequenceID, int32 RemapSourceParentIndex, FMovieSceneSequenceID TargetSequenceID, IMovieScenePlayer& Player)
 {
-	FMovieSceneSequenceID NewSequenceID = FMovieSceneSequenceID(uint32(SequenceID));
+	ensureMsgf(RemapSourceParentIndex >= 0, TEXT("Invalid parent index specified"));
 
-	if (Space == EMovieSceneObjectBindingSpace::Local && LocalSequenceID != MovieSceneSequenceID::Root)
+	if (SourceSequenceID == MovieSceneSequenceID::Root)
 	{
-		FMovieSceneRootEvaluationTemplateInstance& Instance  = Player.GetEvaluationTemplate();
-		const FMovieSceneSequenceHierarchy&        Hierarchy = Instance.GetCompiledDataManager()->GetHierarchyChecked(Instance.GetCompiledDataID());
-
-		while (LocalSequenceID != MovieSceneSequenceID::Root)
-		{
-			const FMovieSceneSequenceHierarchyNode* CurrentNode = Hierarchy.FindNode(LocalSequenceID);
-			const FMovieSceneSubSequenceData* SubData = Hierarchy.FindSubData(LocalSequenceID);
-			if (!ensureAlwaysMsgf(CurrentNode && SubData, TEXT("Malformed sequence hierarchy")))
-			{
-				return FMovieSceneObjectBindingID(Guid, NewSequenceID);
-			}
-
-			NewSequenceID = NewSequenceID.AccumulateParentID(SubData->DeterministicSequenceID);
-			LocalSequenceID = CurrentNode->ParentID;
-		}
+		return TargetSequenceID;
 	}
 
-	return FMovieSceneObjectBindingID(Guid, NewSequenceID);
+	const FMovieSceneSequenceHierarchy* Hierarchy = Player.GetEvaluationTemplate().GetHierarchy();
+	if (!ensureMsgf(Hierarchy, TEXT("Sequence being evaluated as a sub sequence without any hierarchy. This indicates a bug with either the player or the compiled data.")))
+	{
+		return TargetSequenceID;
+	}
+
+	using namespace UE::MovieScene;
+
+	FSubSequencePath PathFromSource(SourceSequenceID, Hierarchy);
+	PathFromSource.PopGenerations(RemapSourceParentIndex);
+	return PathFromSource.ResolveChildSequenceID(TargetSequenceID);
 }
 
 
-void FMovieSceneObjectBindingID::TryAndMakeLocalTo(FMovieSceneSequenceID MakeRelativeTo, IMovieScenePlayer& Player)
+FMovieSceneSequenceID ResolveExternalSequenceID(FMovieSceneSequenceID SourceSequenceID, int32 RemapSourceParentIndex, FMovieSceneSequenceID TargetSequenceID, const FMovieSceneSequenceHierarchy* Hierarchy)
 {
-	// If it's already local or is a root sequence, we can't make it local to the desired sequence
-	if (Space == EMovieSceneObjectBindingSpace::Local || MakeRelativeTo == MovieSceneSequenceID::Root)
+	ensureMsgf(RemapSourceParentIndex >= 0, TEXT("Invalid parent index specified"));
+
+	if (SourceSequenceID == MovieSceneSequenceID::Root)
 	{
-		return;
+		return TargetSequenceID;
 	}
 
-	FMovieSceneRootEvaluationTemplateInstance& Instance  = Player.GetEvaluationTemplate();
-	const FMovieSceneSequenceHierarchy&        Hierarchy = Instance.GetCompiledDataManager()->GetHierarchyChecked(Instance.GetCompiledDataID());
+	if (!ensureMsgf(Hierarchy, TEXT("Sequence being evaluated as a sub sequence without any hierarchy. This indicates a bug with either the player or the compiled data.")))
+	{
+		return TargetSequenceID;
+	}
+
+	using namespace UE::MovieScene;
+
+	FSubSequencePath PathFromSource(SourceSequenceID, Hierarchy);
+	PathFromSource.PopGenerations(RemapSourceParentIndex);
+	return PathFromSource.ResolveChildSequenceID(TargetSequenceID);
+}
+
+
+FRelativeObjectBindingID::FRelativeObjectBindingID(FMovieSceneSequenceID SourceSequenceID, FMovieSceneSequenceID TargetSequenceID, const FGuid& TargetGuid, UMovieSceneSequence* MasterSequence)
+{
+	UMovieSceneCompiledDataManager* CompiledDataManager = UMovieSceneCompiledDataManager::GetPrecompiledData();
+
+	FMovieSceneSequenceHierarchy TempHierarchy;
+	UMovieSceneCompiledDataManager::CompileHierarchy(MasterSequence, &TempHierarchy, EMovieSceneServerClientMask::All);
+
+	ConstructInternal(SourceSequenceID, TargetSequenceID, TargetGuid, &TempHierarchy);
+}
+
+
+FRelativeObjectBindingID::FRelativeObjectBindingID(FMovieSceneSequenceID SourceSequenceID, FMovieSceneSequenceID TargetSequenceID, const FGuid& TargetGuid, IMovieScenePlayer& Player)
+{
+	ConstructInternal(SourceSequenceID, TargetSequenceID, TargetGuid, Player.GetEvaluationTemplate().GetHierarchy());
+}
+
+
+FRelativeObjectBindingID::FRelativeObjectBindingID(FMovieSceneSequenceID SourceSequenceID, FMovieSceneSequenceID TargetSequenceID, const FGuid& TargetGuid, const FMovieSceneSequenceHierarchy* Hierarchy)
+{
+	ConstructInternal(SourceSequenceID, TargetSequenceID, TargetGuid, Hierarchy);
+}
+
+void FRelativeObjectBindingID::ConstructInternal(FMovieSceneSequenceID SourceSequenceID, FMovieSceneSequenceID TargetSequenceID, const FGuid& TargetGuid, const FMovieSceneSequenceHierarchy* Hierarchy)
+{
+	FSubSequencePath PathFromOwner(SourceSequenceID, Hierarchy);
+	FSubSequencePath PathFromTarget(TargetSequenceID, Hierarchy);
+
+	FMovieSceneSequenceID CommonParent = FSubSequencePath::FindCommonParent(PathFromOwner, PathFromTarget);
+
+	Guid = TargetGuid;
+	SequenceID = PathFromTarget.MakeLocalSequenceID(CommonParent);
+	ResolveParentIndex = PathFromOwner.NumGenerationsFromLeaf(CommonParent);
+}
+
+FRelativeObjectBindingID FFixedObjectBindingID::ConvertToRelative(FMovieSceneSequenceID SourceSequenceID, IMovieScenePlayer& InPlayer) const
+{
+	return FRelativeObjectBindingID(SourceSequenceID, SequenceID, Guid, InPlayer);
+}
+
+FRelativeObjectBindingID FFixedObjectBindingID::ConvertToRelative(FMovieSceneSequenceID SourceSequenceID, const FMovieSceneSequenceHierarchy* Hierarchy) const
+{
+	return FRelativeObjectBindingID(SourceSequenceID, SequenceID, Guid, Hierarchy);
+}
+
 
-	// Find the desired sequence
-	FMovieSceneSequenceID CurrentSequenceID = FMovieSceneSequenceID(uint32(SequenceID));
-	FMovieSceneSequenceID NewSequenceID = MovieSceneSequenceID::Root;
+} // namespace MovieScene
+} // namespace UE
 
-	while (CurrentSequenceID != MovieSceneSequenceID::Root)
+
+
+
+#if WITH_EDITORONLY_DATA
+
+void FMovieSceneObjectBindingID::PostSerialize(const FArchive& Ar)
+{
+	if (Space_DEPRECATED != EMovieSceneObjectBindingSpace::Unused)
 	{
-		const FMovieSceneSequenceHierarchyNode* CurrentNode = Hierarchy.FindNode(CurrentSequenceID);
-		const FMovieSceneSubSequenceData*       SubData     = Hierarchy.FindSubData(CurrentSequenceID);
-
-		if (!ensureAlwaysMsgf(CurrentNode && SubData, TEXT("Malformed sequence hierarchy")))
-		{
-			return;
-		}
-
-		NewSequenceID = NewSequenceID.AccumulateParentID(SubData->DeterministicSequenceID);
-		CurrentSequenceID = CurrentNode->ParentID;
-
-		if (CurrentSequenceID == MakeRelativeTo)
-		{
-			Space = EMovieSceneObjectBindingSpace::Local;
-			SequenceID = int32(NewSequenceID.GetInternalValue());
-			return;
-		}
+		ResolveParentIndex = (Space_DEPRECATED == EMovieSceneObjectBindingSpace::Local) ? 0 : FixedRootSequenceParentIndex;
+		Space_DEPRECATED = EMovieSceneObjectBindingSpace::Unused;
 	}
+}
+
+#endif // WITH_EDITORONLY_DATA
+
+FMovieSceneSequenceID FMovieSceneObjectBindingID::ResolveSequenceID(FMovieSceneSequenceID LocalSequenceID, IMovieScenePlayer& Player) const
+{
+	FMovieSceneSequenceID TargetSequenceID = FMovieSceneSequenceID(SequenceID);
+
+	if (ResolveParentIndex == FixedRootSequenceParentIndex)
+	{
+		return TargetSequenceID;
+	}
+	
+	return UE::MovieScene::ResolveExternalSequenceID(LocalSequenceID, ResolveParentIndex, TargetSequenceID, Player);
+}
+
+FMovieSceneSequenceID FMovieSceneObjectBindingID::ResolveSequenceID(FMovieSceneSequenceID LocalSequenceID, const FMovieSceneSequenceHierarchy* Hierarchy) const
+{
+	FMovieSceneSequenceID TargetSequenceID = FMovieSceneSequenceID(SequenceID);
+
+	if (ResolveParentIndex == FixedRootSequenceParentIndex)
+	{
+		return TargetSequenceID;
+	}
+
+	return UE::MovieScene::ResolveExternalSequenceID(LocalSequenceID, ResolveParentIndex, TargetSequenceID, Hierarchy);
+
+}
+UE::MovieScene::FFixedObjectBindingID FMovieSceneObjectBindingID::ResolveToFixed(FMovieSceneSequenceID RuntimeSequenceID, IMovieScenePlayer& Player) const
+{
+	FMovieSceneSequenceID ThisSequenceID = FMovieSceneSequenceID(SequenceID);
+
+	if (ResolveParentIndex == FixedRootSequenceParentIndex)
+	{
+		return UE::MovieScene::FFixedObjectBindingID(Guid, ThisSequenceID);
+	}
+	else if (ensure(RuntimeSequenceID != MovieSceneSequenceID::Invalid))
+	{
+		ThisSequenceID = UE::MovieScene::ResolveExternalSequenceID(RuntimeSequenceID, ResolveParentIndex, ThisSequenceID, Player);
+
+		return UE::MovieScene::FFixedObjectBindingID(Guid, ThisSequenceID);
+	}
+
+	return UE::MovieScene::FFixedObjectBindingID(Guid, MovieSceneSequenceID::Invalid);
+}
+
+TArrayView<TWeakObjectPtr<>> FMovieSceneObjectBindingID::ResolveBoundObjects(FMovieSceneSequenceID LocalSequenceID, IMovieScenePlayer& Player) const
+{
+	if (!IsValid())
+	{
+		return TArrayView<TWeakObjectPtr<>>();
+	}
+
+	FMovieSceneSequenceID ResolvedSequenceID = ResolveSequenceID(LocalSequenceID, Player);
+	return Player.FindBoundObjects(Guid, ResolvedSequenceID);
+}
+
+TArrayView<TWeakObjectPtr<>> FMovieSceneObjectBindingID::ResolveBoundObjects(const UE::MovieScene::FSequenceInstance& SequenceInstance) const
+{
+	if (!IsValid())
+	{
+		return TArrayView<TWeakObjectPtr<>>();
+	}
+
+	return ResolveBoundObjects(SequenceInstance.GetSequenceID(), *SequenceInstance.GetPlayer());
 }
\ No newline at end of file
diff --git a/Engine/Source/Runtime/MovieScene/Private/MovieSceneSequencePlayer.cpp b/Engine/Source/Runtime/MovieScene/Private/MovieSceneSequencePlayer.cpp
--- a/Engine/Source/Runtime/MovieScene/Private/MovieSceneSequencePlayer.cpp
+++ b/Engine/Source/Runtime/MovieScene/Private/MovieSceneSequencePlayer.cpp
@@ -1058,7 +1058,8 @@ void UMovieSceneSequencePlayer::SetTimeController(TSharedPtr<FMovieSceneTimeCont
 TArray<UObject*> UMovieSceneSequencePlayer::GetBoundObjects(FMovieSceneObjectBindingID ObjectBinding)
 {
 	TArray<UObject*> Objects;
-	for (TWeakObjectPtr<> WeakObject : FindBoundObjects(ObjectBinding.GetGuid(), ObjectBinding.GetSequenceID()))
+
+	for (TWeakObjectPtr<> WeakObject : ObjectBinding.ResolveBoundObjects(MovieSceneSequenceID::Root, *this))
 	{
 		if (UObject* Object = WeakObject.Get())
 		{
diff --git a/Engine/Source/Runtime/MovieScene/Private/Sections/MovieSceneSubSection.cpp b/Engine/Source/Runtime/MovieScene/Private/Sections/MovieSceneSubSection.cpp
--- a/Engine/Source/Runtime/MovieScene/Private/Sections/MovieSceneSubSection.cpp
+++ b/Engine/Source/Runtime/MovieScene/Private/Sections/MovieSceneSubSection.cpp
@@ -10,6 +10,7 @@
 #include "Evaluation/MovieSceneEvaluationTemplate.h"
 #include "Misc/FrameRate.h"
 #include "Logging/MessageLog.h"
+#include "Evaluation/MovieSceneRootOverridePath.h"
 
 TWeakObjectPtr<UMovieSceneSubSection> UMovieSceneSubSection::TheRecordingSection;
 
@@ -424,14 +425,13 @@ void UMovieSceneSubSection::BuildDefaultSubSectionComponents(UMovieSceneEntitySy
 	if (Easing.GetEaseInDuration() > 0 || Easing.GetEaseOutDuration() > 0)
 	{
 		FBuiltInComponentTypes* Components = FBuiltInComponentTypes::Get();
-		FInstanceRegistry* InstanceRegistry = EntityLinker->GetInstanceRegistry();
 
-		const FSequenceInstance& Instance = InstanceRegistry->GetInstance(Params.Sequence.InstanceHandle);
-		const FMovieSceneSequenceID SubSequenceID = GetSequenceID();
+		const FSubSequencePath      PathToRoot         = EntityLinker->GetInstanceRegistry()->GetInstance(Params.Sequence.InstanceHandle).GetSubSequencePath();
+		const FMovieSceneSequenceID ResolvedSequenceID = PathToRoot.ResolveChildSequenceID(this->GetSequenceID());
 
 		OutImportedEntity->AddBuilder(
-			FEntityBuilder()
-				.AddConditional(Components->HierarchicalEasingProvider, SubSequenceID, SubSequenceID.IsValid()));
+			FEntityBuilder().Add(Components->HierarchicalEasingProvider, ResolvedSequenceID)
+		);
 	}
 }
 
diff --git a/Engine/Source/Runtime/MovieScene/Public/Compilation/MovieSceneCompiledDataManager.h b/Engine/Source/Runtime/MovieScene/Public/Compilation/MovieSceneCompiledDataManager.h
--- a/Engine/Source/Runtime/MovieScene/Public/Compilation/MovieSceneCompiledDataManager.h
+++ b/Engine/Source/Runtime/MovieScene/Public/Compilation/MovieSceneCompiledDataManager.h
@@ -25,6 +25,14 @@ struct FMovieSceneGatheredCompilerData;
 class UMovieSceneSubTrack;
 class UMovieSceneSequence;
 
+namespace UE
+{
+namespace MovieScene
+{
+	struct FSubSequencePath;
+}
+}
+
 enum class EMovieSceneSequenceCompilerMask : uint8
 {
 	Hierarchy               = 1 << 0,
@@ -245,9 +253,9 @@ private:
 
 	static bool CompileHierarchy(UMovieSceneSequence* Sequence, const FGatherParameters& Params, FMovieSceneSequenceHierarchy* InOutHierarchy);
 
-	static bool CompileHierarchyImpl(UMovieSceneSequence* Sequence, const FGatherParameters& Params, const FMovieSceneEvaluationOperand& Operand, FMovieSceneRootOverridePath* RootPath, FMovieSceneSequenceHierarchy* InOutHierarchy);
+	static bool CompileHierarchyImpl(UMovieSceneSequence* Sequence, const FGatherParameters& Params, const FMovieSceneEvaluationOperand& Operand, UE::MovieScene::FSubSequencePath* RootPath, FMovieSceneSequenceHierarchy* InOutHierarchy);
 
-	static bool CompileSubTrackHierarchy(UMovieSceneSubTrack* SubTrack, const FGatherParameters& Params, const FMovieSceneEvaluationOperand& Operand, FMovieSceneRootOverridePath* RootPath, FMovieSceneSequenceHierarchy* InOutHierarchy);
+	static bool CompileSubTrackHierarchy(UMovieSceneSubTrack* SubTrack, const FGatherParameters& Params, const FMovieSceneEvaluationOperand& Operand, UE::MovieScene::FSubSequencePath* RootPath, FMovieSceneSequenceHierarchy* InOutHierarchy);
 
 	void CompileTrackTemplateField(FMovieSceneCompiledDataEntry* OutEntry, const FMovieSceneSequenceHierarchy& Hierarchy, FMovieSceneGatheredCompilerData* InCompilerData);
 
diff --git a/Engine/Source/Runtime/MovieScene/Public/EntitySystem/MovieSceneSequenceInstance.h b/Engine/Source/Runtime/MovieScene/Public/EntitySystem/MovieSceneSequenceInstance.h
--- a/Engine/Source/Runtime/MovieScene/Public/EntitySystem/MovieSceneSequenceInstance.h
+++ b/Engine/Source/Runtime/MovieScene/Public/EntitySystem/MovieSceneSequenceInstance.h
@@ -23,6 +23,7 @@ namespace MovieScene
 struct ISequenceUpdater;
 struct FSequenceInstance;
 struct FCompiledDataVolatilityManager;
+struct FSubSequencePath;
 
 /**
  * A sequence instance represents a specific instance of a currently playing sequence, either as a top-level sequence in an IMovieScenePlayer, or as a sub sequence.
@@ -179,7 +180,6 @@ public:
 	 */
 	FMovieSceneEntityID FindEntity(UObject* Owner, uint32 EntityID) const;
 
-
 	/**
 	 * Retrieve the legacy evaluator for this sequence, if it is available (may return nullptr)
 	 */
@@ -188,6 +188,11 @@ public:
 		return LegacyEvaluator.Get();
 	}
 
+	/**
+	 * Retrieve a path for this sequence instance back to the root
+	 */
+	FSubSequencePath GetSubSequencePath() const;
+
 public:
 
 	/**
diff --git a/Engine/Source/Runtime/MovieScene/Public/Evaluation/MovieSceneRootOverridePath.h b/Engine/Source/Runtime/MovieScene/Public/Evaluation/MovieSceneRootOverridePath.h
--- a/Engine/Source/Runtime/MovieScene/Public/Evaluation/MovieSceneRootOverridePath.h
+++ b/Engine/Source/Runtime/MovieScene/Public/Evaluation/MovieSceneRootOverridePath.h
@@ -4,30 +4,70 @@
 
 #include "MovieSceneSequenceID.h"
 
+class IMovieScenePlayer;
+
 struct FMovieSceneSequenceHierarchy;
 
+namespace UE
+{
+namespace MovieScene
+{
+
+
 /**
  * A path of unaccumulated sequence IDs ordered from child->parent->grandparent that is used to generate unique sequenceIDs for inner sequences
  * Optimized for Remap rather than Push/Pop by keeping sequence IDs child-parent order (the order they are required for remapping)
  */
-struct MOVIESCENE_API FMovieSceneRootOverridePath
+struct MOVIESCENE_API FSubSequencePath
 {
+	/**
+	 * Default construction to a root path
+	 */
+	FSubSequencePath();
+
+	/**
+	 * Set up this path from a specific sequence ID that points to a particular sequence in the specified hierarchy
+	 *
+	 * @param LeafID 			ID of the child-most sequence to include in this path
+	 * @param Player 			Player from which to retrieve the hierarchy
+	 */
+	explicit FSubSequencePath(FMovieSceneSequenceID LeafID, IMovieScenePlayer& Player);
+
+	/**
+	 * Set up this path from a specific sequence ID that points to a particular sequence in the specified hierarchy
+	 *
+	 * @param LeafID 			ID of the child-most sequence to include in this path
+	 * @param RootHierarchy 	Hierarchy to get sequence IDs from
+	 */
+	explicit FSubSequencePath(FMovieSceneSequenceID LeafID, const FMovieSceneSequenceHierarchy* RootHierarchy);
+
+
+	/**
+	 * Find the first parent sequence ID that is common to both A and B
+	 *
+	 * @param A					The first sequence path
+	 * @param B					The second sequence path
+	 * @return The leaf-most parent common to both paths, or MovieSceneSequenceID::Root
+	 */
+	static FMovieSceneSequenceID FindCommonParent(const FSubSequencePath& A, const FSubSequencePath& B);
+
+
 	/**
 	 * Remap the specified sequence ID based on the currently evaluating sequence path, to the Root
 	 *
 	 * @param SequenceID			The sequence ID to find a template for
 	 * @return Pointer to a template instance, or nullptr if the ID was not found
 	 */
-	FORCEINLINE_DEBUGGABLE FMovieSceneSequenceID Remap(FMovieSceneSequenceID SequenceID) const
+	/*FORCEINLINE_DEBUGGABLE*/ FMovieSceneSequenceID ResolveChildSequenceID(FMovieSceneSequenceID SequenceID) const
 	{
-		if (LIKELY(!ReverseOverrideRootPath.Num()))
+		if (LIKELY(PathToRoot.Num() == 0))
 		{
 			return SequenceID;
 		}
 
-		for (FMovieSceneSequenceID Parent : ReverseOverrideRootPath)
+		for (FSequenceIDPair Parent : PathToRoot)
 		{
-			SequenceID = SequenceID.AccumulateParentID(Parent);
+			SequenceID = SequenceID.AccumulateParentID(Parent.Unaccumulated);
 		}
 		return SequenceID;
 	}
@@ -37,35 +77,70 @@ struct MOVIESCENE_API FMovieSceneRootOverridePath
 	 */
 	void Reset();
 
+
 	/**
 	 * Set up this path from a specific sequence ID that points to a particular sequence in the specified hierarchy
 	 *
 	 * @param LeafID 			ID of the child-most sequence to include in this path
-	 * @param RootHierarchy 	Hierarchy to get unaccumulated sequence IDs from
+	 * @param RootHierarchy 	Hierarchy to get sequence IDs from
 	 */
-	void Set(FMovieSceneSequenceID LeafID, const FMovieSceneSequenceHierarchy* RootHierarchy);
+	void Reset(FMovieSceneSequenceID LeafID, const FMovieSceneSequenceHierarchy* RootHierarchy);
+
 
 	/**
-	 * Push a new child sequence ID into this path
+	 * Check whether this path contains the specified sequence ID
 	 *
-	 * @param InUnaccumulatedSequenceID		The sequece ID of the child sequence within its outer parent
+	 * @param SequenceID		ID of the sequence to check for
+	 * @return true if this path contains the sequence ID (or SequenceID == MovieSceneSequenceID::Root), false otherwise
 	 */
-	void Push(FMovieSceneSequenceID InUnaccumulatedSequenceID)
-	{
-		ReverseOverrideRootPath.Insert(InUnaccumulatedSequenceID, 0);
-	}
+	bool Contains(FMovieSceneSequenceID SequenceID) const;
 
 
 	/**
-	 * Pop the child-most sequence ID from this path
+	 * Return the number of generations between this path's leaf node, and the specified sequence ID
+	 *
+	 * @param SequenceID		ID of the parent sequence to count generations to
+	 * @return The number of generations between the two nodes. (ie, 0 where SequenceID == Leaf, 1 for Immediate parents, 2 for grandparents etc)
 	 */
-	void Pop()
-	{
-		ReverseOverrideRootPath.RemoveAt(0, 1, false);
-	}
+	int32 NumGenerationsFromLeaf(FMovieSceneSequenceID SequenceID) const;
+
+
+	/**
+	 * Return the number of generations between the root and the specified sequence ID
+	 *
+	 * @param SequenceID		ID of the child sequence to count generations to
+	 * @return The number of generations between the two nodes. (ie, 0 where SequenceID == Root, 1 for Immediate children, 2 for grandchildren)
+	 */
+	int32 NumGenerationsFromRoot(FMovieSceneSequenceID SequenceID) const;
+
+	/**
+	 * 
+	 * 
+	 * @param 
+	 * @return
+	 */
+	FMovieSceneSequenceID MakeLocalSequenceID(FMovieSceneSequenceID ParentSequenceID) const;
+
+	FMovieSceneSequenceID MakeLocalSequenceID(FMovieSceneSequenceID ParentSequenceID, FMovieSceneSequenceID TargetSequenceID) const;
+
+	void PushGeneration(FMovieSceneSequenceID AccumulatedSequenceID, FMovieSceneSequenceID UnaccumulatedSequenceID);
+
+	void PopTo(FMovieSceneSequenceID ParentSequenceID);
+
+	void PopGenerations(int32 NumGenerations);
 
 private:
 
+	struct FSequenceIDPair
+	{
+		FMovieSceneSequenceID Unaccumulated;
+		FMovieSceneSequenceID Accumulated;
+	};
+
 	/** A reverse path of deterministic sequence IDs required to accumulate from local -> root */
-	TArray<FMovieSceneSequenceID, TInlineAllocator<8>> ReverseOverrideRootPath;
-};
\ No newline at end of file
+	TArray<FSequenceIDPair, TInlineAllocator<8>> PathToRoot;
+};
+
+
+} // namespace MovieScene
+} // namespace UE
diff --git a/Engine/Source/Runtime/MovieScene/Public/MovieScene.h b/Engine/Source/Runtime/MovieScene/Public/MovieScene.h
--- a/Engine/Source/Runtime/MovieScene/Public/MovieScene.h
+++ b/Engine/Source/Runtime/MovieScene/Public/MovieScene.h
@@ -1101,12 +1101,12 @@ public:
 	/*
 	 * Tag the specified binding ID with the specified name
 	 */
-	void TagBinding(const FName& NewTag, FMovieSceneObjectBindingID BindingToTag);
+	void TagBinding(const FName& NewTag, const UE::MovieScene::FFixedObjectBindingID& BindingToTag);
 
 	/*
 	 * Remove a tag from the specified object binding
 	 */
-	void UntagBinding(const FName& Tag, FMovieSceneObjectBindingID Binding);
+	void UntagBinding(const FName& Tag, const UE::MovieScene::FFixedObjectBindingID& Binding);
 
 	/*
 	 * Remove the specified tag from any binding and forget about it completely
diff --git a/Engine/Source/Runtime/MovieScene/Public/MovieSceneObjectBindingID.h b/Engine/Source/Runtime/MovieScene/Public/MovieSceneObjectBindingID.h
--- a/Engine/Source/Runtime/MovieScene/Public/MovieSceneObjectBindingID.h
+++ b/Engine/Source/Runtime/MovieScene/Public/MovieSceneObjectBindingID.h
@@ -12,6 +12,155 @@
 struct FMovieSceneSequenceHierarchy;
 
 class IMovieScenePlayer;
+class UMovieSceneSequence;
+
+
+namespace UE
+{
+namespace MovieScene
+{
+
+
+struct FSequenceInstance;
+
+
+/**
+ * A binding ID whose target is resolved relative to the owner of this ID (either internal or external to that sequence through a parent or grandparent)
+ * Should be used for asignment to any serialized FMovieSceneObjectBindingID
+ * Represented by a local object binding ID and a parent index to which it is relative
+ */
+struct FRelativeObjectBindingID
+{
+	FRelativeObjectBindingID(const FGuid& InGuid)
+		: Guid(InGuid)
+		, SequenceID(MovieSceneSequenceID::Root)
+		, ResolveParentIndex(0)
+	{}
+
+	FRelativeObjectBindingID(const FGuid& InGuid, FMovieSceneSequenceID InSequenceID)
+		: Guid(InGuid)
+		, SequenceID(InSequenceID)
+		, ResolveParentIndex(0)
+	{}
+
+	FRelativeObjectBindingID(const FGuid& InGuid, FMovieSceneSequenceID InSequenceID, int32 InParentIndex)
+		: Guid(InGuid)
+		, SequenceID(InSequenceID)
+		, ResolveParentIndex(InParentIndex)
+	{}
+
+	/**
+	 * Construction from a master sequence asset
+	 *
+	 * @param SourceSequenceID      Absolute sequence ID within MasterSequence that represents the start of the path (ie, the sequence that this binding ID is serialized within)
+	 * @param TargetSequenceID      Absolute sequence ID within MasterSequence that holds the target object
+	 * @param TargetGuid            The GUID of the object binding within TargetSequenceID that represents the object at runtime
+	 * @param MasterSequence        The master sequence asset
+	 */
+	MOVIESCENE_API FRelativeObjectBindingID(FMovieSceneSequenceID SourceSequenceID, FMovieSceneSequenceID TargetSequenceID, const FGuid& TargetGuid, UMovieSceneSequence* MasterSequence);
+
+	/**
+	 * Construction from a master sequence Hierarchy
+	 *
+	 * @param SourceSequenceID      Absolute sequence ID within Hierarchy that represents the start of the path (ie, the sequence that this binding ID is serialized within)
+	 * @param TargetSequenceID      Absolute sequence ID within Hierarchy that holds the target object
+	 * @param TargetGuid            The GUID of the object binding within TargetSequenceID that represents the object at runtime
+	 * @param Hierarchy             The hierachy to use for computation of the relative path. Many be nullptr where Source and Targets are both MovieSceneSequenceID::Root
+	 */
+	MOVIESCENE_API FRelativeObjectBindingID(FMovieSceneSequenceID SourceSequenceID, FMovieSceneSequenceID TargetSequenceID, const FGuid& TargetGuid, const FMovieSceneSequenceHierarchy* Hierarchy);
+
+	/**
+	 * Construction from a master sequence Player
+	 *
+	 * @param SourceSequenceID      Absolute sequence ID within Player's master sequence hierarchy that represents the start of the path (ie, the sequence that this binding ID is serialized within)
+	 * @param TargetSequenceID      Absolute sequence ID within Player's master sequence hierarchy that holds the target object
+	 * @param TargetGuid            The GUID of the object binding within TargetSequenceID that represents the object at runtime
+	 * @param Player                The movie scene player that is currently playing the master sequence
+	 */
+	MOVIESCENE_API FRelativeObjectBindingID(FMovieSceneSequenceID SourceSequenceID, FMovieSceneSequenceID TargetSequenceID, const FGuid& TargetGuid, IMovieScenePlayer& Player);
+
+	/** GUID of the Object Binding (ie, FMovieSceneBinding::GetObjectGuid) */
+	FGuid Guid;
+
+	/** The sequence ID that the object binding is found within - relative to nth parent specified by ResolveParentIndex */
+	FMovieSceneSequenceID SequenceID;
+
+	/** The parent index denoting where to resolve this binding from */
+	int32 ResolveParentIndex;
+
+private:
+
+	void ConstructInternal(FMovieSceneSequenceID SourceSequenceID, FMovieSceneSequenceID TargetSequenceID, const FGuid& TargetGuid, const FMovieSceneSequenceHierarchy* Hierarchy);
+};
+
+
+/**
+ * A binding ID that is only resolveable through a fixed root sequence.
+ * Use of this type should be reserved for editor code dealing with a specific sequence hierarchy
+ * where hashing or comparison is required.
+ * The SequenceID member is always an absolute sequence ID
+ */
+struct FFixedObjectBindingID
+{
+	FFixedObjectBindingID()
+	{}
+
+	FFixedObjectBindingID(const FGuid& InGuid, FMovieSceneSequenceID InSequenceID)
+		: Guid(InGuid)
+		, SequenceID(InSequenceID)
+	{}
+
+public:
+
+	friend uint32 GetTypeHash(const FFixedObjectBindingID& A)
+	{
+		return GetTypeHash(A.Guid) ^ GetTypeHash(A.SequenceID);
+	}
+
+	friend bool operator==(const FFixedObjectBindingID& A, const FFixedObjectBindingID& B)
+	{
+		return A.Guid == B.Guid && A.SequenceID == B.SequenceID;
+	}
+
+	friend bool operator!=(const FFixedObjectBindingID& A, const FFixedObjectBindingID& B)
+	{
+		return A.Guid != B.Guid || A.SequenceID != B.SequenceID;
+	}
+
+public:
+
+	/**
+	 * Convert this fixed binding ID to one that is resolved relative to the specified Sequence.
+	 *
+	 * @param SourceSequenceID    The sequence ID to make this fixed ID relative to
+	 * @param InPlayer            The active movie scene player to retrieve a master sequence hierarchy from
+	 * @return A new binding ID that will resolve relative to the specified sequence
+	 */
+	MOVIESCENE_API FRelativeObjectBindingID ConvertToRelative(FMovieSceneSequenceID SourceSequenceID, IMovieScenePlayer& InPlayer) const;
+
+	/**
+	 * Convert this fixed binding ID to one that is resolved relative to the specified Sequence.
+	 *
+	 * @param SourceSequenceID    The sequence ID to make this fixed ID relative to
+	 * @param Hierarchy           The hierarchy that represents the master sequence that SourceSequenceID is contained within
+	 * @return A new binding ID that will resolve relative to the specified sequence
+	 */
+	MOVIESCENE_API FRelativeObjectBindingID ConvertToRelative(FMovieSceneSequenceID SourceSequenceID, const FMovieSceneSequenceHierarchy* Hierarchy) const;
+
+public:
+
+	/** GUID of the Object Binding (ie, FMovieSceneBinding::GetObjectGuid) */
+	FGuid Guid;
+
+	/** The sequence ID that the object binding is found within */
+	FMovieSceneSequenceID SequenceID;
+};
+
+
+} // namespace MovieScene
+} // namespace UE
+
+
 
 /** Enumeration specifying how a movie scene object binding ID relates to the sequence */
 UENUM()
@@ -21,35 +170,108 @@ enum class EMovieSceneObjectBindingSpace : uint8
 	Local,
 	/** The object binding sequence ID resolves from the root of the sequence */
 	Root,
+
+	/** Default value for any newly-constructed bindings */
+	Unused UMETA(Hidden),
 };
 
 
-/** Persistent identifier to a specific object binding within a sequence hierarchy. */
+/**
+ * Persistent identifier to a specific object binding within a sequence hierarchy.
+ * 
+ * Binding IDs come in 3 flavors with Local and External being preferred as they are reslilient towards sequences being authored in isolation or included in other master sequences:
+ *     Local: (ResolveParentIndex == 0) SequenceID relates to _this sequence's_ local hierarchy; represents an object binding within the same sequence as the ID is resolved, or inside one of its sub-sequences. Sequence ID must be remapped at runtime.
+ *     External: (ResolveParentIndex > 0) SequenceID is local to the parent sequence of this one denoted by the parent index (ie, 1 = parent, 2 = grandparent etc). Sequence ID must be remapped at runtime.
+ *     Fixed: Represents a binding anywhere in the sequence; always resolved from the root sequence.
+ *
+ * Fixed bindings will break if the sequence is evaluated inside a different master sequence.
+ */
 USTRUCT(BlueprintType, meta=(HasNativeMake))
 struct FMovieSceneObjectBindingID
 {
 	GENERATED_BODY()
 
-	/** Default construction to an invalid object binding ID */
+	/**
+	 * Default construction to a root empty. Only to be used for UObject serialization and default UObject construction
+	 */
 	FMovieSceneObjectBindingID()
-		: SequenceID(int32(MovieSceneSequenceID::Root.GetInternalValue())), Space(EMovieSceneObjectBindingSpace::Root)
+		// Defaults purposefully match the equivalent of a legacy root binding
+		// such that any pre-existing ::Root space binding upgrades to ResolveParentIndex = -1
+		// and ::Local space bindings will get upgraded in PostSerialize
+		: SequenceID(int32(MovieSceneSequenceID::Root.GetInternalValue()))
+		, ResolveParentIndex(FixedRootSequenceParentIndex)
 	{
+		LegacyInit();
 	}
 
-	/** Construction from an object binding guid, and the specific sequence instance ID in which it resides */
-	FMovieSceneObjectBindingID(const FGuid& InGuid, FMovieSceneSequenceID InSequenceID, EMovieSceneObjectBindingSpace InSpace = EMovieSceneObjectBindingSpace::Root)
-		: SequenceID(int32(InSequenceID.GetInternalValue())), Space(InSpace), Guid(InGuid)
+	/**
+	 * Construction from an external object binding ID
+	 */
+	FMovieSceneObjectBindingID(const UE::MovieScene::FRelativeObjectBindingID& InID)
+		: Guid(InID.Guid)
+		, SequenceID(InID.SequenceID.GetInternalValue())
+		, ResolveParentIndex(InID.ResolveParentIndex)
+	{
+		LegacyInit();
+	}
+
+	/**
+	 * Construction from a fixed object binding ID
+	 */
+	FMovieSceneObjectBindingID(const UE::MovieScene::FFixedObjectBindingID& InID)
+		: Guid(InID.Guid)
+		, SequenceID(InID.SequenceID.GetInternalValue())
+		, ResolveParentIndex(FixedRootSequenceParentIndex)
+	{
+		LegacyInit();
+	}
+
+
+public:
+
+
+	/**
+	 * Assignment from an external object binding ID
+	 */
+	FMovieSceneObjectBindingID& operator=(const UE::MovieScene::FRelativeObjectBindingID& InID)
 	{
+		Guid = InID.Guid;
+		SequenceID = InID.SequenceID.GetInternalValue();
+		ResolveParentIndex = InID.ResolveParentIndex;
+		return *this;
 	}
 
-	FMovieSceneObjectBindingID& operator=(const FMovieSceneObjectBindingID& InObjectBindingID)
+	/**
+	 * Assignment from a fixed object binding ID
+	 */
+	FMovieSceneObjectBindingID& operator=(const UE::MovieScene::FFixedObjectBindingID& InID)
 	{
-		Guid = InObjectBindingID.GetGuid();
-		SequenceID = int32(InObjectBindingID.GetSequenceID().GetInternalValue());
-		Space = InObjectBindingID.GetBindingSpace();
+		Guid = InID.Guid;
+		SequenceID = InID.SequenceID.GetInternalValue();
+		ResolveParentIndex = FixedRootSequenceParentIndex;
 		return *this;
 	}
 
+public:
+
+	friend uint32 GetTypeHash(const FMovieSceneObjectBindingID& A)
+	{
+		return GetTypeHash(A.Guid) ^ GetTypeHash(A.SequenceID);
+	}
+
+	friend bool operator==(const FMovieSceneObjectBindingID& A, const FMovieSceneObjectBindingID& B)
+	{
+		return A.Guid == B.Guid && A.SequenceID == B.SequenceID && A.ResolveParentIndex == B.ResolveParentIndex;
+	}
+
+	friend bool operator!=(const FMovieSceneObjectBindingID& A, const FMovieSceneObjectBindingID& B)
+	{
+		return A.Guid != B.Guid || A.SequenceID != B.SequenceID ||  A.ResolveParentIndex != B.ResolveParentIndex;
+	}
+
+
+public:
+
 	/**
 	 * Check whether this object binding ID has been set to something valied
 	 * @note: does not imply that the ID resolves to a valid object
@@ -60,11 +282,11 @@ struct FMovieSceneObjectBindingID
 	}
 
 	/**
-	 * Access the identifier for the sequence in which the object binding resides
+	 * Check whether this binding is fixed to the root sequence - these bindings should be avoided or fixed up where possible to ensure portability
 	 */
-	FMovieSceneSequenceID GetSequenceID() const
+	bool IsFixedBinding() const
 	{
-		return FMovieSceneSequenceID(uint32(SequenceID));
+		return Guid.IsValid() && ResolveParentIndex == FixedRootSequenceParentIndex;
 	}
 
 	/**
@@ -74,61 +296,133 @@ struct FMovieSceneObjectBindingID
 	{
 		return Guid;
 	}
-	
+
 	/**
 	 * Set the guid that identifies the object binding within the sequence
 	 */
 	void SetGuid(const FGuid& InGuid)
 	{
+#if WITH_EDITORONLY_DATA
+		ensureMsgf(Space_DEPRECATED == EMovieSceneObjectBindingSpace::Unused, TEXT("Assigning a GUID for a binding which has no type specifier - consider using Local / External / Fixed instead"));
+#endif
 		Guid = InGuid;
 	}
-	
+
 	/**
-	 * Access how this binding's sequence ID relates to the master sequence
+	 * Get the relative sequence ID for this binding.
+	 * @note: The meaning of this sequence ID is dependent upon ResolveParentIndex so should only be used where the context is known
 	 */
-	EMovieSceneObjectBindingSpace GetBindingSpace() const
+	FMovieSceneSequenceID GetRelativeSequenceID() const
 	{
-		return Space;
+		return FMovieSceneSequenceID(SequenceID);
 	}
 
+
 	/**
-	 * Resolve this binding ID from a local binding to be accessible from the root, by treating the specified local sequence ID as this binding's root
+	 * Resolve this binding ID to a fixed object binding ID resolvable from the root master sequence
+	 *
+	 * @param SourceSequenceID      The source sequence ID that owns this binding ID
+	 * @param Player                The movie scene player that is currently playing the master sequence
+	 * @return A fixed binding ID whose SequenceID relates to the master sequence hierarchy
 	 */
-	MOVIESCENE_API FMovieSceneObjectBindingID ResolveLocalToRoot(FMovieSceneSequenceID LocalSequenceID, IMovieScenePlayer& Player) const;
+	MOVIESCENE_API UE::MovieScene::FFixedObjectBindingID ResolveToFixed(FMovieSceneSequenceID SourceSequenceID, IMovieScenePlayer& Player) const;
+
+
 
 	/**
-	 * Attempt to reinterpret this binding ID so that it is local to the specified sequence ID. Ineffectual if this binding is Space == Root
+	 * Resolve the sequence ID for this binding for the root master sequence
+	 *
+	 * @param SourceSequenceID      The source sequence ID that owns this binding ID
+	 * @param Player                The movie scene player that is currently playing the master sequence
+	 * @return The sequence ID that holds the target object
 	 */
-	MOVIESCENE_API void TryAndMakeLocalTo(FMovieSceneSequenceID MakeRelativeTo, IMovieScenePlayer& Player);
+	MOVIESCENE_API FMovieSceneSequenceID ResolveSequenceID(FMovieSceneSequenceID SourceSequenceID, IMovieScenePlayer& Player) const;
 
-public:
 
-	friend uint32 GetTypeHash(const FMovieSceneObjectBindingID& A)
-	{
-		return GetTypeHash(A.Guid) ^ GetTypeHash(A.SequenceID);
-	}
+	/**
+	 * Resolve the sequence ID for this binding for the root master sequence
+	 *
+	 * @param SourceSequenceID      The source sequence ID that owns this binding ID
+	 * @param Hierarchy             The hierachy to use for computation of the relative path. Many be nullptr where Source and Targets are both MovieSceneSequenceID::Root
+	 * @return The sequence ID that holds the target object
+	 */
+	MOVIESCENE_API FMovieSceneSequenceID ResolveSequenceID(FMovieSceneSequenceID SourceSequenceID, const FMovieSceneSequenceHierarchy* Hierarchy) const;
 
-	friend bool operator==(const FMovieSceneObjectBindingID& A, const FMovieSceneObjectBindingID& B)
-	{
-		return A.Guid == B.Guid && A.SequenceID == B.SequenceID;
-	}
 
-	friend bool operator!=(const FMovieSceneObjectBindingID& A, const FMovieSceneObjectBindingID& B)
+	/**
+	 * Resolve all the bound objects for this binding ID
+	 *
+	 * @param SourceSequenceID      The source sequence ID that owns this binding ID
+	 * @param Hierarchy             The hierachy to use for computation of the relative path. Many be nullptr where Source and Targets are both MovieSceneSequenceID::Root
+	 * @return An array view of all bound objects
+	 */
+	MOVIESCENE_API TArrayView<TWeakObjectPtr<>> ResolveBoundObjects(FMovieSceneSequenceID SourceSequenceID, IMovieScenePlayer& Player) const;
+
+
+	/**
+	 * Resolve all the bound objects for this binding ID
+	 *
+	 * @param SequenceInstance      The sequence instance for the sub sequence that is resolving the binding
+	 * @return An array view of all bound objects
+	 */
+	MOVIESCENE_API TArrayView<TWeakObjectPtr<>> ResolveBoundObjects(const UE::MovieScene::FSequenceInstance& SequenceInstance) const;
+
+
+	/**
+	 * Reinterpret this binding ID as a fixed binding. Should only be used when ResolveParentIndex == -1, although this is not enforced.
+	 */
+	UE::MovieScene::FFixedObjectBindingID ReinterpretAsFixed() const
 	{
-		return A.Guid != B.Guid || A.SequenceID != B.SequenceID;
+		return UE::MovieScene::FFixedObjectBindingID(Guid, FMovieSceneSequenceID(SequenceID));
 	}
 
+public:
+
+#if WITH_EDITORONLY_DATA
+	void PostSerialize(const FArchive& Ar);
+#endif
+
 private:
 
+	static constexpr int32 FixedRootSequenceParentIndex = -1;
+
+	/** Identifier for the object binding within the sequence */
+	UPROPERTY(EditAnywhere, Category="Binding")
+	FGuid Guid;
+
 	/** Sequence ID stored as an int32 so that it can be used in the blueprint VM */
 	UPROPERTY()
 	int32 SequenceID;
 
-	/** The binding's resolution space */
+	/** The parent index to resolve the sequence ID from. 0 signifies the sequence this binding ID is serialized within, -1 signifies the master sequence (previously EMovieSceneObjectBindingSpace::Root) */
 	UPROPERTY()
-	EMovieSceneObjectBindingSpace Space;
+	int32 ResolveParentIndex;
 
-	/** Identifier for the object binding within the sequence */
-	UPROPERTY(EditAnywhere, Category="Binding")
-	FGuid Guid;
+private:
+	void LegacyInit()
+	{
+#if WITH_EDITORONLY_DATA
+		Space_DEPRECATED = EMovieSceneObjectBindingSpace::Unused;
+#endif
+	}
+
+
+#if WITH_EDITORONLY_DATA
+
+	/** DEPRECATED: The binding's resolution space */
+	UPROPERTY()
+	EMovieSceneObjectBindingSpace Space_DEPRECATED;
+
+#endif // WITH_EDITORONLY_DATA
+};
+
+
+template<> struct TStructOpsTypeTraits<FMovieSceneObjectBindingID> : public TStructOpsTypeTraitsBase2<FMovieSceneObjectBindingID>
+{
+#if WITH_EDITORONLY_DATA
+	enum 
+	{
+		WithPostSerialize = true,
+	};
+#endif
 };
\ No newline at end of file
diff --git a/Engine/Source/Runtime/MovieSceneTracks/Private/Evaluation/MovieScene3DPathTemplate.cpp b/Engine/Source/Runtime/MovieSceneTracks/Private/Evaluation/MovieScene3DPathTemplate.cpp
--- a/Engine/Source/Runtime/MovieSceneTracks/Private/Evaluation/MovieScene3DPathTemplate.cpp
+++ b/Engine/Source/Runtime/MovieSceneTracks/Private/Evaluation/MovieScene3DPathTemplate.cpp
@@ -125,35 +125,18 @@ struct F3DPathExecutionToken
 	{
 		MOVIESCENE_DETAILED_SCOPE_CYCLE_COUNTER(MovieSceneEval_PathTrack_TokenExecute)
 
-		FMovieSceneSequenceID SequenceID = Operand.SequenceID;
-		if (PathBindingID.GetSequenceID().IsValid())
-		{
-			// Ensure that this ID is resolvable from the root, based on the current local sequence ID
-			FMovieSceneObjectBindingID RootBindingID = PathBindingID.ResolveLocalToRoot(SequenceID, Player);
-			SequenceID = RootBindingID.GetSequenceID();
-		}
-
-		// If the transform is set, otherwise use the bound actor's transform
-		FMovieSceneEvaluationOperand PathOperand(SequenceID, PathBindingID.GetGuid());
-		
-		TArrayView<TWeakObjectPtr<>> Objects = Player.FindBoundObjects(PathOperand);
-		if (!Objects.Num())
-		{
-			return;
-		}
-
-		// Only ever deal with one parent
-		UObject* ParentObject = Objects[0].Get();
-		AActor* Actor = Cast<AActor>(ParentObject);
-
 		USplineComponent* SplineComponent = nullptr;
 
-		for (UActorComponent* Component : Actor->GetComponents())
+		// Find the first spline component on any of the objects resolved through the path binding ID
+		for (TWeakObjectPtr<> PathTarget : PathBindingID.ResolveBoundObjects(Operand.SequenceID, Player))
 		{
-			SplineComponent = Cast<USplineComponent>(Component);
-			if (SplineComponent)
+			if (AActor* Actor = Cast<AActor>(PathTarget.Get()))
 			{
-				break;
+				SplineComponent = Actor->FindComponentByClass<USplineComponent>();
+				if (SplineComponent)
+				{
+					break;
+				}
 			}
 		}
 
diff --git a/Engine/Source/Runtime/MovieSceneTracks/Private/Evaluation/MovieSceneActorReferenceTemplate.cpp b/Engine/Source/Runtime/MovieSceneTracks/Private/Evaluation/MovieSceneActorReferenceTemplate.cpp
--- a/Engine/Source/Runtime/MovieSceneTracks/Private/Evaluation/MovieSceneActorReferenceTemplate.cpp
+++ b/Engine/Source/Runtime/MovieSceneTracks/Private/Evaluation/MovieSceneActorReferenceTemplate.cpp
@@ -13,9 +13,7 @@ namespace PropertyTemplate
 	template<>
 	UObject* ConvertFromIntermediateType<UObject*, FMovieSceneObjectBindingID>(const FMovieSceneObjectBindingID& InObjectBinding, const FMovieSceneEvaluationOperand& Operand, FPersistentEvaluationData& PersistentData, IMovieScenePlayer& Player)
 	{
-		FMovieSceneObjectBindingID ResolvedID = InObjectBinding.ResolveLocalToRoot(Operand.SequenceID, Player);
-
-		for (TWeakObjectPtr<>& WeakObject : Player.FindBoundObjects(ResolvedID.GetGuid(), ResolvedID.GetSequenceID()))
+		for (TWeakObjectPtr<> WeakObject : InObjectBinding.ResolveBoundObjects(Operand.SequenceID, Player))
 		{
 			if (AActor* Actor = Cast<AActor>(WeakObject.Get()))
 			{
diff --git a/Engine/Source/Runtime/MovieSceneTracks/Private/Evaluation/MovieSceneAudioTemplate.cpp b/Engine/Source/Runtime/MovieSceneTracks/Private/Evaluation/MovieSceneAudioTemplate.cpp
--- a/Engine/Source/Runtime/MovieSceneTracks/Private/Evaluation/MovieSceneAudioTemplate.cpp
+++ b/Engine/Source/Runtime/MovieSceneTracks/Private/Evaluation/MovieSceneAudioTemplate.cpp
@@ -288,18 +288,8 @@ struct FAudioSectionExecutionToken : IMovieSceneExecutionToken
 			FMovieSceneObjectBindingID AttachBindingID = AttachKey.Object;
 			if (AttachBindingID.IsValid())
 			{
-				FMovieSceneSequenceID SequenceID = Operand.SequenceID;
-				if (AttachBindingID.GetSequenceID().IsValid())
-				{
-					// Ensure that this ID is resolvable from the root, based on the current local sequence ID
-					FMovieSceneObjectBindingID RootBindingID = AttachBindingID.ResolveLocalToRoot(SequenceID, Player);
-					SequenceID = RootBindingID.GetSequenceID();
-				}
-
 				// If the transform is set, otherwise use the bound actor's transform
-				FMovieSceneEvaluationOperand ObjectOperand(SequenceID, AttachBindingID.GetGuid());
-
-				for (TWeakObjectPtr<> WeakObject : Player.FindBoundObjects(ObjectOperand))
+				for (TWeakObjectPtr<> WeakObject : AttachBindingID.ResolveBoundObjects(Operand.SequenceID, Player))
 				{
 					AActor* AttachActor = Cast<AActor>(WeakObject.Get());
 					if (AttachActor)
diff --git a/Engine/Source/Runtime/MovieSceneTracks/Private/Sections/MovieScene3DConstraintSection.cpp b/Engine/Source/Runtime/MovieSceneTracks/Private/Sections/MovieScene3DConstraintSection.cpp
--- a/Engine/Source/Runtime/MovieSceneTracks/Private/Sections/MovieScene3DConstraintSection.cpp
+++ b/Engine/Source/Runtime/MovieSceneTracks/Private/Sections/MovieScene3DConstraintSection.cpp
@@ -10,22 +10,6 @@ UMovieScene3DConstraintSection::UMovieScene3DConstraintSection( const FObjectIni
 	bSupportsInfiniteRange = true;
 }
 
-void UMovieScene3DConstraintSection::SetConstraintId(const FGuid& InConstraintId, const FMovieSceneSequenceID& SequenceID)
-{
-	if (TryModify())
-	{
-		SetConstraintBindingID(FMovieSceneObjectBindingID(InConstraintId, SequenceID, EMovieSceneObjectBindingSpace::Root)); //use root if specifying binding
-	}
-}
-
-void UMovieScene3DConstraintSection::SetConstraintId(const FGuid& InConstraintId)
-{
-	if (TryModify())
-	{
-		SetConstraintBindingID(FMovieSceneObjectBindingID(InConstraintId, MovieSceneSequenceID::Root, EMovieSceneObjectBindingSpace::Local));
-	}
-}
-
 void UMovieScene3DConstraintSection::OnBindingsUpdated(const TMap<FGuid, FGuid>& OldGuidToNewGuidMap)
 {
 	if (OldGuidToNewGuidMap.Contains(ConstraintBindingID.GetGuid()))
@@ -49,7 +33,7 @@ void UMovieScene3DConstraintSection::PostLoad()
 	{
 		if (!ConstraintBindingID.IsValid())
 		{
-			ConstraintBindingID = FMovieSceneObjectBindingID(ConstraintId_DEPRECATED, MovieSceneSequenceID::Root, EMovieSceneObjectBindingSpace::Local);
+			ConstraintBindingID = UE::MovieScene::FRelativeObjectBindingID(ConstraintId_DEPRECATED);
 		}
 		ConstraintId_DEPRECATED.Invalidate();
 	}
diff --git a/Engine/Source/Runtime/MovieSceneTracks/Private/Sections/MovieSceneActorReferenceSection.cpp b/Engine/Source/Runtime/MovieSceneTracks/Private/Sections/MovieSceneActorReferenceSection.cpp
--- a/Engine/Source/Runtime/MovieSceneTracks/Private/Sections/MovieSceneActorReferenceSection.cpp
+++ b/Engine/Source/Runtime/MovieSceneTracks/Private/Sections/MovieSceneActorReferenceSection.cpp
@@ -118,7 +118,7 @@ void UMovieSceneActorReferenceSection::PostLoad()
 
 		if (Guids.IsValidIndex(ActorGuidIndexCurve_DEPRECATED.GetDefaultValue()))
 		{
-			FMovieSceneObjectBindingID DefaultValue(Guids[ActorGuidIndexCurve_DEPRECATED.GetDefaultValue()], MovieSceneSequenceID::Root, EMovieSceneObjectBindingSpace::Local);
+			FMovieSceneObjectBindingID DefaultValue = UE::MovieScene::FRelativeObjectBindingID(Guids[ActorGuidIndexCurve_DEPRECATED.GetDefaultValue()]);
 			ActorReferenceData.SetDefault(DefaultValue);
 		}
 
@@ -126,7 +126,7 @@ void UMovieSceneActorReferenceSection::PostLoad()
 		{
 			if (ensure(Guids.IsValidIndex(It->Value)))
 			{
-				FMovieSceneObjectBindingID BindingID(Guids[It->Value], MovieSceneSequenceID::Root, EMovieSceneObjectBindingSpace::Local);
+				FMovieSceneObjectBindingID BindingID = UE::MovieScene::FRelativeObjectBindingID(Guids[It->Value]);
 				ActorReferenceData.UpgradeLegacyTime(this, It->Time, BindingID);
 			}
 		}
diff --git a/Engine/Source/Runtime/MovieSceneTracks/Private/Sections/MovieSceneCameraCutSection.cpp b/Engine/Source/Runtime/MovieSceneTracks/Private/Sections/MovieSceneCameraCutSection.cpp
--- a/Engine/Source/Runtime/MovieSceneTracks/Private/Sections/MovieSceneCameraCutSection.cpp
+++ b/Engine/Source/Runtime/MovieSceneTracks/Private/Sections/MovieSceneCameraCutSection.cpp
@@ -58,7 +58,7 @@ void UMovieSceneCameraCutSection::PostLoad()
 	{
 		if (!CameraBindingID.IsValid())
 		{
-			CameraBindingID = FMovieSceneObjectBindingID(CameraGuid_DEPRECATED, MovieSceneSequenceID::Root, EMovieSceneObjectBindingSpace::Local);
+			CameraBindingID = UE::MovieScene::FRelativeObjectBindingID(CameraGuid_DEPRECATED);
 		}
 		CameraGuid_DEPRECATED.Invalidate();
 	}
@@ -66,16 +66,9 @@ void UMovieSceneCameraCutSection::PostLoad()
 
 UCameraComponent* UMovieSceneCameraCutSection::GetFirstCamera(IMovieScenePlayer& Player, FMovieSceneSequenceID SequenceID) const
 {
-	if (CameraBindingID.GetSequenceID().IsValid())
+	for (TWeakObjectPtr<> WeakObject : CameraBindingID.ResolveBoundObjects(SequenceID, Player))
 	{
-		// Ensure that this ID is resolvable from the root, based on the current local sequence ID
-		FMovieSceneObjectBindingID RootBindingID = CameraBindingID.ResolveLocalToRoot(SequenceID, Player);
-		SequenceID = RootBindingID.GetSequenceID();
-	}
-
-	for (TWeakObjectPtr<>& WeakObject : Player.FindBoundObjects(CameraBindingID.GetGuid(), SequenceID))
-	{
-		if (UObject* Object = WeakObject .Get())
+		if (UObject* Object = WeakObject.Get())
 		{
 			UCameraComponent* Camera = MovieSceneHelpers::CameraComponentFromRuntimeObject(Object);
 			if (Camera)
diff --git a/Engine/Source/Runtime/MovieSceneTracks/Private/Systems/MovieSceneComponentAttachmentSystem.cpp b/Engine/Source/Runtime/MovieSceneTracks/Private/Systems/MovieSceneComponentAttachmentSystem.cpp
--- a/Engine/Source/Runtime/MovieSceneTracks/Private/Systems/MovieSceneComponentAttachmentSystem.cpp
+++ b/Engine/Source/Runtime/MovieSceneTracks/Private/Systems/MovieSceneComponentAttachmentSystem.cpp
@@ -57,18 +57,9 @@ struct FInitializeAttachParentsTask
 
 	void ForEachEntity(UE::MovieScene::FInstanceHandle InstanceHandle, const FMovieSceneObjectBindingID& BindingID, const UE::MovieScene::FAttachmentComponent& AttachComponent, USceneComponent*& OutAttachedParent)
 	{
-		const FSequenceInstance* TargetInstance = &InstanceRegistry->GetInstance(InstanceHandle);
+		const FSequenceInstance& TargetInstance = InstanceRegistry->GetInstance(InstanceHandle);
 
-		IMovieScenePlayer* Player = TargetInstance->GetPlayer();
-
-		FMovieSceneSequenceID ResolvedSequenceID = BindingID.GetSequenceID();
-		if (!TargetInstance->IsRootSequence())
-		{
-			ResolvedSequenceID = BindingID.ResolveLocalToRoot(TargetInstance->GetSequenceID(), *Player).GetSequenceID();
-			TargetInstance = &InstanceRegistry->GetInstance(TargetInstance->GetRootInstanceHandle());
-		}
-
-		for (TWeakObjectPtr<> WeakObject : Player->FindBoundObjects(BindingID.GetGuid(), ResolvedSequenceID))
+		for (TWeakObjectPtr<> WeakObject : BindingID.ResolveBoundObjects(TargetInstance))
 		{
 			if (AActor* ParentActor = Cast<AActor>(WeakObject.Get()))
 			{
diff --git a/Engine/Source/Runtime/MovieSceneTracks/Private/Systems/WeightAndEasingEvaluatorSystem.cpp b/Engine/Source/Runtime/MovieSceneTracks/Private/Systems/WeightAndEasingEvaluatorSystem.cpp
--- a/Engine/Source/Runtime/MovieSceneTracks/Private/Systems/WeightAndEasingEvaluatorSystem.cpp
+++ b/Engine/Source/Runtime/MovieSceneTracks/Private/Systems/WeightAndEasingEvaluatorSystem.cpp
@@ -49,13 +49,14 @@ void UMovieSceneHierarchicalEasingInstantiatorSystem::OnRun(FSystemTaskPrerequis
 
 		for (int32 Index = 0; Index < Allocation->Num(); ++Index)
 		{
-			const FSequenceInstance& Instance = InstanceRegistry->GetInstance(InstanceHandles[Index]);
-			const FMovieSceneSequenceID SubSequenceID = HierarchicalEasingProviders[Index];
+			const FSequenceInstance* RootInstance = &InstanceRegistry->GetInstance(InstanceHandles[Index]);
+			if (!RootInstance->IsRootSequence())
+			{
+				RootInstance = &InstanceRegistry->GetInstance(RootInstance->GetRootInstanceHandle());
+			}
 
-			// Before we look up the sub-sequence, we need to convert the local ID into an absolute one.
-			const FMovieSceneObjectBindingID LocalSubSequenceBinding(FGuid(), SubSequenceID, EMovieSceneObjectBindingSpace::Local);
-			const FMovieSceneObjectBindingID RootedSubSequenceBinding = LocalSubSequenceBinding.ResolveLocalToRoot(Instance.GetSequenceID(), *Instance.GetPlayer());
-			const FInstanceHandle SubSequenceHandle = Instance.FindSubInstance(RootedSubSequenceBinding.GetSequenceID());
+			const FMovieSceneSequenceID SequenceID = HierarchicalEasingProviders[Index];
+			const FInstanceHandle SubSequenceHandle = RootInstance->FindSubInstance(SequenceID);
 
 			// We use instance handles here because sequence IDs by themselves are only unique to a single hierarchy 
 			// of sequences. If a root sequence is playing twice at the same time, there will be 2 sequence instances
diff --git a/Engine/Source/Runtime/MovieSceneTracks/Private/TrackInstances/MovieSceneCameraCutTrackInstance.cpp b/Engine/Source/Runtime/MovieSceneTracks/Private/TrackInstances/MovieSceneCameraCutTrackInstance.cpp
--- a/Engine/Source/Runtime/MovieSceneTracks/Private/TrackInstances/MovieSceneCameraCutTrackInstance.cpp
+++ b/Engine/Source/Runtime/MovieSceneTracks/Private/TrackInstances/MovieSceneCameraCutTrackInstance.cpp
@@ -107,10 +107,7 @@ namespace MovieScene
 	{
 		static UObject* FindBoundObject(FMovieSceneObjectBindingID BindingID, FMovieSceneSequenceID SequenceID, IMovieScenePlayer& Player)
 		{
-			FMovieSceneObjectBindingID ResolvedID = BindingID.ResolveLocalToRoot(SequenceID, Player);
-
-			FMovieSceneEvaluationOperand Operand(ResolvedID.GetSequenceID(), BindingID.GetGuid());
-			TArrayView<TWeakObjectPtr<>> Objects = Player.FindBoundObjects(Operand);
+			TArrayView<TWeakObjectPtr<>> Objects = BindingID.ResolveBoundObjects(SequenceID, Player);
 			if (Objects.Num() > 0)
 			{
 				return Objects[0].Get();
diff --git a/Engine/Source/Runtime/MovieSceneTracks/Public/Sections/MovieScene3DConstraintSection.h b/Engine/Source/Runtime/MovieSceneTracks/Public/Sections/MovieScene3DConstraintSection.h
--- a/Engine/Source/Runtime/MovieSceneTracks/Public/Sections/MovieScene3DConstraintSection.h
+++ b/Engine/Source/Runtime/MovieSceneTracks/Public/Sections/MovieScene3DConstraintSection.h
@@ -22,12 +22,6 @@ class UMovieScene3DConstraintSection
 
 public:
 
-	/** Sets the constraint id for this section */
-	virtual void SetConstraintId(const FGuid& InId);
-
-	/** Sets the constraint id for this section based off of the pasted in sequence id */
-	virtual void SetConstraintId(const FGuid& InConstraintId, const FMovieSceneSequenceID& SequenceID);
-
 	/** Gets the constraint binding for this Constraint section */
 	UFUNCTION(BlueprintCallable, Category = "Movie Scene Section")
 	const FMovieSceneObjectBindingID& GetConstraintBindingID() const
diff --git a/Engine/Source/Runtime/MovieSceneTracks/Public/Sections/MovieSceneCameraCutSection.h b/Engine/Source/Runtime/MovieSceneTracks/Public/Sections/MovieSceneCameraCutSection.h
--- a/Engine/Source/Runtime/MovieSceneTracks/Public/Sections/MovieSceneCameraCutSection.h
+++ b/Engine/Source/Runtime/MovieSceneTracks/Public/Sections/MovieSceneCameraCutSection.h
@@ -34,7 +34,7 @@ public:
 	/** Sets the camera binding for this CameraCut section. Evaluates from the sequence binding ID */
 	void SetCameraGuid(const FGuid& InGuid)
 	{
-		SetCameraBindingID(FMovieSceneObjectBindingID(InGuid, MovieSceneSequenceID::Root, EMovieSceneObjectBindingSpace::Local));
+		SetCameraBindingID(UE::MovieScene::FRelativeObjectBindingID(InGuid));
 	}
 
 	/** Gets the camera binding for this CameraCut section */
