Sequencer:  Need to to Calc Auto Tangent if we when adding a key when we already have one at that time.
#trivial

diff --git a/Engine/Source/Runtime/MovieScene/Private/Channels/MovieSceneFloatChannel.cpp b/Engine/Source/Runtime/MovieScene/Private/Channels/MovieSceneFloatChannel.cpp
--- a/Engine/Source/Runtime/MovieScene/Private/Channels/MovieSceneFloatChannel.cpp
+++ b/Engine/Source/Runtime/MovieScene/Private/Channels/MovieSceneFloatChannel.cpp
@@ -1041,6 +1041,7 @@ FKeyHandle AddKeyToChannel(FMovieSceneFloatChannel* Channel, FFrameNumber InFram
 	{
 		FMovieSceneFloatValue& Value = ChannelData.GetValues()[ExistingIndex]; //-V758
 		Value.Value = InValue;
+		Channel->AutoSetTangents();
 	}
 	else switch (Interpolation)
 	{
