Sequencer: Tidied up component locking and access patterns

  - Client code ultilizing ForEachAllocation task callbacks no longer need to call Resolve to access component data - this is already done in the task internals and the lock is maintained correctly for the duration of the task
  - Added Read/WriteComponent methods to FEntityAllocation to make custom component access more robust - these functions now return a scope-locked wrapper around the component data

diff --git a/Engine/Plugins/MovieScene/TemplateSequence/Source/TemplateSequence/Private/Section/TemplateSequenceSection.cpp b/Engine/Plugins/MovieScene/TemplateSequence/Source/TemplateSequence/Private/Section/TemplateSequenceSection.cpp
--- a/Engine/Plugins/MovieScene/TemplateSequence/Source/TemplateSequence/Private/Section/TemplateSequenceSection.cpp
+++ b/Engine/Plugins/MovieScene/TemplateSequence/Source/TemplateSequence/Private/Section/TemplateSequenceSection.cpp
@@ -4,6 +4,8 @@
 #include "TemplateSequence.h"
 #include "Systems/TemplateSequenceSystem.h"
 #include "Evaluation/MovieSceneRootOverridePath.h"
+#include "EntitySystem/MovieSceneSequenceInstance.h"
+#include "EntitySystem/MovieSceneEntitySystemLinker.h"
 
 UTemplateSequenceSection::UTemplateSequenceSection(const FObjectInitializer& ObjInitializer)
 	: Super(ObjInitializer)
diff --git a/Engine/Plugins/MovieScene/TemplateSequence/Source/TemplateSequence/Private/SequenceCameraShake.cpp b/Engine/Plugins/MovieScene/TemplateSequence/Source/TemplateSequence/Private/SequenceCameraShake.cpp
--- a/Engine/Plugins/MovieScene/TemplateSequence/Source/TemplateSequence/Private/SequenceCameraShake.cpp
+++ b/Engine/Plugins/MovieScene/TemplateSequence/Source/TemplateSequence/Private/SequenceCameraShake.cpp
@@ -13,6 +13,9 @@
 #include "GameFramework/WorldSettings.h"
 #include "MovieSceneFwd.h"
 #include "MovieSceneTimeHelpers.h"
+#include "EntitySystem/MovieSceneEntitySystemTask.h"
+#include "EntitySystem/MovieScenePropertySystemTypes.inl"
+#include "EntitySystem/MovieSceneEntitySystemLinker.h"
 #include "MovieSceneTracksComponentTypes.h"
 
 #if !IS_MONOLITHIC
diff --git a/Engine/Plugins/MovieScene/TemplateSequence/Source/TemplateSequence/Private/Systems/TemplateSequenceSystem.cpp b/Engine/Plugins/MovieScene/TemplateSequence/Source/TemplateSequence/Private/Systems/TemplateSequenceSystem.cpp
--- a/Engine/Plugins/MovieScene/TemplateSequence/Source/TemplateSequence/Private/Systems/TemplateSequenceSystem.cpp
+++ b/Engine/Plugins/MovieScene/TemplateSequence/Source/TemplateSequence/Private/Systems/TemplateSequenceSystem.cpp
@@ -6,8 +6,12 @@
 #include "TemplateSequence.h"
 #include "Compilation/MovieSceneCompiledDataManager.h"
 #include "Evaluation/MovieSceneEvaluationTemplateInstance.h"
+#include "EntitySystem/MovieSceneEntitySystemLinker.h"
 #include "Sections/TemplateSequenceSection.h"
 #include "EntitySystem/MovieSceneSpawnablesSystem.h"
+#include "EntitySystem/MovieSceneEntitySystemTask.h"
+#include "EntitySystem/MovieSceneEntitySystemLinker.h"
+#include "EntitySystem/MovieSceneEntityFactoryTemplates.h"
 
 namespace UE
 {
@@ -76,19 +80,17 @@ void UTemplateSequenceSystem::OnRun(FSystemTaskPrerequisites& InPrerequisites, F
 	FInstanceRegistry* InstanceRegistry  = Linker->GetInstanceRegistry();
 
 	auto SetupTeardownBindingOverrides = [BuiltInComponents, InstanceRegistry](
-			const FEntityAllocation* Allocation,
-			TRead<FInstanceHandle> InstanceHandleAccessor,
-			TRead<FGuid> ObjectBindingIDAccessor,
-			TRead<FTemplateSequenceComponentData> TemplateSequenceDataAccessor)
+			FEntityAllocationIteratorItem AllocationItem,
+			TRead<FInstanceHandle> InstanceHandles,
+			TRead<FGuid> ObjectBindingIDs,
+			TRead<FTemplateSequenceComponentData> TemplateSequenceDatas)
 	{
-		const bool bHasNeedsLink = Allocation->HasComponent(BuiltInComponents->Tags.NeedsLink);
-		const bool bHasNeedsUnlink = Allocation->HasComponent(BuiltInComponents->Tags.NeedsUnlink);
+		const FComponentMask& Mask = AllocationItem.GetAllocationType();
+		const bool bHasNeedsLink   = Mask.Contains(BuiltInComponents->Tags.NeedsLink);
+		const bool bHasNeedsUnlink = Mask.Contains(BuiltInComponents->Tags.NeedsUnlink);
 
-		TArrayView<const FInstanceHandle> InstanceHandles = InstanceHandleAccessor.ResolveAsArray(Allocation);
-		TArrayView<const FGuid> ObjectBindingIDs = ObjectBindingIDAccessor.ResolveAsArray(Allocation);
-		TArrayView<const FTemplateSequenceComponentData> TemplateSequenceDatas = TemplateSequenceDataAccessor.ResolveAsArray(Allocation);
-
-		for (int32 Index = 0; Index < Allocation->Num(); ++Index)
+		const int32 Num = AllocationItem.GetAllocation()->Num();
+		for (int32 Index = 0; Index < Num; ++Index)
 		{
 			const FSequenceInstance& SequenceInstance = InstanceRegistry->GetInstance(InstanceHandles[Index]);
 			const FGuid& ObjectBindingID = ObjectBindingIDs[Index];
diff --git a/Engine/Source/Runtime/MovieScene/Private/EntitySystem/BuiltInComponentTypes.cpp b/Engine/Source/Runtime/MovieScene/Private/EntitySystem/BuiltInComponentTypes.cpp
--- a/Engine/Source/Runtime/MovieScene/Private/EntitySystem/BuiltInComponentTypes.cpp
+++ b/Engine/Source/Runtime/MovieScene/Private/EntitySystem/BuiltInComponentTypes.cpp
@@ -5,6 +5,8 @@
 #include "EntitySystem/MovieSceneEntityManager.h"
 #include "EntitySystem/MovieSceneEntitySystemLinker.h"
 #include "EntitySystem/MovieSceneComponentRegistry.h"
+#include "EntitySystem/MovieScenePropertyBinding.h"
+#include "EntitySystem/MovieSceneEntityFactoryTemplates.h"
 
 namespace UE
 {
@@ -157,8 +159,8 @@ FBuiltInComponentTypes::FBuiltInComponentTypes()
 
 			virtual void Run(const FEntityRange& ChildRange, const FEntityAllocation* ParentAllocation, TArrayView<const int32> ParentAllocationOffsets) override
 			{
-				TArrayView<const FMovieSceneEntityID> ParentIDs = ParentAllocation->GetEntityIDs();
-				TArrayView<FMovieSceneEntityID>       ChildComponents = GetChildComponents(ChildRange.Allocation);
+				TArrayView<const FMovieSceneEntityID> ParentIDs       = ParentAllocation->GetEntityIDs();
+				TComponentWriter<FMovieSceneEntityID> ChildComponents = GetChildComponents(ChildRange.Allocation);
 
 				for (int32 Index = 0; Index < ChildRange.Num; ++Index)
 				{
diff --git a/Engine/Source/Runtime/MovieScene/Private/EntitySystem/MovieSceneBoundObjectInstantiator.cpp b/Engine/Source/Runtime/MovieScene/Private/EntitySystem/MovieSceneBoundObjectInstantiator.cpp
--- a/Engine/Source/Runtime/MovieScene/Private/EntitySystem/MovieSceneBoundObjectInstantiator.cpp
+++ b/Engine/Source/Runtime/MovieScene/Private/EntitySystem/MovieSceneBoundObjectInstantiator.cpp
@@ -5,10 +5,10 @@
 #include "EntitySystem/MovieSceneEntitySystemTask.h"
 #include "EntitySystem/MovieSceneEntityManager.h"
 #include "EntitySystem/BuiltInComponentTypes.h"
-#include "EntitySystem/MovieSceneEntityFactory.h"
 #include "EntitySystem/MovieSceneInstanceRegistry.h"
 #include "EntitySystem/MovieSceneComponentRegistry.h"
 #include "EntitySystem/MovieSceneEntitySystemLinker.h"
+#include "EntitySystem/MovieSceneEntityFactoryTemplates.h"
 
 #include "IMovieScenePlayer.h"
 
diff --git a/Engine/Source/Runtime/MovieScene/Private/EntitySystem/MovieSceneBoundSceneComponentInstantiator.cpp b/Engine/Source/Runtime/MovieScene/Private/EntitySystem/MovieSceneBoundSceneComponentInstantiator.cpp
--- a/Engine/Source/Runtime/MovieScene/Private/EntitySystem/MovieSceneBoundSceneComponentInstantiator.cpp
+++ b/Engine/Source/Runtime/MovieScene/Private/EntitySystem/MovieSceneBoundSceneComponentInstantiator.cpp
@@ -4,7 +4,7 @@
 #include "EntitySystem/MovieSceneEntityInstantiatorSystem.h"
 #include "EntitySystem/MovieSceneEntitySystemTask.h"
 #include "EntitySystem/MovieSceneEntityManager.h"
-#include "EntitySystem/MovieSceneEntityFactory.h"
+#include "EntitySystem/MovieSceneEntityFactoryTemplates.h"
 #include "EntitySystem/BuiltInComponentTypes.h"
 #include "EntitySystem/MovieSceneInstanceRegistry.h"
 #include "EntitySystem/MovieSceneEntitySystemLinker.h"
diff --git a/Engine/Source/Runtime/MovieScene/Private/EntitySystem/MovieSceneComponentAccessors.cpp b/Engine/Source/Runtime/MovieScene/Private/EntitySystem/MovieSceneComponentAccessors.cpp
--- a/Engine/Source/Runtime/MovieScene/Private/EntitySystem/MovieSceneComponentAccessors.cpp
+++ b/Engine/Source/Runtime/MovieScene/Private/EntitySystem/MovieSceneComponentAccessors.cpp
@@ -12,17 +12,17 @@ namespace MovieScene
 
 #if UE_MOVIESCENE_ENTITY_DEBUG
 
-void AccessorToString(const FRead* In, FEntityManager* EntityManager, FString& OutString)
+void AccessorToString(const FReadAccess* In, FEntityManager* EntityManager, FString& OutString)
 {
 	const FComponentTypeInfo& ComponentTypeInfo = EntityManager->GetComponents()->GetComponentTypeChecked(In->ComponentType);
 	OutString += FString::Printf(TEXT("\n\tRead: %s %s"), *ComponentTypeInfo.DebugInfo->DebugName, ComponentTypeInfo.DebugInfo->DebugTypeName);
 }
-void AccessorToString(const FWrite* In, FEntityManager* EntityManager, FString& OutString)
+void AccessorToString(const FWriteAccess* In, FEntityManager* EntityManager, FString& OutString)
 {
 	const FComponentTypeInfo& ComponentTypeInfo = EntityManager->GetComponents()->GetComponentTypeChecked(In->ComponentType);
 	OutString += FString::Printf(TEXT("\n\tWrite: %s %s"), *ComponentTypeInfo.DebugInfo->DebugName, ComponentTypeInfo.DebugInfo->DebugTypeName);
 }
-void AccessorToString(const FReadOptional* In, FEntityManager* EntityManager, FString& OutString)
+void AccessorToString(const FOptionalReadAccess* In, FEntityManager* EntityManager, FString& OutString)
 {
 	if (In->ComponentType)
 	{
@@ -30,7 +30,7 @@ void AccessorToString(const FReadOptional* In, FEntityManager* EntityManager, FS
 		OutString += FString::Printf(TEXT("\n\tRead (Optional): %s %s"), *ComponentTypeInfo.DebugInfo->DebugName, ComponentTypeInfo.DebugInfo->DebugTypeName);
 	}
 }
-void AccessorToString(const FWriteOptional* In, FEntityManager* EntityManager, FString& OutString)
+void AccessorToString(const FOptionalWriteAccess* In, FEntityManager* EntityManager, FString& OutString)
 {
 	if (In->ComponentType)
 	{
@@ -38,11 +38,11 @@ void AccessorToString(const FWriteOptional* In, FEntityManager* EntityManager, F
 		OutString += FString::Printf(TEXT("\n\tWrite (Optional): %s %s"), *ComponentTypeInfo.DebugInfo->DebugName, ComponentTypeInfo.DebugInfo->DebugTypeName);
 	}
 }
-void AccessorToString(const FReadEntityIDs*, FEntityManager* EntityManager, FString& OutString)
+void AccessorToString(const FEntityIDAccess*, FEntityManager* EntityManager, FString& OutString)
 {
 	OutString += TEXT("\n\tRead: Entity IDs");
 }
-void OneOfAccessorToString(const FReadOptional* In, FEntityManager* EntityManager, FString& OutString)
+void OneOfAccessorToString(const FOptionalReadAccess* In, FEntityManager* EntityManager, FString& OutString)
 {
 	if (In->ComponentType)
 	{
diff --git a/Engine/Source/Runtime/MovieScene/Private/EntitySystem/MovieSceneEntityFactory.cpp b/Engine/Source/Runtime/MovieScene/Private/EntitySystem/MovieSceneEntityFactory.cpp
--- a/Engine/Source/Runtime/MovieScene/Private/EntitySystem/MovieSceneEntityFactory.cpp
+++ b/Engine/Source/Runtime/MovieScene/Private/EntitySystem/MovieSceneEntityFactory.cpp
@@ -6,6 +6,7 @@
 #include "EntitySystem/MovieSceneEntityManager.h"
 #include "EntitySystem/BuiltInComponentTypes.h"
 #include "EntitySystem/MovieSceneEntitySystemTask.h"
+#include "EntitySystem/MovieSceneEntityFactoryTemplates.h"
 
 namespace UE
 {
@@ -94,8 +95,10 @@ void FObjectFactoryBatch::InitializeAllocation(UMovieSceneEntitySystemLinker* Li
 
 	TArrayView<const FMovieSceneEntityID> ParentIDs = ParentAllocation->GetEntityIDs();
 
+	FEntityAllocationWriteContext WriteContext = FEntityAllocationWriteContext::NewAllocation();
+
 	int32 Index = GetCurrentIndex();
-	for (TEntityPtr<const FMovieSceneEntityID, const FMovieSceneEntityID, UObject*> Tuple : FEntityTaskBuilder().ReadEntityIDs().Read(ParentEntity).Write(BoundObject).IterateRange(InChildEntityRange))
+	for (TEntityPtr<const FMovieSceneEntityID, const FMovieSceneEntityID, UObject*> Tuple : FEntityTaskBuilder().ReadEntityIDs().Read(ParentEntity).Write(BoundObject).IterateRange(InChildEntityRange, WriteContext))
 	{
 		FMovieSceneEntityID Parent = Tuple.Get<1>();
 		FMovieSceneEntityID Child = Tuple.Get<0>();
@@ -125,15 +128,12 @@ FBoundObjectTask::FBoundObjectTask(UMovieSceneEntitySystemLinker* InLinker)
 	: Linker(InLinker)
 {}
 
-void FBoundObjectTask::ForEachAllocation(const FEntityAllocation* Allocation, FReadEntityIDs EntityIDAccessor, TRead<FInstanceHandle> InstanceAccessor, TRead<FGuid> ObjectBindingAccessor)
+void FBoundObjectTask::ForEachAllocation(const FEntityAllocation* Allocation, FReadEntityIDs EntityIDs, TRead<FInstanceHandle> Instances, TRead<FGuid> ObjectBindings)
 {
 	FObjectFactoryBatch& Batch = AddBatch(Allocation);
 	Batch.StaleEntitiesToPreserve = &StaleEntitiesToPreserve;
 
 	const int32 Num = Allocation->Num();
-	const FMovieSceneEntityID* EntityIDs      = Allocation->GetRawEntityIDs();
-	const FInstanceHandle*     Instances      = InstanceAccessor.Resolve(Allocation);
-	const FGuid*               ObjectBindings = ObjectBindingAccessor.Resolve(Allocation);
 
 	FInstanceRegistry* InstanceRegistry = Linker->GetInstanceRegistry();
 
@@ -152,7 +152,7 @@ void FBoundObjectTask::ForEachAllocation(const FEntityAllocation* Allocation, FR
 		for (int32 ChildIndex = StartNum; ChildIndex < EntitiesToDiscard.Num(); ++ChildIndex)
 		{
 			FMovieSceneEntityID ChildID = EntitiesToDiscard[ChildIndex];
-			TComponentPtr<UObject* const> ObjectPtr = Linker->EntityManager.ReadComponent(ChildID,  BoundObjectComponent);
+			TOptionalComponentReader<UObject*> ObjectPtr = Linker->EntityManager.ReadComponent(ChildID,  BoundObjectComponent);
 			if (ObjectPtr)
 			{
 				StaleEntitiesToPreserve.Add(MakeTuple(*ObjectPtr, ParentID), ChildID);
diff --git a/Engine/Source/Runtime/MovieScene/Private/EntitySystem/MovieSceneEntityManager.cpp b/Engine/Source/Runtime/MovieScene/Private/EntitySystem/MovieSceneEntityManager.cpp
--- a/Engine/Source/Runtime/MovieScene/Private/EntitySystem/MovieSceneEntityManager.cpp
+++ b/Engine/Source/Runtime/MovieScene/Private/EntitySystem/MovieSceneEntityManager.cpp
@@ -151,7 +151,7 @@ struct FEntityInitializer
 
 	static FEntityAllocation* Initialize(const FEntityManager& EntityManager, const FComponentMask& EntityComponentMask, const FEntityAllocationInitializationInfo& InitInfo)
 	{
-		const uint64 CurrentSystemSerialNumber = EntityManager.GetSystemSerial();
+		const FEntityAllocationWriteContext WriteContext(EntityManager);
 		check(IsValidUint16(InitInfo.NumComponents));
 
 		// Compute the size that we need: struct size + component headers array + entity IDs array.
@@ -170,7 +170,7 @@ struct FEntityInitializer
 		Allocation->Size = 0;
 		Allocation->Capacity = InitInfo.InitialCapacity;
 		Allocation->MaxCapacity = InitInfo.MaxCapacity;
-		Allocation->SerialNumber = CurrentSystemSerialNumber;
+		Allocation->SerialNumber = WriteContext.GetSystemSerial();
 
 		// Fixup pointer offsets:
 		//
@@ -216,7 +216,7 @@ struct FEntityInitializer
 
 				Header->ComponentType = ComponentTypeID;
 				Header->Sizeof = TypeInfo.Sizeof;
-				Header->PostWriteComponents(CurrentSystemSerialNumber);
+				Header->PostWriteComponents(WriteContext);
 				if (TypeInfo.IsTag())
 				{
 					Header->Components = nullptr;
@@ -608,6 +608,8 @@ int32 FEntityManager::FreeEntities(const FFreeEntityOperation& Operation, TSet<F
 		return 0;
 	}
 
+	FEntityAllocationWriteContext WriteContext(*this);
+
 	FBuiltInComponentTypes* BuiltInComponents = FBuiltInComponentTypes::Get();
 	FFreeEntityOperation::FCommitData Committed = Operation.Commit();
 
@@ -679,7 +681,7 @@ int32 FEntityManager::FreeEntities(const FFreeEntityOperation& Operation, TSet<F
 		int32 LastEntityIndex = Num - 1;
 
 		// Modify allocation and headers
-		Allocation->PostModifyStructure(SystemSerialNumber);
+		Allocation->PostModifyStructure(WriteContext);
 
 		// Reverse the set bits so we can iterate backwards
 		const int32 MaskSize = Pair.Value.Mask.Num();
@@ -1437,8 +1439,10 @@ bool FEntityManager::CopyComponent(FMovieSceneEntityID SrcEntityID, FMovieSceneE
 	void* DstValue = DstHeader.GetValuePtr(DstLocation.GetEntryIndexWithinAllocation());
 	const void* SrcValue = SrcHeader.GetValuePtr(SrcLocation.GetEntryIndexWithinAllocation());
 
+	FEntityAllocationWriteContext WriteContext(*this);
+
 	ComponentTypeInfo.CopyItems(DstValue, SrcValue, 1);
-	DstHeader.PostWriteComponents(SystemSerialNumber);
+	DstHeader.PostWriteComponents(WriteContext);
 
 	if (bDidChangeStructure)
 	{
@@ -1658,13 +1662,13 @@ void FEntityManager::TouchEntity(FMovieSceneEntityID EntityID)
 		FEntityAllocation* Allocation = GetAllocation(Location.GetAllocationIndex());
 		check(Allocation);
 
-		const uint64 CurrentSystemSerialNumber = SystemSerialNumber;
+		FEntityAllocationWriteContext WriteContext(*this);
 
 		for (FComponentHeader& Header : Allocation->GetComponentHeaders())
 		{
 			if (!Header.IsTag())
 			{
-				Header.PostWriteComponents(CurrentSystemSerialNumber);
+				Header.PostWriteComponents(WriteContext);
 			}
 		}
 	}
@@ -1983,20 +1987,20 @@ int32 FEntityManager::MigrateEntity(int32 DestAllocationIndex, int32 SourceAlloc
 
 	check(SrcOffset < Src->Num() && DstOffset < Dst->Num());
 
-	const uint64 CurrentSystemSerialNumber = SystemSerialNumber;
-	Src->PostModifyStructureExcludingHeaders(CurrentSystemSerialNumber);
-	Dst->PostModifyStructureExcludingHeaders(CurrentSystemSerialNumber);
+	FEntityAllocationWriteContext WriteContext(*this);
+	Src->PostModifyStructureExcludingHeaders(WriteContext);
+	Dst->PostModifyStructureExcludingHeaders(WriteContext);
 
 	// Iterate all destination component types
 	for (int32 DstHeaderOffset = 0; DstHeaderOffset < Dst->GetNumComponentTypes(); ++DstHeaderOffset)
 	{
 		FComponentHeader* DstComponentHeader = &Dst->ComponentHeaders[DstHeaderOffset];
-		DstComponentHeader->PostWriteComponents(CurrentSystemSerialNumber);
+		DstComponentHeader->PostWriteComponents(WriteContext);
 
 		// Try to find a matching source component type
 		while (SrcComponentHeader != SrcLastComponentHeader && SrcComponentHeader->ComponentType.BitIndex() < DstComponentHeader->ComponentType.BitIndex())
 		{
-			SrcComponentHeader->PostWriteComponents(CurrentSystemSerialNumber);
+			SrcComponentHeader->PostWriteComponents(WriteContext);
 			++SrcComponentHeader;
 		}
 
@@ -2073,8 +2077,8 @@ void FEntityManager::CopyComponents(int32 DestAllocationIndex, int32 DestEntityI
 
 	check(SrcOffset < Src->Num() && DstOffset < Dst->Num());
 
-	const uint64 CurrentSystemSerialNumber = SystemSerialNumber;
-	Dst->PostModifyStructureExcludingHeaders(CurrentSystemSerialNumber);
+	FEntityAllocationWriteContext WriteContext(*this);
+	Dst->PostModifyStructureExcludingHeaders(WriteContext);
 
 	// Iterate all source component types
 	for (int32 DstHeaderOffset = 0;  DstHeaderOffset < Dst->GetNumComponentTypes(); ++DstHeaderOffset)
@@ -2085,7 +2089,7 @@ void FEntityManager::CopyComponents(int32 DestAllocationIndex, int32 DestEntityI
 			continue;
 		}
 
-		DstComponentHeader->PostWriteComponents(CurrentSystemSerialNumber);
+		DstComponentHeader->PostWriteComponents(WriteContext);
 
 		// Try to find a matching source component type
 		while (SrcComponentHeader != SrcLastComponentHeader && SrcComponentHeader->ComponentType.BitIndex() < DstComponentHeader->ComponentType.BitIndex())
@@ -2189,12 +2193,12 @@ int32 FEntityManager::AddEntityToAllocation(int32 AllocationIndex, FMovieSceneEn
 	const int32 ActualEntityOffset   = Allocation->Num();
 	FEntityInitializer::AddEntity(Allocation, ActualEntityOffset, ID);
 
-	const uint64 CurrentSystemSerialNumber = SystemSerialNumber;
-	Allocation->PostModifyStructureExcludingHeaders(CurrentSystemSerialNumber);
+	FEntityAllocationWriteContext WriteContext(*this);
+	Allocation->PostModifyStructureExcludingHeaders(WriteContext);
 
 	for (FComponentHeader& Header : Allocation->GetComponentHeaders())
 	{
-		Header.PostWriteComponents(CurrentSystemSerialNumber);
+		Header.PostWriteComponents(WriteContext);
 
 		if (!Header.IsTag())
 		{
@@ -2225,13 +2229,13 @@ void FEntityManager::RemoveEntityFromAllocation(int32 AllocationIndex, int32 Ent
 
 	const bool bHadCapacity = Allocation->Num() != Allocation->GetMaxCapacity();
 
-	const uint64 CurrentSystemSerialNumber = SystemSerialNumber;
-	Allocation->PostModifyStructureExcludingHeaders(CurrentSystemSerialNumber);
+	FEntityAllocationWriteContext WriteContext(*this);
+	Allocation->PostModifyStructureExcludingHeaders(WriteContext);
 
 	// Destruct its components
 	for (FComponentHeader& Header : Allocation->GetComponentHeaders())
 	{
-		Header.PostWriteComponents(CurrentSystemSerialNumber);
+		Header.PostWriteComponents(WriteContext);
 
 		if (Header.IsTag())
 		{
diff --git a/Engine/Source/Runtime/MovieScene/Private/EntitySystem/MovieSceneEntitySystemLinker.cpp b/Engine/Source/Runtime/MovieScene/Private/EntitySystem/MovieSceneEntitySystemLinker.cpp
--- a/Engine/Source/Runtime/MovieScene/Private/EntitySystem/MovieSceneEntitySystemLinker.cpp
+++ b/Engine/Source/Runtime/MovieScene/Private/EntitySystem/MovieSceneEntitySystemLinker.cpp
@@ -3,6 +3,7 @@
 #include "EntitySystem/MovieSceneEntitySystemLinker.h"
 #include "EntitySystem/MovieSceneEntitySystemTypes.h"
 #include "EntitySystem/MovieSceneEntitySystem.h"
+#include "EntitySystem/MovieSceneEntitySystemTask.h"
 #include "EntitySystem/MovieSceneEntityMutations.h"
 #include "EntitySystem/MovieSceneComponentRegistry.h"
 #include "EntitySystem/BuiltInComponentTypes.h"
diff --git a/Engine/Source/Runtime/MovieScene/Private/EntitySystem/MovieSceneEntitySystemTypes.cpp b/Engine/Source/Runtime/MovieScene/Private/EntitySystem/MovieSceneEntitySystemTypes.cpp
--- a/Engine/Source/Runtime/MovieScene/Private/EntitySystem/MovieSceneEntitySystemTypes.cpp
+++ b/Engine/Source/Runtime/MovieScene/Private/EntitySystem/MovieSceneEntitySystemTypes.cpp
@@ -1,6 +1,8 @@
 // Copyright Epic Games, Inc. All Rights Reserved.
 
 #include "EntitySystem/MovieSceneEntitySystemTypes.h"
+#include "EntitySystem/MovieSceneEntityManager.h"
+#include "EntitySystem/MovieSceneComponentAccessors.h"
 
 namespace UE
 {
@@ -102,5 +104,41 @@ bool FEntityComponentFilter::IsValid() const
 	return false;
 }
 
+FEntityAllocationWriteContext::FEntityAllocationWriteContext(const FEntityManager& EntityManager)
+	: SystemSerial(EntityManager.GetSystemSerial())
+{}
+
+FComponentReader FEntityAllocation::ReadComponentsErased(FComponentTypeID ComponentType) const
+{
+	const FComponentHeader& Header = GetComponentHeaderChecked(ComponentType);
+	return FComponentReader(&Header);
+}
+
+FComponentWriter FEntityAllocation::WriteComponentsErased(FComponentTypeID ComponentType, FEntityAllocationWriteContext InWriteContext) const
+{
+	const FComponentHeader& Header = GetComponentHeaderChecked(ComponentType);
+	return FComponentWriter(&Header, InWriteContext);
+}
+
+FOptionalComponentReader FEntityAllocation::TryReadComponentsErased(FComponentTypeID ComponentType) const
+{
+	if (const FComponentHeader* Header = FindComponentHeader(ComponentType))
+	{
+		return FOptionalComponentReader(Header);
+	}
+
+	return FOptionalComponentReader();
+}
+
+FOptionalComponentWriter FEntityAllocation::TryWriteComponentsErased(FComponentTypeID ComponentType, FEntityAllocationWriteContext InWriteContext) const
+{
+	if (const FComponentHeader* Header = FindComponentHeader(ComponentType))
+	{
+		return FOptionalComponentWriter(Header, InWriteContext);
+	}
+
+	return FOptionalComponentWriter();
+}
+
 } // namespace MovieScene
 } // namespace UE
\ No newline at end of file
diff --git a/Engine/Source/Runtime/MovieScene/Private/EntitySystem/MovieScenePropertyRegistry.cpp b/Engine/Source/Runtime/MovieScene/Private/EntitySystem/MovieScenePropertyRegistry.cpp
--- a/Engine/Source/Runtime/MovieScene/Private/EntitySystem/MovieScenePropertyRegistry.cpp
+++ b/Engine/Source/Runtime/MovieScene/Private/EntitySystem/MovieScenePropertyRegistry.cpp
@@ -1,6 +1,11 @@
 // Copyright Epic Games, Inc. All Rights Reserved.
 
 #include "EntitySystem/MovieScenePropertyRegistry.h"
+#include "MovieSceneFwd.h"
+#include "UObject/Class.h"
+#include "UObject/Object.h"
+#include "EntitySystem/MovieScenePropertyBinding.h"
+#include "MovieSceneCommonHelpers.h"
 
 
 namespace UE
@@ -8,7 +13,80 @@ namespace UE
 namespace MovieScene
 {
 
+TOptional<uint16> ComputeFastPropertyPtrOffset(UClass* ObjectClass, const FMovieScenePropertyBinding& PropertyBinding)
+{
+	using namespace UE::MovieScene;
+
+	FProperty* Property = ObjectClass->FindPropertyByName(PropertyBinding.PropertyName);
+	// @todo: Constructing FNames from strings is _very_ costly and we really shouldn't be doing this at runtime.
+	UFunction* Setter   = ObjectClass->FindFunctionByName(*(FString("Set") + PropertyBinding.PropertyName.ToString()));
+	if (Property && !Setter)
+	{
+		UObject* DefaultObject   = ObjectClass->GetDefaultObject();
+		uint8*   PropertyAddress = Property->ContainerPtrToValuePtr<uint8>(DefaultObject);
+		int32    PropertyOffset  = PropertyAddress - reinterpret_cast<uint8*>(DefaultObject);
+
+		if (PropertyOffset >= 0 && PropertyOffset < int32(uint16(0xFFFF)))
+		{
+			return uint16(PropertyOffset);
+		}
+	}
+
+	return TOptional<uint16>();
+}
+
+TOptional<FResolvedFastProperty> FPropertyRegistry::ResolveFastProperty(UObject* Object, const FMovieScenePropertyBinding& PropertyBinding, FCustomAccessorView CustomAccessors)
+{
+	UClass* Class = Object->GetClass();
+
+	if (CustomAccessors.Num() != 0)
+	{
+		const int32 CustomPropertyIndex = CustomAccessors.FindCustomAccessorIndex(Class, PropertyBinding.PropertyPath);
+		if (CustomPropertyIndex != INDEX_NONE)
+		{
+			check(CustomPropertyIndex < MAX_uint16);
+
+			// This property has a custom property accessor that can apply properties through a static function ptr.
+			// Just add the function ptrs to the property entity so they can be called directly
+			return FResolvedFastProperty(TInPlaceType<FCustomPropertyIndex>(), FCustomPropertyIndex{ static_cast<uint16>(CustomPropertyIndex) });
+		}
+	}
+
+	if (PropertyBinding.CanUseClassLookup())
+	{
+		TOptional<uint16> FastPtrOffset = ComputeFastPropertyPtrOffset(Class, PropertyBinding);
+		if (FastPtrOffset.IsSet())
+		{
+			// This property/object combination has no custom setter function and a constant property offset from the base ptr for all instances of the object.
+			return FResolvedFastProperty(TInPlaceType<uint16>(), FastPtrOffset.GetValue());
+		}
+	}
+
+	return TOptional<FResolvedFastProperty>();
+}
+
+TOptional<FResolvedProperty> FPropertyRegistry::ResolveProperty(UObject* Object, const FMovieScenePropertyBinding& PropertyBinding, FCustomAccessorView CustomAccessors)
+{
+	TOptional< FResolvedFastProperty > FastProperty = FPropertyRegistry::ResolveFastProperty(Object, PropertyBinding, CustomAccessors);
+	if (FastProperty.IsSet())
+	{
+		if (const FCustomPropertyIndex* CustomIndex = FastProperty->TryGet<FCustomPropertyIndex>())
+		{
+			return FResolvedProperty(TInPlaceType<FCustomPropertyIndex>(), *CustomIndex);
+		}
+		return FResolvedProperty(TInPlaceType<uint16>(), FastProperty->Get<uint16>());
+	}
+
+	// None of the above optimized paths can apply to this property (probably because it has a setter function or because it is within a compound property), so we must use the slow property bindings
+	TSharedPtr<FTrackInstancePropertyBindings> SlowBindings = MakeShared<FTrackInstancePropertyBindings>(PropertyBinding.PropertyName, PropertyBinding.PropertyPath.ToString());
+	if (SlowBindings->GetProperty(*Object) == nullptr)
+	{
+		UE_LOG(LogMovieScene, Warning, TEXT("Unable to resolve property '%s' from '%s' instance '%s'"), *PropertyBinding.PropertyPath.ToString(), *Object->GetClass()->GetName(), *Object->GetName());
+		return TOptional<FResolvedProperty>();
+	}
 
+	return FResolvedProperty(TInPlaceType<TSharedPtr<FTrackInstancePropertyBindings>>(), SlowBindings);
+}
 
 
 } // namespace MovieScene
diff --git a/Engine/Source/Runtime/MovieScene/Private/EntitySystem/TrackInstance/MovieSceneTrackInstanceSystem.cpp b/Engine/Source/Runtime/MovieScene/Private/EntitySystem/TrackInstance/MovieSceneTrackInstanceSystem.cpp
--- a/Engine/Source/Runtime/MovieScene/Private/EntitySystem/TrackInstance/MovieSceneTrackInstanceSystem.cpp
+++ b/Engine/Source/Runtime/MovieScene/Private/EntitySystem/TrackInstance/MovieSceneTrackInstanceSystem.cpp
@@ -10,7 +10,7 @@
 #include "EntitySystem/BuiltInComponentTypes.h"
 #include "EntitySystem/MovieSceneInstanceRegistry.h"
 #include "EntitySystem/MovieSceneEntitySystemLinker.h"
-#include "EntitySystem/MovieSceneEntityFactory.h"
+#include "EntitySystem/MovieSceneEntityFactoryTemplates.h"
 #include "MovieSceneSection.h"
 
 DECLARE_CYCLE_STAT(TEXT("Generic Track Instances"), MovieSceneEval_GenericTrackInstances, STATGROUP_MovieSceneECS);
@@ -34,12 +34,13 @@ struct FTrackInstanceInputComponentInitializer : TChildEntityInitializer<FMovieS
 	{
 		check(ParentAllocationOffsets.Num() == ChildRange.Num);
 
-		TArrayView<const FMovieSceneTrackInstanceComponent> TrackInstanceComponents = TRead<FMovieSceneTrackInstanceComponent>(GetParentComponent()).ResolveAsArray(ParentAllocation);
+		FEntityAllocationWriteContext FreshWriteContext = FEntityAllocationWriteContext::NewAllocation();
 
-		TArrayView<FTrackInstanceInputComponent> Inputs       = TWrite<FTrackInstanceInputComponent>(GetChildComponent()).ResolveAsArray(ChildRange.Allocation);
-		TArrayView<UObject* const>               BoundObjects = TReadOptional<UObject*>(FBuiltInComponentTypes::Get()->BoundObject).ResolveAsArray(ChildRange.Allocation);
+		TComponentReader<FMovieSceneTrackInstanceComponent> TrackInstances  = ParentAllocation->ReadComponents(GetParentComponent());
+		TComponentWriter<FTrackInstanceInputComponent>      Inputs          = ChildRange.Allocation->WriteComponents(GetChildComponent(), FreshWriteContext);
+		TOptionalComponentReader<UObject*>                  BoundObjects    = ChildRange.Allocation->TryReadComponents(FBuiltInComponentTypes::Get()->BoundObject);
 
-		if (BoundObjects.Num() != 0)
+		if (BoundObjects)
 		{
 			for (int32 Index = 0; Index < ChildRange.Num; ++Index)
 			{
@@ -47,7 +48,7 @@ struct FTrackInstanceInputComponentInitializer : TChildEntityInitializer<FMovieS
 				const int32 ChildIndex  = ChildRange.ComponentStartOffset + Index;
 
 				// Initialize the output index
-				Inputs[ChildIndex].OutputIndex = Instantiator->MakeOutput(BoundObjects[ChildIndex], TrackInstanceComponents[ParentIndex].TrackInstanceClass);
+				Inputs[ChildIndex].OutputIndex = Instantiator->MakeOutput(BoundObjects[ChildIndex], TrackInstances[ParentIndex].TrackInstanceClass);
 			}
 		}
 		else for (int32 Index = 0; Index < ChildRange.Num; ++Index)
@@ -56,7 +57,7 @@ struct FTrackInstanceInputComponentInitializer : TChildEntityInitializer<FMovieS
 			const int32 ChildIndex  = ChildRange.ComponentStartOffset + Index;
 
 			// Initialize the output index
-			Inputs[ChildIndex].OutputIndex = Instantiator->MakeOutput(nullptr, TrackInstanceComponents[ParentIndex].TrackInstanceClass);
+			Inputs[ChildIndex].OutputIndex = Instantiator->MakeOutput(nullptr, TrackInstances[ParentIndex].TrackInstanceClass);
 		}
 	}
 };
diff --git a/Engine/Source/Runtime/MovieScene/Private/Sections/MovieSceneHookSection.cpp b/Engine/Source/Runtime/MovieScene/Private/Sections/MovieSceneHookSection.cpp
--- a/Engine/Source/Runtime/MovieScene/Private/Sections/MovieSceneHookSection.cpp
+++ b/Engine/Source/Runtime/MovieScene/Private/Sections/MovieSceneHookSection.cpp
@@ -4,6 +4,7 @@
 #include "EntitySystem/MovieSceneEvaluationHookSystem.h"
 #include "Channels/MovieSceneChannelProxy.h"
 #include "EntitySystem/BuiltInComponentTypes.h"
+#include "EntitySystem/MovieSceneEntitySystemLinker.h"
 
 UMovieSceneHookSection::UMovieSceneHookSection(const FObjectInitializer& ObjInit)
 	: Super(ObjInit)
diff --git a/Engine/Source/Runtime/MovieScene/Private/Sections/MovieSceneSubSection.cpp b/Engine/Source/Runtime/MovieScene/Private/Sections/MovieSceneSubSection.cpp
--- a/Engine/Source/Runtime/MovieScene/Private/Sections/MovieSceneSubSection.cpp
+++ b/Engine/Source/Runtime/MovieScene/Private/Sections/MovieSceneSubSection.cpp
@@ -11,6 +11,9 @@
 #include "Misc/FrameRate.h"
 #include "Logging/MessageLog.h"
 #include "Evaluation/MovieSceneRootOverridePath.h"
+#include "EntitySystem/MovieSceneEntitySystemLinker.h"
+#include "EntitySystem/MovieSceneInstanceRegistry.h"
+#include "EntitySystem/IMovieSceneEntityProvider.h"
 
 TWeakObjectPtr<UMovieSceneSubSection> UMovieSceneSubSection::TheRecordingSection;
 
diff --git a/Engine/Source/Runtime/MovieScene/Public/EntitySystem/BuiltInComponentTypes.h b/Engine/Source/Runtime/MovieScene/Public/EntitySystem/BuiltInComponentTypes.h
--- a/Engine/Source/Runtime/MovieScene/Public/EntitySystem/BuiltInComponentTypes.h
+++ b/Engine/Source/Runtime/MovieScene/Public/EntitySystem/BuiltInComponentTypes.h
@@ -17,54 +17,14 @@
 
 enum class EMovieSceneBlendType : uint8;
 struct FMovieSceneFloatChannel;
+struct FMovieScenePropertyBinding;
+
 class UMovieSceneSection;
 class UMovieSceneTrackInstance;
 class UMovieSceneBlenderSystem;
 class FTrackInstancePropertyBindings;
 
 
-/**
- * Source property binding information for an entity on a moviescene timeline
- * Comprises a leaf property name and a path and a cached boolean signifying whether the binding is allowed to perform a fast class-wise property lookup
- */
-USTRUCT()
-struct FMovieScenePropertyBinding
-{
-	GENERATED_BODY()
-
-	FMovieScenePropertyBinding()
-		: bCanUseClassLookup(false)
-	{}
-
-	FMovieScenePropertyBinding(FName InPropertyName, const FString& InPropertyPath)
-		: PropertyName(InPropertyName), PropertyPath(*InPropertyPath)
-	{
-		bCanUseClassLookup = !(InPropertyPath.Contains(TEXT(".")) || InPropertyPath.Contains(TEXT("/")) || InPropertyPath.Contains(TEXT("\\")) || InPropertyPath.Contains(TEXT("[")));
-	}
-
-	friend bool operator==(FMovieScenePropertyBinding A, FMovieScenePropertyBinding B)
-	{
-		return A.PropertyName == B.PropertyName && A.PropertyPath == B.PropertyPath;
-	}
-
-	bool CanUseClassLookup() const
-	{
-		return bCanUseClassLookup;
-	}
-
-	/** Leaf name of the property to animate */
-	UPROPERTY()
-	FName PropertyName;
-
-	/** Full path to the property from the object including struct and array indirection */
-	UPROPERTY()
-	FName PropertyPath;
-
-	/** True if this property can be considered for fast property offset resolution(ie the property address is _always_ a constant offset from the obejct ptr), false otherwise */
-	UPROPERTY()
-	bool bCanUseClassLookup;
-};
-
 /**
  * Easing component data.
  */
@@ -320,16 +280,5 @@ private:
 };
 
 
-
-#if UE_MOVIESCENE_ENTITY_DEBUG
-
-	template<> struct TComponentDebugType<FMovieScenePropertyBinding>
-	{
-		static const EComponentDebugType Type = EComponentDebugType::Property;
-	};
-
-#endif // UE_MOVIESCENE_ENTITY_DEBUG
-
-
 } // namespace MovieScene
 } // namespace UE
diff --git a/Engine/Source/Runtime/MovieScene/Public/EntitySystem/MovieSceneComponentAccessors.h b/Engine/Source/Runtime/MovieScene/Public/EntitySystem/MovieSceneComponentAccessors.h
--- a/Engine/Source/Runtime/MovieScene/Public/EntitySystem/MovieSceneComponentAccessors.h
+++ b/Engine/Source/Runtime/MovieScene/Public/EntitySystem/MovieSceneComponentAccessors.h
@@ -5,6 +5,8 @@
 #include "MovieSceneEntityIDs.h"
 #include "Async/TaskGraphInterfaces.h"
 #include "EntitySystem/MovieSceneSystemTaskDependencies.h"
+#include "EntitySystem/EntityAllocationIterator.h"
+#include "EntitySystem/MovieSceneComponentPtr.h"
 
 #include <initializer_list>
 
@@ -14,225 +16,84 @@ namespace MovieScene
 {
 
 
-struct FComponentAccessor
+struct FComponentAccess
 {
-	enum { SupportsDirectEntityIteration = true };
 	FComponentTypeID ComponentType;
 };
-struct FRead : FComponentAccessor
+struct FReadAccess : FComponentAccess
 {
-	FRead(FComponentTypeID InComponentType) : FComponentAccessor{ InComponentType } {}
+	FReadAccess(FComponentTypeID InComponentType) : FComponentAccess{ InComponentType } {}
 };
-struct FWrite : FComponentAccessor
+struct FWriteAccess : FComponentAccess
 {
-	FWrite(FComponentTypeID InComponentType) : FComponentAccessor{ InComponentType } {}
+	FWriteAccess(FComponentTypeID InComponentType) : FComponentAccess{ InComponentType } {}
 };
 
 
-struct FOptionalComponentAccessor
+struct FOptionalComponentAccess
 {
-	enum { SupportsDirectEntityIteration = true };
 	FComponentTypeID ComponentType;
 };
-struct FReadOptional : FOptionalComponentAccessor
+struct FOptionalReadAccess : FOptionalComponentAccess
 {
-	FReadOptional(FComponentTypeID InComponentType) : FOptionalComponentAccessor{ InComponentType } {}
+	FOptionalReadAccess(FComponentTypeID InComponentType) : FOptionalComponentAccess{ InComponentType } {}
 };
-struct FWriteOptional : FOptionalComponentAccessor
+struct FOptionalWriteAccess : FOptionalComponentAccess
 {
-	FWriteOptional(FComponentTypeID InComponentType) : FOptionalComponentAccessor{ InComponentType } {}
+	FOptionalWriteAccess(FComponentTypeID InComponentType) : FOptionalComponentAccess{ InComponentType } {}
 };
 
 
-struct FErasedIterState
+struct FEntityIDAccess
 {
-	const uint8* ComponentPtr;
-	int32 Sizeof;
+	using AccessType = const FMovieSceneEntityID;
 
-	void operator++()
+	FORCEINLINE TRead<FMovieSceneEntityID> LockComponentData(const FEntityAllocation* Allocation, FEntityAllocationWriteContext WriteContext) const
 	{
-		ComponentPtr += Sizeof;
-	}
-
-	const void* operator*()
-	{
-		return ComponentPtr;
-	}
-
-	const void* operator[](const int32 Index)
-	{
-		return ComponentPtr += Index*Sizeof;
-	}
-};
-
-struct FErasedOptionalIterState
-{
-	uint8* ComponentPtr;
-	int32 Sizeof;
-
-	void operator++()
-	{
-		if (ComponentPtr)
-		{
-			ComponentPtr += Sizeof;
-		}
-	}
-
-	void* operator*()
-	{
-		return ComponentPtr;
-	}
-};
-
-template<typename T>
-struct TComponentIterState
-{
-	T* ComponentPtr;
-
-	void operator++()
-	{
-		++ComponentPtr;
-	}
-	T& operator*()
-	{
-		return *ComponentPtr;
-	}
-};
-
-template<typename T, typename ProjectionType>
-struct TProjectedComponentIterState
-{
-	T* ComponentPtr;
-	ProjectionType Projection;
-
-	void operator++()
-	{
-		++ComponentPtr;
-	}
-	auto operator*()
-	{
-		return Invoke(Projection, *ComponentPtr);
-	}
-};
-
-template<typename T>
-struct TOptionalIterState
-{
-	T* ComponentPtr;
-
-	void operator++()
-	{
-		if (ComponentPtr)
-		{
-			++ComponentPtr;
-		}
-	}
-	T* operator*()
-	{
-		return ComponentPtr;
+		return TRead<FMovieSceneEntityID>(Allocation->GetRawEntityIDs());
 	}
 };
 
-struct FReadEntityIDs
-{
-	enum { SupportsDirectEntityIteration = true };
-
-	using Type        = FMovieSceneEntityID;
-	using AccessType  = const FMovieSceneEntityID;
-
-	TComponentIterState<AccessType> CreateIterState(const FEntityAllocation* Allocation) const
-	{
-		AccessType* ComponentArray = Resolve(Allocation);
-		return TComponentIterState<AccessType>{ ComponentArray };
-	}
-
-	AccessType* Resolve(const FEntityAllocation* Allocation) const
-	{
-		return Allocation->GetRawEntityIDs();
-	}
-
-	TArrayView<AccessType> ResolveAsArray(const FEntityAllocation* Allocation) const
-	{
-		AccessType* Resolved = Resolve(Allocation);
-		return MakeArrayView(Resolved, Allocation->Num());
-	}
-};
-
-
 
 
 template<typename T>
-struct TRead : FRead
+struct TReadAccess : FReadAccess
 {
-	using Type        = T;
-	using AccessType  = const T;
+	using AccessType = const T;
 
-	TRead(FComponentTypeID InComponentTypeID)
-		: FRead{ InComponentTypeID }
+	TReadAccess(FComponentTypeID InComponentTypeID)
+		: FReadAccess{ InComponentTypeID }
 	{}
 
-	TComponentIterState<AccessType> CreateIterState(const FEntityAllocation* Allocation) const
-	{
-		AccessType* ComponentArray = Resolve(Allocation);
-		return TComponentIterState<AccessType>{ ComponentArray };
-	}
-
-	AccessType* Resolve(const FEntityAllocation* Allocation) const
+	FORCEINLINE TComponentLock<TRead<T>> LockComponentData(const FEntityAllocation* Allocation, FEntityAllocationWriteContext WriteContext) const
 	{
-		const FComponentHeader& Header = Allocation->GetComponentHeaderChecked(this->ComponentType);
-		return reinterpret_cast<AccessType*>(Header.Components);
-	}
-
-	TArrayView<AccessType> ResolveAsArray(const FEntityAllocation* Allocation) const
-	{
-		AccessType* Resolved = Resolve(Allocation);
-		return MakeArrayView(Resolved, Allocation->Num());
+		return Allocation->ReadComponents(ComponentType.ReinterpretCast<T>());
 	}
 };
 
 
-struct FReadErased : FRead
+struct FErasedReadAccess : FReadAccess
 {
-	using Type        = uint8;
-	using AccessType  = const uint8;
-
-	FReadErased(FComponentTypeID InComponentTypeID)
-		: FRead{ InComponentTypeID }
+	FErasedReadAccess(FComponentTypeID InComponentTypeID)
+		: FReadAccess{ InComponentTypeID }
 	{}
 
-	FErasedIterState CreateIterState(const FEntityAllocation* Allocation) const
+	FORCEINLINE FComponentReader LockComponentData(const FEntityAllocation* Allocation, FEntityAllocationWriteContext WriteContext) const
 	{
-		const FComponentHeader& Header = Allocation->GetComponentHeaderChecked(this->ComponentType);
-		return FErasedIterState{ Header.Components, Header.Sizeof };
+		return Allocation->ReadComponentsErased(ComponentType);
 	}
 };
 
 
-template<typename T, typename ProjectionType>
-struct TReadProjected : FRead
+struct FErasedWriteAccess : FWriteAccess
 {
-	using Type        = T;
-	using AccessType  = const T;
-
-	ProjectionType Projection;
-
-	static_assert(TNot<TIsReferenceType<ProjectionType>>::Value, "Reference projections are not supported");
-
-	TReadProjected(FComponentTypeID InComponentTypeID, ProjectionType InProjection)
-		: FRead{ InComponentTypeID }
-		, Projection(InProjection)
+	FErasedWriteAccess(FComponentTypeID InComponentTypeID)
+		: FWriteAccess{ InComponentTypeID }
 	{}
 
-	TProjectedComponentIterState<AccessType, ProjectionType> CreateIterState(const FEntityAllocation* Allocation) const
-	{
-		AccessType* ComponentArray = Resolve(Allocation);
-		return TProjectedComponentIterState<AccessType, ProjectionType>{ ComponentArray, Projection };
-	}
-
-	AccessType* Resolve(const FEntityAllocation* Allocation) const
+	FORCEINLINE FComponentWriter LockComponentData(const FEntityAllocation* Allocation, FEntityAllocationWriteContext WriteContext) const
 	{
-		const FComponentHeader& Header = Allocation->GetComponentHeaderChecked(this->ComponentType);
-		return reinterpret_cast<AccessType*>(Header.Components);
+		return Allocation->WriteComponentsErased(ComponentType, WriteContext);
 	}
 };
 
@@ -240,31 +101,17 @@ struct TReadProjected : FRead
 
 
 template<typename T>
-struct TWrite : FWrite
+struct TWriteAccess : FWriteAccess
 {
-	using Type        = T;
-	using AccessType  = T;
+	using AccessType = T;
 
-	TWrite(FComponentTypeID InComponentTypeID)
-		: FWrite{ InComponentTypeID }
+	TWriteAccess(FComponentTypeID InComponentTypeID)
+		: FWriteAccess{ InComponentTypeID }
 	{}
 
-	TComponentIterState<AccessType> CreateIterState(const FEntityAllocation* Allocation) const
+	FORCEINLINE TComponentLock<TWrite<T>> LockComponentData(const FEntityAllocation* Allocation, FEntityAllocationWriteContext WriteContext) const
 	{
-		AccessType* ComponentArray = Resolve(Allocation);
-		return TComponentIterState<AccessType>{ ComponentArray };
-	}
-
-	AccessType* Resolve(const FEntityAllocation* Allocation) const
-	{
-		const FComponentHeader& Header = Allocation->GetComponentHeaderChecked(this->ComponentType);
-		return reinterpret_cast<AccessType*>(Header.Components);
-	}
-
-	TArrayView<AccessType> ResolveAsArray(const FEntityAllocation* Allocation) const
-	{
-		AccessType* Resolved = Resolve(Allocation);
-		return MakeArrayView(Resolved, Allocation->Num());
+		return Allocation->WriteComponents(ComponentType.ReinterpretCast<T>(), WriteContext);
 	}
 };
 
@@ -272,35 +119,17 @@ struct TWrite : FWrite
 
 
 template<typename T>
-struct TReadOptional : FReadOptional
+struct TOptionalReadAccess : FOptionalReadAccess
 {
-	using Type        = T;
-	using AccessType  = const T;
+	using AccessType = const T;
 
-	TReadOptional(FComponentTypeID InComponentTypeID)
-		: FReadOptional{InComponentTypeID}
+	TOptionalReadAccess(FComponentTypeID InComponentTypeID)
+		: FOptionalReadAccess{InComponentTypeID}
 	{}
 
-	TOptionalIterState<AccessType> CreateIterState(const FEntityAllocation* Allocation) const
-	{
-		AccessType* ComponentArray = Resolve(Allocation);
-		return TOptionalIterState<AccessType>{ ComponentArray };
-	}
-
-	AccessType* Resolve(const FEntityAllocation* Allocation) const
+	FORCEINLINE TComponentLock<TReadOptional<T>> LockComponentData(const FEntityAllocation* Allocation, FEntityAllocationWriteContext WriteContext) const
 	{
-		const FComponentHeader* Header = Allocation->FindComponentHeader(this->ComponentType);
-		if (Header)
-		{
-			return reinterpret_cast<AccessType*>(Header->Components);
-		}
-		return nullptr;
-	}
-
-	TArrayView<AccessType> ResolveAsArray(const FEntityAllocation* Allocation) const
-	{
-		AccessType* Resolved = Resolve(Allocation);
-		return MakeArrayView(Resolved, Resolved ? Allocation->Num() : 0);
+		return Allocation->TryReadComponents(ComponentType.ReinterpretCast<T>());
 	}
 };
 
@@ -308,35 +137,17 @@ struct TReadOptional : FReadOptional
 
 
 template<typename T>
-struct TWriteOptional : FWriteOptional
+struct TOptionalWriteAccess : FOptionalWriteAccess
 {
-	using Type        = T;
-	using AccessType  = T;
+	using AccessType = T;
 
-	TWriteOptional(FComponentTypeID InComponentTypeID)
-		: FWriteOptional{ InComponentTypeID }
+	TOptionalWriteAccess(FComponentTypeID InComponentTypeID)
+		: FOptionalWriteAccess{ InComponentTypeID }
 	{}
 
-	TOptionalIterState<AccessType> CreateIterState(const FEntityAllocation* Allocation) const
-	{
-		AccessType* ComponentArray = Resolve(Allocation);
-		return TOptionalIterState<AccessType>{ ComponentArray };
-	}
-
-	AccessType* Resolve(const FEntityAllocation* Allocation) const
-	{
-		const FComponentHeader* Header = Allocation->FindComponentHeader(this->ComponentType);
-		if (Header)
-		{
-			return reinterpret_cast<AccessType*>(Header->Components);
-		}
-		return nullptr;
-	}
-
-	TArrayView<AccessType> ResolveAsArray(const FEntityAllocation* Allocation) const
+	FORCEINLINE TComponentLock<TWriteOptional<T>> LockComponentData(const FEntityAllocation* Allocation, FEntityAllocationWriteContext WriteContext) const
 	{
-		AccessType* Resolved = Resolve(Allocation);
-		return MakeArrayView(Resolved, Resolved ? Allocation->Num() : 0);
+		return Allocation->TryWriteComponents(ComponentType.ReinterpretCast<T>(), WriteContext);
 	}
 };
 
@@ -344,159 +155,83 @@ struct TWriteOptional : FWriteOptional
 
 
 template<typename... T>
-struct TReadOneOf
+struct TReadOneOfAccessor
 {
-	enum { SupportsDirectEntityIteration = false };
-
-	using AccessType  = TTuple<const T*...>;
+	using AccessType = TMultiComponentLock<TReadOptional<T>...>;
 
-	TReadOneOf(TComponentTypeID<T>... InComponentTypeIDs)
+	TReadOneOfAccessor(TComponentTypeID<T>... InComponentTypeIDs)
 		: ComponentTypes{ InComponentTypeIDs... }
 	{}
 
-	void CreateIterState(...) const
-	{
-		static_assert(!TAnd<TIsSame<T, T>...>::Value, "PerEntity iteration is not supported for ReadOneOf accessors. Must use PerAllocation iterators and resolve the tuple manually.");
-	}
-
-	AccessType Resolve(const FEntityAllocation* Allocation) const
+	FORCEINLINE TMultiComponentLock<TReadOptional<T>...> LockComponentData(const FEntityAllocation* Allocation, FEntityAllocationWriteContext WriteContext) const
 	{
-		auto ResolveImpl = [Allocation](const auto& InRead)
-		{
-			return InRead.Resolve(Allocation);
-		};
-
-		return TransformTuple(ComponentTypes, ResolveImpl);
+		return TransformTuple(ComponentTypes, [Allocation, WriteContext](auto In){ return In.LockComponentData(Allocation, WriteContext); });
 	}
 
-	TTuple<TArrayView<const T*>...> ResolveAsArrays(const FEntityAllocation* Allocation) const
-	{
-		const int32 AllocationSize = Allocation->Num();
-		AccessType Resolved = Resolve(Allocation);
-
-		auto ResolveImpl = [AllocationSize](const auto& InRead)
-		{
-			return MakeArrayView(InRead, InRead ? AllocationSize : 0);
-		};
-		return TransformTuple(Resolved, ResolveImpl);
-	}
-
-	void ResolveAsArrays(const FEntityAllocation* Allocation, TArrayView<const T*>*... OutArray) const
-	{
-		const int32 AllocationSize = Allocation->Num();
-		TTuple<TArrayView<const T> *...> OutArrayTuple(MakeTuple(OutArray...));
-		AccessType Resolved = Resolve(Allocation);
-
-		auto ResolveImpl = [AllocationSize](const auto& InRead, auto* OutWrite)
-		{
-			*OutWrite = MakeArrayView(InRead, InRead ? AllocationSize : 0);
-		};
-		VisitTupleElements(ResolveImpl, Resolved, OutArrayTuple);
-	}
-
-	TTuple< TReadOptional<T>... > ComponentTypes;
+	TTuple< TOptionalReadAccess<T>... > ComponentTypes;
 };
 
 
 
 
 template<typename... T>
-struct TReadOneOrMoreOf
+struct TReadOneOrMoreOfAccessor
 {
-	enum { SupportsDirectEntityIteration = false };
-
-	using AccessType = TTuple<const T *...>;
+	using AccessType = TMultiComponentLock<TReadOptional<T>...>;
 
-	TReadOneOrMoreOf(TComponentTypeID<T>... InComponentTypeIDs)
+	TReadOneOrMoreOfAccessor(TComponentTypeID<T>... InComponentTypeIDs)
 		: ComponentTypes{ InComponentTypeIDs... }
 	{}
 
-	void CreateIterState(...) const
+	FORCEINLINE TMultiComponentLock<TReadOptional<T>...> LockComponentData(const FEntityAllocation* Allocation, FEntityAllocationWriteContext WriteContext) const
 	{
-		static_assert(!TAnd<TIsSame<T, T>...>::Value, "PerEntity iteration is not supported for ReadOneOf accessors. Must use PerAllocation iterators and resolve the tuple manually.");
+		return TransformTuple(ComponentTypes, [Allocation, WriteContext](auto In){ return In.LockComponentData(Allocation, WriteContext); });
 	}
 
-	AccessType Resolve(const FEntityAllocation* Allocation) const
-	{
-		auto ResolveImpl = [Allocation](const auto& InRead)
-		{
-			return InRead.Resolve(Allocation);
-		};
-
-		return TransformTuple(ComponentTypes, ResolveImpl);
-	}
-
-	TTuple<TArrayView<const T>...> ResolveAsArrays(const FEntityAllocation* Allocation) const
-	{
-		const int32 AllocationSize = Allocation->Num();
-		AccessType Resolved = Resolve(Allocation);
-
-		auto ResolveImpl = [AllocationSize](const auto& InRead)
-		{
-			return MakeArrayView(InRead, InRead ? AllocationSize : 0);
-		};
-		return TransformTuple(Resolved, ResolveImpl);
-	}
-
-	void ResolveAsArrays(const FEntityAllocation* Allocation, TArrayView<const T>*... OutArray) const
-	{
-		const int32 AllocationSize = Allocation->Num();
-		TTuple<TArrayView<const T>*...> OutArrayTuple(MakeTuple(OutArray...));
-		AccessType Resolved = Resolve(Allocation);
-
-		auto ResolveImpl = [AllocationSize](const auto& InRead, auto* OutWrite)
-		{
-			*OutWrite = MakeArrayView(InRead, InRead ? AllocationSize : 0);
-		};
-		VisitTupleElements(ResolveImpl, Resolved, OutArrayTuple);
-	}
-
-	TTuple< TReadOptional<T>... > ComponentTypes;
+	TTuple< TOptionalReadAccess<T>... > ComponentTypes;
 };
 
 
-
-
-inline void AddAccessorToFilter(const FReadEntityIDs*, FEntityComponentFilter* OutFilter)
+inline void AddAccessorToFilter(const FEntityIDAccess*, FEntityComponentFilter* OutFilter)
 {
 }
-inline void AddAccessorToFilter(const FComponentAccessor* In, FEntityComponentFilter* OutFilter)
+inline void AddAccessorToFilter(const FComponentAccess* In, FEntityComponentFilter* OutFilter)
 {
 	check(In->ComponentType);
 	OutFilter->All({ In->ComponentType });
 }
-inline void AddAccessorToFilter(const FOptionalComponentAccessor* In, FEntityComponentFilter* OutFilter)
+inline void AddAccessorToFilter(const FOptionalComponentAccess* In, FEntityComponentFilter* OutFilter)
 {
 }
 template<typename... T>
-void AddAccessorToFilter(const TReadOneOf<T...>* In, FEntityComponentFilter* OutFilter)
+void AddAccessorToFilter(const TReadOneOfAccessor<T...>* In, FEntityComponentFilter* OutFilter)
 {
 	FComponentMask Mask;
-	VisitTupleElements([&Mask](FReadOptional Composite){ if (Composite.ComponentType) { Mask.Set(Composite.ComponentType); } }, In->ComponentTypes);
+	VisitTupleElements([&Mask](FOptionalReadAccess Composite){ if (Composite.ComponentType) { Mask.Set(Composite.ComponentType); } }, In->ComponentTypes);
 
 	check(Mask.NumComponents() != 0);
 	OutFilter->Complex(Mask, EComplexFilterMode::OneOf);
 }
 template<typename... T>
-void AddAccessorToFilter(const TReadOneOrMoreOf<T...>* In, FEntityComponentFilter* OutFilter)
+void AddAccessorToFilter(const TReadOneOrMoreOfAccessor<T...>* In, FEntityComponentFilter* OutFilter)
 {
 	FComponentMask Mask;
-	VisitTupleElements([&Mask](FReadOptional Composite) { if (Composite.ComponentType) { Mask.Set(Composite.ComponentType); } }, In->ComponentTypes);
+	VisitTupleElements([&Mask](FOptionalReadAccess Composite) { if (Composite.ComponentType) { Mask.Set(Composite.ComponentType); } }, In->ComponentTypes);
 
 	check(Mask.NumComponents() != 0);
 	OutFilter->Complex(Mask, EComplexFilterMode::OneOrMoreOf);
 }
 
 
-inline void PopulatePrerequisites(const FReadEntityIDs*, const FSystemTaskPrerequisites& InPrerequisites, FGraphEventArray* OutGatheredPrereqs)
+inline void PopulatePrerequisites(const FEntityIDAccess*, const FSystemTaskPrerequisites& InPrerequisites, FGraphEventArray* OutGatheredPrereqs)
 {
 }
-inline void PopulatePrerequisites(const FComponentAccessor* In, const FSystemTaskPrerequisites& InPrerequisites, FGraphEventArray* OutGatheredPrereqs)
+inline void PopulatePrerequisites(const FComponentAccess* In, const FSystemTaskPrerequisites& InPrerequisites, FGraphEventArray* OutGatheredPrereqs)
 {
 	check(In->ComponentType);
 	InPrerequisites.FilterByComponent(*OutGatheredPrereqs, In->ComponentType);
 }
-inline void PopulatePrerequisites(const FOptionalComponentAccessor* In, const FSystemTaskPrerequisites& InPrerequisites, FGraphEventArray* OutGatheredPrereqs)
+inline void PopulatePrerequisites(const FOptionalComponentAccess* In, const FSystemTaskPrerequisites& InPrerequisites, FGraphEventArray* OutGatheredPrereqs)
 {
 	if (In->ComponentType)
 	{
@@ -505,20 +240,20 @@ inline void PopulatePrerequisites(const FOptionalComponentAccessor* In, const FS
 	}
 }
 template<typename... T>
-void PopulatePrerequisites(const TReadOneOf<T...>* In, const FSystemTaskPrerequisites& InPrerequisites, FGraphEventArray* OutGatheredPrereqs)
+void PopulatePrerequisites(const TReadOneOfAccessor<T...>* In, const FSystemTaskPrerequisites& InPrerequisites, FGraphEventArray* OutGatheredPrereqs)
 {
 	VisitTupleElements(
-		[&InPrerequisites, OutGatheredPrereqs](FReadOptional Composite)
+		[&InPrerequisites, OutGatheredPrereqs](FOptionalReadAccess Composite)
 		{
 			PopulatePrerequisites(&Composite, InPrerequisites, OutGatheredPrereqs);
 		}
 	, In->ComponentTypes);
 }
 template<typename... T>
-void PopulatePrerequisites(const TReadOneOrMoreOf<T...>* In, const FSystemTaskPrerequisites& InPrerequisites, FGraphEventArray* OutGatheredPrereqs)
+void PopulatePrerequisites(const TReadOneOrMoreOfAccessor<T...>* In, const FSystemTaskPrerequisites& InPrerequisites, FGraphEventArray* OutGatheredPrereqs)
 {
 	VisitTupleElements(
-		[&InPrerequisites, OutGatheredPrereqs](FReadOptional Composite)
+		[&InPrerequisites, OutGatheredPrereqs](FOptionalReadAccess Composite)
 		{
 			PopulatePrerequisites(&Composite, InPrerequisites, OutGatheredPrereqs);
 		}
@@ -527,12 +262,12 @@ void PopulatePrerequisites(const TReadOneOrMoreOf<T...>* In, const FSystemTaskPr
 
 
 
-inline void PopulateSubsequents(const FWrite* In, const FGraphEventRef& InEvent, FSystemSubsequentTasks& OutSubsequents)
+inline void PopulateSubsequents(const FWriteAccess* In, const FGraphEventRef& InEvent, FSystemSubsequentTasks& OutSubsequents)
 {
 	check(In->ComponentType);
 	OutSubsequents.AddComponentTask(In->ComponentType, InEvent);
 }
-inline void PopulateSubsequents(const FWriteOptional* In, const FGraphEventRef& InEvent, FSystemSubsequentTasks& OutSubsequents)
+inline void PopulateSubsequents(const FOptionalWriteAccess* In, const FGraphEventRef& InEvent, FSystemSubsequentTasks& OutSubsequents)
 {
 	if (In->ComponentType)
 	{
@@ -543,95 +278,15 @@ inline void PopulateSubsequents(const void* In, const FGraphEventRef& InEvent, F
 {
 }
 
-
-
-inline void LockHeader(const FReadEntityIDs* In, FEntityAllocation* Allocation)
-{
-}
-inline void LockHeader(const FRead* In, FEntityAllocation* Allocation)
-{
-	Allocation->GetComponentHeaderChecked(In->ComponentType).ReadWriteLock.ReadLock();
-}
-inline void LockHeader(const FReadOptional* In, FEntityAllocation* Allocation)
-{
-	if (const FComponentHeader* Header = Allocation->FindComponentHeader(In->ComponentType))
-	{
-		Header->ReadWriteLock.ReadLock();
-	}
-}
-inline void LockHeader(const FWrite* In, FEntityAllocation* Allocation)
-{
-	FComponentHeader& Header = Allocation->GetComponentHeaderChecked(In->ComponentType);
-	Header.ReadWriteLock.WriteLock();
-}
-inline void LockHeader(const FWriteOptional* In, FEntityAllocation* Allocation)
-{
-	if (FComponentHeader* Header = Allocation->FindComponentHeader(In->ComponentType))
-	{
-		Header->ReadWriteLock.WriteLock();
-	}
-}
-template<typename... T>
-void LockHeader(const TReadOneOf<T...>* In, FEntityAllocation* Allocation)
-{
-	VisitTupleElements([Allocation](FReadOptional It){ LockHeader(&It, Allocation); }, In->ComponentTypes);
-}
-template<typename... T>
-void LockHeader(const TReadOneOrMoreOf<T...>* In, FEntityAllocation* Allocation)
-{
-	VisitTupleElements([Allocation](FReadOptional It) { LockHeader(&It, Allocation); }, In->ComponentTypes);
-}
-
-
-inline void UnlockHeader(const FReadEntityIDs* In, FEntityAllocation* Allocation, uint64 InSystemSerial)
-{
-}
-inline void UnlockHeader(const FRead* In, FEntityAllocation* Allocation, uint64 InSystemSerial)
-{
-	Allocation->GetComponentHeaderChecked(In->ComponentType).ReadWriteLock.ReadUnlock();
-}
-inline void UnlockHeader(const FReadOptional* In, FEntityAllocation* Allocation, uint64 InSystemSerial)
-{
-	if (FComponentHeader* Header = Allocation->FindComponentHeader(In->ComponentType))
-	{
-		Header->ReadWriteLock.ReadUnlock();
-	}
-}
-inline void UnlockHeader(const FWrite* In, FEntityAllocation* Allocation, uint64 InSystemSerial)
-{
-	FComponentHeader& Header = Allocation->GetComponentHeaderChecked(In->ComponentType);
-	Header.PostWriteComponents(InSystemSerial);
-	Header.ReadWriteLock.WriteUnlock();
-}
-inline void UnlockHeader(const FWriteOptional* In, FEntityAllocation* Allocation, uint64 InSystemSerial)
-{
-	if (FComponentHeader* Header = Allocation->FindComponentHeader(In->ComponentType))
-	{
-		Header->PostWriteComponents(InSystemSerial);
-		Header->ReadWriteLock.WriteUnlock();
-	}
-}
-template<typename... T>
-void UnlockHeader(const TReadOneOf<T...>* In, FEntityAllocation* Allocation, uint64 InSystemSerial)
-{
-	VisitTupleElements([Allocation, InSystemSerial](FReadOptional It){ UnlockHeader(&It, Allocation, InSystemSerial); }, In->ComponentTypes);
-}
-template<typename... T>
-void UnlockHeader(const TReadOneOrMoreOf<T...>* In, FEntityAllocation* Allocation, uint64 InSystemSerial)
-{
-	VisitTupleElements([Allocation, InSystemSerial](FReadOptional It) { UnlockHeader(&It, Allocation, InSystemSerial); }, In->ComponentTypes);
-}
-
-
-inline bool HasBeenWrittenToSince(const FReadEntityIDs* In, FEntityAllocation* Allocation, uint64 InSystemSerial)
+inline bool HasBeenWrittenToSince(const FEntityIDAccess* In, FEntityAllocation* Allocation, uint64 InSystemSerial)
 {
 	return Allocation->HasStructureChangedSince(InSystemSerial);
 }
-inline bool HasBeenWrittenToSince(const FComponentAccessor* In, FEntityAllocation* Allocation, uint64 InSystemSerial)
+inline bool HasBeenWrittenToSince(const FComponentAccess* In, FEntityAllocation* Allocation, uint64 InSystemSerial)
 {
 	return Allocation->GetComponentHeaderChecked(In->ComponentType).HasBeenWrittenToSince(InSystemSerial);
 }
-inline bool HasBeenWrittenToSince(const FOptionalComponentAccessor* In, FEntityAllocation* Allocation, uint64 InSystemSerial)
+inline bool HasBeenWrittenToSince(const FOptionalComponentAccess* In, FEntityAllocation* Allocation, uint64 InSystemSerial)
 {
 	if (FComponentHeader* Header = Allocation->FindComponentHeader(In->ComponentType))
 	{
@@ -640,63 +295,63 @@ inline bool HasBeenWrittenToSince(const FOptionalComponentAccessor* In, FEntityA
 	return false;
 }
 template<typename... T>
-bool HasBeenWrittenToSince(const TReadOneOf<T...>* In, FEntityAllocation* Allocation, uint64 InSystemSerial)
+bool HasBeenWrittenToSince(const TReadOneOfAccessor<T...>* In, FEntityAllocation* Allocation, uint64 InSystemSerial)
 {
 	bool bAnyWrittenTo = false;
-	VisitTupleElements([Allocation, &bAnyWrittenTo, InSystemSerial](FReadOptional It){ bAnyWrittenTo |= HasBeenWrittenToSince(&It, Allocation, InSystemSerial); }, In->ComponentTypes);
+	VisitTupleElements([Allocation, &bAnyWrittenTo, InSystemSerial](FOptionalReadAccess It){ bAnyWrittenTo |= HasBeenWrittenToSince(&It, Allocation, InSystemSerial); }, In->ComponentTypes);
 	return bAnyWrittenTo;
 }
 template<typename... T>
-bool HasBeenWrittenToSince(const TReadOneOrMoreOf<T...>* In, FEntityAllocation* Allocation, uint64 InSystemSerial)
+bool HasBeenWrittenToSince(const TReadOneOrMoreOfAccessor<T...>* In, FEntityAllocation* Allocation, uint64 InSystemSerial)
 {
 	bool bAnyWrittenTo = false;
-	VisitTupleElements([Allocation, &bAnyWrittenTo, InSystemSerial](FReadOptional It) { bAnyWrittenTo |= HasBeenWrittenToSince(&It, Allocation, InSystemSerial); }, In->ComponentTypes);
+	VisitTupleElements([Allocation, &bAnyWrittenTo, InSystemSerial](FOptionalReadAccess It) { bAnyWrittenTo |= HasBeenWrittenToSince(&It, Allocation, InSystemSerial); }, In->ComponentTypes);
 	return bAnyWrittenTo;
 }
 
 
 
-inline bool IsAccessorValid(const FReadEntityIDs*)
+inline bool IsAccessorValid(const FEntityIDAccess*)
 {
 	return true;
 }
-inline bool IsAccessorValid(const FComponentAccessor* In)
+inline bool IsAccessorValid(const FComponentAccess* In)
 {
 	return In->ComponentType != FComponentTypeID::Invalid();
 }
-inline bool IsAccessorValid(const FOptionalComponentAccessor* In)
+inline bool IsAccessorValid(const FOptionalComponentAccess* In)
 {
 	return true;
 }
 template<typename... T>
-inline bool IsAccessorValid(const TReadOneOf<T...>* In)
+inline bool IsAccessorValid(const TReadOneOfAccessor<T...>* In)
 {
 	bool bValid = false;
-	VisitTupleElements([&bValid](FReadOptional It){ bValid |= IsAccessorValid(&It); }, In->ComponentTypes);
+	VisitTupleElements([&bValid](FOptionalReadAccess It){ bValid |= IsAccessorValid(&It); }, In->ComponentTypes);
 	return bValid;
 }
 template<typename... T>
-inline bool IsAccessorValid(const TReadOneOrMoreOf<T...>* In)
+inline bool IsAccessorValid(const TReadOneOrMoreOfAccessor<T...>* In)
 {
 	bool bValid = false;
-	VisitTupleElements([&bValid](FReadOptional It) { bValid |= IsAccessorValid(&It); }, In->ComponentTypes);
+	VisitTupleElements([&bValid](FOptionalReadAccess It) { bValid |= IsAccessorValid(&It); }, In->ComponentTypes);
 	return bValid;
 }
 
 #if UE_MOVIESCENE_ENTITY_DEBUG
 
-	MOVIESCENE_API void AccessorToString(const FRead* In, FEntityManager* EntityManager, FString& OutString);
-	MOVIESCENE_API void AccessorToString(const FWrite* In, FEntityManager* EntityManager, FString& OutString);
-	MOVIESCENE_API void AccessorToString(const FReadOptional* In, FEntityManager* EntityManager, FString& OutString);
-	MOVIESCENE_API void AccessorToString(const FWriteOptional* In, FEntityManager* EntityManager, FString& OutString);
-	MOVIESCENE_API void AccessorToString(const FReadEntityIDs*, FEntityManager* EntityManager, FString& OutString);
-	MOVIESCENE_API void OneOfAccessorToString(const FReadOptional* In, FEntityManager* EntityManager, FString& OutString);
+	MOVIESCENE_API void AccessorToString(const FReadAccess* In, FEntityManager* EntityManager, FString& OutString);
+	MOVIESCENE_API void AccessorToString(const FWriteAccess* In, FEntityManager* EntityManager, FString& OutString);
+	MOVIESCENE_API void AccessorToString(const FOptionalReadAccess* In, FEntityManager* EntityManager, FString& OutString);
+	MOVIESCENE_API void AccessorToString(const FOptionalWriteAccess* In, FEntityManager* EntityManager, FString& OutString);
+	MOVIESCENE_API void AccessorToString(const FEntityIDAccess*, FEntityManager* EntityManager, FString& OutString);
+	MOVIESCENE_API void OneOfAccessorToString(const FOptionalReadAccess* In, FEntityManager* EntityManager, FString& OutString);
 
 	template<typename... T>
-	void AccessorToString(const TReadOneOf<T...>* In, FEntityManager* EntityManager, FString& OutString)
+	void AccessorToString(const TReadOneOfAccessor<T...>* In, FEntityManager* EntityManager, FString& OutString)
 	{
 		TArray<FString> Strings;
-		VisitTupleElements([&Strings, EntityManager](FReadOptional ReadOptional)
+		VisitTupleElements([&Strings, EntityManager](FOptionalReadAccess ReadOptional)
 		{
 			OneOfAccessorToString(&ReadOptional, EntityManager, Strings.Emplace_GetRef());
 		}, In->ComponentTypes);
@@ -705,10 +360,10 @@ inline bool IsAccessorValid(const TReadOneOrMoreOf<T...>* In)
 	}
 
 	template<typename... T>
-	void AccessorToString(const TReadOneOrMoreOf<T...>* In, FEntityManager* EntityManager, FString& OutString)
+	void AccessorToString(const TReadOneOrMoreOfAccessor<T...>* In, FEntityManager* EntityManager, FString& OutString)
 	{
 		TArray<FString> Strings;
-		VisitTupleElements([&Strings, EntityManager](FReadOptional ReadOptional)
+		VisitTupleElements([&Strings, EntityManager](FOptionalReadAccess ReadOptional)
 			{
 				OneOfAccessorToString(&ReadOptional, EntityManager, Strings.Emplace_GetRef());
 			}, In->ComponentTypes);
@@ -718,11 +373,6 @@ inline bool IsAccessorValid(const TReadOneOrMoreOf<T...>* In)
 
 #endif // UE_MOVIESCENE_ENTITY_DEBUG
 
-template<typename T>
-struct TSupportsDirectEntityIteration
-{
-	enum { Value = T::SupportsDirectEntityIteration };
-};
 
 } // namespace MovieScene
-} // namespace UE
\ No newline at end of file
+} // namespace UE
diff --git a/Engine/Source/Runtime/MovieScene/Public/EntitySystem/MovieSceneComponentPtr.h b/Engine/Source/Runtime/MovieScene/Public/EntitySystem/MovieSceneComponentPtr.h
--- a/Engine/Source/Runtime/MovieScene/Public/EntitySystem/MovieSceneComponentPtr.h
+++ b/Engine/Source/Runtime/MovieScene/Public/EntitySystem/MovieSceneComponentPtr.h
@@ -2,8 +2,6 @@
 
 #pragma once
 
-#include "Templates/IsConst.h"
-#include "HAL/CriticalSection.h"
 #include "EntitySystem/MovieSceneEntitySystemTypes.h"
 
 namespace UE
@@ -12,107 +10,509 @@ namespace MovieScene
 {
 
 
-template<typename ComponentType>
+/** Facade for any component data type */
+template<typename T>
 struct TComponentPtr
 {
-	static constexpr bool bReadOnly = TIsConst<ComponentType>::Value;
+	/** Value type is either a T or const T& for read-only types, or T& for writeable types */
+	using ValueType = typename TChooseClass<TIsConst<T>::Value, typename TCallTraits<T>::ParamType, T&>::Result;
 
 	TComponentPtr()
-		: Header(nullptr)
-		, ComponentPtr(nullptr)
-		, SystemSerialNumber(0)
+		: ComponentPtr(nullptr)
 	{}
 
-	explicit TComponentPtr(FComponentHeader& InHeader, uint16 EntityOffset, uint64 InSystemSerialNumber)
+	explicit TComponentPtr(T* InComponentPtr)
+		: ComponentPtr(InComponentPtr)
+	{}
+
+	TComponentPtr(const TComponentPtr& RHS) = default;
+	TComponentPtr& operator=(const TComponentPtr& RHS) = default;
+
+	TComponentPtr(TComponentPtr&& RHS)
+		: ComponentPtr(RHS.ComponentPtr)
+	{
+		RHS.ComponentPtr = nullptr;
+	}
+
+	TComponentPtr& operator=(TComponentPtr&& RHS)
+	{
+		ComponentPtr = RHS.ComponentPtr;
+		RHS.ComponentPtr = nullptr;
+		return *this;
+	}
+
+	FORCEINLINE explicit operator bool() const
+	{
+		return IsValid();
+	}
+
+	FORCEINLINE bool IsValid() const
+	{
+		return ComponentPtr != nullptr;
+	}
+
+	/** Explicitly convert this component data to its underlying pointer type */
+	FORCEINLINE T* AsPtr() const
+	{
+		return ComponentPtr;
+	}
+
+	/** Retrieve this component data as an array view of the specified size (usually the size will be FEntityAllocation::Num()) */
+	FORCEINLINE TArrayView<T> AsArray(int32 Num) const
+	{
+		return MakeArrayView(ComponentPtr, Num);
+	}
+
+	/** Retrieve a slice of this component data as an array view of the specified size and offset (usually the size will be FEntityAllocation::Num()) */
+	FORCEINLINE TArrayView<T> Slice(int32 StartIndex, int32 Num) const
+	{
+		return MakeArrayView(ComponentPtr + StartIndex, Num);
+	}
+
+	FORCEINLINE T* operator->() const
+	{
+		return ComponentPtr;
+	}
+
+	FORCEINLINE ValueType operator*() const
+	{
+		return *ComponentPtr;
+	}
+
+	FORCEINLINE ValueType operator[](int32 Index) const
+	{
+		return ComponentPtr[Index];
+	}
+
+	FORCEINLINE operator T*() const
+	{
+		return ComponentPtr;
+	}
+
+protected:
+
+	T* ComponentPtr;
+};
+
+
+/**
+ * Typed write-lock for component data
+ */
+template<typename Accessor> struct TComponentLock;
+template<typename LockType> struct TComponentLockMixin;
+
+template<>
+struct TComponentLockMixin<FScopedHeaderReadLock>
+{
+	TComponentLockMixin() = default;
+	explicit TComponentLockMixin(const FComponentHeader* InHeader)
+		: Lock(InHeader)
+	{}
+
+private:
+	FScopedHeaderReadLock Lock;
+};
+
+template<>
+struct TComponentLockMixin<FScopedHeaderWriteLock>
+{
+	TComponentLockMixin() = default;
+	explicit TComponentLockMixin(const FComponentHeader* InHeader, FEntityAllocationWriteContext InWriteContext)
+		: Lock(InHeader, InWriteContext)
+	{}
+
+private:
+	FScopedHeaderWriteLock Lock;
+};
+
+struct FReadErased
+{
+	FReadErased(const FComponentHeader* InHeader, int32 ComponentOffset = 0)
+		: ComponentPtr(InHeader->Components + ComponentOffset*InHeader->Sizeof)
+		, Sizeof(InHeader->Sizeof)
+	{}
+
+	FORCEINLINE const void* AsPtr() const
+	{
+		return ComponentPtr;
+	}
+	FORCEINLINE const void* operator[](int32 Index) const
 	{
-		SystemSerialNumber = InSystemSerialNumber;
-		Header = &InHeader;
+		return ComponentPtr + Sizeof*Index;
+	}
+	FORCEINLINE const void* ComponentAtIndex(int32 Index) const
+	{
+		return ComponentPtr + Sizeof*Index;
+	}
+	FORCEINLINE bool IsValid() const
+	{
+		return ComponentPtr != nullptr;
+	}
+
+protected:
+	FReadErased()
+		: ComponentPtr(nullptr)
+		, Sizeof(0)
+	{}
+
+	const uint8* ComponentPtr;
+	int32 Sizeof;
+};
 
-		if (bReadOnly)
+struct FReadErasedOptional : FReadErased
+{
+	FReadErasedOptional() = default;
+	explicit FReadErasedOptional(const FComponentHeader* InHeader, int32 ComponentOffset = 0)
+	{
+		if (InHeader)
 		{
-			InHeader.ReadWriteLock.ReadLock();
+			ComponentPtr = InHeader->Components + ComponentOffset*InHeader->Sizeof;
+			Sizeof = InHeader->Sizeof;
 		}
-		else
+	}
+};
+
+template<typename T>
+struct TRead : TComponentPtr<const T>
+{
+	using ValueType = typename TCallTraits<T>::ParamType;
+
+	TRead() = default;
+	explicit TRead(const T* ComponentPtr)
+		: TComponentPtr<const T>(ComponentPtr)
+	{}
+	explicit TRead(const FComponentHeader* InHeader, int32 ComponentOffset = 0)
+		: TComponentPtr<const T>(reinterpret_cast<const T*>(InHeader->Components) + ComponentOffset)
+	{}
+	ValueType ComponentAtIndex(int32 Index) const
+	{
+		return (*this)[Index];
+	}
+
+	UE_DEPRECATED(5.00, "It is no longer necessary to call Resolve - please change to using the TComponentPtr API.")
+	const T* Resolve(const FEntityAllocation* Allocation) const
+	{
+		return this->AsPtr();
+	}
+
+	UE_DEPRECATED(5.00, "It is no longer necessary to call ResolveAsArray - please change to TComponentPtr::AsArray.")
+	TArrayView<const T> ResolveAsArray(const FEntityAllocation* Allocation) const
+	{
+		return MakeArrayView(this->AsPtr(), Allocation->Num());
+	}
+};
+using FReadEntityIDs = TRead<FMovieSceneEntityID>;
+
+template<typename T>
+struct TReadOptional : TComponentPtr<const T>
+{
+	TReadOptional() = default;
+	explicit TReadOptional(const FComponentHeader* InHeader, int32 ComponentOffset = 0)
+	{
+		if (InHeader)
 		{
-			InHeader.ReadWriteLock.WriteLock();
+			this->ComponentPtr = reinterpret_cast<const T*>(InHeader->Components) + ComponentOffset;
 		}
+	}
 
-		ComponentPtr = static_cast<ComponentType*>(InHeader.GetValuePtr(EntityOffset));
+	UE_DEPRECATED(5.00, "It is no longer necessary to call Resolve - please change to using the TComponentPtr API.")
+	const T* Resolve(const FEntityAllocation* Allocation) const
+	{
+		return this->AsPtr();
 	}
 
-	~TComponentPtr()
+	UE_DEPRECATED(5.00, "It is no longer necessary to call ResolveAsArray - please change to TComponentPtr::AsArray.")
+	TArrayView<const T> ResolveAsArray(const FEntityAllocation* Allocation) const
 	{
-		Release();
+		const T* Ptr = this->AsPtr();
+		return MakeArrayView(Ptr, Ptr ? Allocation->Num() : 0);
 	}
+};
 
-	TComponentPtr(TComponentPtr&& RHS)
-		: Header(RHS.Header)
-		, ComponentPtr(RHS.ComponentPtr)
-		, SystemSerialNumber(RHS.SystemSerialNumber)
+
+struct FWriteErased
+{
+	FWriteErased()
+		: ComponentPtr(nullptr)
+		, Sizeof(0)
+	{}
+	FWriteErased(const FComponentHeader* InHeader, int32 ComponentOffset = 0)
+		: ComponentPtr(InHeader->Components + ComponentOffset*InHeader->Sizeof)
+		, Sizeof(InHeader->Sizeof)
+	{}
+	FORCEINLINE void* AsPtr() const
 	{
-		RHS.ComponentPtr = nullptr;
-		RHS.Header = nullptr;
+		return ComponentPtr;
 	}
-	TComponentPtr& operator=(TComponentPtr&& RHS)
+	FORCEINLINE void* operator[](int32 Index) const
+	{
+		return ComponentPtr + Sizeof*Index;
+	}
+	FORCEINLINE bool IsValid() const
+	{
+		return ComponentPtr != nullptr;
+	}
+
+protected:
+	uint8* ComponentPtr;
+	int32 Sizeof;
+};
+
+struct FWriteErasedOptional : FWriteErased
+{
+	FWriteErasedOptional() = default;
+	FWriteErasedOptional(const FComponentHeader* InHeader, int32 ComponentOffset = 0)
 	{
-		Release();
+		if (InHeader)
+		{
+			ComponentPtr = InHeader->Components + ComponentOffset*InHeader->Sizeof;
+			Sizeof = InHeader->Sizeof;
+		}
+	}
+};
 
-		ComponentPtr  = RHS.ComponentPtr;
-		Header = RHS.Header;
-		SystemSerialNumber = RHS.SystemSerialNumber;
+template<typename T>
+struct TWrite : TComponentPtr<T>
+{
+	TWrite() = default;
+	explicit TWrite(const FComponentHeader* InHeader, int32 ComponentOffset = 0)
+		: TComponentPtr<T>(reinterpret_cast<T*>(InHeader->Components) + ComponentOffset)
+	{}
 
-		RHS.ComponentPtr = nullptr;
-		RHS.Header = nullptr;
-		return *this;
+	UE_DEPRECATED(5.00, "It is no longer necessary to call Resolve - please change to using the TComponentPtr API.")
+	T* Resolve(const FEntityAllocation* Allocation) const
+	{
+		return this->AsPtr();
 	}
 
-	TComponentPtr(const TComponentPtr&) = delete;
-	void operator=(const TComponentPtr&) = delete;
+	UE_DEPRECATED(5.00, "It is no longer necessary to call ResolveAsArray - please change to TComponentPtr::AsArray.")
+	TArrayView<T> ResolveAsArray(const FEntityAllocation* Allocation) const
+	{
+		return MakeArrayView(this->AsPtr(), Allocation->Num());
+	}
+};
 
-	explicit operator bool()
+template<typename T>
+struct TWriteOptional : TComponentPtr<T>
+{
+	TWriteOptional() = default;
+	explicit TWriteOptional(const FComponentHeader* InHeader, int32 ComponentOffset = 0)
 	{
-		return ComponentPtr != nullptr;
+		if (InHeader)
+		{
+			this->ComponentPtr = reinterpret_cast<T*>(InHeader->Components) + ComponentOffset;
+		}
 	}
 
-	operator ComponentType*() const
+	UE_DEPRECATED(5.00, "It is no longer necessary to call Resolve - please change to using the TComponentPtr API.")
+	T* Resolve(const FEntityAllocation* Allocation) const
 	{
-		return ComponentPtr;
+		return this->AsPtr();
 	}
 
-	ComponentType* operator->() const
+	UE_DEPRECATED(5.00, "It is no longer necessary to call ResolveAsArray - please change to TComponentPtr::AsArray.")
+	TArrayView<T> ResolveAsArray(const FEntityAllocation* Allocation) const
 	{
-		return ComponentPtr;
+		T* Ptr = this->AsPtr();
+		return MakeArrayView(Ptr, Ptr ? Allocation->Num() : 0);
 	}
+};
+
 
-	ComponentType& operator*() const
+
+template<typename AccessorType> struct TComponentLock;
+
+template<>
+struct TComponentLock<FReadErased> : TComponentLockMixin<FScopedHeaderReadLock>, FReadErased
+{
+	explicit TComponentLock(const FComponentHeader* InHeader, int32 ComponentOffset = 0)
+		: TComponentLockMixin<FScopedHeaderReadLock>(InHeader)
+		, FReadErased(InHeader, ComponentOffset)
+	{}
+
+	const void* ComponentAtIndex(int32 Index) const
 	{
-		checkSlow(ComponentPtr);
-		return *ComponentPtr;
+		return (*this)[Index];
+	}
+};
+template<>
+struct TComponentLock<FReadErasedOptional> : TComponentLockMixin<FScopedHeaderReadLock>, FReadErasedOptional
+{
+	TComponentLock() = default;
+	explicit TComponentLock(const FComponentHeader* InHeader, int32 ComponentOffset = 0)
+	{
+		if (InHeader)
+		{
+			*static_cast<TComponentLockMixin<FScopedHeaderReadLock>*>(this) = TComponentLockMixin<FScopedHeaderReadLock>(InHeader);
+			*static_cast<FReadErasedOptional*>(this) = FReadErasedOptional(InHeader, ComponentOffset);
+		}
 	}
 
-private:
+	explicit operator bool() const
+	{
+		return this->ComponentPtr != nullptr;
+	}
+
+	const void* ComponentAtIndex(int32 Index) const
+	{
+		return this->ComponentPtr != nullptr ? (*this)[Index] : nullptr;
+	}
+};
+template<>
+struct TComponentLock<FWriteErased> : TComponentLockMixin<FScopedHeaderWriteLock>, FWriteErased
+{
+	explicit TComponentLock(const FComponentHeader* InHeader, FEntityAllocationWriteContext InWriteContext, int32 ComponentOffset = 0)
+		: TComponentLockMixin<FScopedHeaderWriteLock>(InHeader, InWriteContext)
+		, FWriteErased(InHeader, ComponentOffset)
+	{}
 
-	void Release()
+	void* ComponentAtIndex(int32 Index) const
+	{
+		return (*this)[Index];
+	}
+};
+template<>
+struct TComponentLock<FWriteErasedOptional> : TComponentLockMixin<FScopedHeaderWriteLock>, FWriteErasedOptional
+{
+	TComponentLock() = default;
+	explicit TComponentLock(const FComponentHeader* InHeader, FEntityAllocationWriteContext InWriteContext, int32 ComponentOffset = 0)
 	{
-		if (Header)
+		if (InHeader)
 		{
-			if (bReadOnly)
-			{
-				Header->ReadWriteLock.ReadUnlock();
-			}
-			else
+			*static_cast<TComponentLockMixin<FScopedHeaderWriteLock>*>(this) = TComponentLockMixin<FScopedHeaderWriteLock>(InHeader, InWriteContext);
+			*static_cast<FWriteErasedOptional*>(this) = FWriteErasedOptional(InHeader, ComponentOffset);
+		}
+	}
+
+	explicit operator bool() const
+	{
+		return this->ComponentPtr != nullptr;
+	}
+
+	void* ComponentAtIndex(int32 Index) const
+	{
+		return this->ComponentPtr != nullptr ? (*this)[Index] : nullptr;
+	}
+};
+template<typename T>
+struct TComponentLock<TRead<T>> : TComponentLockMixin<FScopedHeaderReadLock>, TRead<T>
+{
+	explicit TComponentLock(const FComponentHeader* InHeader, int32 ComponentOffset = 0)
+		: TComponentLockMixin<FScopedHeaderReadLock>(InHeader)
+		, TRead<T>(InHeader, ComponentOffset)
+	{}
+};
+
+template<typename T>
+struct TComponentLock<TReadOptional<T>> : TComponentLockMixin<FScopedHeaderReadLock>, TReadOptional<T>
+{
+	TComponentLock() = default;
+	explicit TComponentLock(const FComponentHeader* InHeader, int32 ComponentOffset = 0)
+	{
+		if (InHeader)
+		{
+			*static_cast<TComponentLockMixin<FScopedHeaderReadLock>*>(this) = TComponentLockMixin<FScopedHeaderReadLock>(InHeader);
+			*static_cast<TReadOptional<T>*>(this) = TReadOptional<T>(InHeader, ComponentOffset);
+		}
+	}
+
+	explicit operator bool() const
+	{
+		return this->ComponentPtr != nullptr;
+	}
+
+	const T* ComponentAtIndex(int32 Index) const
+	{
+		return this->ComponentPtr != nullptr ? (*this)[Index] : nullptr;
+	}
+};
+template<typename T>
+struct TComponentLock<TWrite<T>> : TComponentLockMixin<FScopedHeaderWriteLock>, TWrite<T>
+{
+	explicit TComponentLock(const FComponentHeader* InHeader, FEntityAllocationWriteContext InWriteContext, int32 ComponentOffset = 0)
+		: TComponentLockMixin<FScopedHeaderWriteLock>(InHeader, InWriteContext)
+		, TWrite<T>(InHeader, ComponentOffset)
+	{}
+
+	T& ComponentAtIndex(int32 Index) const
+	{
+		return (*this)[Index];
+	}
+};
+
+template<typename T>
+struct TComponentLock<TWriteOptional<T>> : TComponentLockMixin<FScopedHeaderWriteLock>, TWriteOptional<T>
+{
+	TComponentLock() = default;
+	explicit TComponentLock(const FComponentHeader* InHeader, FEntityAllocationWriteContext InWriteContext, int32 ComponentOffset = 0)
+	{
+		if (InHeader)
+		{
+			*static_cast<TComponentLockMixin<FScopedHeaderWriteLock>*>(this) = TComponentLockMixin<FScopedHeaderWriteLock>(InHeader, InWriteContext);
+			*static_cast<TWriteOptional<T>*>(this) = TWriteOptional<T>(InHeader, ComponentOffset);
+		}
+	}
+
+	explicit operator bool() const
+	{
+		return this->ComponentPtr != nullptr;
+	}
+
+	T* ComponentAtIndex(int32 Index) const
+	{
+		return this->ComponentPtr != nullptr ? (*this)[Index] : nullptr;
+	}
+};
+
+
+using FComponentReader = TComponentLock<FReadErased>;
+using FOptionalComponentReader = TComponentLock<FReadErasedOptional>;
+using FComponentWriter = TComponentLock<FWriteErased>;
+using FOptionalComponentWriter = TComponentLock<FWriteErasedOptional>;
+template<typename T> using TComponentReader = TComponentLock<TRead<T>>;
+template<typename T> using TOptionalComponentReader = TComponentLock<TReadOptional<T>>;
+template<typename T> using TComponentWriter = TComponentLock<TWrite<T>>;
+template<typename T> using TOptionalComponentWriter = TComponentLock<TWriteOptional<T>>;
+
+
+template<typename... T> using TMultiComponentLock = TTuple<TComponentLock<T>...>;
+
+template<typename... T>
+struct TMultiComponentData
+{
+	TMultiComponentData(const TMultiComponentLock<T...>& InAggregateLock)
+	{
+		auto Init = [](const auto& InLock, auto& OutData)
+		{
+			if (InLock)
 			{
-				Header->PostWriteComponents(SystemSerialNumber);
-				Header->ReadWriteLock.WriteUnlock();
+				OutData = InLock;
 			}
-		}
+		};
+		VisitTupleElements(Init, InAggregateLock, this->Data);
+	}
+
+	template<int Index>
+	auto* Get() const
+	{
+		return Data.template Get<Index>().AsPtr();
 	}
 
-	FComponentHeader* Header;
-	ComponentType* ComponentPtr;
-	uint64 SystemSerialNumber;
+	template<int Index>
+	auto GetAsArray(int32 ArraySize) const
+	{
+		auto* ComponentPtr = Data.template Get<Index>().AsPtr();
+		return MakeArrayView(ComponentPtr, ComponentPtr ? ArraySize : 0);
+	}
+
+private:
+
+	TTuple<T...> Data;
 };
 
+template<typename... T> using TMultiReadOptional = TMultiComponentData<TReadOptional<T>...>;
+template<typename... T> using TReadOneOrMoreOf = TMultiComponentData<TReadOptional<T>...>;
+template<typename... T> using TReadOneOf = TMultiComponentData<TReadOptional<T>...>;
 
 } // namespace UE
-} // namespace MovieScene
\ No newline at end of file
+} // namespace MovieScene
diff --git a/Engine/Source/Runtime/MovieScene/Public/EntitySystem/MovieSceneComponentRegistry.h b/Engine/Source/Runtime/MovieScene/Public/EntitySystem/MovieSceneComponentRegistry.h
--- a/Engine/Source/Runtime/MovieScene/Public/EntitySystem/MovieSceneComponentRegistry.h
+++ b/Engine/Source/Runtime/MovieScene/Public/EntitySystem/MovieSceneComponentRegistry.h
@@ -141,49 +141,5 @@ private:
 };
 
 
-template<typename T>
-TComponentTypeID<T> FComponentRegistry::NewComponentType(const TCHAR* const DebugName, EComponentTypeFlags Flags)
-{
-	static const uint32 ComponentTypeSize = sizeof(T);
-	static_assert(ComponentTypeSize < TNumericLimits<decltype(FComponentTypeInfo::Sizeof)>::Max(), "Type too large to be used as component data");
-
-	static const uint32 Alignment = alignof(T);
-	static_assert(Alignment < TNumericLimits<decltype(FComponentTypeInfo::Alignment)>::Max(), "Type alignment too large to be used as component data");
-
-	FComponentTypeInfo NewTypeInfo;
-
-	NewTypeInfo.Sizeof                     = ComponentTypeSize;
-	NewTypeInfo.Alignment                  = Alignment;
-	NewTypeInfo.bIsZeroConstructType       = TIsZeroConstructType<T>::Value;
-	NewTypeInfo.bIsTriviallyDestructable   = TIsTriviallyDestructible<T>::Value;
-	NewTypeInfo.bIsTriviallyCopyAssignable = TIsTriviallyCopyAssignable<T>::Value;
-	NewTypeInfo.bIsPreserved               = EnumHasAnyFlags(Flags, EComponentTypeFlags::Preserved);
-	NewTypeInfo.bIsMigratedToOutput        = EnumHasAnyFlags(Flags, EComponentTypeFlags::MigrateToOutput);
-	NewTypeInfo.bHasReferencedObjects      = !TIsSame< FNotImplemented*, decltype( AddReferencedObjectForComponent((FReferenceCollector*)0, (T*)0) ) >::Value;
-
-#if UE_MOVIESCENE_ENTITY_DEBUG
-	NewTypeInfo.DebugInfo                = MakeUnique<FComponentTypeDebugInfo>();
-	NewTypeInfo.DebugInfo->DebugName     = DebugName;
-	NewTypeInfo.DebugInfo->DebugTypeName = GetGeneratedTypeName<T>();
-	NewTypeInfo.DebugInfo->Type          = TComponentDebugType<T>::Type;
-#endif
-
-	if (!NewTypeInfo.bIsZeroConstructType || !NewTypeInfo.bIsTriviallyDestructable || !NewTypeInfo.bIsTriviallyCopyAssignable || NewTypeInfo.bHasReferencedObjects)
-	{
-		NewTypeInfo.MakeComplexComponentOps<T>();
-	}
-
-	FComponentTypeID    NewTypeID = NewComponentTypeInternal(MoveTemp(NewTypeInfo));
-	TComponentTypeID<T> TypedTypeID = NewTypeID.ReinterpretCast<T>();
-
-	if (EnumHasAnyFlags(Flags, EComponentTypeFlags::CopyToChildren))
-	{
-		Factories.DefineChildComponent(TDuplicateChildEntityInitializer<T>(TypedTypeID));
-	}
-
-	return TypedTypeID;
-}
-
-
-}
-}
\ No newline at end of file
+} // namespace MovieScene
+} // namespace UE
\ No newline at end of file
diff --git a/Engine/Source/Runtime/MovieScene/Public/EntitySystem/MovieSceneEntityFactory.h b/Engine/Source/Runtime/MovieScene/Public/EntitySystem/MovieSceneEntityFactory.h
--- a/Engine/Source/Runtime/MovieScene/Public/EntitySystem/MovieSceneEntityFactory.h
+++ b/Engine/Source/Runtime/MovieScene/Public/EntitySystem/MovieSceneEntityFactory.h
@@ -9,12 +9,8 @@
 #include "Misc/InlineValue.h"
 
 #include "EntitySystem/MovieSceneEntityIDs.h"
-#include "EntitySystem/MovieSceneEntityRange.h"
-#include "EntitySystem/MovieSceneEntitySystemTask.h"
 #include "EntitySystem/MovieSceneEntitySystemTypes.h"
 #include "EntitySystem/MovieSceneEntityFactoryTypes.h"
-#include "EntitySystem/MovieSceneComponentAccessors.h"
-#include "EntitySystem/MovieSceneSequenceInstanceHandle.h"
 #include "EntitySystem/MovieSceneEntitySystemDirectedGraph.h"
 
 
@@ -23,184 +19,6 @@ namespace UE
 namespace MovieScene
 {
 
-struct FInstanceRegistry;
-
-template<typename ParentComponentType, typename ChildComponentType>
-struct TChildEntityInitializer : FChildEntityInitializer
-{
-	explicit TChildEntityInitializer(TComponentTypeID<ParentComponentType> InParentComponent, TComponentTypeID<ChildComponentType> InChildComponent)
-		: FChildEntityInitializer(InParentComponent, InChildComponent)
-	{}
-
-	TComponentTypeID<ParentComponentType> GetParentComponent() const
-	{
-		return ParentComponent.ReinterpretCast<ParentComponentType>();
-	}
-
-	TComponentTypeID<ChildComponentType> GetChildComponent() const
-	{
-		return ChildComponent.ReinterpretCast<ChildComponentType>();
-	}
-
-	TArrayView<const ParentComponentType> GetParentComponents(const FEntityAllocation* Allocation) const
-	{
-		return TRead<ParentComponentType>(GetParentComponent()).ResolveAsArray(Allocation);
-	}
-
-	TArrayView<ChildComponentType> GetChildComponents(const FEntityAllocation* Allocation) const
-	{
-		return TWrite<ChildComponentType>(GetChildComponent()).ResolveAsArray(Allocation);
-	}
-};
-
-// Callback must be compatible with form void(const ParentComponentType, ChildComponentType&);
-template<typename ParentComponentType, typename ChildComponentType, typename InitializerType>
-struct TStaticChildEntityInitializer : TChildEntityInitializer<ParentComponentType, ChildComponentType>
-{
-	InitializerType Callback;
-
-	explicit TStaticChildEntityInitializer(TComponentTypeID<ParentComponentType> InParentComponent, TComponentTypeID<ChildComponentType> InChildComponent, InitializerType InCallback)
-		: TChildEntityInitializer<ParentComponentType, ChildComponentType>(InParentComponent, InChildComponent)
-		, Callback(InCallback)
-	{}
-
-	virtual void Run(const FEntityRange& ChildRange, const FEntityAllocation* ParentAllocation, TArrayView<const int32> ParentAllocationOffsets)
-	{
-		TArrayView<const ParentComponentType> ParentComponents = TRead<ParentComponentType>(this->GetParentComponent()).ResolveAsArray(ParentAllocation);
-		TArrayView<ChildComponentType>        ChildComponents  = TWrite<ChildComponentType>(this->GetChildComponent()).ResolveAsArray(ChildRange.Allocation);
-
-		for (int32 Index = 0; Index < ChildRange.Num; ++Index)
-		{
-			const int32 ParentIndex = ParentAllocationOffsets[Index];
-			const int32 ChildIndex  = ChildRange.ComponentStartOffset + Index;
-
-			Callback(ParentComponents[ParentIndex], ChildComponents[ChildIndex]);
-		}
-	}
-};
-
-template<typename ComponentType>
-struct TDuplicateChildEntityInitializer : FChildEntityInitializer
-{
-	explicit TDuplicateChildEntityInitializer(TComponentTypeID<ComponentType> InComponent)
-		: FChildEntityInitializer(InComponent, InComponent)
-	{}
-
-	TComponentTypeID<ComponentType> GetComponent() const
-	{
-		return ParentComponent.ReinterpretCast<ComponentType>();
-	}
-
-	virtual void Run(const FEntityRange& ChildRange, const FEntityAllocation* ParentAllocation, TArrayView<const int32> ParentAllocationOffsets)
-	{
-		TArrayView<const ComponentType> ParentComponents = TRead<ComponentType>(GetComponent()).ResolveAsArray(ParentAllocation);
-		TArrayView<ComponentType>       ChildComponents  = TWrite<ComponentType>(GetComponent()).ResolveAsArray(ChildRange.Allocation).Slice(ChildRange.ComponentStartOffset, ChildRange.Num);
-
-		for (int32 Index = 0; Index < ParentAllocationOffsets.Num(); ++Index)
-		{
-			ChildComponents[Index] = ParentComponents[ParentAllocationOffsets[Index]];
-		}
-	}
-};
-
-struct FObjectFactoryBatch : FChildEntityFactory
-{
-	void Add(int32 EntityIndex, UObject* BoundObject);
-
-	virtual void GenerateDerivedType(FComponentMask& OutNewEntityType) override;
-
-	virtual void InitializeAllocation(UMovieSceneEntitySystemLinker* Linker, const FComponentMask& ParentType, const FComponentMask& ChildType, const FEntityAllocation* ParentAllocation, TArrayView<const int32> ParentAllocationOffsets, const FEntityRange& InChildEntityRange) override;
-
-	virtual void PostInitialize(UMovieSceneEntitySystemLinker* InLinker) override;
-
-	virtual void ResolveObjects(FInstanceRegistry* InstanceRegistry, FInstanceHandle InstanceHandle, int32 InEntityIndex, const FGuid& ObjectBinding) = 0;
-
-	TMap<TTuple<UObject*, FMovieSceneEntityID>, FMovieSceneEntityID>* StaleEntitiesToPreserve;
-
-private:
-	TSortedMap<FMovieSceneEntityID, FMovieSceneEntityID> PreservedEntities;
-	TArray<UObject*> ObjectsToAssign;
-};
-
-struct FBoundObjectTask
-{
-	FBoundObjectTask(UMovieSceneEntitySystemLinker* InLinker);
-	virtual ~FBoundObjectTask(){}
-
-	virtual FObjectFactoryBatch& AddBatch(const FEntityAllocation* Parent) = 0;
-	virtual void Apply() = 0;
-
-	void ForEachAllocation(const FEntityAllocation* Allocation, FReadEntityIDs EntityIDAccessor, TRead<FInstanceHandle> InstanceAccessor, TRead<FGuid> ObjectBindingAccessor);
-
-	void PostTask();
-
-private:
-
-	TMap<TTuple<UObject*, FMovieSceneEntityID>, FMovieSceneEntityID> StaleEntitiesToPreserve;
-	TArray<FMovieSceneEntityID> EntitiesToDiscard;
-
-protected:
-
-	UMovieSceneEntitySystemLinker* Linker;
-};
-
-template<typename BatchType>
-struct TBoundObjectTask : FBoundObjectTask
-{
-	TBoundObjectTask(UMovieSceneEntitySystemLinker* InLinker)
-		: FBoundObjectTask(InLinker)
-	{}
-
-private:
-
-	virtual FObjectFactoryBatch& AddBatch(const FEntityAllocation* Parent) override
-	{
-		return Batches.Add(Parent);
-	}
-
-	virtual void Apply() override
-	{
-		for (TTuple<const FEntityAllocation*, BatchType>& Pair : Batches)
-		{
-			// Determine the type for the new entities
-			const FEntityAllocation* ParentAllocation = Pair.Key;
-			if (Pair.Value.Num() != 0)
-			{
-				Pair.Value.Apply(Linker, ParentAllocation);
-			}
-		}
-	}
-
-	TMap<const FEntityAllocation*, BatchType> Batches;
-};
-
-
-
-template<typename ComponentTypeA, typename ComponentTypeB>
-struct TMutualEntityInitializer : FMutualEntityInitializer
-{
-	using CallbackType = void (*)(ComponentTypeA* ComponentsA, ComponentTypeB* ComponentsB, int32 Num);
-
-	explicit TMutualEntityInitializer(TComponentTypeID<ComponentTypeA> InComponentA, TComponentTypeID<ComponentTypeB> InComponentB, CallbackType InCallback)
-		: FMutualEntityInitializer(InComponentA, InComponentB)
-		, Callback(InCallback)
-	{}
-
-private:
-
-	virtual void Run(const FEntityRange& Range) override
-	{
-		TEntityRange<ComponentTypeA, ComponentTypeB> Iter = FEntityTaskBuilder()
-			.Write(ComponentA.ReinterpretCast<ComponentTypeA>())
-			.Write(ComponentB.ReinterpretCast<ComponentTypeB>())
-			.IterateRange(Range);
-
-		Callback(Iter.template GetRawUnchecked<0>(), Iter.template GetRawUnchecked<1>(), Iter.Num());
-	}
-
-	CallbackType Callback;
-};
-
 
 enum class EComplexInclusivityFilterMode
 {
@@ -287,25 +105,18 @@ struct FEntityFactories
 
 	/**
 	 * Makes the given component automatically copied from a parent entity to all its children entities.
+	 * @note: include "EntitySystem/MovieSceneEntityFactoryTemplates.h" for definition
 	 */
 	template<typename ComponentType>
-	void DuplicateChildComponent(TComponentTypeID<ComponentType> InComponent)
-	{
-		DefineChildComponent(TDuplicateChildEntityInitializer<ComponentType>(InComponent));
-	}
+	void DuplicateChildComponent(TComponentTypeID<ComponentType> InComponent);
 
 	/**
 	 * Specifies that if a component is present on a parent entity, the given child component should
 	 * be created on any child entity, and initialized with the given initializer.
+	 * @note: include "EntitySystem/MovieSceneEntityFactoryTemplates.h" for definition
 	 */
 	template<typename ParentComponent, typename ChildComponent, typename InitializerCallback>
-	void DefineChildComponent(TComponentTypeID<ParentComponent> InParentType, TComponentTypeID<ChildComponent> InChildType, InitializerCallback&& InInitializer)
-	{
-		using FInitializer = TStaticChildEntityInitializer<ParentComponent, ChildComponent, InitializerCallback>;
-
-		DefineChildComponent(InParentType, InChildType);
-		ChildInitializers.Add(FInitializer(InParentType, InChildType, Forward<InitializerCallback>(InInitializer)));
-	}
+	void DefineChildComponent(TComponentTypeID<ParentComponent> InParentType, TComponentTypeID<ChildComponent> InChildType, InitializerCallback&& InInitializer);
 
 	/**
 	 * Adds the definition for a child component. The helper methods above are easier and preferrable.
@@ -327,14 +138,10 @@ struct FEntityFactories
 
 	/**
 	 * Specifies that if an entity matches the given filter, the specified components should be created on it.
+	 * @note: include "EntitySystem/MovieSceneEntityFactoryTemplates.h" for definition
 	 */
 	template<typename... ComponentTypes>
-	void DefineComplexInclusiveComponents(const FComplexInclusivityFilter& InFilter, ComponentTypes... InComponents)
-	{
-		FComponentMask ComponentsToInclude { InComponents... };
-		FComplexInclusivity NewComplexInclusivity { InFilter, ComponentsToInclude };
-		DefineComplexInclusiveComponents(NewComplexInclusivity);
-	}
+	void DefineComplexInclusiveComponents(const FComplexInclusivityFilter& InFilter, ComponentTypes... InComponents);
 
 	/**
 	 * Specifies that if an entity matches the given filter, the specified component should be created on it.
diff --git a/Engine/Source/Runtime/MovieScene/Public/EntitySystem/MovieSceneEntityFactoryTemplates.h b/Engine/Source/Runtime/MovieScene/Public/EntitySystem/MovieSceneEntityFactoryTemplates.h
new file mode 100644
--- /dev/null
+++ b/Engine/Source/Runtime/MovieScene/Public/EntitySystem/MovieSceneEntityFactoryTemplates.h
@@ -0,0 +1,272 @@
+// Copyright Epic Games, Inc. All Rights Reserved.
+
+#pragma once
+
+#include "Containers/Array.h"
+#include "Containers/ArrayView.h"
+#include "Containers/Map.h"
+#include "Containers/SortedMap.h"
+
+#include "EntitySystem/MovieSceneEntityIDs.h"
+#include "EntitySystem/MovieSceneEntityRange.h"
+#include "EntitySystem/MovieSceneEntityFactory.h"
+#include "EntitySystem/MovieSceneEntitySystemTask.h"
+#include "EntitySystem/MovieSceneComponentRegistry.h"
+#include "EntitySystem/MovieSceneEntitySystemTypes.h"
+#include "EntitySystem/MovieSceneComponentAccessors.h"
+#include "EntitySystem/MovieSceneSequenceInstanceHandle.h"
+#include "EntitySystem/MovieSceneEntitySystemDirectedGraph.h"
+
+
+namespace UE
+{
+namespace MovieScene
+{
+
+struct FInstanceRegistry;
+
+template<typename ParentComponentType, typename ChildComponentType>
+struct TChildEntityInitializer : FChildEntityInitializer
+{
+	explicit TChildEntityInitializer(TComponentTypeID<ParentComponentType> InParentComponent, TComponentTypeID<ChildComponentType> InChildComponent)
+		: FChildEntityInitializer(InParentComponent, InChildComponent)
+	{}
+
+	TComponentTypeID<ParentComponentType> GetParentComponent() const
+	{
+		return ParentComponent.ReinterpretCast<ParentComponentType>();
+	}
+
+	TComponentTypeID<ChildComponentType> GetChildComponent() const
+	{
+		return ChildComponent.ReinterpretCast<ChildComponentType>();
+	}
+
+	TComponentLock<TRead<ParentComponentType>> GetParentComponents(const FEntityAllocation* Allocation) const
+	{
+		return Allocation->ReadComponents(GetParentComponent());
+	}
+
+	TComponentLock<TWrite<ChildComponentType>> GetChildComponents(const FEntityAllocation* Allocation) const
+	{
+		return Allocation->WriteComponents(GetChildComponent(), FEntityAllocationWriteContext::NewAllocation());
+	}
+};
+
+// Callback must be compatible with form void(const ParentComponentType, ChildComponentType&);
+template<typename ParentComponentType, typename ChildComponentType, typename InitializerType>
+struct TStaticChildEntityInitializer : TChildEntityInitializer<ParentComponentType, ChildComponentType>
+{
+	InitializerType Callback;
+
+	explicit TStaticChildEntityInitializer(TComponentTypeID<ParentComponentType> InParentComponent, TComponentTypeID<ChildComponentType> InChildComponent, InitializerType InCallback)
+		: TChildEntityInitializer<ParentComponentType, ChildComponentType>(InParentComponent, InChildComponent)
+		, Callback(InCallback)
+	{}
+
+	virtual void Run(const FEntityRange& ChildRange, const FEntityAllocation* ParentAllocation, TArrayView<const int32> ParentAllocationOffsets)
+	{
+		TComponentLock<TRead<ParentComponentType>> ParentComponents = ParentAllocation->ReadComponents(this->GetParentComponent());
+		TComponentLock<TWrite<ChildComponentType>> ChildComponents  = ChildRange.Allocation->WriteComponents(this->GetChildComponent(), FEntityAllocationWriteContext::NewAllocation());
+
+		for (int32 Index = 0; Index < ChildRange.Num; ++Index)
+		{
+			const int32 ParentIndex = ParentAllocationOffsets[Index];
+			const int32 ChildIndex  = ChildRange.ComponentStartOffset + Index;
+
+			Callback(ParentComponents[ParentIndex], ChildComponents[ChildIndex]);
+		}
+	}
+};
+
+template<typename ComponentType>
+struct TDuplicateChildEntityInitializer : FChildEntityInitializer
+{
+	explicit TDuplicateChildEntityInitializer(TComponentTypeID<ComponentType> InComponent)
+		: FChildEntityInitializer(InComponent, InComponent)
+	{}
+
+	TComponentTypeID<ComponentType> GetComponent() const
+	{
+		return ParentComponent.ReinterpretCast<ComponentType>();
+	}
+
+	virtual void Run(const FEntityRange& ChildRange, const FEntityAllocation* ParentAllocation, TArrayView<const int32> ParentAllocationOffsets)
+	{
+		TComponentLock<TRead<ComponentType>>  ParentComponents = ParentAllocation->ReadComponents(GetComponent());
+		TComponentLock<TWrite<ComponentType>> ChildComponents  = ChildRange.Allocation->WriteComponents(GetComponent(), FEntityAllocationWriteContext::NewAllocation());
+
+		TArrayView<ComponentType> ChildComponentSlice = ChildComponents.Slice(ChildRange.ComponentStartOffset, ChildRange.Num);
+		for (int32 Index = 0; Index < ParentAllocationOffsets.Num(); ++Index)
+		{
+			ChildComponentSlice[Index] = ParentComponents[ParentAllocationOffsets[Index]];
+		}
+	}
+};
+
+struct FObjectFactoryBatch : FChildEntityFactory
+{
+	void Add(int32 EntityIndex, UObject* BoundObject);
+
+	virtual void GenerateDerivedType(FComponentMask& OutNewEntityType) override;
+
+	virtual void InitializeAllocation(UMovieSceneEntitySystemLinker* Linker, const FComponentMask& ParentType, const FComponentMask& ChildType, const FEntityAllocation* ParentAllocation, TArrayView<const int32> ParentAllocationOffsets, const FEntityRange& InChildEntityRange) override;
+
+	virtual void PostInitialize(UMovieSceneEntitySystemLinker* InLinker) override;
+
+	virtual void ResolveObjects(FInstanceRegistry* InstanceRegistry, FInstanceHandle InstanceHandle, int32 InEntityIndex, const FGuid& ObjectBinding) = 0;
+
+	TMap<TTuple<UObject*, FMovieSceneEntityID>, FMovieSceneEntityID>* StaleEntitiesToPreserve;
+
+private:
+	TSortedMap<FMovieSceneEntityID, FMovieSceneEntityID> PreservedEntities;
+	TArray<UObject*> ObjectsToAssign;
+};
+
+struct FBoundObjectTask
+{
+	FBoundObjectTask(UMovieSceneEntitySystemLinker* InLinker);
+	virtual ~FBoundObjectTask(){}
+
+	virtual FObjectFactoryBatch& AddBatch(const FEntityAllocation* Parent) = 0;
+	virtual void Apply() = 0;
+
+	void ForEachAllocation(const FEntityAllocation* Allocation, FReadEntityIDs EntityIDs, TRead<FInstanceHandle> Instances, TRead<FGuid> ObjectBindings);
+
+	void PostTask();
+
+private:
+
+	TMap<TTuple<UObject*, FMovieSceneEntityID>, FMovieSceneEntityID> StaleEntitiesToPreserve;
+	TArray<FMovieSceneEntityID> EntitiesToDiscard;
+
+protected:
+
+	UMovieSceneEntitySystemLinker* Linker;
+};
+
+template<typename BatchType>
+struct TBoundObjectTask : FBoundObjectTask
+{
+	TBoundObjectTask(UMovieSceneEntitySystemLinker* InLinker)
+		: FBoundObjectTask(InLinker)
+	{}
+
+private:
+
+	virtual FObjectFactoryBatch& AddBatch(const FEntityAllocation* Parent) override
+	{
+		return Batches.Add(Parent);
+	}
+
+	virtual void Apply() override
+	{
+		for (TTuple<const FEntityAllocation*, BatchType>& Pair : Batches)
+		{
+			// Determine the type for the new entities
+			const FEntityAllocation* ParentAllocation = Pair.Key;
+			if (Pair.Value.Num() != 0)
+			{
+				Pair.Value.Apply(Linker, ParentAllocation);
+			}
+		}
+	}
+
+	TMap<const FEntityAllocation*, BatchType> Batches;
+};
+
+
+
+template<typename ComponentTypeA, typename ComponentTypeB>
+struct TMutualEntityInitializer : FMutualEntityInitializer
+{
+	using CallbackType = void (*)(ComponentTypeA* ComponentsA, ComponentTypeB* ComponentsB, int32 Num);
+
+	explicit TMutualEntityInitializer(TComponentTypeID<ComponentTypeA> InComponentA, TComponentTypeID<ComponentTypeB> InComponentB, CallbackType InCallback)
+		: FMutualEntityInitializer(InComponentA, InComponentB)
+		, Callback(InCallback)
+	{}
+
+private:
+
+	virtual void Run(const FEntityRange& Range) override
+	{
+		TComponentLock<TWrite<ComponentTypeA>> A = Range.Allocation->WriteComponents(ComponentA.ReinterpretCast<ComponentTypeA>(), FEntityAllocationWriteContext::NewAllocation());
+		TComponentLock<TWrite<ComponentTypeB>> B = Range.Allocation->WriteComponents(ComponentB.ReinterpretCast<ComponentTypeB>(), FEntityAllocationWriteContext::NewAllocation());
+
+		Callback(&A[Range.ComponentStartOffset], &B[Range.ComponentStartOffset], Range.Num);
+	}
+
+	CallbackType Callback;
+};
+
+
+
+template<typename ComponentType>
+inline void FEntityFactories::DuplicateChildComponent(TComponentTypeID<ComponentType> InComponent)
+{
+	DefineChildComponent(TDuplicateChildEntityInitializer<ComponentType>(InComponent));
+}
+
+template<typename ParentComponent, typename ChildComponent, typename InitializerCallback>
+inline void FEntityFactories::DefineChildComponent(TComponentTypeID<ParentComponent> InParentType, TComponentTypeID<ChildComponent> InChildType, InitializerCallback&& InInitializer)
+{
+	using FInitializer = TStaticChildEntityInitializer<ParentComponent, ChildComponent, InitializerCallback>;
+
+	DefineChildComponent(InParentType, InChildType);
+	ChildInitializers.Add(FInitializer(InParentType, InChildType, Forward<InitializerCallback>(InInitializer)));
+}
+
+template<typename... ComponentTypes>
+inline void FEntityFactories::DefineComplexInclusiveComponents(const FComplexInclusivityFilter& InFilter, ComponentTypes... InComponents)
+{
+	FComponentMask ComponentsToInclude { InComponents... };
+	FComplexInclusivity NewComplexInclusivity { InFilter, ComponentsToInclude };
+	DefineComplexInclusiveComponents(NewComplexInclusivity);
+}
+
+template<typename T>
+TComponentTypeID<T> FComponentRegistry::NewComponentType(const TCHAR* const DebugName, EComponentTypeFlags Flags)
+{
+	static const uint32 ComponentTypeSize = sizeof(T);
+	static_assert(ComponentTypeSize < TNumericLimits<decltype(FComponentTypeInfo::Sizeof)>::Max(), "Type too large to be used as component data");
+
+	static const uint32 Alignment = alignof(T);
+	static_assert(Alignment < TNumericLimits<decltype(FComponentTypeInfo::Alignment)>::Max(), "Type alignment too large to be used as component data");
+
+	FComponentTypeInfo NewTypeInfo;
+
+	NewTypeInfo.Sizeof                     = ComponentTypeSize;
+	NewTypeInfo.Alignment                  = Alignment;
+	NewTypeInfo.bIsZeroConstructType       = TIsZeroConstructType<T>::Value;
+	NewTypeInfo.bIsTriviallyDestructable   = TIsTriviallyDestructible<T>::Value;
+	NewTypeInfo.bIsTriviallyCopyAssignable = TIsTriviallyCopyAssignable<T>::Value;
+	NewTypeInfo.bIsPreserved               = EnumHasAnyFlags(Flags, EComponentTypeFlags::Preserved);
+	NewTypeInfo.bIsMigratedToOutput        = EnumHasAnyFlags(Flags, EComponentTypeFlags::MigrateToOutput);
+	NewTypeInfo.bHasReferencedObjects      = !TIsSame< FNotImplemented*, decltype( AddReferencedObjectForComponent((FReferenceCollector*)0, (T*)0) ) >::Value;
+
+#if UE_MOVIESCENE_ENTITY_DEBUG
+	NewTypeInfo.DebugInfo                = MakeUnique<FComponentTypeDebugInfo>();
+	NewTypeInfo.DebugInfo->DebugName     = DebugName;
+	NewTypeInfo.DebugInfo->DebugTypeName = GetGeneratedTypeName<T>();
+	NewTypeInfo.DebugInfo->Type          = TComponentDebugType<T>::Type;
+#endif
+
+	if (!NewTypeInfo.bIsZeroConstructType || !NewTypeInfo.bIsTriviallyDestructable || !NewTypeInfo.bIsTriviallyCopyAssignable || NewTypeInfo.bHasReferencedObjects)
+	{
+		NewTypeInfo.MakeComplexComponentOps<T>();
+	}
+
+	FComponentTypeID    NewTypeID = NewComponentTypeInternal(MoveTemp(NewTypeInfo));
+	TComponentTypeID<T> TypedTypeID = NewTypeID.ReinterpretCast<T>();
+
+	if (EnumHasAnyFlags(Flags, EComponentTypeFlags::CopyToChildren))
+	{
+		Factories.DefineChildComponent(TDuplicateChildEntityInitializer<T>(TypedTypeID));
+	}
+
+	return TypedTypeID;
+}
+
+}	// using namespace MovieScene
+}	// using namespace UE
diff --git a/Engine/Source/Runtime/MovieScene/Public/EntitySystem/MovieSceneEntityManager.h b/Engine/Source/Runtime/MovieScene/Public/EntitySystem/MovieSceneEntityManager.h
--- a/Engine/Source/Runtime/MovieScene/Public/EntitySystem/MovieSceneEntityManager.h
+++ b/Engine/Source/Runtime/MovieScene/Public/EntitySystem/MovieSceneEntityManager.h
@@ -15,7 +15,6 @@
 #include "EntitySystem/MovieSceneEntityFactoryTypes.h"
 #include "EntitySystem/MovieSceneComponentTypeHandler.h"
 #include "EntitySystem/MovieSceneComponentTypeInfo.h"
-#include "EntitySystem/MovieSceneComponentPtr.h"
 #include "Algo/AnyOf.h"
 #include "Algo/AllOf.h"
 
@@ -400,19 +399,19 @@ public:
 	 * @return A scoped component ptr that points directly to the entity's component, or nullptr if it does not exist
 	 */
 	template<typename T>
-	TComponentPtr<const T> ReadComponent(FMovieSceneEntityID Entity, TComponentTypeID<T> ComponentTypeID) const
+	TComponentLock<TReadOptional<T>> ReadComponent(FMovieSceneEntityID Entity, TComponentTypeID<T> ComponentTypeID) const
 	{
 		check(Entity);
 
 		if (!ComponentTypeID)
 		{
-			return TComponentPtr<const T>();
+			return TComponentLock<TReadOptional<T>>();
 		}
 
 		FEntityLocation Location = EntityLocations[Entity.AsIndex()];
 		if (!Location.IsValid())
 		{
-			return TComponentPtr<const T>();
+			return TComponentLock<TReadOptional<T>>();
 		}
 
 		FEntityAllocation* Allocation = EntityAllocations[Location.GetAllocationIndex()];
@@ -422,11 +421,11 @@ public:
 		{
 			if (Header.ComponentType == ComponentTypeID)
 			{
-				return TComponentPtr<const T>(Header, ComponentOffset, SystemSerialNumber);
+				return TComponentLock<TReadOptional<T>>(&Header, ComponentOffset);
 			}
 		}
 
-		return TComponentPtr<const T>();
+		return TComponentLock<TReadOptional<T>>();
 	}
 
 
@@ -441,7 +440,7 @@ public:
 	template<typename T>
 	T ReadComponentChecked(FMovieSceneEntityID Entity, TComponentTypeID<T> ComponentTypeID) const
 	{
-		TComponentPtr<const T> Value = ReadComponent(Entity, ComponentTypeID);
+		TComponentLock<TReadOptional<T>> Value = ReadComponent(Entity, ComponentTypeID);
 		check(Value);
 		return *Value;
 	}
@@ -456,19 +455,19 @@ public:
 	 * @return A scoped component ptr that points directly to the entity's component, or nullptr if it does not exist
 	 */
 	template<typename T>
-	TComponentPtr<T> WriteComponent(FMovieSceneEntityID Entity, TComponentTypeID<T> ComponentTypeID)
+	TComponentLock<TWriteOptional<T>> WriteComponent(FMovieSceneEntityID Entity, TComponentTypeID<T> ComponentTypeID)
 	{
 		check(Entity);
 
 		if (!ComponentTypeID)
 		{
-			return TComponentPtr<T>();
+			return TComponentLock<TWriteOptional<T>>();
 		}
 
 		FEntityLocation Location = EntityLocations[Entity.AsIndex()];
 		if (!Location.IsValid())
 		{
-			return TComponentPtr<T>();
+			return TComponentLock<TWriteOptional<T>>();
 		}
 
 		FEntityAllocation* Allocation = EntityAllocations[Location.GetAllocationIndex()];
@@ -478,11 +477,11 @@ public:
 		{
 			if (Header.ComponentType == ComponentTypeID)
 			{
-				return TComponentPtr<T>(Header, ComponentOffset, SystemSerialNumber);
+				return TComponentLock<TWriteOptional<T>>(Header, FEntityAllocationWriteContext(*this), ComponentOffset);
 			}
 		}
 
-		return TComponentPtr<T>();
+		return TComponentLock<TWriteOptional<T>>();
 	}
 
 
@@ -497,7 +496,7 @@ public:
 	template<typename T, typename ValueType>
 	void WriteComponentChecked(FMovieSceneEntityID Entity, TComponentTypeID<T> ComponentTypeID, ValueType&& Value) const
 	{
-		TComponentPtr<T> ComponentPtr = WriteComponent(Entity, ComponentTypeID);
+		TComponentLock<TWriteOptional<T>> ComponentPtr = WriteComponent(Entity, ComponentTypeID);
 		check(ComponentPtr);
 		*ComponentPtr = Forward<ValueType>(Value);
 	}
diff --git a/Engine/Source/Runtime/MovieScene/Public/EntitySystem/MovieSceneEntitySystemTask.h b/Engine/Source/Runtime/MovieScene/Public/EntitySystem/MovieSceneEntitySystemTask.h
--- a/Engine/Source/Runtime/MovieScene/Public/EntitySystem/MovieSceneEntitySystemTask.h
+++ b/Engine/Source/Runtime/MovieScene/Public/EntitySystem/MovieSceneEntitySystemTask.h
@@ -9,6 +9,7 @@
 #include "EntitySystem/MovieSceneComponentAccessors.h"
 #include "EntitySystem/MovieSceneEntityManager.h"
 #include "EntitySystem/MovieSceneSystemTaskDependencies.h"
+#include "EntitySystem/MovieSceneComponentPtr.h"
 
 #include "Templates/AndOrNot.h"
 
@@ -22,6 +23,12 @@ namespace MovieScene
 DECLARE_CYCLE_STAT(TEXT("Aquire Component Access Locks"), MovieSceneEval_AquireComponentAccessLocks, STATGROUP_MovieSceneECS);
 DECLARE_CYCLE_STAT(TEXT("Release Component Access Locks"), MovieSceneEval_ReleaseComponentAccessLocks, STATGROUP_MovieSceneECS);
 
+template<typename> struct TReadAccess;
+template<typename> struct TOptionalReadAccess;
+template<typename> struct TWriteAccess;
+template<typename> struct TOptionalWriteAccess;
+template<typename...> struct TReadOneOfAccessor;
+template<typename...> struct TReadOneOrMoreOfAccessor;
 
 template<typename...> struct TFilteredEntityTask;
 template<typename...> struct TEntityTaskComponents;
@@ -228,7 +235,7 @@ struct TEntityTaskComponents : TEntityTaskComponentsImpl<TMakeIntegerSequence<in
 	 * For example:
 	 * struct FForEachAllocation
 	 * {
-	 *     void ForEachAllocation(FEntityAllocation* InAllocation, const TFilteredEntityTask< FReadEntityIDs, TRead<FMovieSceneFloatChannel> >& InputTask);
+	 *     void ForEachAllocation(FEntityAllocation* InAllocation, const TFilteredEntityTask< FEntityIDAccess, TRead<FMovieSceneFloatChannel> >& InputTask);
 	 * };
 	 *
 	 * TComponentTypeID<FMovieSceneFloatChannel> FloatChannelComponent = ...;
@@ -395,14 +402,12 @@ public:
 template<int... Indices, typename... T>
 struct TEntityTaskComponentsImpl<TIntegerSequence<int, Indices...>, T...>
 {
-	using EntityRangeType = TEntityRange<typename T::AccessType...>;
-
 	/**
 	 * Read the entity ID along with this task. Passed to the task as an FMovieSceneEntityID
 	 */
-	TEntityTaskComponents< T..., FReadEntityIDs > ReadEntityIDs() const
+	TEntityTaskComponents< T..., FEntityIDAccess > ReadEntityIDs() const
 	{
-		return TEntityTaskComponents< T..., FReadEntityIDs >( Accessors.template Get<Indices>()..., FReadEntityIDs{} );
+		return TEntityTaskComponents< T..., FEntityIDAccess >( Accessors.template Get<Indices>()..., FEntityIDAccess{} );
 	}
 
 	/**
@@ -412,9 +417,9 @@ struct TEntityTaskComponentsImpl<TIntegerSequence<int, Indices...>, T...>
 	 * @param ComponentType   A valid component type to read.
 	 */
 	template<typename U>
-	TEntityTaskComponents<T..., TRead<U>> Read(TComponentTypeID<U> ComponentType) const
+	TEntityTaskComponents<T..., TReadAccess<U>> Read(TComponentTypeID<U> ComponentType) const
 	{
-		return TEntityTaskComponents<T..., TRead<U> >(Accessors.template Get<Indices>()..., TRead<U>(ComponentType));
+		return TEntityTaskComponents<T..., TReadAccess<U> >(Accessors.template Get<Indices>()..., TReadAccess<U>(ComponentType));
 	}
 
 	/**
@@ -424,9 +429,9 @@ struct TEntityTaskComponentsImpl<TIntegerSequence<int, Indices...>, T...>
 	 * @param ComponentTypes  The component types to visit
 	 */
 	template<typename... U>
-	TEntityTaskComponents<T..., TReadOneOf<U...>> ReadOneOf(TComponentTypeID<U>... ComponentTypes) const
+	TEntityTaskComponents<T..., TReadOneOfAccessor<U...>> ReadOneOf(TComponentTypeID<U>... ComponentTypes) const
 	{
-		return TEntityTaskComponents<T..., TReadOneOf<U...> >( Accessors.template Get<Indices>()..., TReadOneOf<U...>(ComponentTypes...) );
+		return TEntityTaskComponents<T..., TReadOneOfAccessor<U...> >( Accessors.template Get<Indices>()..., TReadOneOfAccessor<U...>(ComponentTypes...) );
 	}
 
 	/**
@@ -435,9 +440,9 @@ struct TEntityTaskComponentsImpl<TIntegerSequence<int, Indices...>, T...>
 	 * @param ComponentTypes  The component types to visit
 	 */
 	template<typename... U>
-	TEntityTaskComponents<T..., TReadOneOrMoreOf<U...>> ReadOneOrMoreOf(TComponentTypeID<U>... ComponentTypes) const
+	TEntityTaskComponents<T..., TReadOneOrMoreOfAccessor<U...>> ReadOneOrMoreOf(TComponentTypeID<U>... ComponentTypes) const
 	{
-		return TEntityTaskComponents<T..., TReadOneOrMoreOf<U...> >(Accessors.template Get<Indices>()..., TReadOneOrMoreOf<U...>(ComponentTypes...));
+		return TEntityTaskComponents<T..., TReadOneOrMoreOfAccessor<U...> >(Accessors.template Get<Indices>()..., TReadOneOrMoreOfAccessor<U...>(ComponentTypes...));
 	}
 
 	/**
@@ -446,33 +451,9 @@ struct TEntityTaskComponentsImpl<TIntegerSequence<int, Indices...>, T...>
 	 * @param ComponentTypes  The component types to visit
 	 */
 	template<typename... U>
-	TEntityTaskComponents<T..., TRead<U>...> ReadAllOf(TComponentTypeID<U>... ComponentTypes) const
-	{
-		return TEntityTaskComponents<T..., TRead<U>... >( Accessors.template Get<Indices>()..., TRead<U>(ComponentTypes)... );
-	}
-
-	/**
-	 * Projected Read of a component value. Type passed to the task will be the result of the invocation of the projection.
-	 * @note Supplying an invalid ComponentType will be handled gracefully, but will result in no task being dispatched.
-	 * @note For example:
-	 * 
-	 * struct FForEachEntity
-	 * {
-	 *     void ForEachEntity(float InX);
-	 * };
-	 *
-	 * TComponentTypeID<FVector> VectorComponent = ...;
-	 *
-	 * FGraphEventRef Task = FEntityTaskBuilder()
-	 * .ReadEntityIDs()
-	 * .ReadProjected(VectorComponent, &FVector::X)
-	 * .Dispatch_PerEntity<FForEachEntity>(...);
-	 */
-	template<typename U, typename ProjectionType>
-	TEntityTaskComponents<T..., TReadProjected< U, typename TRemoveReference<ProjectionType>::Type >> ReadProjected(TComponentTypeID<U> ComponentType, ProjectionType&& Projection) const
+	TEntityTaskComponents<T..., TReadAccess<U>...> ReadAllOf(TComponentTypeID<U>... ComponentTypes) const
 	{
-		using FReturnType = TReadProjected< U, typename TRemoveReference<ProjectionType>::Type >;
-		return TEntityTaskComponents<T..., FReturnType >(Accessors.template Get<Indices>()..., FReturnType(ComponentType, Forward<ProjectionType>(Projection)));
+		return TEntityTaskComponents<T..., TReadAccess<U>... >( Accessors.template Get<Indices>()..., TReadAccess<U>(ComponentTypes)... );
 	}
 
 	/**
@@ -481,9 +462,9 @@ struct TEntityTaskComponentsImpl<TIntegerSequence<int, Indices...>, T...>
 	 *
 	 * @param ComponentType   A valid component type to read.
 	 */
-	TEntityTaskComponents<T..., FReadErased> ReadErased(FComponentTypeID ComponentType) const
+	TEntityTaskComponents<T..., FErasedReadAccess> ReadErased(FComponentTypeID ComponentType) const
 	{
-		return TEntityTaskComponents<T..., FReadErased >(Accessors.template Get<Indices>()..., FReadErased(ComponentType));
+		return TEntityTaskComponents<T..., FErasedReadAccess >(Accessors.template Get<Indices>()..., FErasedReadAccess(ComponentType));
 	}
 
 	/**
@@ -491,9 +472,9 @@ struct TEntityTaskComponentsImpl<TIntegerSequence<int, Indices...>, T...>
 	 * @note Always passed to the task as a const T* pointer which must be checked for null
 	 */
 	template<typename U>
-	TEntityTaskComponents<T..., TReadOptional<U>> ReadOptional(TComponentTypeID<U> ComponentType) const
+	TEntityTaskComponents<T..., TOptionalReadAccess<U>> ReadOptional(TComponentTypeID<U> ComponentType) const
 	{
-		return TEntityTaskComponents<T..., TReadOptional<U> >(Accessors.template Get<Indices>()..., TReadOptional<U>(ComponentType));
+		return TEntityTaskComponents<T..., TOptionalReadAccess<U> >(Accessors.template Get<Indices>()..., TOptionalReadAccess<U>(ComponentType));
 	}
 
 	/**
@@ -503,9 +484,20 @@ struct TEntityTaskComponentsImpl<TIntegerSequence<int, Indices...>, T...>
 	 * @param ComponentType   A valid component type to read.
 	 */
 	template<typename U>
-	TEntityTaskComponents<T..., TWrite<U>> Write(TComponentTypeID<U> ComponentType) const
+	TEntityTaskComponents<T..., TWriteAccess<U>> Write(TComponentTypeID<U> ComponentType) const
 	{
-		return TEntityTaskComponents<T..., TWrite<U> >(Accessors.template Get<Indices>()..., TWrite<U>(ComponentType));
+		return TEntityTaskComponents<T..., TWriteAccess<U> >(Accessors.template Get<Indices>()..., TWriteAccess<U>(ComponentType));
+	}
+
+	/**
+	 * Write the type-erased value of a component. Passed to the task as a void*
+	 * @note Supplying an invalid ComponentType will be handled gracefully, but will result in no task being dispatched.
+	 *
+	 * @param ComponentType   A valid component type to read.
+	 */
+	TEntityTaskComponents<T..., FErasedWriteAccess> WriteErased(FComponentTypeID ComponentType) const
+	{
+		return TEntityTaskComponents<T..., FErasedWriteAccess >(Accessors.template Get<Indices>()..., FErasedWriteAccess(ComponentType));
 	}
 
 
@@ -514,9 +506,9 @@ struct TEntityTaskComponentsImpl<TIntegerSequence<int, Indices...>, T...>
 	 * @note Always passed to the task as a T* pointer which must be checked for null
 	 */
 	template<typename U>
-	TEntityTaskComponents<T..., TWriteOptional<U>> WriteOptional(TComponentTypeID<U> ComponentType) const
+	TEntityTaskComponents<T..., TOptionalWriteAccess<U>> WriteOptional(TComponentTypeID<U> ComponentType) const
 	{
-		return TEntityTaskComponents<T..., TWriteOptional<U> >(Accessors.template Get<Indices>()..., TWriteOptional<U>(ComponentType));
+		return TEntityTaskComponents<T..., TOptionalWriteAccess<U> >(Accessors.template Get<Indices>()..., TOptionalWriteAccess<U>(ComponentType));
 	}
 
 	bool HasBeenWrittenToSince(uint32 InSystemVersion)
@@ -568,55 +560,35 @@ struct TEntityTaskComponentsImpl<TIntegerSequence<int, Indices...>, T...>
 		(void)Temp;
 	}
 
-	/** Lock the component headers that we need to access */
-	void Lock(FEntityAllocation* Allocation) const
-	{
-		//SCOPE_CYCLE_COUNTER(MovieSceneEval_AquireComponentAccessLocks);
-
-		int Temp[] = { (LockHeader(&Accessors.template Get<Indices>(), Allocation), 0)..., 0 };
-		(void)Temp;
-	}
-
-	/** Unlock the component headers that we accessed in Lock() */
-	void Unlock(FEntityAllocation* Allocation, uint64 SystemSerial) const
-	{
-		//SCOPE_CYCLE_COUNTER(MovieSceneEval_ReleaseComponentAccessLocks);
-
-		int Temp[] = { ( UnlockHeader(&Accessors.template Get<Indices>(), Allocation, SystemSerial), 0)..., 0 };
-		(void)Temp;
-	}
-
 	/**
 	 * Perform a thread-safe iteration of the specified allocation using this task, inline on the current thread
+	 * @note: This is highly unsafe as it circumvents all the thread-safety mechanisms that protect component data
 	 *
 	 * @param Allocation  The allocation to iterate
 	 * @return An iterator that defines the full range of entities in the allocation
 	 */
-	EntityRangeType IterateAllocation(const FEntityAllocation* Allocation) const
+	TEntityRange<typename T::AccessType...> IterateAllocation(const FEntityAllocation* Allocation, FEntityAllocationWriteContext WriteContext) const
 	{
-		EntityRangeType Result(Allocation->Num());
-
-		static_assert(TAnd<TSupportsDirectEntityIteration<T>...>::Value, "Accessor type(s) do not support direct entity iteration - must be resolved separately.");
-
 		checkf(IsValid(), TEXT("Attempting to use a component pack with an invalid component type."));
 
-		int Unused[] = { (Result.template GetComponentArrayReference<Indices>() = Accessors.template Get<Indices>().Resolve(Allocation), 0)..., 0 };
-		(void)Unused;
+		auto LockedComponentData = MakeTuple( Accessors.template Get<Indices>().LockComponentData(Allocation, WriteContext)... );
 
-		return Result;
+		// WARNING: This is highly unsafe as it circumvents all the thread-safety mechanisms that protect component data
+		return TEntityRange<typename T::AccessType...>(Allocation->Num(), LockedComponentData.template Get<Indices>().AsPtr()... );
 	}
 
 	/**
 	 * Perform a thread-safe iteration of the specified entity range using this task, inline on the current thread
+	 * @note: This is highly unsafe as it circumvents all the thread-safety mechanisms that protect component data
 	 *
 	 * @param Allocation  The allocation to iterate
 	 * @return An iterator that defines the full range of entities in the allocation
 	 */
-	EntityRangeType IterateRange(const FEntityRange& EntityRange) const
+	TEntityRange<typename T::AccessType...> IterateRange(const FEntityRange& EntityRange, FEntityAllocationWriteContext WriteContext) const
 	{
 		check(EntityRange.ComponentStartOffset >= 0 && EntityRange.ComponentStartOffset + EntityRange.Num <= EntityRange.Allocation->Num());
 
-		EntityRangeType Result = IterateAllocation(EntityRange.Allocation);
+		TEntityRange<typename T::AccessType...> Result = IterateAllocation(EntityRange.Allocation, WriteContext);
 		Result.Slice(EntityRange.ComponentStartOffset, EntityRange.Num);
 		return Result;
 	}
@@ -659,29 +631,24 @@ struct TEntityTaskComponentsImpl<TIntegerSequence<int, Indices...>, T...>
 	template<typename Callback>
 	void Iterate_PerEntityImpl(FEntityManager* EntityManager, const FEntityComponentFilter& Filter, Callback&& InCallback) const
 	{
+		using TupleType = TTuple< decltype(Accessors.template Get<Indices>().LockComponentData(nullptr, DeclVal<FEntityAllocationWriteContext>()))... >;
+
 		if (IsValid())
 		{
-			const uint64 SystemSerial = EntityManager->GetSystemSerial();
+			FEntityAllocationWriteContext WriteContext(*EntityManager);
+
 			for (FEntityAllocation* Allocation : EntityManager->Iterate(&Filter))
 			{
 				FEntityIterationResult Result;
 
-				// Lock on the components we want to access
-				Lock(Allocation);
-
-				auto IterState = MakeTuple( Accessors.template Get<Indices>().CreateIterState(Allocation)... );
+				// Lock the components we want to access
+				TupleType ComponentData( Accessors.template Get<Indices>().LockComponentData(Allocation, WriteContext)... );
 
 				const int32 Num = Allocation->Num();
 				for (int32 ComponentOffset = 0; ComponentOffset < Num && Result.Value; ++ComponentOffset )
 				{
-					Result = (InCallback( *IterState.template Get<Indices>()... ), Result);
-
-					int Temp[] = { (++IterState.template Get<Indices>(), 0)..., 0 };
-					(void)Temp;
+					Result = (InCallback( ComponentData.template Get<Indices>().ComponentAtIndex(ComponentOffset)... ), Result);
 				}
-
-				// Unlock from the components we wanted to access
-				Unlock(Allocation, SystemSerial);
 			}
 		}
 	}
@@ -698,18 +665,15 @@ struct TEntityTaskComponentsImpl<TIntegerSequence<int, Indices...>, T...>
 	{
 		if (IsValid())
 		{
-			const uint64 SystemSerial = EntityManager->GetSystemSerial();
+			FEntityAllocationWriteContext WriteContext(*EntityManager);
 			for (FEntityAllocationIteratorItem Item : EntityManager->Iterate(&Filter))
 			{
 				FEntityAllocation* Allocation = Item;
 
 				// Lock on the components we want to access
-				Lock(Allocation);
+				auto ComponentData = MakeTuple( Accessors.template Get<Indices>().LockComponentData(Allocation, WriteContext)... );
 
-				FEntityIterationResult Result = (InCallback(Item, Accessors.template Get<Indices>()...), FEntityIterationResult{});
-
-				// Unlock from the components we wanted to access
-				Unlock(Allocation, SystemSerial);
+				FEntityIterationResult Result = (InCallback(Item, ComponentData.template Get<Indices>()...), FEntityIterationResult{});
 
 				if (!Result)
 				{
@@ -1104,13 +1068,13 @@ struct TEntityTaskBase
 	explicit TEntityTaskBase(FEntityManager* InEntityManager, const TEntityTaskComponents<ComponentTypes...>& InComponents)
 		: FilteredTask(InComponents)
 		, EntityManager(InEntityManager)
-		, SystemSerial(InEntityManager->GetSystemSerial())
+		, WriteContext(*InEntityManager)
 	{}
 
 	explicit TEntityTaskBase(FEntityManager* InEntityManager, const TFilteredEntityTask<ComponentTypes...>& InFilteredTask)
 		: FilteredTask(InFilteredTask)
 		, EntityManager(InEntityManager)
-		, SystemSerial(InEntityManager->GetSystemSerial())
+		, WriteContext(*InEntityManager)
 	{}
 
 	void Run(TaskImpl& TaskImplInstance)
@@ -1121,11 +1085,7 @@ struct TEntityTaskBase
 
 		for (FEntityAllocation* Allocation : EntityManager->Iterate(&FilteredTask.GetFilter()))
 		{
-			FilteredTask.GetComponents().Lock(Allocation);
-
-			Caller::ForEachEntity(TaskImplInstance, Allocation, FilteredTask.GetComponents());
-
-			FilteredTask.GetComponents().Unlock(Allocation, SystemSerial);
+			Caller::ForEachEntityImpl(TaskImplInstance, Allocation, WriteContext, FilteredTask.GetComponents());
 		}
 
 		PostTask(&TaskImplInstance);
@@ -1149,7 +1109,7 @@ private:
 
 	TFilteredEntityTask<ComponentTypes...> FilteredTask;
 	FEntityManager* EntityManager;
-	uint64 SystemSerial;
+	FEntityAllocationWriteContext WriteContext;
 };
 
 
@@ -1222,13 +1182,13 @@ struct TEntityAllocationTaskBase
 	explicit TEntityAllocationTaskBase(FEntityManager* InEntityManager, const TEntityTaskComponents<ComponentTypes...>& InComponents)
 		: ComponentFilter(InComponents)
 		, EntityManager(InEntityManager)
-		, SystemSerial(InEntityManager->GetSystemSerial())
+		, WriteContext(*InEntityManager)
 	{}
 
 	explicit TEntityAllocationTaskBase(FEntityManager* InEntityManager, const TFilteredEntityTask<ComponentTypes...>& InComponentFilter)
 		: ComponentFilter(InComponentFilter)
 		, EntityManager(InEntityManager)
-		, SystemSerial(InEntityManager->GetSystemSerial())
+		, WriteContext(*InEntityManager)
 	{}
 
 	void Run(TaskImpl& TaskImplInstance)
@@ -1239,15 +1199,7 @@ struct TEntityAllocationTaskBase
 
 		for (FEntityAllocationIteratorItem Item : EntityManager->Iterate(&ComponentFilter.GetFilter()))
 		{
-			FEntityAllocation* Allocation = Item;
-
-			// Lock on the components we want to access
-			ComponentFilter.GetComponents().Lock(Allocation);
-
-			Caller::ForEachAllocation(TaskImplInstance, Item, ComponentFilter.GetComponents());
-
-			// Unlock from the components we wanted to access
-			ComponentFilter.GetComponents().Unlock(Allocation, SystemSerial);
+			Caller::ForEachAllocationImpl(TaskImplInstance, Item, WriteContext, ComponentFilter.GetComponents());
 		}
 
 		PostTask(&TaskImplInstance);
@@ -1273,7 +1225,7 @@ protected:
 
 	TFilteredEntityTask<ComponentTypes...> ComponentFilter;
 	FEntityManager* EntityManager;
-	uint64 SystemSerial;
+	FEntityAllocationWriteContext WriteContext;
 };
 
 template<typename TaskImpl, typename... ComponentTypes>
@@ -1345,24 +1297,26 @@ template<typename...> struct TEntityTaskCaller_AutoExpansion;
 template<typename TaskImpl, int... Indices>
 struct TEntityTaskCaller_AutoExpansion<TaskImpl, TIntegerSequence<int, Indices...>>
 {
-	template<typename... ComponentTypes>
-	FORCEINLINE static void ForEachEntity(TaskImpl& TaskImplInstance, const FEntityAllocation* Allocation, const TEntityTaskComponents<ComponentTypes...>& Components)
+	template<typename... AccessorTypes>
+	FORCEINLINE static void ForEachEntityImpl(TaskImpl& TaskImplInstance, const FEntityAllocation* Allocation, FEntityAllocationWriteContext WriteContext, const TEntityTaskComponents<AccessorTypes...>& Components)
 	{
-		auto IterState = MakeTuple( Components.template GetAccessor<Indices>().CreateIterState(Allocation)... );
+		FEntityIterationResult Result;
+
+		// Lock the components we want to access
+		auto LockedComponentData = MakeTuple( Components.template GetAccessor<Indices>().LockComponentData(Allocation, WriteContext)... );
 
 		const int32 Num = Allocation->Num();
-		for (int32 ComponentOffset = 0; ComponentOffset < Num; ++ComponentOffset )
+		for (int32 ComponentOffset = 0; ComponentOffset < Num && Result.Value; ++ComponentOffset )
 		{
-			TaskImplInstance.ForEachEntity( *IterState.template Get<Indices>()... );
-			int Temp[] = { (++IterState.template Get<Indices>(), 0)... };
-			(void)Temp;
+			Result = (TaskImplInstance.ForEachEntity(LockedComponentData.template Get<Indices>().ComponentAtIndex(ComponentOffset)... ), Result);
 		}
 	}
 
-	template<typename... ComponentTypes>
-	FORCEINLINE static void ForEachAllocation(TaskImpl& TaskImplInstance, FEntityAllocationIteratorItem Item, const TEntityTaskComponents<ComponentTypes...>& Components)
+	template<typename... AccessorTypes>
+	FORCEINLINE static void ForEachAllocationImpl(TaskImpl& TaskImplInstance, FEntityAllocationIteratorItem Item, FEntityAllocationWriteContext WriteContext, const TEntityTaskComponents<AccessorTypes...>& Components)
 	{
-		TaskImplInstance.ForEachAllocation(Item, Components.template GetAccessor<Indices>()...);
+		auto LockedComponentData = MakeTuple( Components.template GetAccessor<Indices>().LockComponentData(Item.GetAllocation(), WriteContext)... );
+		TaskImplInstance.ForEachAllocation(Item, LockedComponentData.template Get<Indices>()...);
 	}
 };
 
@@ -1374,17 +1328,13 @@ struct TEntityTaskCaller<TaskImpl, NumComponents, true> : TEntityTaskCaller_Auto
 template<typename TaskImpl, int32 NumComponents>
 struct TEntityTaskCaller<TaskImpl, NumComponents, false>
 {
-	template<typename... ComponentTypes>
-	FORCEINLINE static void ForEachEntity(TaskImpl& TaskImplInstance, const FEntityAllocation* Allocation, const TEntityTaskComponents<ComponentTypes...>& Components)
+	FORCEINLINE static void ForEachEntityImpl(...)
 	{
-		for (const auto& Entity : Components.IterateAllocation(Allocation))
-		{
-			TaskImplInstance.ForEachEntity(Entity);
-		}
+		static_assert(TNot<TIsSame<TaskImpl, TaskImpl>>::Value, "non-expanded entity iteration is not supported");
 	}
 
-	template<typename... ComponentTypes>
-	FORCEINLINE static void ForEachAllocation(TaskImpl& TaskImplInstance, FEntityAllocationIteratorItem Item, const TEntityTaskComponents<ComponentTypes...>& Components)
+	template<typename... AccessorTypes>
+	FORCEINLINE static void ForEachAllocationImpl(TaskImpl& TaskImplInstance, FEntityAllocationIteratorItem Item, const TEntityTaskComponents<AccessorTypes...>& Components)
 	{
 		TaskImplInstance.ForEachAllocation(Item, Components);
 	}
diff --git a/Engine/Source/Runtime/MovieScene/Public/EntitySystem/MovieSceneEntitySystemTypes.h b/Engine/Source/Runtime/MovieScene/Public/EntitySystem/MovieSceneEntitySystemTypes.h
--- a/Engine/Source/Runtime/MovieScene/Public/EntitySystem/MovieSceneEntitySystemTypes.h
+++ b/Engine/Source/Runtime/MovieScene/Public/EntitySystem/MovieSceneEntitySystemTypes.h
@@ -20,6 +20,17 @@ namespace UE
 namespace MovieScene
 {
 
+template<typename T> struct TComponentLock;
+
+struct FReadErased;
+struct FReadErasedOptional;
+struct FWriteErased;
+struct FWriteErasedOptional;
+template<typename T> struct TRead;
+template<typename T> struct TReadOptional;
+template<typename T> struct TWrite;
+template<typename T> struct TWriteOptional;
+
 enum class ESystemPhase : uint8
 {
 	/** Null phase which indicates that the system never runs, but still exists in the reference graph */
@@ -360,6 +371,32 @@ private:
 	TArray<FComplexMask> ComplexMasks;
 };
 
+
+struct FEntityAllocationWriteContext
+{
+	MOVIESCENE_API FEntityAllocationWriteContext(const FEntityManager& EntityManager);
+
+	static FEntityAllocationWriteContext NewAllocation()
+	{
+		return FEntityAllocationWriteContext();
+	}
+
+	uint64 GetSystemSerial() const
+	{
+		return SystemSerial;
+	}
+
+private:
+
+	FEntityAllocationWriteContext()
+		: SystemSerial(0)
+	{}
+
+	uint64 SystemSerial;
+};
+
+
+
 struct FComponentHeader
 {
 	mutable uint8* Components;
@@ -368,7 +405,7 @@ struct FComponentHeader
 
 private:
 
-	uint64 SerialNumber;
+	mutable uint64 SerialNumber;
 
 public:
 
@@ -404,9 +441,9 @@ public:
 		return Components + Sizeof*Offset;
 	}
 
-	void PostWriteComponents(uint64 InSystemSerial)
+	void PostWriteComponents(FEntityAllocationWriteContext InWriteContext) const
 	{
-		SerialNumber = FMath::Max(SerialNumber, InSystemSerial);
+		SerialNumber = FMath::Max(SerialNumber, InWriteContext.GetSystemSerial());
 	}
 
 	bool HasBeenWrittenToSince(uint64 InSystemSerial) const
@@ -416,6 +453,49 @@ public:
 };
 
 
+/**
+ * Moveable scoped read-lock type for component headers
+ */
+struct FScopedHeaderReadLock
+{
+	FScopedHeaderReadLock();
+	FScopedHeaderReadLock(const FComponentHeader* InHeader);
+
+	FScopedHeaderReadLock(const FScopedHeaderReadLock& RHS) = delete;
+	void operator=(const FScopedHeaderReadLock& RHS) = delete;
+
+	FScopedHeaderReadLock(FScopedHeaderReadLock&& RHS);
+	FScopedHeaderReadLock& operator=(FScopedHeaderReadLock&& RHS);
+
+	~FScopedHeaderReadLock();
+
+private:
+	const FComponentHeader* Header;
+};
+
+
+/**
+ * Moveable scoped write-lock type for component headers
+ */
+struct FScopedHeaderWriteLock
+{
+	FScopedHeaderWriteLock();
+	FScopedHeaderWriteLock(const FComponentHeader* InHeader, FEntityAllocationWriteContext InWriteContext);
+
+	FScopedHeaderWriteLock(const FScopedHeaderWriteLock& RHS) = delete;
+	void operator=(const FScopedHeaderWriteLock& RHS) = delete;
+
+	FScopedHeaderWriteLock(FScopedHeaderWriteLock&& RHS);
+	FScopedHeaderWriteLock& operator=(FScopedHeaderWriteLock&& RHS);
+
+	~FScopedHeaderWriteLock();
+
+private:
+	const FComponentHeader* Header;
+	FEntityAllocationWriteContext WriteContext;
+};
+
+
 /**
  * FEntityAllocation is the authoritative storage of entity-component data within an FEntityManager. It stores component data in separate contiguous arrays, aligned to a cache line.
  * Storing component data in this way allows for cache-efficient and concurrent access to each component array in isolation. It also allows for write access to component arrays
@@ -586,18 +666,18 @@ struct FEntityAllocation
 	/**
 	 * Called when this allocation has been modified. Will invalidate any cached data based of this allocation's serial number
 	 */
-	void PostModifyStructureExcludingHeaders(uint64 InSystemSerial)
+	void PostModifyStructureExcludingHeaders(FEntityAllocationWriteContext InWriteContext)
 	{
-		SerialNumber = FMath::Max(SerialNumber, InSystemSerial);
+		SerialNumber = FMath::Max(SerialNumber, InWriteContext.GetSystemSerial());
 	}
 
 
-	void PostModifyStructure(uint64 InSystemSerial)
+	void PostModifyStructure(FEntityAllocationWriteContext InWriteContext)
 	{
-		SerialNumber = FMath::Max(SerialNumber, InSystemSerial);
+		SerialNumber = FMath::Max(SerialNumber, InWriteContext.GetSystemSerial());
 		for (FComponentHeader& Header : GetComponentHeaders())
 		{
-			Header.PostWriteComponents(InSystemSerial);
+			Header.PostWriteComponents(InWriteContext);
 		}
 	}
 
@@ -645,6 +725,72 @@ struct FEntityAllocation
 		return int32(Capacity) - int32(Size);
 	}
 
+	/**
+	 * Read type-erased component data for the specified component type
+	 */
+	UE_NODISCARD MOVIESCENE_API TComponentLock<FReadErased> ReadComponentsErased(FComponentTypeID ComponentType) const;
+
+	/**
+	 * Write type-erased component data for the specified component type
+	 */
+	UE_NODISCARD MOVIESCENE_API TComponentLock<FWriteErased> WriteComponentsErased(FComponentTypeID ComponentType, FEntityAllocationWriteContext InWriteContext) const;
+
+	/**
+	 * Attempt to read type-erased component data for the specified component type
+	 */
+	UE_NODISCARD MOVIESCENE_API TComponentLock<FReadErasedOptional> TryReadComponentsErased(FComponentTypeID ComponentType) const;
+
+	/**
+	 * Attempt to write type-erased component data for the specified component type
+	 */
+	UE_NODISCARD MOVIESCENE_API TComponentLock<FWriteErasedOptional> TryWriteComponentsErased(FComponentTypeID ComponentType, FEntityAllocationWriteContext InWriteContext) const;
+
+	/**
+	 * Read typed component data for the specified component type
+	 */
+	template<typename T>
+	UE_NODISCARD TComponentLock<TRead<T>> ReadComponents(TComponentTypeID<T> ComponentType) const
+	{
+		const FComponentHeader& Header = GetComponentHeaderChecked(ComponentType);
+		return TComponentLock<TRead<T>>(&Header);
+	}
+
+	/**
+	 * Write typed component data for the specified component type
+	 */
+	template<typename T>
+	UE_NODISCARD TComponentLock<TReadOptional<T>> TryReadComponents(TComponentTypeID<T> ComponentType) const
+	{
+		if (const FComponentHeader* Header = FindComponentHeader(ComponentType))
+		{
+			return TComponentLock<TReadOptional<T>>(Header);
+		}
+		return TComponentLock<TReadOptional<T>>();
+	}
+
+	/**
+	 * Write typed component data for the specified component type
+	 */
+	template<typename T>
+	UE_NODISCARD TComponentLock<TWrite<T>> WriteComponents(TComponentTypeID<T> ComponentType, FEntityAllocationWriteContext InWriteContext) const
+	{
+		const FComponentHeader& Header = GetComponentHeaderChecked(ComponentType);
+		return TComponentLock<TWrite<T>>(&Header, InWriteContext);
+	}
+
+	/**
+	 * Attempt to write typed component data for the specified component type
+	 */
+	template<typename T>
+	UE_NODISCARD TComponentLock<TWriteOptional<T>> TryWriteComponents(TComponentTypeID<T> ComponentType, FEntityAllocationWriteContext InWriteContext) const
+	{
+		if (const FComponentHeader* Header = FindComponentHeader(ComponentType))
+		{
+			return TComponentLock<TWriteOptional<T>>(Header, InWriteContext);
+		}
+		return TComponentLock<TWriteOptional<T>>();
+	}
+
 private:
 
 	friend struct FEntityInitializer;
@@ -696,6 +842,81 @@ struct FEntityInfo
 	FMovieSceneEntityID EntityID;
 };
 
+inline FScopedHeaderReadLock::FScopedHeaderReadLock()
+	: Header(nullptr)
+{}
+inline FScopedHeaderReadLock::FScopedHeaderReadLock(const FComponentHeader* InHeader)
+	: Header(InHeader)
+{
+	InHeader->ReadWriteLock.ReadLock();
+}
+inline FScopedHeaderReadLock::FScopedHeaderReadLock(FScopedHeaderReadLock&& RHS)
+	: Header(RHS.Header)
+{
+	RHS.Header = nullptr;
+}
+inline FScopedHeaderReadLock& FScopedHeaderReadLock::operator=(FScopedHeaderReadLock&& RHS)
+{
+	if (Header)
+	{
+		Header->ReadWriteLock.ReadUnlock();
+	}
+
+	Header = RHS.Header;
+	RHS.Header = nullptr;
+	return *this;
+}
+
+inline FScopedHeaderReadLock::~FScopedHeaderReadLock()
+{
+	if (Header)
+	{
+		Header->ReadWriteLock.ReadUnlock();
+	}
+}
+
+inline FScopedHeaderWriteLock::FScopedHeaderWriteLock()
+	: Header(nullptr)
+	, WriteContext(FEntityAllocationWriteContext::NewAllocation())
+{}
+
+inline FScopedHeaderWriteLock::FScopedHeaderWriteLock(const FComponentHeader* InHeader, FEntityAllocationWriteContext InWriteContext)
+	: Header(InHeader)
+	, WriteContext(InWriteContext)
+{
+	InHeader->ReadWriteLock.WriteLock();
+}
+
+inline FScopedHeaderWriteLock::FScopedHeaderWriteLock(FScopedHeaderWriteLock&& RHS)
+	: Header(RHS.Header)
+	, WriteContext(RHS.WriteContext)
+{
+	RHS.Header = nullptr;
+}
+inline FScopedHeaderWriteLock& FScopedHeaderWriteLock::operator=(FScopedHeaderWriteLock&& RHS)
+{
+	if (Header)
+	{
+		Header->PostWriteComponents(WriteContext);
+		Header->ReadWriteLock.WriteUnlock();
+	}
+
+	Header = RHS.Header;
+	WriteContext = RHS.WriteContext;
+
+	RHS.Header = nullptr;
+	return *this;
+}
+inline FScopedHeaderWriteLock::~FScopedHeaderWriteLock()
+{
+	if (Header)
+	{
+		Header->PostWriteComponents(WriteContext);
+		Header->ReadWriteLock.WriteUnlock();
+	}
+}
+
+
 } // namespace MovieScene
 } // namespace UE
 
diff --git a/Engine/Source/Runtime/MovieScene/Public/EntitySystem/MovieSceneOverlappingEntityTracker.h b/Engine/Source/Runtime/MovieScene/Public/EntitySystem/MovieSceneOverlappingEntityTracker.h
--- a/Engine/Source/Runtime/MovieScene/Public/EntitySystem/MovieSceneOverlappingEntityTracker.h
+++ b/Engine/Source/Runtime/MovieScene/Public/EntitySystem/MovieSceneOverlappingEntityTracker.h
@@ -10,6 +10,7 @@
 #include "EntitySystem/MovieSceneEntitySystemTypes.h"
 #include "EntitySystem/MovieSceneComponentAccessors.h"
 #include "EntitySystem/MovieSceneEntitySystemTask.h"
+#include "EntitySystem/MovieSceneEntitySystemLinker.h"
 #include "EntitySystem/BuiltInComponentTypes.h"
 
 namespace UE
@@ -228,12 +229,11 @@ struct TOverlappingEntityTracker
 
 protected:
 
-	void VisitLinkedAllocationImpl(const FEntityAllocation* Allocation, TRead<KeyType> ReadKeys)
+	void VisitLinkedAllocationImpl(const FEntityAllocation* Allocation, TRead<KeyType> Keys)
 	{
 		const int32 Num = Allocation->Num();
 
 		const FMovieSceneEntityID* EntityIDs = Allocation->GetRawEntityIDs();
-		const KeyType*             Keys      = ReadKeys.Resolve(Allocation);
 
 		if (Allocation->HasComponent(FBuiltInComponentTypes::Get()->Tags.RestoreState))
 		{
diff --git a/Engine/Source/Runtime/MovieScene/Public/EntitySystem/MovieScenePartialProperties.h b/Engine/Source/Runtime/MovieScene/Public/EntitySystem/MovieScenePartialProperties.h
--- a/Engine/Source/Runtime/MovieScene/Public/EntitySystem/MovieScenePartialProperties.h
+++ b/Engine/Source/Runtime/MovieScene/Public/EntitySystem/MovieScenePartialProperties.h
@@ -55,8 +55,8 @@ struct TSetPartialPropertyValues
 {
 	using IntermediateType = typename ProjectionType::IntermediateType;
 
-	using FThreeWayAccessor  = TReadOneOf< FCustomPropertyIndex, uint16, TSharedPtr<FTrackInstancePropertyBindings> >;
-	using FTwoWayAccessor    = TReadOneOf< uint16, TSharedPtr<FTrackInstancePropertyBindings> >;
+	using FThreeWayAccessor  = TMultiReadOptional<FCustomPropertyIndex, uint16, TSharedPtr<FTrackInstancePropertyBindings>>;
+	using FTwoWayAccessor    = TMultiReadOptional<uint16, TSharedPtr<FTrackInstancePropertyBindings>>;
 
 	explicit TSetPartialPropertyValues(ICustomPropertyRegistration* InCustomProperties, const ProjectionType& InProjections)
 		: CustomProperties(InCustomProperties)
diff --git a/Engine/Source/Runtime/MovieScene/Public/EntitySystem/MovieScenePartialProperties.inl b/Engine/Source/Runtime/MovieScene/Public/EntitySystem/MovieScenePartialProperties.inl
--- a/Engine/Source/Runtime/MovieScene/Public/EntitySystem/MovieScenePartialProperties.inl
+++ b/Engine/Source/Runtime/MovieScene/Public/EntitySystem/MovieScenePartialProperties.inl
@@ -3,6 +3,7 @@
 #pragma once
 
 #include "EntitySystem/MovieScenePartialProperties.h"
+#include "EntitySystem/MovieSceneOperationalTypeConversions.h"
 #include "MovieSceneCommonHelpers.h"
 
 namespace UE
@@ -79,8 +80,6 @@ TSetPartialPropertyValues<PropertyType, ProjectionType>::TSetPartialPropertyValu
 template<typename PropertyType, typename ProjectionType>
 void TSetPartialPropertyValues<PropertyType, ProjectionType>::ForEachAllocation(const FEntityAllocation* Allocation, TRead<UObject*> BoundObjectComponents, FThreeWayAccessor PropertyBindingComponents)
 {
-	using FPropertyTuple = TTuple< const FCustomPropertyIndex*, const uint16*, const TSharedPtr<FTrackInstancePropertyBindings>* >;
-
 	// ----------------------------------------------------------------------------------------------------------------------------
 	// For partially animated composites, we first retrieve the current properties for the allocation, then go through and patch in
 	// All the animated values, then apply the properties to objects
@@ -88,20 +87,17 @@ void TSetPartialPropertyValues<PropertyType, ProjectionType>::ForEachAllocation(
 	const int32 Num = Allocation->Num();
 	IntermediateValues.SetNumUninitialized(Num);
 
-	UObject* const* RawObjectPtr = BoundObjectComponents.Resolve(Allocation);
-	FPropertyTuple  Properties   = PropertyBindingComponents.Resolve(Allocation);
-
-	if (const FCustomPropertyIndex* Custom = Properties.template Get<0>())
+	if (const FCustomPropertyIndex* Custom = PropertyBindingComponents.template Get<0>())
 	{
-		ForEachCustom(Allocation, MakeArrayView(RawObjectPtr, Num), MakeArrayView(Custom, Num));
+		ForEachCustom(Allocation, BoundObjectComponents.AsArray(Num), MakeArrayView(Custom, Num));
 	}
-	else if (const uint16* Fast = Properties.template Get<1>())
+	else if (const uint16* Fast = PropertyBindingComponents.template Get<1>())
 	{
-		ForEachFast(Allocation, MakeArrayView(RawObjectPtr, Num), MakeArrayView(Fast, Num));
+		ForEachFast(Allocation, BoundObjectComponents.AsArray(Num), MakeArrayView(Fast, Num));
 	}
-	else if (const TSharedPtr<FTrackInstancePropertyBindings>* Slow = Properties.template Get<2>())
+	else if (const TSharedPtr<FTrackInstancePropertyBindings>* Slow = PropertyBindingComponents.template Get<2>())
 	{
-		ForEachSlow(Allocation, MakeArrayView(RawObjectPtr, Num), MakeArrayView(Slow, Num));
+		ForEachSlow(Allocation, BoundObjectComponents.AsArray(Num), MakeArrayView(Slow, Num));
 	}
 }
 
@@ -109,8 +105,6 @@ void TSetPartialPropertyValues<PropertyType, ProjectionType>::ForEachAllocation(
 template<typename PropertyType, typename ProjectionType>
 void TSetPartialPropertyValues<PropertyType, ProjectionType>::ForEachAllocation(const FEntityAllocation* Allocation, TRead<UObject*> BoundObjectComponents, FTwoWayAccessor PropertyBindingComponents)
 {
-	using FPropertyTuple = TTuple< const uint16*, const TSharedPtr<FTrackInstancePropertyBindings>* >;
-
 	// ----------------------------------------------------------------------------------------------------------------------------
 	// For partially animated composites, we first retrieve the current properties for the allocation, then go through and patch in
 	// All the animated values, then apply the properties to objects
@@ -118,16 +112,13 @@ void TSetPartialPropertyValues<PropertyType, ProjectionType>::ForEachAllocation(
 	const int32 Num = Allocation->Num();
 	IntermediateValues.SetNumUninitialized(Num);
 
-	UObject* const* RawObjectPtr = BoundObjectComponents.Resolve(Allocation);
-	FPropertyTuple  Properties   = PropertyBindingComponents.Resolve(Allocation);
-
-	if (const uint16* Fast = Properties.template Get<0>())
+	if (const uint16* Fast = PropertyBindingComponents.template Get<0>())
 	{
-		ForEachFast(Allocation, MakeArrayView(RawObjectPtr, Num), MakeArrayView(Fast, Num));
+		ForEachFast(Allocation, BoundObjectComponents.AsArray(Num), MakeArrayView(Fast, Num));
 	}
-	else if (const TSharedPtr<FTrackInstancePropertyBindings>* Slow = Properties.template Get<1>())
+	else if (const TSharedPtr<FTrackInstancePropertyBindings>* Slow = PropertyBindingComponents.template Get<1>())
 	{
-		ForEachSlow(Allocation, MakeArrayView(RawObjectPtr, Num), MakeArrayView(Slow, Num));
+		ForEachSlow(Allocation, BoundObjectComponents.AsArray(Num), MakeArrayView(Slow, Num));
 	}
 }
 
diff --git a/Engine/Source/Runtime/MovieScene/Public/EntitySystem/MovieScenePreAnimatedPropertyHelper.h b/Engine/Source/Runtime/MovieScene/Public/EntitySystem/MovieScenePreAnimatedPropertyHelper.h
--- a/Engine/Source/Runtime/MovieScene/Public/EntitySystem/MovieScenePreAnimatedPropertyHelper.h
+++ b/Engine/Source/Runtime/MovieScene/Public/EntitySystem/MovieScenePreAnimatedPropertyHelper.h
@@ -18,6 +18,7 @@
 
 #include "IMovieScenePlayer.h"
 #include "MovieSceneExecutionToken.h"
+#include "MovieSceneCommonHelpers.h"
 
 namespace UE
 {
diff --git a/Engine/Source/Runtime/MovieScene/Public/EntitySystem/MovieScenePropertyBinding.h b/Engine/Source/Runtime/MovieScene/Public/EntitySystem/MovieScenePropertyBinding.h
new file mode 100644
--- /dev/null
+++ b/Engine/Source/Runtime/MovieScene/Public/EntitySystem/MovieScenePropertyBinding.h
@@ -0,0 +1,70 @@
+// Copyright Epic Games, Inc. All Rights Reserved.
+
+#pragma once
+
+#include "UObject/NameTypes.h"
+#include "EntitySystem/MovieSceneComponentDebug.h"
+
+#include "MovieScenePropertyBinding.generated.h"
+
+
+/**
+ * Source property binding information for an entity on a moviescene timeline
+ * Comprises a leaf property name and a path and a cached boolean signifying whether the binding is allowed to perform a fast class-wise property lookup
+ */
+USTRUCT()
+struct FMovieScenePropertyBinding
+{
+	GENERATED_BODY()
+
+	FMovieScenePropertyBinding()
+		: bCanUseClassLookup(false)
+	{}
+
+	FMovieScenePropertyBinding(FName InPropertyName, const FString& InPropertyPath)
+		: PropertyName(InPropertyName), PropertyPath(*InPropertyPath)
+	{
+		bCanUseClassLookup = !(InPropertyPath.Contains(TEXT(".")) || InPropertyPath.Contains(TEXT("/")) || InPropertyPath.Contains(TEXT("\\")) || InPropertyPath.Contains(TEXT("[")));
+	}
+
+	friend bool operator==(FMovieScenePropertyBinding A, FMovieScenePropertyBinding B)
+	{
+		return A.PropertyName == B.PropertyName && A.PropertyPath == B.PropertyPath;
+	}
+
+	bool CanUseClassLookup() const
+	{
+		return bCanUseClassLookup;
+	}
+
+	/** Leaf name of the property to animate */
+	UPROPERTY()
+	FName PropertyName;
+
+	/** Full path to the property from the object including struct and array indirection */
+	UPROPERTY()
+	FName PropertyPath;
+
+	/** True if this property can be considered for fast property offset resolution(ie the property address is _always_ a constant offset from the obejct ptr), false otherwise */
+	UPROPERTY()
+	bool bCanUseClassLookup;
+};
+
+
+
+#if UE_MOVIESCENE_ENTITY_DEBUG
+
+namespace UE
+{
+namespace MovieScene
+{
+
+	template<> struct TComponentDebugType<FMovieScenePropertyBinding>
+	{
+		static const EComponentDebugType Type = EComponentDebugType::Property;
+	};
+
+} // namespace MovieScene
+} // namespace UE
+
+#endif // UE_MOVIESCENE_ENTITY_DEBUG
\ No newline at end of file
diff --git a/Engine/Source/Runtime/MovieScene/Public/EntitySystem/MovieScenePropertyComponentHandler.h b/Engine/Source/Runtime/MovieScene/Public/EntitySystem/MovieScenePropertyComponentHandler.h
--- a/Engine/Source/Runtime/MovieScene/Public/EntitySystem/MovieScenePropertyComponentHandler.h
+++ b/Engine/Source/Runtime/MovieScene/Public/EntitySystem/MovieScenePropertyComponentHandler.h
@@ -9,6 +9,7 @@
 #include "EntitySystem/MovieSceneDecompositionQuery.h"
 #include "EntitySystem/MovieSceneBlenderSystem.h"
 #include "EntitySystem/MovieScenePreAnimatedPropertyHelper.h"
+#include "EntitySystem/MovieScenePropertySystemTypes.inl"
 
 #include "EntitySystem/MovieSceneOperationalTypeConversions.h"
 
@@ -229,7 +230,7 @@ struct TPropertyComponentHandlerImpl<TIntegerSequence<int, Indices...>, Property
 		}
 
 		// Get the initial value in case we have a value without a full-weighted absolute channel.
-		TComponentPtr<const OperationalType> InitialValueComponent;
+		TOptionalComponentReader<OperationalType> InitialValueComponent;
 		if (InParams.PropertyEntityID)
 		{
 			const FEntityManager& EntityManager = Blender->GetLinker()->EntityManager;
@@ -288,7 +289,7 @@ struct TPropertyComponentHandlerImpl<TIntegerSequence<int, Indices...>, Property
 		}
 
 		// Get the initial value in case we have a value without a full-weighted absolute channel.
-		TComponentPtr<const OperationalType> InitialValueComponent;
+		TOptionalComponentReader<OperationalType> InitialValueComponent;
 		if (InParams.PropertyEntityID)
 		{
 			const FEntityManager& EntityManager = Blender->GetLinker()->EntityManager;
diff --git a/Engine/Source/Runtime/MovieScene/Public/EntitySystem/MovieScenePropertyRegistry.h b/Engine/Source/Runtime/MovieScene/Public/EntitySystem/MovieScenePropertyRegistry.h
--- a/Engine/Source/Runtime/MovieScene/Public/EntitySystem/MovieScenePropertyRegistry.h
+++ b/Engine/Source/Runtime/MovieScene/Public/EntitySystem/MovieScenePropertyRegistry.h
@@ -11,11 +11,14 @@
 #include "EntitySystem/IMovieScenePropertyComponentHandler.h"
 
 #include "Misc/InlineValue.h"
+#include "Misc/TVariant.h"
 #include <initializer_list>
 
+struct FMovieScenePropertyBinding;
 
 class UMovieSceneBlenderSystem;
 class UMovieSceneEntitySystemLinker;
+class FTrackInstancePropertyBindings;
 
 namespace UE
 {
@@ -95,6 +98,11 @@ struct FPropertyCompositeDefinition
 	uint16 CompositeOffset;
 };
 
+/** Type aliases for a property that resolved to either a fast pointer offset (type index 0), or a custom property index (specific to the path of the property - type index 1) */
+using FResolvedFastProperty = TVariant<uint16, UE::MovieScene::FCustomPropertyIndex>;
+/** Type aliases for a property that resolved to either a fast pointer offset (type index 0), or a custom property index (specific to the path of the property - type index 1) with a fallback to a slow property binding (type index 2) */
+using FResolvedProperty = TVariant<uint16, UE::MovieScene::FCustomPropertyIndex, TSharedPtr<FTrackInstancePropertyBindings>>;
+
 template<typename PropertyType, typename OperationalType> struct TPropertyDefinitionBuilder;
 template<typename PropertyType, typename OperationalType, typename... Composites> struct TCompositePropertyDefinitionBuilder;
 
@@ -111,6 +119,26 @@ public:
 	FPropertyRegistry(FPropertyRegistry&&) = delete;
 	FPropertyRegistry(const FPropertyRegistry&) = delete;
 
+	/**
+	 * Resolve a property to either a fast ptr offset, or a custom property accessor based on the specified array
+	 *
+	 * @param Object          The object to resolve the property for
+	 * @param PropertyBinding The property binding to resolve
+	 * @param CustomAccessors A view to an array of custom accessors (as retrieved from ICustomPropertyRegistration::GetAccessors)
+	 * @return An optional variant specifying the resolved property if it resolved successfully
+	 */
+	static TOptional< FResolvedFastProperty > ResolveFastProperty(UObject* Object, const FMovieScenePropertyBinding& PropertyBinding, FCustomAccessorView CustomAccessors);
+
+	/**
+	 * Resolve a property to either a fast ptr offset, or a custom property accessor based on the specified array falling back to a slow instance binding if possible
+	 *
+	 * @param Object          The object to resolve the property for
+	 * @param PropertyBinding The property binding to resolve
+	 * @param CustomAccessors A view to an array of custom accessors (as retrieved from ICustomPropertyRegistration::GetAccessors)
+	 * @return An optional variant specifying the resolved property if it resolved successfully
+	 */
+	static TOptional< FResolvedProperty > ResolveProperty(UObject* Object, const FMovieScenePropertyBinding& PropertyBinding, FCustomAccessorView CustomAccessors);
+
 	/**
 	 * Define a new animatable composite property type from its components.
 	 * 
diff --git a/Engine/Source/Runtime/MovieScene/Public/EntitySystem/MovieScenePropertySystemTypes.h b/Engine/Source/Runtime/MovieScene/Public/EntitySystem/MovieScenePropertySystemTypes.h
--- a/Engine/Source/Runtime/MovieScene/Public/EntitySystem/MovieScenePropertySystemTypes.h
+++ b/Engine/Source/Runtime/MovieScene/Public/EntitySystem/MovieScenePropertySystemTypes.h
@@ -3,18 +3,16 @@
 #pragma once
 
 #include "CoreTypes.h"
+#include "UObject/Class.h"
+#include "UObject/Object.h"
 #include "UObject/NameTypes.h"
 #include "Templates/SharedPointer.h"
 #include "Templates/UnrealTypeTraits.h"
-#include "Containers/SparseArray.h"
-#include "MovieSceneCommonHelpers.h"
-#include "EntitySystem/MovieSceneEntitySystemTypes.h"
-#include "EntitySystem/MovieSceneComponentRegistry.h"
-#include "EntitySystem/MovieSceneEntitySystemLinker.h"
 #include "EntitySystem/MovieSceneComponentAccessors.h"
-#include "EntitySystem/MovieSceneOperationalTypeConversions.h"
 
 class UClass;
+class FTrackInstancePropertyBindings;
+
 
 namespace UE
 {
@@ -130,12 +128,33 @@ struct FCustomAccessorView
 		return ViewNum;
 	}
 
+	int32 FindCustomAccessorIndex(UClass* ClassType, FName PropertyPath) const
+	{
+		UClass* StopIterationAt = UObject::StaticClass();
+
+		while (ClassType != StopIterationAt)
+		{
+			for (int32 Index = 0; Index < ViewNum; ++Index)
+			{
+				const FCustomPropertyAccessor& Accessor = (*this)[Index];
+				if (Accessor.Class == ClassType && Accessor.PropertyPath == PropertyPath)
+				{
+					return Index;
+				}
+			}
+			ClassType = ClassType->GetSuperClass();
+		}
+
+		return INDEX_NONE;
+	}
+
 private:
 	const uint8* Base;
 	int32 ViewNum;
 	int32 Stride;
 };
 
+
 struct ICustomPropertyRegistration
 {
 	virtual ~ICustomPropertyRegistration() {}
@@ -215,13 +234,7 @@ struct TSetPropertyValues
 	/**
 	 * Run before this task executes any logic over entities and components
 	 */
-	void PreTask()
-	{
-		if (CustomProperties)
-		{
-			CustomAccessors = CustomProperties->GetAccessors();
-		}
-	}
+	void PreTask();
 
 
 	/**
@@ -234,11 +247,8 @@ struct TSetPropertyValues
 	 *     .Read( TComponentTypeID<PropertyType>(...) )
 	 *     .Dispatch_PerEntity<TSetPropertyValues<PropertyType>>(...);
 	 */
-	void ForEachEntity(UObject* InObject, FCustomPropertyIndex CustomIndex, ParamType ValueToSet)
-	{
-		const TCustomPropertyAccessor<PropertyType>& CustomAccessor = static_cast<const TCustomPropertyAccessor<PropertyType>&>(CustomAccessors[CustomIndex.Value]);
-		CustomAccessor.Functions.Setter(InObject, ValueToSet);
-	}
+	void ForEachEntity(UObject* InObject, FCustomPropertyIndex CustomIndex, ParamType ValueToSet);
+	static void ForEachEntity(UObject* InObject, const FCustomPropertyAccessor& CustomAccessor, ParamType ValueToSet);
 
 
 	/**
@@ -251,15 +261,7 @@ struct TSetPropertyValues
 	 *     .Read( TComponentTypeID<PropertyType>(...) )
 	 *     .Dispatch_PerEntity<TSetPropertyValues<PropertyType>>(...);
 	 */
-	void ForEachEntity(UObject* InObject, uint16 PropertyOffset, ParamType ValueToSet)
-	{
-		// Would really like to avoid branching here, but if we encounter this data the options are either handle it gracefully, stomp a vtable, or report a fatal error.
-		if (ensureAlwaysMsgf(PropertyOffset != 0, TEXT("Invalid property offset specified (ptr+%d bytes) for property on object %s. This would otherwise overwrite the object's vfptr."), PropertyOffset, *InObject->GetName()))
-		{
-			PropertyType* PropertyAddress = reinterpret_cast<PropertyType*>( reinterpret_cast<uint8*>(InObject) + PropertyOffset );
-			*PropertyAddress = ValueToSet;
-		}
-	}
+	static void ForEachEntity(UObject* InObject, uint16 PropertyOffset, ParamType ValueToSet);
 
 
 	/**
@@ -272,15 +274,12 @@ struct TSetPropertyValues
 	 *     .Read( TComponentTypeID<PropertyType>(...) )
 	 *     .Dispatch_PerEntity<TSetPropertyValues<PropertyType>>(...);
 	 */
-	void ForEachEntity(UObject* InObject, TSharedPtr<FTrackInstancePropertyBindings> PropertyBindings, ParamType ValueToSet)
-	{
-		PropertyBindings->CallFunction<PropertyType>(*InObject, ValueToSet);
-	}
+	static void ForEachEntity(UObject* InObject, TSharedPtr<FTrackInstancePropertyBindings> PropertyBindings, ParamType ValueToSet);
 
 public:
 
-	using FThreeWayAccessor  = TReadOneOf< FCustomPropertyIndex, uint16, TSharedPtr<FTrackInstancePropertyBindings> >;
-	using FTwoWayAccessor    = TReadOneOf< uint16, TSharedPtr<FTrackInstancePropertyBindings> >;
+	using FThreeWayAccessor  = TMultiReadOptional<FCustomPropertyIndex, uint16, TSharedPtr<FTrackInstancePropertyBindings>>;
+	using FTwoWayAccessor    = TMultiReadOptional<uint16, TSharedPtr<FTrackInstancePropertyBindings>>;
 
 	/**
 	 * Task callback that applies properties for a whole allocation of entities with either an FCustomPropertyIndex, uint16 or TSharedPtr<FTrackInstancePropertyBindings> property component
@@ -292,37 +291,7 @@ public:
 	 *     .Read(      TComponentTypeID<PropertyType>(...) )
 	 *     .Dispatch_PerAllocation<TSetPropertyValues<PropertyType>>(...);
 	 */
-	void ForEachAllocation(const FEntityAllocation* Allocation, TRead<UObject*> BoundObjectComponents, FThreeWayAccessor PropertyBindingComponents, TRead<PropertyType> PropertyValueComponents)
-	{
-		using FPropertyTuple = TTuple< const FCustomPropertyIndex*, const uint16*, const TSharedPtr<FTrackInstancePropertyBindings>* >;
-
-		UObject* const *    Objects    = BoundObjectComponents.Resolve(Allocation);
-		FPropertyTuple      Properties = PropertyBindingComponents.Resolve(Allocation);
-		const PropertyType* Input      = PropertyValueComponents.Resolve(Allocation);
-
-		const int32 Num = Allocation->Num();
-		if (const FCustomPropertyIndex* Custom = Properties.template Get<0>())
-		{
-			for (int32 Index = 0; Index < Num; ++Index)
-			{
-				ForEachEntity(*Objects++, *Custom++, *Input++);
-			}
-		}
-		else if (const uint16* Fast = Properties.template Get<1>())
-		{
-			for (int32 Index = 0; Index < Num; ++Index)
-			{
-				ForEachEntity(*Objects++, *Fast++, *Input++);
-			}
-		}
-		else if (const TSharedPtr<FTrackInstancePropertyBindings>* Slow = Properties.template Get<2>())
-		{
-			for (int32 Index = 0; Index < Num; ++Index)
-			{
-				ForEachEntity(*Objects++, *Slow++, *Input++);
-			}
-		}
-	}
+	void ForEachAllocation(const FEntityAllocation* Allocation, TRead<UObject*> BoundObjectComponents, FThreeWayAccessor ResolvedPropertyComponents, TRead<PropertyType> PropertyValueComponents);
 
 
 	/**
@@ -335,30 +304,7 @@ public:
 	 *     .Read(      TComponentTypeID<PropertyType>(...) )
 	 *     .Dispatch_PerAllocation<TSetPropertyValues<PropertyType>>(...);
 	 */
-	void ForEachAllocation(const FEntityAllocation* Allocation, TRead<UObject*> BoundObjectComponents, FTwoWayAccessor PropertyBindingComponents, TRead<PropertyType> PropertyValueComponents)
-	{
-		using FPropertyTuple = TTuple< const uint16*, const TSharedPtr<FTrackInstancePropertyBindings>* >;
-
-		UObject* const *    Objects    = BoundObjectComponents.Resolve(Allocation);
-		FPropertyTuple      Properties = PropertyBindingComponents.Resolve(Allocation);
-		const PropertyType* Input      = PropertyValueComponents.Resolve(Allocation);
-
-		const int32 Num = Allocation->Num();
-		if (const uint16* Fast = Properties.template Get<0>())
-		{
-			for (int32 Index = 0; Index < Num; ++Index)
-			{
-				ForEachEntity(*Objects++, *Fast++, *Input++);
-			}
-		}
-		else if (const TSharedPtr<FTrackInstancePropertyBindings>* Slow = Properties.template Get<1>())
-		{
-			for (int32 Index = 0; Index < Num; ++Index)
-			{
-				ForEachEntity(*Objects++, *Slow++, *Input++);
-			}
-		}
-	}
+	void ForEachAllocation(const FEntityAllocation* Allocation, TRead<UObject*> BoundObjectComponents, FTwoWayAccessor ResolvedPropertyComponents, TRead<PropertyType> PropertyValueComponents);
 
 private:
 
@@ -401,13 +347,7 @@ struct TGetPropertyValues
 		: CustomProperties(InCustomProperties)
 	{}
 
-	void PreTask()
-	{
-		if (CustomProperties)
-		{
-			CustomAccessors = CustomProperties->GetAccessors();
-		}
-	}
+	void PreTask();
 
 	/**
 	 * Task callback that retrieves the object's current value via a custom native setter function, and writes it to the specified output variable
@@ -419,11 +359,7 @@ struct TGetPropertyValues
 	 *     .Write( TComponentTypeID<OperationalType>(...) )
 	 *     .Dispatch_PerEntity<TGetPropertyValues<PropertyType, OperationalType>>(...);
 	 */
-	void ForEachEntity(UObject* InObject, FCustomPropertyIndex CustomPropertyIndex, OperationalType& OutValue)
-	{
-		const TCustomPropertyAccessor<PropertyType>& CustomAccessor = static_cast<const TCustomPropertyAccessor<PropertyType>&>(CustomAccessors[CustomPropertyIndex.Value]);
-		ConvertOperationalProperty(CustomAccessor.Functions.Getter(InObject), OutValue);
-	}
+	void ForEachEntity(UObject* InObject, FCustomPropertyIndex CustomPropertyIndex, OperationalType& OutValue);
 
 	/**
 	 * Task callback that retrieves the object's current value via a fast pointer offset, and writes it to the specified output variable
@@ -435,16 +371,7 @@ struct TGetPropertyValues
 	 *     .Read( TComponentTypeID<OperationalType>(...) )
 	 *     .Dispatch_PerEntity<TGetPropertyValues<PropertyType, OperationalType>>(...);
 	 */
-	void ForEachEntity(UObject* InObject, uint16 PropertyOffset, OperationalType& OutValue)
-	{
-		// Would really like to avoid branching here, but if we encounter this data the options are either handle it gracefully, stomp a vtable, or report a fatal error.
-		if (ensureAlwaysMsgf(PropertyOffset != 0, TEXT("Invalid property offset specified (ptr+%d bytes) for property on object %s. This would otherwise overwrite the object's vfptr."), PropertyOffset, *InObject->GetName()))
-		{
-			PropertyType* PropertyAddress = reinterpret_cast<PropertyType*>( reinterpret_cast<uint8*>(InObject) + PropertyOffset );
-
-			ConvertOperationalProperty(*PropertyAddress, OutValue);
-		}
-	}
+	void ForEachEntity(UObject* InObject, uint16 PropertyOffset, OperationalType& OutValue);
 
 	/**
 	 * Task callback that retrieves the object's current value via a slow (legacy) track instance binding, and writes it to the specified output variable
@@ -456,15 +383,12 @@ struct TGetPropertyValues
 	 *     .Read( TComponentTypeID<OperationalType>(...) )
 	 *     .Dispatch_PerEntity<TGetPropertyValues<PropertyType, OperationalType>>(...);
 	 */
-	void ForEachEntity(UObject* InObject, TSharedPtr<FTrackInstancePropertyBindings> PropertyBindings, OperationalType& OutValue)
-	{
-		ConvertOperationalProperty(PropertyBindings->GetCurrentValue<PropertyType>(*InObject), OutValue);
-	}
+	void ForEachEntity(UObject* InObject, TSharedPtr<FTrackInstancePropertyBindings> PropertyBindings, OperationalType& OutValue);
 
 public:
 
-	using FThreeWayAccessor  = TReadOneOf< FCustomPropertyIndex, uint16, TSharedPtr<FTrackInstancePropertyBindings> >;
-	using FTwoWayAccessor    = TReadOneOf< uint16, TSharedPtr<FTrackInstancePropertyBindings> >;
+	using FThreeWayAccessor  = TMultiReadOptional<FCustomPropertyIndex, uint16, TSharedPtr<FTrackInstancePropertyBindings>>;
+	using FTwoWayAccessor    = TMultiReadOptional<uint16, TSharedPtr<FTrackInstancePropertyBindings>>;
 
 
 	/**
@@ -477,37 +401,7 @@ public:
 	 *     .Write(     TComponentTypeID<OperationalType>(...) )
 	 *     .Dispatch_PerAllocation<TGetPropertyValues<PropertyType, OperationalType>>(...);
 	 */
-	void ForEachAllocation(const FEntityAllocation* Allocation, TRead<UObject*> BoundObjectComponents, FThreeWayAccessor PropertyBindingComponents, TWrite<OperationalType> OutValueComponents)
-	{
-		using FPropertyTuple = TTuple< const FCustomPropertyIndex*, const uint16*, const TSharedPtr<FTrackInstancePropertyBindings>* >;
-
-		UObject* const *    Objects    = BoundObjectComponents.Resolve(Allocation);
-		FPropertyTuple      Properties = PropertyBindingComponents.Resolve(Allocation);
-		OperationalType*   Output     = OutValueComponents.Resolve(Allocation);
-
-		const int32 Num = Allocation->Num();
-		if (const FCustomPropertyIndex* Custom = Properties.template Get<0>())
-		{
-			for (int32 Index = 0; Index < Num; ++Index)
-			{
-				ForEachEntity(Objects[Index], Custom[Index], Output[Index]);
-			}
-		}
-		else if (const uint16* Fast = Properties.template Get<1>())
-		{
-			for (int32 Index = 0; Index < Num; ++Index)
-			{
-				ForEachEntity(Objects[Index], Fast[Index], Output[Index]);
-			}
-		}
-		else if (const TSharedPtr<FTrackInstancePropertyBindings>* Slow = Properties.template Get<2>())
-		{
-			for (int32 Index = 0; Index < Num; ++Index)
-			{
-				ForEachEntity(Objects[Index], Slow[Index], Output[Index]);
-			}
-		}
-	}
+	void ForEachAllocation(const FEntityAllocation* Allocation, TRead<UObject*> BoundObjectComponents, FThreeWayAccessor ResolvedPropertyComponents, TWrite<OperationalType> OutValueComponents);
 
 
 	/**
@@ -520,30 +414,7 @@ public:
 	 *     .Read(      TComponentTypeID<OperationalType>(...) )
 	 *     .Dispatch_PerAllocation<TGetPropertyValues<PropertyType ,OperationalType>>(...);
 	 */
-	void ForEachAllocation(const FEntityAllocation* Allocation, TRead<UObject*> BoundObjectComponents, FTwoWayAccessor PropertyBindingComponents, TWrite<OperationalType> OutValueComponents)
-	{
-		using FPropertyTuple = TTuple< const uint16*, const TSharedPtr<FTrackInstancePropertyBindings>* >;
-
-		UObject* const *    Objects    = BoundObjectComponents.Resolve(Allocation);
-		FPropertyTuple      Properties = PropertyBindingComponents.Resolve(Allocation);
-		OperationalType*   Output     = OutValueComponents.Resolve(Allocation);
-
-		const int32 Num = Allocation->Num();
-		if (const uint16* Fast = Properties.template Get<0>())
-		{
-			for (int32 Index = 0; Index < Num; ++Index)
-			{
-				ForEachEntity(Objects[Index], Fast[Index], Output[Index]);
-			}
-		}
-		else if (const TSharedPtr<FTrackInstancePropertyBindings>* Slow = Properties.template Get<1>())
-		{
-			for (int32 Index = 0; Index < Num; ++Index)
-			{
-				ForEachEntity(Objects[Index], Slow[Index], Output[Index]);
-			}
-		}
-	}
+	void ForEachAllocation(const FEntityAllocation* Allocation, TRead<UObject*> BoundObjectComponents, FTwoWayAccessor ResolvedPropertyComponents, TWrite<OperationalType> OutValueComponents);
 
 private:
 
@@ -594,13 +465,7 @@ struct TSetCompositePropertyValuesImpl<TIntegerSequence<int, Indices...>, Proper
 	/**
 	 * Run before this task executes any logic over entities and components
 	 */
-	void PreTask()
-	{
-		if (CustomProperties)
-		{
-			CustomAccessors = CustomProperties->GetAccessors();
-		}
-	}
+	void PreTask();
 
 	/**
 	 * Task callback that applies a value to an object property via a custom native setter function
@@ -614,13 +479,7 @@ struct TSetCompositePropertyValuesImpl<TIntegerSequence<int, Indices...>, Proper
 	 *     .Read( TComponentTypeID<CompositeType[N-1]>(...) )
 	 *     .Dispatch_PerEntity<TSetCompositePropertyValues<PropertyType, CompositeTypes...>>(..., Projection);
 	 */
-	void ForEachEntity(UObject* InObject, FCustomPropertyIndex CustomPropertyIndex, typename TCallTraits<CompositeTypes>::ParamType... CompositeResults)
-	{
-		const TCustomPropertyAccessor<PropertyType>& CustomAccessor = static_cast<const TCustomPropertyAccessor<PropertyType>&>(CustomAccessors[CustomPropertyIndex.Value]);
-
-		PropertyType Result = Invoke(Projection, CompositeResults...);
-		CustomAccessor.Functions.Setter(InObject, Result);
-	}
+	void ForEachEntity(UObject* InObject, FCustomPropertyIndex CustomPropertyIndex, typename TCallTraits<CompositeTypes>::ParamType... CompositeResults);
 
 
 	/**
@@ -635,17 +494,7 @@ struct TSetCompositePropertyValuesImpl<TIntegerSequence<int, Indices...>, Proper
 	 *     .Read( TComponentTypeID<CompositeType[N-1]>(...) )
 	 *     .Dispatch_PerEntity<TSetCompositePropertyValues<PropertyType, CompositeTypes...>>(...);
 	 */
-	void ForEachEntity(UObject* InObject, uint16 PropertyOffset, typename TCallTraits<CompositeTypes>::ParamType... CompositeResults)
-	{
-		// Would really like to avoid branching here, but if we encounter this data the options are either handle it gracefully, stomp a vtable, or report a fatal error.
-		if (ensureAlwaysMsgf(PropertyOffset != 0, TEXT("Invalid property offset specified (ptr+%d bytes) for property on object %s. This would otherwise overwrite the object's vfptr."), PropertyOffset, *InObject->GetName()))
-		{
-			PropertyType Result = Invoke(Projection, CompositeResults...);
-
-			PropertyType* PropertyAddress = reinterpret_cast<PropertyType*>( reinterpret_cast<uint8*>(InObject) + PropertyOffset );
-			*PropertyAddress = Result;
-		}
-	}
+	void ForEachEntity(UObject* InObject, uint16 PropertyOffset, typename TCallTraits<CompositeTypes>::ParamType... CompositeResults);
 
 
 	/**
@@ -660,16 +509,12 @@ struct TSetCompositePropertyValuesImpl<TIntegerSequence<int, Indices...>, Proper
 	 *     .Read( TComponentTypeID<CompositeType[N-1]>(...) )
 	 *     .Dispatch_PerEntity<TSetCompositePropertyValues<PropertyType, CompositeTypes...>>(...);
 	 */
-	void ForEachEntity(UObject* InObject, TSharedPtr<FTrackInstancePropertyBindings> PropertyBindings, typename TCallTraits<CompositeTypes>::ParamType... CompositeResults)
-	{
-		PropertyType Result = Invoke(Projection, CompositeResults...);
-		PropertyBindings->CallFunction<PropertyType>(*InObject, Result);
-	}
+	void ForEachEntity(UObject* InObject, TSharedPtr<FTrackInstancePropertyBindings> PropertyBindings, typename TCallTraits<CompositeTypes>::ParamType... CompositeResults);
 
 public:
 
-	using FThreeWayAccessor  = TReadOneOf< FCustomPropertyIndex, uint16, TSharedPtr<FTrackInstancePropertyBindings> >;
-	using FTwoWayAccessor    = TReadOneOf< uint16, TSharedPtr<FTrackInstancePropertyBindings> >;
+	using FThreeWayAccessor  = TMultiReadOptional<FCustomPropertyIndex, uint16, TSharedPtr<FTrackInstancePropertyBindings>>;
+	using FTwoWayAccessor    = TMultiReadOptional<uint16, TSharedPtr<FTrackInstancePropertyBindings>>;
 
 
 	/**
@@ -684,41 +529,7 @@ public:
 	 *     .Read( TComponentTypeID<CompositeType[N-1]>(...) )
 	 *     .Dispatch_PerAllocation<TSetCompositePropertyValues<PropertyType, CompositeTypes...>>(...);
 	 */
-	void ForEachAllocation(const FEntityAllocation* Allocation, TRead<UObject*> BoundObjectComponents, FThreeWayAccessor PropertyBindingComponents, TRead<CompositeTypes>... VariadicComponents)
-	{
-		using FPropertyTuple = TTuple< const FCustomPropertyIndex*, const uint16*, const TSharedPtr<FTrackInstancePropertyBindings>* >;
-
-		UObject* const * Objects    = BoundObjectComponents.Resolve(Allocation);
-		FPropertyTuple   Properties = PropertyBindingComponents.Resolve(Allocation);
-
-		auto CompositesState = MakeTuple( VariadicComponents.CreateIterState(Allocation)... );
-
-		const int32 Num = Allocation->Num();
-		if (const FCustomPropertyIndex* Custom = Properties.template Get<0>())
-		{
-			for (int32 ComponentOffset = 0; ComponentOffset < Num; ++ComponentOffset )
-			{
-				ForEachEntity( *Objects++, *Custom++, *CompositesState.template Get<Indices>()... );
-				VisitTupleElements([](auto& It){ ++It; }, CompositesState);
-			}
-		}
-		else if (const uint16* Fast = Properties.template Get<1>())
-		{
-			for (int32 ComponentOffset = 0; ComponentOffset < Num; ++ComponentOffset )
-			{
-				ForEachEntity( *Objects++, *Fast++, *CompositesState.template Get<Indices>()... );
-				VisitTupleElements([](auto& It){ ++It; }, CompositesState);
-			}
-		}
-		else if (const TSharedPtr<FTrackInstancePropertyBindings>* Slow = Properties.template Get<2>())
-		{
-			for (int32 ComponentOffset = 0; ComponentOffset < Num; ++ComponentOffset )
-			{
-				ForEachEntity( *Objects++, *Slow++, *CompositesState.template Get<Indices>()... );
-				VisitTupleElements([](auto& It){ ++It; }, CompositesState);
-			}
-		}
-	}
+	void ForEachAllocation(const FEntityAllocation* Allocation, TRead<UObject*> BoundObjectComponents, FThreeWayAccessor ResolvedPropertyComponents, TRead<CompositeTypes>... VariadicComponents);
 
 
 	/**
@@ -733,33 +544,8 @@ public:
 	 *     .Read( TComponentTypeID<CompositeType[N-1]>(...) )
 	 *     .Dispatch_PerAllocation<TSetCompositePropertyValues<PropertyType, CompositeTypes...>>(...);
 	 */
-	void ForEachAllocation(const FEntityAllocation* Allocation, TRead<UObject*> BoundObjectComponents, FTwoWayAccessor PropertyBindingComponents, TRead<CompositeTypes>... VariadicComponents)
-	{
-		using FPropertyTuple = TTuple< const uint16*, const TSharedPtr<FTrackInstancePropertyBindings>* >;
-
-		UObject* const * Objects    = BoundObjectComponents.Resolve(Allocation);
-		FPropertyTuple   Properties = PropertyBindingComponents.Resolve(Allocation);
+	void ForEachAllocation(const FEntityAllocation* Allocation, TRead<UObject*> BoundObjectComponents, FTwoWayAccessor ResolvedPropertyComponents, TRead<CompositeTypes>... VariadicComponents);
 
-		auto CompositesState = MakeTuple( VariadicComponents.CreateIterState(Allocation)... );
-
-		const int32 Num = Allocation->Num();
-		if (const uint16* Fast = Properties.template Get<0>())
-		{
-			for (int32 ComponentOffset = 0; ComponentOffset < Num; ++ComponentOffset )
-			{
-				ForEachEntity( *Objects++, *Fast++, *CompositesState.template Get<Indices>()... );
-				VisitTupleElements([](auto& It){ ++It; }, CompositesState);
-			}
-		}
-		else if (const TSharedPtr<FTrackInstancePropertyBindings>* Slow = Properties.template Get<1>())
-		{
-			for (int32 ComponentOffset = 0; ComponentOffset < Num; ++ComponentOffset )
-			{
-				ForEachEntity( *Objects++, *Slow++, *CompositesState.template Get<Indices>()... );
-				VisitTupleElements([](auto& It){ ++It; }, CompositesState);
-			}
-		}
-	}
 private:
 
 	ICustomPropertyRegistration* CustomProperties;
diff --git a/Engine/Source/Runtime/MovieScene/Public/EntitySystem/MovieScenePropertySystemTypes.inl b/Engine/Source/Runtime/MovieScene/Public/EntitySystem/MovieScenePropertySystemTypes.inl
new file mode 100644
--- /dev/null
+++ b/Engine/Source/Runtime/MovieScene/Public/EntitySystem/MovieScenePropertySystemTypes.inl
@@ -0,0 +1,279 @@
+// Copyright Epic Games, Inc. All Rights Reserved.
+
+#include "EntitySystem/MovieScenePropertySystemTypes.h"
+#include "CoreTypes.h"
+#include "UObject/NameTypes.h"
+#include "Templates/SharedPointer.h"
+#include "Templates/UnrealTypeTraits.h"
+#include "Containers/SparseArray.h"
+#include "MovieSceneCommonHelpers.h"
+#include "EntitySystem/MovieSceneEntitySystemTypes.h"
+#include "EntitySystem/MovieSceneComponentRegistry.h"
+#include "EntitySystem/MovieSceneEntitySystemLinker.h"
+#include "EntitySystem/MovieSceneComponentAccessors.h"
+#include "EntitySystem/MovieSceneOperationalTypeConversions.h"
+#include "EntitySystem/MovieScenePropertyBinding.h"
+
+
+namespace UE
+{
+namespace MovieScene
+{
+
+template<typename PropertyType>
+void TSetPropertyValues<PropertyType>::PreTask()
+{
+	if (CustomProperties)
+	{
+		CustomAccessors = CustomProperties->GetAccessors();
+	}
+}
+
+template<typename PropertyType>
+void TSetPropertyValues<PropertyType>::ForEachEntity(UObject* InObject, FCustomPropertyIndex CustomIndex, ParamType ValueToSet)
+{
+	ForEachEntity(InObject, CustomAccessors[CustomIndex.Value], ValueToSet);
+}
+
+template<typename PropertyType>
+void TSetPropertyValues<PropertyType>::ForEachEntity(UObject* InObject, const FCustomPropertyAccessor& BaseCustomAccessor, ParamType ValueToSet)
+{
+
+	const TCustomPropertyAccessor<PropertyType>& CustomAccessor = static_cast<const TCustomPropertyAccessor<PropertyType>&>(BaseCustomAccessor);
+	CustomAccessor.Functions.Setter(InObject, ValueToSet);
+}
+
+template<typename PropertyType>
+void TSetPropertyValues<PropertyType>::ForEachEntity(UObject* InObject, uint16 PropertyOffset, ParamType ValueToSet)
+{
+	// Would really like to avoid branching here, but if we encounter this data the options are either handle it gracefully, stomp a vtable, or report a fatal error.
+	if (ensureAlwaysMsgf(PropertyOffset != 0, TEXT("Invalid property offset specified (ptr+%d bytes) for property on object %s. This would otherwise overwrite the object's vfptr."), PropertyOffset, *InObject->GetName()))
+	{
+		PropertyType* PropertyAddress = reinterpret_cast<PropertyType*>( reinterpret_cast<uint8*>(InObject) + PropertyOffset );
+		*PropertyAddress = ValueToSet;
+	}
+}
+
+template<typename PropertyType>
+void TSetPropertyValues<PropertyType>::ForEachEntity(UObject* InObject, TSharedPtr<FTrackInstancePropertyBindings> PropertyBindings, ParamType ValueToSet)
+{
+	PropertyBindings->CallFunction<PropertyType>(*InObject, ValueToSet);
+}
+
+template<typename PropertyType>
+void TSetPropertyValues<PropertyType>::ForEachAllocation(const FEntityAllocation* Allocation, TRead<UObject*> BoundObjectComponents, FThreeWayAccessor ResolvedPropertyComponents, TRead<PropertyType> PropertyValueComponents)
+{
+	const int32 Num = Allocation->Num();
+	if (const FCustomPropertyIndex* Custom = ResolvedPropertyComponents.template Get<0>())
+	{
+		for (int32 Index = 0; Index < Num; ++Index)
+		{
+			ForEachEntity(BoundObjectComponents[Index], *Custom++, PropertyValueComponents[Index]);
+		}
+	}
+	else if (const uint16* Fast = ResolvedPropertyComponents.template Get<1>())
+	{
+		for (int32 Index = 0; Index < Num; ++Index)
+		{
+			ForEachEntity(BoundObjectComponents[Index], *Fast++, PropertyValueComponents[Index]);
+		}
+	}
+	else if (const TSharedPtr<FTrackInstancePropertyBindings>* Slow = ResolvedPropertyComponents.template Get<2>())
+	{
+		for (int32 Index = 0; Index < Num; ++Index)
+		{
+			ForEachEntity(BoundObjectComponents[Index], *Slow++, PropertyValueComponents[Index]);
+		}
+	}
+}
+
+template<typename PropertyType>
+void TSetPropertyValues<PropertyType>::ForEachAllocation(const FEntityAllocation* Allocation, TRead<UObject*> BoundObjectComponents, FTwoWayAccessor ResolvedPropertyComponents, TRead<PropertyType> PropertyValueComponents)
+{
+	const int32 Num = Allocation->Num();
+	if (const uint16* Fast = ResolvedPropertyComponents.template Get<0>())
+	{
+		for (int32 Index = 0; Index < Num; ++Index)
+		{
+			ForEachEntity(BoundObjectComponents[Index], *Fast++, PropertyValueComponents[Index]);
+		}
+	}
+	else if (const TSharedPtr<FTrackInstancePropertyBindings>* Slow = ResolvedPropertyComponents.template Get<1>())
+	{
+		for (int32 Index = 0; Index < Num; ++Index)
+		{
+			ForEachEntity(BoundObjectComponents[Index], *Slow++, PropertyValueComponents[Index]);
+		}
+	}
+}
+
+
+template<typename PropertyType, typename OperationalType>
+void TGetPropertyValues<PropertyType, OperationalType>::PreTask()
+{
+	if (CustomProperties)
+	{
+		CustomAccessors = CustomProperties->GetAccessors();
+	}
+}
+
+template<typename PropertyType, typename OperationalType>
+void TGetPropertyValues<PropertyType, OperationalType>::ForEachEntity(UObject* InObject, FCustomPropertyIndex CustomPropertyIndex, OperationalType& OutValue)
+{
+	const TCustomPropertyAccessor<PropertyType>& CustomAccessor = static_cast<const TCustomPropertyAccessor<PropertyType>&>(CustomAccessors[CustomPropertyIndex.Value]);
+	ConvertOperationalProperty(CustomAccessor.Functions.Getter(InObject), OutValue);
+}
+
+template<typename PropertyType, typename OperationalType>
+void TGetPropertyValues<PropertyType, OperationalType>::ForEachEntity(UObject* InObject, uint16 PropertyOffset, OperationalType& OutValue)
+{
+	// Would really like to avoid branching here, but if we encounter this data the options are either handle it gracefully, stomp a vtable, or report a fatal error.
+	if (ensureAlwaysMsgf(PropertyOffset != 0, TEXT("Invalid property offset specified (ptr+%d bytes) for property on object %s. This would otherwise overwrite the object's vfptr."), PropertyOffset, *InObject->GetName()))
+	{
+		PropertyType* PropertyAddress = reinterpret_cast<PropertyType*>( reinterpret_cast<uint8*>(InObject) + PropertyOffset );
+
+		ConvertOperationalProperty(*PropertyAddress, OutValue);
+	}
+}
+
+template<typename PropertyType, typename OperationalType>
+void TGetPropertyValues<PropertyType, OperationalType>::ForEachEntity(UObject* InObject, TSharedPtr<FTrackInstancePropertyBindings> PropertyBindings, OperationalType& OutValue)
+{
+	ConvertOperationalProperty(PropertyBindings->GetCurrentValue<PropertyType>(*InObject), OutValue);
+}
+
+template<typename PropertyType, typename OperationalType>
+void TGetPropertyValues<PropertyType, OperationalType>::ForEachAllocation(const FEntityAllocation* Allocation, TRead<UObject*> BoundObjectComponents, FThreeWayAccessor ResolvedPropertyComponents, TWrite<OperationalType> OutValueComponents)
+{
+	const int32 Num = Allocation->Num();
+	if (const FCustomPropertyIndex* Custom = ResolvedPropertyComponents.template Get<0>())
+	{
+		for (int32 Index = 0; Index < Num; ++Index)
+		{
+			ForEachEntity(BoundObjectComponents[Index], Custom[Index], OutValueComponents[Index]);
+		}
+	}
+	else if (const uint16* Fast = ResolvedPropertyComponents.template Get<1>())
+	{
+		for (int32 Index = 0; Index < Num; ++Index)
+		{
+			ForEachEntity(BoundObjectComponents[Index], Fast[Index], OutValueComponents[Index]);
+		}
+	}
+	else if (const TSharedPtr<FTrackInstancePropertyBindings>* Slow = ResolvedPropertyComponents.template Get<2>())
+	{
+		for (int32 Index = 0; Index < Num; ++Index)
+		{
+			ForEachEntity(BoundObjectComponents[Index], Slow[Index], OutValueComponents[Index]);
+		}
+	}
+}
+
+template<typename PropertyType, typename OperationalType>
+void TGetPropertyValues<PropertyType, OperationalType>::ForEachAllocation(const FEntityAllocation* Allocation, TRead<UObject*> BoundObjectComponents, FTwoWayAccessor ResolvedPropertyComponents, TWrite<OperationalType> OutValueComponents)
+{
+	const int32 Num = Allocation->Num();
+	if (const uint16* Fast = ResolvedPropertyComponents.template Get<0>())
+	{
+		for (int32 Index = 0; Index < Num; ++Index)
+		{
+			ForEachEntity(BoundObjectComponents[Index], Fast[Index], OutValueComponents[Index]);
+		}
+	}
+	else if (const TSharedPtr<FTrackInstancePropertyBindings>* Slow = ResolvedPropertyComponents.template Get<1>())
+	{
+		for (int32 Index = 0; Index < Num; ++Index)
+		{
+			ForEachEntity(BoundObjectComponents[Index], Slow[Index], OutValueComponents[Index]);
+		}
+	}
+}
+
+
+template<typename PropertyType, typename ProjectionType, typename... CompositeTypes, int... Indices>
+void TSetCompositePropertyValuesImpl<TIntegerSequence<int, Indices...>, PropertyType, ProjectionType, CompositeTypes...>::PreTask()
+{
+	if (CustomProperties)
+	{
+		CustomAccessors = CustomProperties->GetAccessors();
+	}
+}
+
+template<typename PropertyType, typename ProjectionType, typename... CompositeTypes, int... Indices>
+void TSetCompositePropertyValuesImpl<TIntegerSequence<int, Indices...>, PropertyType, ProjectionType, CompositeTypes...>::ForEachEntity(UObject* InObject, FCustomPropertyIndex CustomPropertyIndex, typename TCallTraits<CompositeTypes>::ParamType... CompositeResults)
+{
+	const TCustomPropertyAccessor<PropertyType>& CustomAccessor = static_cast<const TCustomPropertyAccessor<PropertyType>&>(CustomAccessors[CustomPropertyIndex.Value]);
+
+	PropertyType Result = Invoke(Projection, CompositeResults...);
+	CustomAccessor.Functions.Setter(InObject, Result);
+}
+
+template<typename PropertyType, typename ProjectionType, typename... CompositeTypes, int... Indices>
+void TSetCompositePropertyValuesImpl<TIntegerSequence<int, Indices...>, PropertyType, ProjectionType, CompositeTypes...>::ForEachEntity(UObject* InObject, uint16 PropertyOffset, typename TCallTraits<CompositeTypes>::ParamType... CompositeResults)
+{
+	// Would really like to avoid branching here, but if we encounter this data the options are either handle it gracefully, stomp a vtable, or report a fatal error.
+	if (ensureAlwaysMsgf(PropertyOffset != 0, TEXT("Invalid property offset specified (ptr+%d bytes) for property on object %s. This would otherwise overwrite the object's vfptr."), PropertyOffset, *InObject->GetName()))
+	{
+		PropertyType Result = Invoke(Projection, CompositeResults...);
+
+		PropertyType* PropertyAddress = reinterpret_cast<PropertyType*>( reinterpret_cast<uint8*>(InObject) + PropertyOffset );
+		*PropertyAddress = Result;
+	}
+}
+
+template<typename PropertyType, typename ProjectionType, typename... CompositeTypes, int... Indices>
+void TSetCompositePropertyValuesImpl<TIntegerSequence<int, Indices...>, PropertyType, ProjectionType, CompositeTypes...>::ForEachEntity(UObject* InObject, TSharedPtr<FTrackInstancePropertyBindings> PropertyBindings, typename TCallTraits<CompositeTypes>::ParamType... CompositeResults)
+{
+	PropertyType Result = Invoke(Projection, CompositeResults...);
+	PropertyBindings->CallFunction<PropertyType>(*InObject, Result);
+}
+
+template<typename PropertyType, typename ProjectionType, typename... CompositeTypes, int... Indices>
+void TSetCompositePropertyValuesImpl<TIntegerSequence<int, Indices...>, PropertyType, ProjectionType, CompositeTypes...>::ForEachAllocation(const FEntityAllocation* Allocation, TRead<UObject*> BoundObjectComponents, FThreeWayAccessor ResolvedPropertyComponents, TRead<CompositeTypes>... VariadicComponents)
+{
+	const int32 Num = Allocation->Num();
+	if (const FCustomPropertyIndex* Custom = ResolvedPropertyComponents.template Get<0>())
+	{
+		for (int32 Index = 0; Index < Num; ++Index )
+		{
+			ForEachEntity( BoundObjectComponents[Index], Custom[Index], VariadicComponents[Index]... );
+		}
+	}
+	else if (const uint16* Fast = ResolvedPropertyComponents.template Get<1>())
+	{
+		for (int32 Index = 0; Index < Num; ++Index )
+		{
+			ForEachEntity( BoundObjectComponents[Index], Fast[Index], VariadicComponents[Index]... );
+		}
+	}
+	else if (const TSharedPtr<FTrackInstancePropertyBindings>* Slow = ResolvedPropertyComponents.template Get<2>())
+	{
+		for (int32 Index = 0; Index < Num; ++Index )
+		{
+			ForEachEntity( BoundObjectComponents[Index], Slow[Index], VariadicComponents[Index]... );
+		}
+	}
+}
+
+template<typename PropertyType, typename ProjectionType, typename... CompositeTypes, int... Indices>
+void TSetCompositePropertyValuesImpl<TIntegerSequence<int, Indices...>, PropertyType, ProjectionType, CompositeTypes...>::ForEachAllocation(const FEntityAllocation* Allocation, TRead<UObject*> BoundObjectComponents, FTwoWayAccessor ResolvedPropertyComponents, TRead<CompositeTypes>... VariadicComponents)
+{
+	const int32 Num = Allocation->Num();
+	if (const uint16* Fast = ResolvedPropertyComponents.template Get<0>())
+	{
+		for (int32 Index = 0; Index < Num; ++Index )
+		{
+			ForEachEntity( BoundObjectComponents[Index], Fast[Index], VariadicComponents[Index]... );
+		}
+	}
+	else if (const TSharedPtr<FTrackInstancePropertyBindings>* Slow = ResolvedPropertyComponents.template Get<1>())
+	{
+		for (int32 Index = 0; Index < Num; ++Index )
+		{
+			ForEachEntity( BoundObjectComponents[Index], Slow[Index], VariadicComponents[Index]... );
+		}
+	}
+}
+
+} // namespace MovieScene
+} // namespace UE
\ No newline at end of file
diff --git a/Engine/Source/Runtime/MovieSceneTracks/Private/EntitySystem/Interrogation/MovieSceneInterrogatedPropertyInstantiator.cpp b/Engine/Source/Runtime/MovieSceneTracks/Private/EntitySystem/Interrogation/MovieSceneInterrogatedPropertyInstantiator.cpp
--- a/Engine/Source/Runtime/MovieSceneTracks/Private/EntitySystem/Interrogation/MovieSceneInterrogatedPropertyInstantiator.cpp
+++ b/Engine/Source/Runtime/MovieSceneTracks/Private/EntitySystem/Interrogation/MovieSceneInterrogatedPropertyInstantiator.cpp
@@ -124,7 +124,7 @@ UClass* UMovieSceneInterrogatedPropertyInstantiatorSystem::ResolveBlenderClass(T
 
 	for (FMovieSceneEntityID Input : Inputs)
 	{
-		TComponentPtr<const TSubclassOf<UMovieSceneBlenderSystem>> BlenderTypeComponent = Linker->EntityManager.ReadComponent(Input, BuiltInComponents->BlenderType);
+		TOptionalComponentReader<TSubclassOf<UMovieSceneBlenderSystem>> BlenderTypeComponent = Linker->EntityManager.ReadComponent(Input, BuiltInComponents->BlenderType);
 		if (BlenderTypeComponent)
 		{
 			BlenderClass = BlenderTypeComponent->Get();
@@ -269,12 +269,12 @@ void UMovieSceneInterrogatedPropertyInstantiatorSystem::OnRun(FSystemTaskPrerequ
 	{
 		TArrayView<const FPropertyDefinition> AllProperties = BuiltInComponents->PropertyRegistry.GetProperties();
 
-		auto LinkCallback =  [this, AllProperties](const FEntityAllocation* Allocation, TRead<FInterrogationKey> InterrogationChannelAccessor)
+		auto LinkCallback =  [this, AllProperties](const FEntityAllocation* Allocation, TRead<FInterrogationKey> InterrogationChannels)
 		{
 			const int32 PropertyDefinitionIndex = Algo::IndexOfByPredicate(AllProperties, [=](const FPropertyDefinition& InDefinition){ return Allocation->HasComponent(InDefinition.PropertyType); });
 			if (PropertyDefinitionIndex != INDEX_NONE)
 			{
-				this->PropertyTracker.VisitLinkedAllocation(Allocation, InterrogationChannelAccessor);
+				this->PropertyTracker.VisitLinkedAllocation(Allocation, InterrogationChannels);
 			}
 		};
 		auto UnlinkCallback = [this, AllProperties](const FEntityAllocation* Allocation)
diff --git a/Engine/Source/Runtime/MovieSceneTracks/Private/EntitySystem/Interrogation/MovieSceneInterrogationLinker.cpp b/Engine/Source/Runtime/MovieSceneTracks/Private/EntitySystem/Interrogation/MovieSceneInterrogationLinker.cpp
--- a/Engine/Source/Runtime/MovieSceneTracks/Private/EntitySystem/Interrogation/MovieSceneInterrogationLinker.cpp
+++ b/Engine/Source/Runtime/MovieSceneTracks/Private/EntitySystem/Interrogation/MovieSceneInterrogationLinker.cpp
@@ -5,6 +5,7 @@
 #include "EntitySystem/MovieSceneEntitySystem.h"
 #include "EntitySystem/MovieSceneEntitySystemRunner.h"
 #include "EntitySystem/MovieSceneEntitySystemTask.h"
+#include "EntitySystem/MovieSceneEntitySystemLinker.h"
 
 #include "MovieSceneTimeHelpers.h"
 #include "MovieSceneSequence.h"
@@ -570,19 +571,14 @@ void FSystemInterrogator::QueryWorldSpaceTransforms(const TBitArray<>& ChannelsT
 
 		if (Linker->EntityManager.Contains(Filter))
 		{
-			auto PopulatePartial = [&VariableTransformsByChannel, &PredicateBits, &BaseValues, Components](const FEntityAllocation* Allocation, TRead<FInterrogationKey> KeysAccessor)
+			auto PopulatePartial = [&VariableTransformsByChannel, &PredicateBits, &BaseValues, Components](const FEntityAllocation* Allocation, const FInterrogationKey* Keys)
 			{
 				const int32 Num = Allocation->Num();
-				const FInterrogationKey* Keys = KeysAccessor.Resolve(Allocation);
 
-				const FComponentHeader* Headers[9];
+				TOptionalComponentReader<float> FloatComponents[9];
 				for (int32 Index = 0; Index < 9; ++Index)
 				{
-					Headers[Index] = Allocation->FindComponentHeader(Components->FloatResult[Index]);
-					if (Headers[Index])
-					{
-						Headers[Index]->ReadWriteLock.ReadLock();
-					}
+					FloatComponents[Index] = Allocation->TryReadComponents(Components->FloatResult[Index]);
 				}
 
 				for (int32 ComponentIndex = 0; ComponentIndex < Num; ++ComponentIndex)
@@ -597,27 +593,19 @@ void FSystemInterrogator::QueryWorldSpaceTransforms(const TBitArray<>& ChannelsT
 						FRotator  Rotation = Transform.GetRotation().Rotator();
 						FVector   Scale    = Transform.GetScale3D();
 
-						if (Headers[0]) { Location.X     = *static_cast<const float*>(Headers[0]->GetValuePtr(ComponentIndex)); }
-						if (Headers[1]) { Location.Y     = *static_cast<const float*>(Headers[1]->GetValuePtr(ComponentIndex)); }
-						if (Headers[2]) { Location.Z     = *static_cast<const float*>(Headers[2]->GetValuePtr(ComponentIndex)); }
-						if (Headers[3]) { Rotation.Roll  = *static_cast<const float*>(Headers[3]->GetValuePtr(ComponentIndex)); }
-						if (Headers[4]) { Rotation.Pitch = *static_cast<const float*>(Headers[4]->GetValuePtr(ComponentIndex)); }
-						if (Headers[5]) { Rotation.Yaw   = *static_cast<const float*>(Headers[5]->GetValuePtr(ComponentIndex)); }
-						if (Headers[6]) { Scale.X        = *static_cast<const float*>(Headers[6]->GetValuePtr(ComponentIndex)); }
-						if (Headers[7]) { Scale.Y        = *static_cast<const float*>(Headers[7]->GetValuePtr(ComponentIndex)); }
-						if (Headers[8]) { Scale.Z        = *static_cast<const float*>(Headers[8]->GetValuePtr(ComponentIndex)); }
+						if (FloatComponents[0]) { Location.X     = FloatComponents[0][ComponentIndex]; }
+						if (FloatComponents[1]) { Location.Y     = FloatComponents[1][ComponentIndex]; }
+						if (FloatComponents[2]) { Location.Z     = FloatComponents[2][ComponentIndex]; }
+						if (FloatComponents[3]) { Rotation.Roll  = FloatComponents[3][ComponentIndex]; }
+						if (FloatComponents[4]) { Rotation.Pitch = FloatComponents[4][ComponentIndex]; }
+						if (FloatComponents[5]) { Rotation.Yaw   = FloatComponents[5][ComponentIndex]; }
+						if (FloatComponents[6]) { Scale.X        = FloatComponents[6][ComponentIndex]; }
+						if (FloatComponents[7]) { Scale.Y        = FloatComponents[7][ComponentIndex]; }
+						if (FloatComponents[8]) { Scale.Z        = FloatComponents[8][ComponentIndex]; }
 
 						VariableTransformsByChannel[ChannelIndex][InterrogationKey.InterrogationIndex] = FTransform(Rotation, Location, Scale);
 					}
 				}
-
-				for (const FComponentHeader* Header : Headers)
-				{
-					if (Header)
-					{
-						Header->ReadWriteLock.ReadUnlock();
-					}
-				}
 			};
 
 			FEntityTaskBuilder()
@@ -796,19 +784,14 @@ void FSystemInterrogator::QueryLocalSpaceTransforms(const TBitArray<>& ChannelsT
 
 		if (Linker->EntityManager.Contains(Filter))
 		{
-			auto PopulatePartial = [&OnGetOutputForChannel, &PredicateBits, &BaseValues, Components](const FEntityAllocation* Allocation, TRead<FInterrogationKey> KeysAccessor)
+			auto PopulatePartial = [&OnGetOutputForChannel, &PredicateBits, &BaseValues, Components](const FEntityAllocation* Allocation, const FInterrogationKey* Keys)
 			{
 				const int32 Num = Allocation->Num();
-				const FInterrogationKey* Keys = KeysAccessor.Resolve(Allocation);
 
-				const FComponentHeader* Headers[9];
+				TOptionalComponentReader<float> FloatComponents[9];
 				for (int32 Index = 0; Index < 9; ++Index)
 				{
-					Headers[Index] = Allocation->FindComponentHeader(Components->FloatResult[Index]);
-					if (Headers[Index])
-					{
-						Headers[Index]->ReadWriteLock.ReadLock();
-					}
+					FloatComponents[Index] = Allocation->TryReadComponents(Components->FloatResult[Index]);
 				}
 
 				for (int32 ComponentIndex = 0; ComponentIndex < Num; ++ComponentIndex)
@@ -818,27 +801,19 @@ void FSystemInterrogator::QueryLocalSpaceTransforms(const TBitArray<>& ChannelsT
 					{
 						FIntermediate3DTransform Transform = BaseValues[ComponentIndex];
 
-						if (Headers[0]) { Transform.T_X = *static_cast<const float*>(Headers[0]->GetValuePtr(ComponentIndex)); }
-						if (Headers[1]) { Transform.T_Y = *static_cast<const float*>(Headers[1]->GetValuePtr(ComponentIndex)); }
-						if (Headers[2]) { Transform.T_Z = *static_cast<const float*>(Headers[2]->GetValuePtr(ComponentIndex)); }
-						if (Headers[3]) { Transform.R_X = *static_cast<const float*>(Headers[3]->GetValuePtr(ComponentIndex)); }
-						if (Headers[4]) { Transform.R_Y = *static_cast<const float*>(Headers[4]->GetValuePtr(ComponentIndex)); }
-						if (Headers[5]) { Transform.R_Z = *static_cast<const float*>(Headers[5]->GetValuePtr(ComponentIndex)); }
-						if (Headers[6]) { Transform.S_X = *static_cast<const float*>(Headers[6]->GetValuePtr(ComponentIndex)); }
-						if (Headers[7]) { Transform.S_Y = *static_cast<const float*>(Headers[7]->GetValuePtr(ComponentIndex)); }
-						if (Headers[8]) { Transform.S_Z = *static_cast<const float*>(Headers[8]->GetValuePtr(ComponentIndex)); }
+						if (FloatComponents[0]) { Transform.T_X = FloatComponents[0][ComponentIndex]; }
+						if (FloatComponents[1]) { Transform.T_Y = FloatComponents[1][ComponentIndex]; }
+						if (FloatComponents[2]) { Transform.T_Z = FloatComponents[2][ComponentIndex]; }
+						if (FloatComponents[3]) { Transform.R_X = FloatComponents[3][ComponentIndex]; }
+						if (FloatComponents[4]) { Transform.R_Y = FloatComponents[4][ComponentIndex]; }
+						if (FloatComponents[5]) { Transform.R_Z = FloatComponents[5][ComponentIndex]; }
+						if (FloatComponents[6]) { Transform.S_X = FloatComponents[6][ComponentIndex]; }
+						if (FloatComponents[7]) { Transform.S_Y = FloatComponents[7][ComponentIndex]; }
+						if (FloatComponents[8]) { Transform.S_Z = FloatComponents[8][ComponentIndex]; }
 
 						OnGetOutputForChannel(InterrogationKey.Channel)[InterrogationKey.InterrogationIndex] = Transform;
 					}
 				}
-
-				for (const FComponentHeader* Header : Headers)
-				{
-					if (Header)
-					{
-						Header->ReadWriteLock.ReadUnlock();
-					}
-				}
 			};
 
 			FEntityTaskBuilder()
diff --git a/Engine/Source/Runtime/MovieSceneTracks/Private/MovieSceneTracksComponentTypes.cpp b/Engine/Source/Runtime/MovieSceneTracks/Private/MovieSceneTracksComponentTypes.cpp
--- a/Engine/Source/Runtime/MovieSceneTracks/Private/MovieSceneTracksComponentTypes.cpp
+++ b/Engine/Source/Runtime/MovieSceneTracks/Private/MovieSceneTracksComponentTypes.cpp
@@ -7,6 +7,7 @@
 #include "EntitySystem/MovieSceneComponentRegistry.h"
 #include "EntitySystem/MovieSceneBlenderSystem.h"
 #include "EntitySystem/MovieScenePropertyComponentHandler.h"
+#include "EntitySystem/MovieSceneEntityFactoryTemplates.h"
 
 #include "MovieSceneObjectBindingID.h"
 #include "GameFramework/Actor.h"
diff --git a/Engine/Source/Runtime/MovieSceneTracks/Private/Sections/MovieSceneEventRepeaterSection.cpp b/Engine/Source/Runtime/MovieSceneTracks/Private/Sections/MovieSceneEventRepeaterSection.cpp
--- a/Engine/Source/Runtime/MovieSceneTracks/Private/Sections/MovieSceneEventRepeaterSection.cpp
+++ b/Engine/Source/Runtime/MovieSceneTracks/Private/Sections/MovieSceneEventRepeaterSection.cpp
@@ -7,6 +7,7 @@
 #include "EntitySystem/MovieSceneEntityManager.h"
 #include "EntitySystem/BuiltInComponentTypes.h"
 #include "EntitySystem/MovieSceneSequenceInstance.h"
+#include "EntitySystem/MovieSceneEntitySystemLinker.h"
 #include "Systems/MovieSceneEventSystems.h"
 #include "MovieSceneTracksComponentTypes.h"
 
diff --git a/Engine/Source/Runtime/MovieSceneTracks/Private/Systems/MovieSceneComponentAttachmentSystem.cpp b/Engine/Source/Runtime/MovieSceneTracks/Private/Systems/MovieSceneComponentAttachmentSystem.cpp
--- a/Engine/Source/Runtime/MovieSceneTracks/Private/Systems/MovieSceneComponentAttachmentSystem.cpp
+++ b/Engine/Source/Runtime/MovieSceneTracks/Private/Systems/MovieSceneComponentAttachmentSystem.cpp
@@ -101,8 +101,8 @@ struct FAttachmentHandler
 
 		for (FMovieSceneEntityID Entity : Inputs)
 		{
-			TComponentPtr<USceneComponent* const>     AttachParentComponent = EntityManager->ReadComponent(Entity, TrackComponents->AttachParent);
-			TComponentPtr<const FAttachmentComponent> AttachmentComponent   = EntityManager->ReadComponent(Entity, TrackComponents->AttachComponent);
+			TOptionalComponentReader<USceneComponent*>     AttachParentComponent = EntityManager->ReadComponent(Entity, TrackComponents->AttachParent);
+			TOptionalComponentReader<FAttachmentComponent> AttachmentComponent   = EntityManager->ReadComponent(Entity, TrackComponents->AttachComponent);
 			if (AttachParentComponent && AttachmentComponent)
 			{
 				if (USceneComponent* AttachParent = *AttachParentComponent)
diff --git a/Engine/Source/Runtime/MovieSceneTracks/Private/Systems/MovieSceneDeferredComponentMovementSystem.cpp b/Engine/Source/Runtime/MovieSceneTracks/Private/Systems/MovieSceneDeferredComponentMovementSystem.cpp
--- a/Engine/Source/Runtime/MovieSceneTracks/Private/Systems/MovieSceneDeferredComponentMovementSystem.cpp
+++ b/Engine/Source/Runtime/MovieSceneTracks/Private/Systems/MovieSceneDeferredComponentMovementSystem.cpp
@@ -3,6 +3,7 @@
 #include "Systems/MovieSceneDeferredComponentMovementSystem.h"
 #include "EntitySystem/MovieSceneEntitySystemTask.h"
 #include "EntitySystem/MovieSceneEntitySystemRunner.h"
+#include "EntitySystem/MovieSceneEntitySystemLinker.h"
 #include "MovieSceneTracksComponentTypes.h"
 
 #include "GameFramework/Actor.h"
diff --git a/Engine/Source/Runtime/MovieSceneTracks/Private/Systems/MovieSceneHierarchicalBiasSystem.cpp b/Engine/Source/Runtime/MovieSceneTracks/Private/Systems/MovieSceneHierarchicalBiasSystem.cpp
--- a/Engine/Source/Runtime/MovieSceneTracks/Private/Systems/MovieSceneHierarchicalBiasSystem.cpp
+++ b/Engine/Source/Runtime/MovieSceneTracks/Private/Systems/MovieSceneHierarchicalBiasSystem.cpp
@@ -28,22 +28,19 @@ struct FHierarchicalBiasTask
 		return MaxBiasByChannel.Num() != 0;
 	}
 
-	void ForEachAllocation(const FEntityAllocation* Allocation, FReadEntityIDs EntityIDAccessor, TRead<uint16> BlendChannelAcessor, TReadOptional<int16> HBiasAccessor)
+	void ForEachAllocation(const FEntityAllocation* Allocation, TRead<FMovieSceneEntityID> EntityIDs, TRead<uint16> BlendChannels, TReadOptional<int16> OptHBiases)
 	{
-		TArrayView<const FMovieSceneEntityID> EntityIDs = EntityIDAccessor.ResolveAsArray(Allocation);
-		TArrayView<const int16>  HBiases                = HBiasAccessor.ResolveAsArray(Allocation);
-		TArrayView<const uint16> BlendChannels          = BlendChannelAcessor.ResolveAsArray(Allocation);
-
-		if (HBiases.Num())
+		const int32 Num = Allocation->Num();
+		if (OptHBiases)
 		{
-			for (int32 Index = 0; Index < BlendChannels.Num(); ++Index)
+			for (int32 Index = 0; Index < Num; ++Index)
 			{
-				VisitChannel(EntityIDs[Index], BlendChannels[Index], HBiases[Index]);
+				VisitChannel(EntityIDs[Index], BlendChannels[Index], OptHBiases[Index]);
 			}
 		}
 		else
 		{
-			for (int32 Index = 0; Index < BlendChannels.Num(); ++Index)
+			for (int32 Index = 0; Index < Num; ++Index)
 			{
 				VisitChannel(EntityIDs[Index], BlendChannels[Index], 0);
 			}
diff --git a/Engine/Source/Runtime/MovieSceneTracks/Private/Systems/MovieSceneLevelVisibilitySystem.cpp b/Engine/Source/Runtime/MovieSceneTracks/Private/Systems/MovieSceneLevelVisibilitySystem.cpp
--- a/Engine/Source/Runtime/MovieSceneTracks/Private/Systems/MovieSceneLevelVisibilitySystem.cpp
+++ b/Engine/Source/Runtime/MovieSceneTracks/Private/Systems/MovieSceneLevelVisibilitySystem.cpp
@@ -5,6 +5,7 @@
 #include "EntitySystem/MovieSceneEntitySystemTask.h"
 #include "EntitySystem/MovieSceneMasterInstantiatorSystem.h"
 #include "EntitySystem/MovieSceneSpawnablesSystem.h"
+#include "EntitySystem/MovieSceneEntitySystemLinker.h"
 #include "MovieSceneTracksComponentTypes.h"
 #include "MovieSceneExecutionToken.h"
 #include "IMovieScenePlayer.h"
@@ -435,26 +436,20 @@ void UMovieSceneLevelVisibilitySystem::OnRun(FSystemTaskPrerequisites& InPrerequ
 	FMovieSceneTracksComponentTypes* TracksComponents = FMovieSceneTracksComponentTypes::Get();
 
 	auto ApplyLevelVisibilities = [this, BuiltInComponents](
-		const FEntityAllocation* Allocation,
-		FReadEntityIDs EntityIDAccessor,
-		TRead<FInstanceHandle> InstanceHandleAccessor,
-		TRead<FLevelVisibilityComponentData> LevelVisibilityAccessor,
-		TReadOptional<int16> HBiasAccessor)
+		FEntityAllocationIteratorItem AllocationItem,
+		TRead<FMovieSceneEntityID> EntityIDs,
+		TRead<FInstanceHandle> InstanceHandles,
+		TRead<FLevelVisibilityComponentData> LevelVisibilityData,
+		TReadOptional<int16> OptHBiases)
 	{
-		const bool bHasNeedsLink = Allocation->HasComponent(BuiltInComponents->Tags.NeedsLink);
-		const bool bHasNeedsUnlink = Allocation->HasComponent(BuiltInComponents->Tags.NeedsUnlink);
-		const bool bHasHBias = Allocation->HasComponent(BuiltInComponents->HierarchicalBias);
+		const FComponentMask& AllocationType = AllocationItem.GetAllocationType();
+		const bool bHasNeedsLink = AllocationType.Contains(BuiltInComponents->Tags.NeedsLink);
+		const bool bHasNeedsUnlink = AllocationType.Contains(BuiltInComponents->Tags.NeedsUnlink);
 
-		const FMovieSceneEntityID* EntityIDs = EntityIDAccessor.Resolve(Allocation);
-		const FInstanceHandle* InstanceHandles = InstanceHandleAccessor.Resolve(Allocation);
-		const FLevelVisibilityComponentData* LevelVisibilityData = LevelVisibilityAccessor.Resolve(Allocation);
-		const int16* HBiases = HBiasAccessor.Resolve(Allocation);
-
-		for (int32 Index = 0; Index < Allocation->Num(); ++Index)
+		for (int32 Index = 0; Index < AllocationItem.GetAllocation()->Num(); ++Index)
 		{
 			const FMovieSceneEntityID EntityID = EntityIDs[Index];
 			const UMovieSceneLevelVisibilitySection* Section = LevelVisibilityData[Index].Section;
-			const int16 HBias = (bHasHBias ? HBiases[Index] : 0);
 
 			if (!ensure(Section))
 			{
@@ -466,6 +461,7 @@ void UMovieSceneLevelVisibilitySystem::OnRun(FSystemTaskPrerequisites& InPrerequ
 
 			if (bHasNeedsLink)
 			{
+				const int16 HBias = (OptHBiases ? OptHBiases[Index] : 0);
 				SharedData.AssignLevelVisibilityOverrides(InstanceHandles[Index], LevelNames, Visibility, HBias, EntityID);
 			}
 			if (bHasNeedsUnlink)
diff --git a/Engine/Source/Runtime/MovieSceneTracks/Private/Systems/MovieScenePiecewiseFloatBlenderSystem.cpp b/Engine/Source/Runtime/MovieSceneTracks/Private/Systems/MovieScenePiecewiseFloatBlenderSystem.cpp
--- a/Engine/Source/Runtime/MovieSceneTracks/Private/Systems/MovieScenePiecewiseFloatBlenderSystem.cpp
+++ b/Engine/Source/Runtime/MovieSceneTracks/Private/Systems/MovieScenePiecewiseFloatBlenderSystem.cpp
@@ -65,14 +65,11 @@ struct FAccumulationTask
 	{}
 
 	/** Task entry point - iterates the allocation's headers and accumulates float results for any required components */
-	void ForEachAllocation(FEntityAllocationIteratorItem InItem, TRead<uint16> BlendID, TReadOptional<float> EasingAndWeightResult)
+	void ForEachAllocation(FEntityAllocationIteratorItem InItem, TRead<uint16> BlendIDs, TReadOptional<float> OptionalEasingAndWeights)
 	{
 		const FEntityAllocation* Allocation = InItem;
 		const FComponentMask& AllocationType = InItem;
 
-		const uint16* BlendIDs         = BlendID.Resolve(Allocation);
-		const float*  EasingAndWeights = EasingAndWeightResult.Resolve(Allocation);
-
 		for (const FComponentHeader& ComponentHeader : Allocation->GetComponentHeaders())
 		{
 			if (TArray<FBlendResult>* AccumulationBuffer = AccumulationBuffers->Find(ComponentHeader.ComponentType))
@@ -80,7 +77,7 @@ struct FAccumulationTask
 				ComponentHeader.ReadWriteLock.ReadLock();
 
 				const float* FloatResults = static_cast<const float*>(ComponentHeader.GetValuePtr(0));
-				AccumulateResults(Allocation, FloatResults, BlendIDs, EasingAndWeights, *AccumulationBuffer);
+				AccumulateResults(Allocation, FloatResults, BlendIDs, OptionalEasingAndWeights, *AccumulationBuffer);
 
 				ComponentHeader.ReadWriteLock.ReadUnlock();
 			}
@@ -127,30 +124,19 @@ struct FAdditiveFromBaseBlendTask
 {
 	TSortedMap<FComponentTypeID, FAdditiveFromBaseBuffer>* AccumulationBuffers;
 
-	void ForEachAllocation(FEntityAllocationIteratorItem InItem, TRead<uint16> BlendID, TReadOptional<float> EasingAndWeightResult)
+	void ForEachAllocation(FEntityAllocationIteratorItem InItem, TRead<uint16> BlendIDs, TReadOptional<float> EasingAndWeightResults)
 	{
 		FEntityAllocation* Allocation = InItem;
 		const FComponentMask& AllocationType = InItem;
 
-		const uint16* BlendIDs         = BlendID.Resolve(Allocation);
-		const float*  EasingAndWeights = EasingAndWeightResult.Resolve(Allocation);
-
 		for (const FComponentHeader& ComponentHeader : Allocation->GetComponentHeaders())
 		{
 			if (FAdditiveFromBaseBuffer* Buffer = AccumulationBuffers->Find(ComponentHeader.ComponentType))
 			{
-				const FComponentHeader& BaseValueHeader = Allocation->GetComponentHeaderChecked(Buffer->BaseComponent);
-
-				BaseValueHeader.ReadWriteLock.ReadLock();
-				ComponentHeader.ReadWriteLock.ReadLock();
-
-				const float* BaseValues   = static_cast<const float*>(BaseValueHeader.GetValuePtr(0));
-				const float* FloatResults = static_cast<const float*>(ComponentHeader.GetValuePtr(0));
+				TComponentReader<float> BaseValues = Allocation->ReadComponents(Buffer->BaseComponent.ReinterpretCast<float>());
+				TComponentReader<float> FloatResults(&ComponentHeader);
 
-				AccumulateResults(Allocation, FloatResults, BaseValues, BlendIDs, EasingAndWeights, Buffer->Buffer);
-
-				ComponentHeader.ReadWriteLock.ReadUnlock();
-				BaseValueHeader.ReadWriteLock.ReadUnlock();
+				AccumulateResults(Allocation, FloatResults.AsPtr(), BaseValues.AsPtr(), BlendIDs, EasingAndWeightResults, Buffer->Buffer);
 			}
 		}
 	}
@@ -188,14 +174,14 @@ private:
 /** Task that combines all accumulated blends for any tracked property type that has blend inputs/outputs */
 struct FCombineBlends
 {
-	explicit FCombineBlends(const TBitArray<>& InCachedRelevantProperties, const FAccumulationBuffers* InAccumulationBuffers, uint64 InSystemSerial)
+	explicit FCombineBlends(const TBitArray<>& InCachedRelevantProperties, const FAccumulationBuffers* InAccumulationBuffers, FEntityAllocationWriteContext InWriteContext)
 		: CachedRelevantProperties(InCachedRelevantProperties)
 		, AccumulationBuffers(InAccumulationBuffers)
 		, PropertyRegistry(&FBuiltInComponentTypes::Get()->PropertyRegistry)
-		, SystemSerial(InSystemSerial)
+		, WriteContext(InWriteContext)
 	{}
 
-	void ForEachAllocation(FEntityAllocationIteratorItem InItem, TRead<uint16> BlendID)
+	void ForEachAllocation(FEntityAllocationIteratorItem InItem, TRead<uint16> BlendIDs)
 	{
 		FEntityAllocation* Allocation = InItem;
 		const FComponentMask& AllocationType = InItem;
@@ -206,7 +192,7 @@ struct FCombineBlends
 			const FPropertyDefinition& PropertyDefinition = PropertyRegistry->GetDefinition(FCompositePropertyTypeID::FromIndex(PropertyIndex.GetIndex()));
 			if (AllocationType.Contains(PropertyDefinition.PropertyType))
 			{
-				ProcessPropertyType(Allocation, AllocationType, PropertyDefinition, BlendID);
+				ProcessPropertyType(Allocation, AllocationType, PropertyDefinition, BlendIDs);
 				return;
 			}
 		}
@@ -214,17 +200,13 @@ struct FCombineBlends
 
 private:
 
-	void ProcessPropertyType(FEntityAllocation* Allocation, const FComponentMask& AllocationType, const FPropertyDefinition& PropertyDefinition, TRead<uint16> BlendID)
+	void ProcessPropertyType(FEntityAllocation* Allocation, const FComponentMask& AllocationType, const FPropertyDefinition& PropertyDefinition, const uint16* BlendIDs)
 	{
 		TArrayView<const FPropertyCompositeDefinition> Composites = PropertyRegistry->GetComposites(PropertyDefinition);
 
-		const FComponentHeader* InitialValuesHeader = Allocation->FindComponentHeader(PropertyDefinition.InitialValueType);
-		if (InitialValuesHeader)
-		{
-			InitialValuesHeader->ReadWriteLock.ReadLock();
-		}
 
-		const uint16* BlendIDs = BlendID.Resolve(Allocation);
+		FOptionalComponentReader OptInitialValues = Allocation->TryReadComponentsErased(PropertyDefinition.InitialValueType);
+
 		for (int32 CompositeIndex = 0; CompositeIndex < Composites.Num(); ++CompositeIndex)
 		{
 			if ( (PropertyDefinition.FloatCompositeMask & (1 << CompositeIndex)) == 0)
@@ -244,16 +226,13 @@ private:
 				const uint16 InitialValueProjectionOffset = Composites[CompositeIndex].CompositeOffset;
 
 				// Open the float result channel for write
-				FComponentHeader& ResultsHeader = Allocation->GetComponentHeaderChecked(ResultComponent);
-				ResultsHeader.ReadWriteLock.WriteLock();
+				TComponentWriter<float> FloatResults = Allocation->WriteComponents(ResultComponent, WriteContext);
 
-				float* FloatResults = static_cast<float*>(ResultsHeader.GetValuePtr(0));
-
-				if (InitialValuesHeader)
+				if (OptInitialValues)
 				{
 					for (int32 Index = 0; Index < Allocation->Num(); ++Index)
 					{
-						const float InitialValue = *reinterpret_cast<const float*>(reinterpret_cast<const uint8*>(InitialValuesHeader->GetValuePtr(Index)) + InitialValueProjectionOffset);
+						const float InitialValue = *reinterpret_cast<const float*>(static_cast<const uint8*>(OptInitialValues[Index]) + InitialValueProjectionOffset);
 						BlendResultsWithInitial(Results, BlendIDs[Index], InitialValue, FloatResults[Index]);
 					}
 				}
@@ -264,16 +243,8 @@ private:
 						BlendResults(Results, BlendIDs[Index], FloatResults[Index]);
 					}
 				}
-
-				ResultsHeader.PostWriteComponents(SystemSerial);
-				ResultsHeader.ReadWriteLock.WriteUnlock();
 			}
 		}
-
-		if (InitialValuesHeader)
-		{
-			InitialValuesHeader->ReadWriteLock.ReadUnlock();
-		}
 	}
 
 	void BlendResultsWithInitial(const FAccumulationResult& Results, uint16 BlendID, const float InitialValue, float& OutFinalBlendResult)
@@ -341,7 +312,7 @@ private:
 	TBitArray<> CachedRelevantProperties;
 	const FAccumulationBuffers* AccumulationBuffers;
 	const FPropertyRegistry* PropertyRegistry;
-	uint64 SystemSerial;
+	FEntityAllocationWriteContext WriteContext;
 };
 
 
@@ -513,7 +484,7 @@ void UMovieScenePiecewiseFloatBlenderSystem::OnRun(FSystemTaskPrerequisites& InP
 	.Read(BuiltInComponents->BlendChannelOutput)
 	.FilterAny(BlendedPropertyMask)
 	.SetStat(GET_STATID(MovieSceneEval_BlendCombineFloatValues))
-	.Dispatch_PerAllocation<FCombineBlends>(&Linker->EntityManager, Prereqs, &Subsequents, CachedRelevantProperties, &AccumulationBuffers, Linker->EntityManager.GetSystemSerial());
+	.Dispatch_PerAllocation<FCombineBlends>(&Linker->EntityManager, Prereqs, &Subsequents, CachedRelevantProperties, &AccumulationBuffers, FEntityAllocationWriteContext(Linker->EntityManager));
 }
 
 void UMovieScenePiecewiseFloatBlenderSystem::ReinitializeAccumulationBuffers()
@@ -646,24 +617,12 @@ FGraphEventRef UMovieScenePiecewiseFloatBlenderSystem::DispatchDecomposeTask(con
 			EntitiesToDecompose.Append(Params.Query.Entities.GetData(), Params.Query.Entities.Num());
 		}
 
-		void ForEachAllocation(const FEntityAllocation* Allocation, FReadEntityIDs EntityToDecomposeIDComponent, TRead<uint16> BlendChannelComponent, TRead<float> FloatResultComponent, TReadOptional<float> OptionalWeightComponent)
-		{
-			TArrayView<const FMovieSceneEntityID> EntityToDecomposeIDs = EntityToDecomposeIDComponent.ResolveAsArray(Allocation);
-			ForEachAllocationImpl(Allocation, EntityToDecomposeIDs, BlendChannelComponent, FloatResultComponent, OptionalWeightComponent);
-		}
-
-		void ForEachAllocation(const FEntityAllocation* Allocation, TRead<FMovieSceneEntityID> EntityToDecomposeIDComponent, TRead<uint16> BlendChannelComponent, TRead<float> FloatResultComponent, TReadOptional<float> OptionalWeightComponent)
-		{
-			TArrayView<const FMovieSceneEntityID> EntityToDecomposeIDs = EntityToDecomposeIDComponent.ResolveAsArray(Allocation);
-			ForEachAllocationImpl(Allocation, EntityToDecomposeIDs, BlendChannelComponent, FloatResultComponent, OptionalWeightComponent);
-		}
-
-		void ForEachAllocationImpl(const FEntityAllocation* Allocation, TArrayView<const FMovieSceneEntityID> EntityToDecomposeIDs, TRead<uint16> BlendChannelComponent, TRead<float> FloatResultComponent, TReadOptional<float> OptionalWeightComponent)
+		void ForEachAllocation(const FEntityAllocation* Allocation, TRead<FMovieSceneEntityID> EntityToDecomposeIDs, TRead<uint16> BlendChannels, TRead<float> FloatResultComponent, TReadOptional<float> OptionalWeightComponent)
 		{
 			const bool bAdditive = Allocation->HasComponent(AdditiveBlendTag);
 
-			TArrayView<const uint16> BlendChannels = BlendChannelComponent.ResolveAsArray(Allocation);
-			for (int32 EntityIndex = 0; EntityIndex < BlendChannels.Num(); ++EntityIndex)
+			const int32 Num = Allocation->Num();
+			for (int32 EntityIndex = 0; EntityIndex < Num; ++EntityIndex)
 			{
 				if (BlendChannels[EntityIndex] != DecomposeBlendChannel)
 				{
@@ -672,9 +631,8 @@ FGraphEventRef UMovieScenePiecewiseFloatBlenderSystem::DispatchDecomposeTask(con
 
 				// We've found a contributor for this blend channel
 				const FMovieSceneEntityID EntityToDecompose = EntityToDecomposeIDs[EntityIndex];
-				TArrayView<const float>   Weights           = OptionalWeightComponent.ResolveAsArray(Allocation);
-				const float               Weight            = (Weights.Num() != 0 ? Weights[EntityIndex] : 1.f);
-				const float               FloatResult       = FloatResultComponent.ResolveAsArray(Allocation)[EntityIndex];
+				const float               Weight            = OptionalWeightComponent ? OptionalWeightComponent[EntityIndex] : 1.f;
+				const float               FloatResult       = FloatResultComponent[EntityIndex];
 
 				if (EntitiesToDecompose.Contains(EntityToDecompose))
 				{
diff --git a/Engine/Source/Runtime/MovieSceneTracks/Private/Systems/MovieScenePropertyInstantiator.cpp b/Engine/Source/Runtime/MovieSceneTracks/Private/Systems/MovieScenePropertyInstantiator.cpp
--- a/Engine/Source/Runtime/MovieSceneTracks/Private/Systems/MovieScenePropertyInstantiator.cpp
+++ b/Engine/Source/Runtime/MovieSceneTracks/Private/Systems/MovieScenePropertyInstantiator.cpp
@@ -2,6 +2,7 @@
 
 #include "Systems/MovieScenePropertyInstantiator.h"
 #include "EntitySystem/MovieSceneEntityBuilder.h"
+#include "EntitySystem/MovieScenePropertyBinding.h"
 #include "EntitySystem/MovieSceneEntitySystemLinker.h"
 #include "EntitySystem/MovieSceneBlenderSystem.h"
 #include "EntitySystem/MovieScenePropertyRegistry.h"
@@ -107,7 +108,7 @@ void UMovieScenePropertyInstantiatorSystem::DiscoverInvalidatedProperties(TBitAr
 
 	PropertyStats.SetNum(Properties.Num());
 
-	auto VisitNewProperties = [this, Properties, &OutInvalidatedProperties](const FEntityAllocation* Allocation, FReadEntityIDs EntityIDsAccessor, TRead<UObject*> ObjectComponents, TRead<FMovieScenePropertyBinding> PropertyBindingComponents)
+	auto VisitNewProperties = [this, Properties, &OutInvalidatedProperties](const FEntityAllocation* Allocation, const FMovieSceneEntityID* EntityIDs, UObject* const * ObjectPtrs, const FMovieScenePropertyBinding* PropertyPtrs)
 	{
 		const int32 PropertyDefinitionIndex = Algo::IndexOfByPredicate(Properties, [=](const FPropertyDefinition& InDefinition){ return Allocation->HasComponent(InDefinition.PropertyType); });
 		if (PropertyDefinitionIndex == INDEX_NONE)
@@ -119,10 +120,6 @@ void UMovieScenePropertyInstantiatorSystem::DiscoverInvalidatedProperties(TBitAr
 
 		FCustomAccessorView CustomAccessors = PropertyDefinition.CustomPropertyRegistration ? PropertyDefinition.CustomPropertyRegistration->GetAccessors() : FCustomAccessorView();
 
-		UObject* const * ObjectPtrs = ObjectComponents.Resolve(Allocation);
-		const FMovieSceneEntityID* EntityIDs = EntityIDsAccessor.Resolve(Allocation);
-		const FMovieScenePropertyBinding* PropertyPtrs = PropertyBindingComponents.Resolve(Allocation);
-
 		for (int32 Index = 0; Index < Allocation->Num(); ++Index)
 		{
 			const int32 PropertyIndex = this->ResolveProperty(CustomAccessors, ObjectPtrs[Index], PropertyPtrs[Index], PropertyDefinitionIndex);
@@ -387,16 +384,16 @@ void UMovieScenePropertyInstantiatorSystem::InitializeFastPath(const FPropertyPa
 	Params.PropertyInfo->PropertyEntityID = SoleContributor;
 
 	check(Params.PropertyInfo->BlendChannel == INVALID_BLEND_CHANNEL);
-	switch (Params.PropertyInfo->Property->GetIndex())
+	switch (Params.PropertyInfo->Property.GetIndex())
 	{
 	case 0:
-		Linker->EntityManager.AddComponent(SoleContributor, BuiltInComponents->FastPropertyOffset,  Params.PropertyInfo->Property->template Get<uint16>());
+		Linker->EntityManager.AddComponent(SoleContributor, BuiltInComponents->FastPropertyOffset,  Params.PropertyInfo->Property.template Get<uint16>());
 		break;
 	case 1:
-		Linker->EntityManager.AddComponent(SoleContributor, BuiltInComponents->CustomPropertyIndex, Params.PropertyInfo->Property->template Get<FCustomPropertyIndex>());
+		Linker->EntityManager.AddComponent(SoleContributor, BuiltInComponents->CustomPropertyIndex, Params.PropertyInfo->Property.template Get<FCustomPropertyIndex>());
 		break;
 	case 2:
-		Linker->EntityManager.AddComponent(SoleContributor, BuiltInComponents->SlowProperty,        Params.PropertyInfo->Property->template Get<FSlowPropertyPtr>());
+		Linker->EntityManager.AddComponent(SoleContributor, BuiltInComponents->SlowProperty,        Params.PropertyInfo->Property.template Get<FSlowPropertyPtr>());
 		break;
 	}
 }
@@ -413,8 +410,8 @@ void UMovieScenePropertyInstantiatorSystem::InitializeBlendPath(const FPropertyP
 	for (auto ContributorIt = Contributors.CreateConstKeyIterator(Params.PropertyInfoIndex); ContributorIt; ++ContributorIt)
 	{
 		FMovieSceneEntityID Contributor = ContributorIt.Value();
-		
-		TComponentPtr<const TSubclassOf<UMovieSceneBlenderSystem>> BlenderTypeComponent = Linker->EntityManager.ReadComponent(Contributor, BuiltInComponents->BlenderType);
+
+		TOptionalComponentReader<TSubclassOf<UMovieSceneBlenderSystem>> BlenderTypeComponent = Linker->EntityManager.ReadComponent(Contributor, BuiltInComponents->BlenderType);
 		if (BlenderTypeComponent)
 		{
 			BlenderClass = BlenderTypeComponent->Get();
@@ -458,12 +455,12 @@ void UMovieScenePropertyInstantiatorSystem::InitializeBlendPath(const FPropertyP
 		NewMask.Set(Params.PropertyDefinition->PropertyType);
 
 		FMovieSceneEntityID NewEntityID;
-		switch (Params.PropertyInfo->Property->GetIndex())
+		switch (Params.PropertyInfo->Property.GetIndex())
 		{
 		// Never seen this property before
 		case 0:
 			NewEntityID = FEntityBuilder()
-			.Add(BuiltInComponents->FastPropertyOffset,      Params.PropertyInfo->Property->template Get<uint16>())
+			.Add(BuiltInComponents->FastPropertyOffset,      Params.PropertyInfo->Property.template Get<uint16>())
 			.Add(BuiltInComponents->BoundObject,             Params.PropertyInfo->BoundObject)
 			.Add(BuiltInComponents->BlendChannelOutput,      Params.PropertyInfo->BlendChannel)
 			.AddTagConditional(BuiltInComponents->Tags.MigratedFromFastPath, Params.PropertyInfo->PropertyEntityID.IsValid())
@@ -475,7 +472,7 @@ void UMovieScenePropertyInstantiatorSystem::InitializeBlendPath(const FPropertyP
 
 		case 1:
 			NewEntityID = FEntityBuilder()
-			.Add(BuiltInComponents->CustomPropertyIndex, Params.PropertyInfo->Property->template Get<FCustomPropertyIndex>())
+			.Add(BuiltInComponents->CustomPropertyIndex, Params.PropertyInfo->Property.template Get<FCustomPropertyIndex>())
 			.Add(BuiltInComponents->BoundObject,         Params.PropertyInfo->BoundObject)
 			.Add(BuiltInComponents->BlendChannelOutput,  Params.PropertyInfo->BlendChannel)
 			.AddTagConditional(BuiltInComponents->Tags.MigratedFromFastPath, Params.PropertyInfo->PropertyEntityID.IsValid())
@@ -487,7 +484,7 @@ void UMovieScenePropertyInstantiatorSystem::InitializeBlendPath(const FPropertyP
 
 		case 2:
 			NewEntityID = FEntityBuilder()
-			.Add(BuiltInComponents->SlowProperty,            Params.PropertyInfo->Property->template Get<FSlowPropertyPtr>())
+			.Add(BuiltInComponents->SlowProperty,            Params.PropertyInfo->Property.template Get<FSlowPropertyPtr>())
 			.Add(BuiltInComponents->BoundObject,             Params.PropertyInfo->BoundObject)
 			.Add(BuiltInComponents->BlendChannelOutput,      Params.PropertyInfo->BlendChannel)
 			.AddTagConditional(BuiltInComponents->Tags.MigratedFromFastPath, Params.PropertyInfo->PropertyEntityID.IsValid())
@@ -534,50 +531,6 @@ void UMovieScenePropertyInstantiatorSystem::InitializeBlendPath(const FPropertyP
 	}
 }
 
-int32 UMovieScenePropertyInstantiatorSystem::FindCustomAccessorIndex(UE::MovieScene::FCustomAccessorView Accessors, UClass* ClassType, FName PropertyPath)
-{
-	using namespace UE::MovieScene;
-
-	UClass* StopIterationAt = UObject::StaticClass();
-
-	while (ClassType != StopIterationAt)
-	{
-		for (int32 Index = 0; Index < Accessors.Num(); ++Index)
-		{
-			const FCustomPropertyAccessor& Accessor = Accessors[Index];
-			if (Accessor.Class == ClassType && Accessor.PropertyPath == PropertyPath)
-			{
-				return Index;
-			}
-		}
-		ClassType = ClassType->GetSuperClass();
-	}
-
-	return INDEX_NONE;
-}
-
-TOptional<uint16> UMovieScenePropertyInstantiatorSystem::ComputeFastPropertyPtrOffset(UClass* ObjectClass, const FMovieScenePropertyBinding& PropertyBinding)
-{
-	using namespace UE::MovieScene;
-
-	FProperty* Property = ObjectClass->FindPropertyByName(PropertyBinding.PropertyName);
-	// @todo: Constructing FNames from strings is _very_ costly and we really shouldn't be doing this at runtime.
-	UFunction* Setter   = ObjectClass->FindFunctionByName(*(FString("Set") + PropertyBinding.PropertyName.ToString()));
-	if (Property && !Setter)
-	{
-		UObject* DefaultObject   = ObjectClass->GetDefaultObject();
-		uint8*   PropertyAddress = Property->ContainerPtrToValuePtr<uint8>(DefaultObject);
-		int32    PropertyOffset  = PropertyAddress - reinterpret_cast<uint8*>(DefaultObject);
-
-		if (ensureMsgf(PropertyOffset >= 0 && PropertyOffset < int32(uint16(0xFFFF)), TEXT("Property offset of more than 65535 bytes - this is most likely an error and is not supported by fast property accessors.")))
-		{
-			return uint16(PropertyOffset);
-		}
-	}
-
-	return TOptional<uint16>();
-}
-
 int32 UMovieScenePropertyInstantiatorSystem::ResolveProperty(UE::MovieScene::FCustomAccessorView CustomAccessors, UObject* Object, const FMovieScenePropertyBinding& PropertyBinding, int32 PropertyDefinitionIndex)
 {
 	using namespace UE::MovieScene;
@@ -588,51 +541,19 @@ int32 UMovieScenePropertyInstantiatorSystem::ResolveProperty(UE::MovieScene::FCu
 		return *ExistingPropertyIndex;
 	}
 
+	TOptional<FResolvedProperty> ResolvedProperty = FPropertyRegistry::ResolveProperty(Object, PropertyBinding, CustomAccessors);
+	if (!ResolvedProperty.IsSet())
+	{
+		UE_LOG(LogMovieScene, Warning, TEXT("Unable to resolve property '%s' from '%s' instance '%s'"), *PropertyBinding.PropertyPath.ToString(), *Object->GetClass()->GetName(), *Object->GetName());
+		return INDEX_NONE;
+	}
 
-	FObjectPropertyInfo NewInfo;
+	FObjectPropertyInfo NewInfo(MoveTemp(ResolvedProperty.GetValue()));
 
 	NewInfo.BoundObject  = Object;
 	NewInfo.PropertyPath = PropertyBinding.PropertyPath;
 	NewInfo.PropertyDefinitionIndex = PropertyDefinitionIndex;
 
-	UClass* Class = Object->GetClass();
-
-	if (CustomAccessors.Num() != 0)
-	{
-		const int32 CustomPropertyIndex = FindCustomAccessorIndex(CustomAccessors, Class, PropertyBinding.PropertyPath);
-		if (CustomPropertyIndex != INDEX_NONE)
-		{
-			check(CustomPropertyIndex < MAX_uint16);
-
-			// This property has a custom property accessor that can apply properties through a static function ptr.
-			// Just add the function ptrs to the property entity so they can be called directly
-			NewInfo.Property.Emplace(TInPlaceType<FCustomPropertyIndex>(), FCustomPropertyIndex{ static_cast<uint16>(CustomPropertyIndex) });
-		}
-	}
-
-	if (PropertyBinding.CanUseClassLookup())
-	{
-		TOptional<uint16> FastPtrOffset = ComputeFastPropertyPtrOffset(Class, PropertyBinding);
-		if (FastPtrOffset.IsSet())
-		{
-			// This property/object combination has no custom setter function and a constant property offset from the base ptr for all instances of the object.
-			NewInfo.Property.Emplace(TInPlaceType<uint16>(), FastPtrOffset.GetValue());
-		}
-	}
-
-	// None of the above optimized paths can apply to this property (probably because it has a setter function or because it is within a compound property), so we must use the slow property bindings
-	if (!NewInfo.Property.IsSet())
-	{
-		TSharedPtr<FTrackInstancePropertyBindings> SlowBindings = MakeShared<FTrackInstancePropertyBindings>(PropertyBinding.PropertyName, PropertyBinding.PropertyPath.ToString());
-		if (SlowBindings->GetProperty(*Object) == nullptr)
-		{
-			UE_LOG(LogMovieScene, Warning, TEXT("Unable to resolve property '%s' from '%s' instance '%s'"), *PropertyBinding.PropertyPath.ToString(), *Object->GetClass()->GetName(), *Object->GetName());
-			return INDEX_NONE;
-		}
-
-		NewInfo.Property.Emplace(TInPlaceType<FSlowPropertyPtr>(), SlowBindings);
-	}
-
 	const int32 NewPropertyIndex = ResolvedProperties.Add(NewInfo);
 
 	ObjectPropertyToResolvedIndex.Add(Key, NewPropertyIndex);
@@ -646,7 +567,7 @@ UE::MovieScene::FPropertyRecomposerPropertyInfo UMovieScenePropertyInstantiatorS
 {
 	using namespace UE::MovieScene;
 
-	TComponentPtr<const FMovieScenePropertyBinding> PropertyBinding = Linker->EntityManager.ReadComponent(EntityID, BuiltInComponents->PropertyBinding);
+	TOptionalComponentReader<FMovieScenePropertyBinding> PropertyBinding = Linker->EntityManager.ReadComponent(EntityID, BuiltInComponents->PropertyBinding);
 	if (!PropertyBinding)
 	{
 		return FPropertyRecomposerPropertyInfo::Invalid();
diff --git a/Engine/Source/Runtime/MovieSceneTracks/Private/Systems/MovieSceneQuaternionInterpolationRotationSystem.cpp b/Engine/Source/Runtime/MovieSceneTracks/Private/Systems/MovieSceneQuaternionInterpolationRotationSystem.cpp
--- a/Engine/Source/Runtime/MovieSceneTracks/Private/Systems/MovieSceneQuaternionInterpolationRotationSystem.cpp
+++ b/Engine/Source/Runtime/MovieSceneTracks/Private/Systems/MovieSceneQuaternionInterpolationRotationSystem.cpp
@@ -4,6 +4,8 @@
 #include "Channels/MovieSceneFloatChannel.h"
 #include "EntitySystem/BuiltInComponentTypes.h"
 #include "EntitySystem/MovieSceneEvalTimeSystem.h"
+#include "EntitySystem/MovieSceneEntitySystemLinker.h"
+#include "EntitySystem/MovieSceneEntitySystemTask.h"
 #include "MovieSceneTracksComponentTypes.h"
 
 DECLARE_CYCLE_STAT(TEXT("MovieScene: Quat-interp-rot channel system"), MovieSceneEval_QuatInterpRotChannelSystem, STATGROUP_MovieSceneECS);
@@ -36,32 +38,31 @@ struct FEvaluateQuaternionInterpolationRotationChannels
 				FrameRange.SetUpperBoundValue(Times[Index2]);
 			}
 		}
-	};
+	}
 
 	void ForEachAllocation(
 			const FEntityAllocation* Allocation,
-			TRead<FFrameTime> FrameTimeAccessor, 
-			TReadOneOrMoreOf<FSourceFloatChannel, FSourceFloatChannel, FSourceFloatChannel> RotChannelAccessors,
-			TWriteOptional<float> OutResultXAccessor, TWriteOptional<float> OutResultYAccessor, TWriteOptional<float> OutResultZAccessor)
+			TRead<FFrameTime> FrameTimes, 
+			TMultiReadOptional<FSourceFloatChannel, FSourceFloatChannel, FSourceFloatChannel> RotChannelAccessors,
+			float* OutResultXs, float* OutResultYs, float* OutResultZs)
 	{
-		TArrayView<const FFrameTime> FrameTimes = FrameTimeAccessor.ResolveAsArray(Allocation);
-
-		TArrayView<const FSourceFloatChannel> RotationXs, RotationYs, RotationZs;
-		RotChannelAccessors.ResolveAsArrays(Allocation, &RotationXs, &RotationYs, &RotationZs);
-
-		TArrayView<float> OutResultXs = OutResultXAccessor.ResolveAsArray(Allocation);
-		TArrayView<float> OutResultYs = OutResultYAccessor.ResolveAsArray(Allocation);
-		TArrayView<float> OutResultZs = OutResultZAccessor.ResolveAsArray(Allocation);
-
-		check(OutResultXs.Num() == RotationXs.Num() && OutResultYs.Num() == RotationYs.Num() && OutResultZs.Num() == RotationZs.Num());
+		const FSourceFloatChannel* RotationXs = RotChannelAccessors.Get<0>();
+		const FSourceFloatChannel* RotationYs = RotChannelAccessors.Get<1>();
+		const FSourceFloatChannel* RotationZs = RotChannelAccessors.Get<2>();
+
+		check(
+			(OutResultXs != nullptr) == (RotationXs != nullptr) &&
+			(OutResultYs != nullptr) == (RotationYs != nullptr) &&
+			(OutResultZs != nullptr) == (RotationZs != nullptr) );
 		
 		const int32 AllocationSize = Allocation->Num();
 		for (int32 Index = 0; Index < AllocationSize; ++Index)
 		{
 			const FFrameTime FrameTime = FrameTimes[Index];
-			const FSourceFloatChannel* RotationX = RotationXs.Num() > 0 ? &RotationXs[Index] : nullptr;
-			const FSourceFloatChannel* RotationY = RotationYs.Num() > 0 ? &RotationYs[Index] : nullptr;
-			const FSourceFloatChannel* RotationZ = RotationZs.Num() > 0 ? &RotationZs[Index] : nullptr;
+
+			const FSourceFloatChannel* RotationX = RotationXs ? &RotationXs[Index] : nullptr;
+			const FSourceFloatChannel* RotationY = RotationYs ? &RotationYs[Index] : nullptr;
+			const FSourceFloatChannel* RotationZ = RotationZs ? &RotationZs[Index] : nullptr;
 
 			// Find the closest keyframes before/after the current time on the 3 rotation channels.
 			TRange<FFrameNumber> FrameRange(TNumericLimits<FFrameNumber>::Min(), TNumericLimits<FFrameNumber>::Max());
diff --git a/Engine/Source/Runtime/MovieSceneTracks/Private/Systems/MovieSceneTransformOriginSystem.cpp b/Engine/Source/Runtime/MovieSceneTracks/Private/Systems/MovieSceneTransformOriginSystem.cpp
--- a/Engine/Source/Runtime/MovieSceneTracks/Private/Systems/MovieSceneTransformOriginSystem.cpp
+++ b/Engine/Source/Runtime/MovieSceneTracks/Private/Systems/MovieSceneTransformOriginSystem.cpp
@@ -23,21 +23,10 @@ struct FAssignTransformOriginLocation
 {
 	const TSparseArray<FTransform>* TransformOriginsByInstanceID;
 
-	void ForEachAllocation(const FEntityAllocation* Allocation, TRead<FInstanceHandle> InstanceAccessor, TRead<UObject*> BoundObjectAccessor,
-		TWriteOptional<float> LocationXAccessor, TWriteOptional<float> LocationYAccessor, TWriteOptional<float> LocationZAccessor,
-		TWriteOptional<FSourceFloatChannelFlags> FlagsXAccessor, TWriteOptional<FSourceFloatChannelFlags> FlagsYAccessor, TWriteOptional<FSourceFloatChannelFlags> FlagsZAccessor)
+	void ForEachAllocation(const FEntityAllocation* Allocation, TRead<FInstanceHandle> Instances, TRead<UObject*> BoundObjects,
+		TWriteOptional<float> LocationX, TWriteOptional<float> LocationY, TWriteOptional<float> LocationZ,
+		TWriteOptional<FSourceFloatChannelFlags> FlagsX, TWriteOptional<FSourceFloatChannelFlags> FlagsY, TWriteOptional<FSourceFloatChannelFlags> FlagsZ)
 	{
-		const FInstanceHandle* Instances = InstanceAccessor.Resolve(Allocation);
-		const UObject* const * BoundObjects = BoundObjectAccessor.Resolve(Allocation);
-
-		float* LocationX = LocationXAccessor.Resolve(Allocation);
-		float* LocationY = LocationYAccessor.Resolve(Allocation);
-		float* LocationZ = LocationZAccessor.Resolve(Allocation);
-
-		FSourceFloatChannelFlags* FlagsX = FlagsXAccessor.Resolve(Allocation);
-		FSourceFloatChannelFlags* FlagsY = FlagsYAccessor.Resolve(Allocation);
-		FSourceFloatChannelFlags* FlagsZ = FlagsZAccessor.Resolve(Allocation);
-
 		TransformLocation(Instances, BoundObjects, LocationX, LocationY, LocationZ, FlagsX, FlagsY, FlagsZ, Allocation->Num());
 	}
 
@@ -83,21 +72,10 @@ struct FAssignTransformOriginRotation
 {
 	const TSparseArray<FTransform>* TransformOriginsByInstanceID;
 
-	void ForEachAllocation(const FEntityAllocation* Allocation, TRead<FInstanceHandle> InstanceAccessor, TRead<UObject*> BoundObjectAccessor,
-		TWriteOptional<float> RotationXAccessor, TWriteOptional<float> RotationYAccessor, TWriteOptional<float> RotationZAccessor,
-		TWriteOptional<FSourceFloatChannelFlags> FlagsXAccessor, TWriteOptional<FSourceFloatChannelFlags> FlagsYAccessor, TWriteOptional<FSourceFloatChannelFlags> FlagsZAccessor)
+	void ForEachAllocation(const FEntityAllocation* Allocation, TRead<FInstanceHandle> Instances, TRead<UObject*> BoundObjects,
+		TWriteOptional<float> RotationX, TWriteOptional<float> RotationY, TWriteOptional<float> RotationZ,
+		TWriteOptional<FSourceFloatChannelFlags> FlagsX, TWriteOptional<FSourceFloatChannelFlags> FlagsY, TWriteOptional<FSourceFloatChannelFlags> FlagsZ)
 	{
-		const FInstanceHandle* Instances = InstanceAccessor.Resolve(Allocation);
-		const UObject* const * BoundObjects = BoundObjectAccessor.Resolve(Allocation);
-
-		float* RotationX = RotationXAccessor.Resolve(Allocation);
-		float* RotationY = RotationYAccessor.Resolve(Allocation);
-		float* RotationZ = RotationZAccessor.Resolve(Allocation);
-
-		FSourceFloatChannelFlags* FlagsX = FlagsXAccessor.Resolve(Allocation);
-		FSourceFloatChannelFlags* FlagsY = FlagsYAccessor.Resolve(Allocation);
-		FSourceFloatChannelFlags* FlagsZ = FlagsZAccessor.Resolve(Allocation);
-
 		TransformRotation(Instances, BoundObjects, RotationX, RotationY, RotationZ, FlagsX, FlagsY, FlagsZ, Allocation->Num());
 	}
 
diff --git a/Engine/Source/Runtime/MovieSceneTracks/Private/Systems/WeightAndEasingEvaluatorSystem.cpp b/Engine/Source/Runtime/MovieSceneTracks/Private/Systems/WeightAndEasingEvaluatorSystem.cpp
--- a/Engine/Source/Runtime/MovieSceneTracks/Private/Systems/WeightAndEasingEvaluatorSystem.cpp
+++ b/Engine/Source/Runtime/MovieSceneTracks/Private/Systems/WeightAndEasingEvaluatorSystem.cpp
@@ -42,11 +42,8 @@ void UMovieSceneHierarchicalEasingInstantiatorSystem::OnRun(FSystemTaskPrerequis
 	//
 	// We allocate the hierarchical easing channel for their sub-sequence.
 	//
-	auto VisitNewEasingProviders = [this, InstanceRegistry, EvaluatorSystem](const FEntityAllocation* Allocation, TRead<FInstanceHandle> InstanceHandleAccessor, TRead<FMovieSceneSequenceID> HierarchicalEasingProviderAccessor)
+	auto VisitNewEasingProviders = [this, InstanceRegistry, EvaluatorSystem](const FEntityAllocation* Allocation, const FInstanceHandle* InstanceHandles, const FMovieSceneSequenceID* HierarchicalEasingProviders)
 	{
-		TArrayView<const FInstanceHandle> InstanceHandles = InstanceHandleAccessor.ResolveAsArray(Allocation);
-		TArrayView<const FMovieSceneSequenceID> HierarchicalEasingProviders = HierarchicalEasingProviderAccessor.ResolveAsArray(Allocation);
-
 		for (int32 Index = 0; Index < Allocation->Num(); ++Index)
 		{
 			const FSequenceInstance* RootInstance = &InstanceRegistry->GetInstance(InstanceHandles[Index]);
@@ -89,11 +86,8 @@ void UMovieSceneHierarchicalEasingInstantiatorSystem::OnRun(FSystemTaskPrerequis
 	//
 	// We need to assign them to the appropriate hierarchical easing channel that we created in step 1.
 	//
-	auto VisitNewEasings = [this, InstanceRegistry](const FEntityAllocation* Allocation, TRead<FInstanceHandle> InstanceHandleAccessor, TWrite<uint16> HierarchicalEasingAccessor)
+	auto VisitNewEasings = [this, InstanceRegistry](const FEntityAllocation* Allocation, const FInstanceHandle* InstanceHandles, uint16* HierarchicalEasings)
 	{
-		TArrayView<const FInstanceHandle> InstanceHandles = InstanceHandleAccessor.ResolveAsArray(Allocation);
-		TArrayView<uint16> HierarchicalEasings = HierarchicalEasingAccessor.ResolveAsArray(Allocation);
-
 		for (int32 Index = 0; Index < Allocation->Num(); ++Index)
 		{
 			const FInstanceHandle& InstanceHandle = InstanceHandles[Index];
@@ -113,11 +107,8 @@ void UMovieSceneHierarchicalEasingInstantiatorSystem::OnRun(FSystemTaskPrerequis
 
 	// Step 3: Visit removed hierarchical easing providers, so we can free up our channels.
 	//
-	auto VisitRemovedEasingProviders = [this, InstanceRegistry, EvaluatorSystem](const FEntityAllocation* Allocation, TRead<FInstanceHandle> InstanceHandleAccessor, TRead<FMovieSceneSequenceID> HierarchicalEasingProviderAccessor)
+	auto VisitRemovedEasingProviders = [this, InstanceRegistry, EvaluatorSystem](const FEntityAllocation* Allocation, const FInstanceHandle* InstanceHandles, const FMovieSceneSequenceID* HierarchicalEasingProviders)
 	{
-		TArrayView<const FInstanceHandle> InstanceHandles = InstanceHandleAccessor.ResolveAsArray(Allocation);
-		TArrayView<const FMovieSceneSequenceID> HierarchicalEasingProviders = HierarchicalEasingProviderAccessor.ResolveAsArray(Allocation);
-
 		for (int32 Index = 0; Index < Allocation->Num(); ++Index)
 		{
 			const FSequenceInstance& Instance = InstanceRegistry->GetInstance(InstanceHandles[Index]);
@@ -159,52 +150,38 @@ struct FEvaluateEasings
 
 	void ForEachAllocation(
 			const FEntityAllocation* InAllocation, 
-			TRead<FFrameTime> TimeData,
-			TReadOptional<FEasingComponentData> EasingData,
-			TReadOptional<float> WeightData,
-			TReadOptional<FInstanceHandle> InstanceHandleData,
-			TReadOptional<FMovieSceneSequenceID> HierarchicalEasingProviderData,
-			TWrite<float> WeightAndEasingResultData)
+			TRead<FFrameTime> Times,
+			TReadOptional<FEasingComponentData> OptEasing,
+			TReadOptional<float> OptManualWeights,
+			TReadOptional<FInstanceHandle> OptInstanceHandles,
+			TReadOptional<FMovieSceneSequenceID> OptHierarchicalEasingProviders,
+			TWrite<float> Results)
 	{
-		const int32 Num          = InAllocation->Num();
-		const FFrameTime* Times  = TimeData.Resolve(InAllocation);
-		float* Results           = WeightAndEasingResultData.Resolve(InAllocation);
+		const int32 Num = InAllocation->Num();
 
 		// Initialize our result array.
+		for (int32 Idx = 0; Idx < Num; ++Idx)
 		{
-			float* CurResult = Results;
-			for (int32 Idx = 0; Idx < Num; ++Idx)
-			{
-				*(CurResult++) = 1.f;
-			}
+			Results[Idx] = 1.f;
 		}
 
 		// Compute and add easing weight.
-		const bool bHasEasing = InAllocation->HasComponent(EasingData.ComponentType);
-		if (bHasEasing)
+		if (OptEasing)
 		{
-			float* CurResult = Results;
-			const FFrameTime* CurTime = Times;
-			const FEasingComponentData* CurEasing = EasingData.Resolve(InAllocation);
 			for (int32 Idx = 0; Idx < Num; ++Idx)
 			{
-				const float EasingWeight = (CurEasing++)->Section->EvaluateEasing(*(CurTime++));
-				*(CurResult++) *= FMath::Max(EasingWeight, 0.f);
+				const float EasingWeight = OptEasing[Idx].Section->EvaluateEasing(Times[Idx]);
+				Results[Idx] *= FMath::Max(EasingWeight, 0.f);
 			}
 		}
 
 		// Manual weight has already been computed by the float channel evaluator system, so we
 		// just need to pick up the result and combine it.
-		const bool bHasWeight = InAllocation->HasComponent(WeightData.ComponentType);
-		if (bHasWeight)
+		if (OptManualWeights)
 		{
-			float* CurResult = Results;
-			const FFrameTime* CurTime = Times;
-			const float* CurWeight = WeightData.Resolve(InAllocation);
 			for (int32 Idx = 0; Idx < Num; ++Idx)
 			{
-				const float CustomWeight = *(CurWeight++);
-				*(CurResult++) *= FMath::Max(CustomWeight, 0.f);
+				Results[Idx] *= FMath::Max(OptManualWeights[Idx], 0.f);
 			}
 		}
 
@@ -216,20 +193,14 @@ struct FEvaluateEasings
 		// Note that we need to check for instance handles because in interrogation evaluations, there are no
 		// instance handles.
 		//
-		const bool bIsProviders = InAllocation->HasComponent(HierarchicalEasingProviderData.ComponentType);
-		const bool bHasInstances = InAllocation->HasComponent(InstanceHandleData.ComponentType);
-		if (bIsProviders && bHasInstances)
+		if (OptInstanceHandles && OptHierarchicalEasingProviders)
 		{
 			const FInstanceRegistry* InstanceRegistry = EvaluatorSystem->GetLinker()->GetInstanceRegistry();
 
-			float* CurResult = Results;
-			const FInstanceHandle* InstanceHandles = InstanceHandleData.Resolve(InAllocation);
-			const FMovieSceneSequenceID* CurSubSequenceIDs = HierarchicalEasingProviderData.Resolve(InAllocation);
 			for (int32 Idx = 0; Idx < Num; ++Idx)
 			{
-				const FSequenceInstance& Instance = InstanceRegistry->GetInstance(*(InstanceHandles++));
-				const FInstanceHandle CurSubSequenceHandle = Instance.FindSubInstance(*(CurSubSequenceIDs++));
-				EvaluatorSystem->SetSubSequenceEasing(CurSubSequenceHandle, *(CurResult++));
+				const FSequenceInstance& Instance = InstanceRegistry->GetInstance(OptInstanceHandles[Idx]);
+				EvaluatorSystem->SetSubSequenceEasing(OptInstanceHandles[Idx], Results[Idx]);
 			}
 		}
 	}
@@ -267,36 +238,23 @@ struct FAccumulateHierarchicalEasings
 
 struct FPropagateHierarchicalEasings
 {
-	const FInstanceRegistry* InstanceRegistry;
 	TSparseArray<FHierarchicalEasingChannelData>* EasingChannels;
 
-	FPropagateHierarchicalEasings(
-			const FInstanceRegistry* InInstanceRegistry, 
-			TSparseArray<FHierarchicalEasingChannelData>* InEasingChannels)
-		: InstanceRegistry(InInstanceRegistry)
-		, EasingChannels(InEasingChannels)
+	FPropagateHierarchicalEasings(TSparseArray<FHierarchicalEasingChannelData>* InEasingChannels)
+		: EasingChannels(InEasingChannels)
 	{}
 
-	void ForEachAllocation(
-			const FEntityAllocation* Allocation,
-			TRead<FInstanceHandle> InstanceHandleData,
-			TRead<uint16> HierarchicalEasingChannelAccessor,
-			TWrite<float> WeightAndEasingResultData)
+	void ForEachAllocation(const FEntityAllocation* Allocation, TRead<uint16> HierarchicalEasingChannels, TWrite<float> WeightAndEasingResults)
 	{
-		TArrayView<const FInstanceHandle> InstanceHandles = InstanceHandleData.ResolveAsArray(Allocation);
-		TArrayView<const uint16> HierarchicalEasingChannels = HierarchicalEasingChannelAccessor.ResolveAsArray(Allocation);
-		TArrayView<float> WeightAndEasingResults = WeightAndEasingResultData.ResolveAsArray(Allocation);
-
-		for (int32 Index = 0; Index < Allocation->Num(); ++Index)
+		const int32 Num = Allocation->Num();
+		for (int32 Index = 0; Index < Num; ++Index)
 		{
-			const FSequenceInstance& SequenceInstance = InstanceRegistry->GetInstance(InstanceHandles[Index]);
 			const uint16 HierarchicalEasingChannel = HierarchicalEasingChannels[Index];
-			float& WeightAndEasingResult = WeightAndEasingResults[Index];
 
 			if (ensure(EasingChannels->IsValidIndex(HierarchicalEasingChannel)))
 			{
 				const FHierarchicalEasingChannelData& EasingChannel = (*EasingChannels)[HierarchicalEasingChannel];
-				WeightAndEasingResult *= EasingChannel.FinalEasingResult;
+				WeightAndEasingResults[Index] *= EasingChannel.FinalEasingResult;
 			}
 		}
 	}
@@ -423,11 +381,10 @@ void UWeightAndEasingEvaluatorSystem::OnRun(FSystemTaskPrerequisites& InPrerequi
 		// Step 3: Apply hierarchical easing results to all entities inside affected sub-sequences.
 		//
 		FEntityTaskBuilder()
-			.Read(Components->InstanceHandle)
 			.Read(Components->HierarchicalEasingChannel)
 			.Write(Components->WeightAndEasingResult)
 			.SetStat(GET_STATID(MovieSceneEval_EvaluateEasingTask))
-			.Dispatch_PerAllocation<FPropagateHierarchicalEasings>(&Linker->EntityManager, PropagatePrereqs, &Subsequents, InstanceRegistry, &EasingChannels);
+			.Dispatch_PerAllocation<FPropagateHierarchicalEasings>(&Linker->EntityManager, PropagatePrereqs, &Subsequents, &EasingChannels);
 	}
 }
 
diff --git a/Engine/Source/Runtime/MovieSceneTracks/Public/EntitySystem/Interrogation/MovieSceneInterrogationLinker.h b/Engine/Source/Runtime/MovieSceneTracks/Public/EntitySystem/Interrogation/MovieSceneInterrogationLinker.h
--- a/Engine/Source/Runtime/MovieSceneTracks/Public/EntitySystem/Interrogation/MovieSceneInterrogationLinker.h
+++ b/Engine/Source/Runtime/MovieSceneTracks/Public/EntitySystem/Interrogation/MovieSceneInterrogationLinker.h
@@ -4,6 +4,7 @@
 
 #include "CoreTypes.h"
 #include "UObject/GCObject.h"
+#include "Containers/SortedMap.h"
 #include "EntitySystem/MovieSceneEntitySystemTypes.h"
 #include "MovieSceneTracksComponentTypes.h"
 #include "Evaluation/MovieSceneEvaluationField.h"
diff --git a/Engine/Source/Runtime/MovieSceneTracks/Public/Systems/MovieScenePropertyInstantiator.h b/Engine/Source/Runtime/MovieSceneTracks/Public/Systems/MovieScenePropertyInstantiator.h
--- a/Engine/Source/Runtime/MovieSceneTracks/Public/Systems/MovieScenePropertyInstantiator.h
+++ b/Engine/Source/Runtime/MovieSceneTracks/Public/Systems/MovieScenePropertyInstantiator.h
@@ -7,6 +7,7 @@
 #include "EntitySystem/MovieSceneEntityInstantiatorSystem.h"
 #include "EntitySystem/MovieSceneDecompositionQuery.h"
 #include "EntitySystem/BuiltInComponentTypes.h"
+#include "EntitySystem/MovieScenePropertyRegistry.h"
 #include "EntitySystem/MovieScenePropertyComponentHandler.h"
 #include "EntitySystem/MovieScenePreAnimatedStateSystem.h"
 
@@ -93,14 +94,15 @@ private:
 
 	struct FObjectPropertyInfo
 	{
-		FObjectPropertyInfo()
-			: BoundObject(nullptr)
+		FObjectPropertyInfo(UE::MovieScene::FResolvedProperty&& InProperty)
+			: Property(MoveTemp(InProperty))
+			, BoundObject(nullptr)
 			, BlendChannel(INVALID_BLEND_CHANNEL)
 			, bWantsRestoreState(false)
 		{}
 
 		/** Variant of the property itself as either a pointer offset, a custom property index, or slow track instance bindings object */
-		TOptional< TVariant<uint16, UE::MovieScene::FCustomPropertyIndex, FSlowPropertyPtr> > Property;
+		UE::MovieScene::FResolvedProperty Property;
 
 		/** POinter to the blender system to use for this property, if its blended */
 		TWeakObjectPtr<UMovieSceneBlenderSystem> Blender;
@@ -138,8 +140,6 @@ private:
 	bool PropertySupportsFastPath(const FPropertyParameters& Params) const;
 	void InitializeFastPath(const FPropertyParameters& Params);
 	void InitializeBlendPath(const FPropertyParameters& Params);
-	int32 FindCustomAccessorIndex(UE::MovieScene::FCustomAccessorView Accessors, UClass* ClassType, FName PropertyPath);
-	TOptional<uint16> ComputeFastPropertyPtrOffset(UClass* ObjectClass, const FMovieScenePropertyBinding& PropertyBinding);
 	int32 ResolveProperty(UE::MovieScene::FCustomAccessorView CustomAccessors, UObject* Object, const FMovieScenePropertyBinding& PropertyBinding, int32 PropertyDefinitionIndex);
 
 	UE::MovieScene::FPropertyRecomposerPropertyInfo FindPropertyFromSource(FMovieSceneEntityID EntityID, UObject* Object) const;
diff --git a/Engine/Source/Runtime/MovieSceneTracks/Public/Tracks/MovieScenePropertyTrack.h b/Engine/Source/Runtime/MovieSceneTracks/Public/Tracks/MovieScenePropertyTrack.h
--- a/Engine/Source/Runtime/MovieSceneTracks/Public/Tracks/MovieScenePropertyTrack.h
+++ b/Engine/Source/Runtime/MovieSceneTracks/Public/Tracks/MovieScenePropertyTrack.h
@@ -6,7 +6,7 @@
 #include "UObject/ObjectMacros.h"
 #include "MovieSceneNameableTrack.h"
 #include "MovieSceneCommonHelpers.h"
-#include "EntitySystem/BuiltInComponentTypes.h"
+#include "EntitySystem/MovieScenePropertyBinding.h"
 #include "MovieScenePropertyTrack.generated.h"
 
 /**
diff --git a/Engine/Source/Runtime/UMG/Private/Animation/MovieSceneUMGComponentTypes.cpp b/Engine/Source/Runtime/UMG/Private/Animation/MovieSceneUMGComponentTypes.cpp
--- a/Engine/Source/Runtime/UMG/Private/Animation/MovieSceneUMGComponentTypes.cpp
+++ b/Engine/Source/Runtime/UMG/Private/Animation/MovieSceneUMGComponentTypes.cpp
@@ -5,6 +5,7 @@
 #include "EntitySystem/MovieScenePropertyComponentHandler.h"
 #include "EntitySystem/BuiltInComponentTypes.h"
 #include "MovieSceneTracksComponentTypes.h"
+#include "EntitySystem/MovieSceneEntityFactoryTemplates.h"
 
 #include "Components/Widget.h"
 
