Sequencer: Added a method to retrieve the first component in an FComponentMask

- This is useful when you want to find the result of a bitwise AND bewteen 2 masks and you know there will only be a single bit result

diff --git a/Engine/Source/Runtime/MovieScene/Public/EntitySystem/MovieSceneEntityIDs.h b/Engine/Source/Runtime/MovieScene/Public/EntitySystem/MovieSceneEntityIDs.h
--- a/Engine/Source/Runtime/MovieScene/Public/EntitySystem/MovieSceneEntityIDs.h
+++ b/Engine/Source/Runtime/MovieScene/Public/EntitySystem/MovieSceneEntityIDs.h
@@ -131,6 +131,11 @@ struct FComponentMask
 
 	void Remove(FComponentTypeID InComponentType);
 
+	/**
+	 * Find the first component type ID in this mask, or Invalid if the mask is empty.
+	 */
+	FComponentTypeID First() const;
+
 	FORCEINLINE void  Reset()
 	{
 		Bits.Reset();
@@ -448,6 +453,16 @@ inline void FComponentMask::Remove(FComponentTypeID InComponentType)
 	}
 }
 
+inline FComponentTypeID FComponentMask::First() const
+{
+	const int32 FirstBit = Bits.Find(true);
+	if (FirstBit != INDEX_NONE)
+	{
+		return FComponentTypeID::FromBitIndex(FirstBit);
+	}
+	return FComponentTypeID::Invalid();
+}
+
 template<typename T>
 inline TComponentTypeID<T> FComponentTypeID::ReinterpretCast() const
 {
