Sequencer: Added a function for mutations to initialize allocations that were mutated

diff --git a/Engine/Source/Runtime/MovieScene/Private/EntitySystem/MovieSceneEntityManager.cpp b/Engine/Source/Runtime/MovieScene/Private/EntitySystem/MovieSceneEntityManager.cpp
--- a/Engine/Source/Runtime/MovieScene/Private/EntitySystem/MovieSceneEntityManager.cpp
+++ b/Engine/Source/Runtime/MovieScene/Private/EntitySystem/MovieSceneEntityManager.cpp
@@ -1644,6 +1644,8 @@ int32 FEntityManager::MutateAll(const FEntityComponentFilter& Filter, const IMov
 		EntityAllocationMasks[Pair.Key] = Pair.Value;
 		EntityAllocations[Pair.Key] = NewAllocation;
 
+		Mutation.InitializeAllocation(NewAllocation, Pair.Value);
+
 		DestroyAllocation(SourceAllocation);
 	}
 
diff --git a/Engine/Source/Runtime/MovieScene/Public/EntitySystem/MovieSceneEntityMutations.h b/Engine/Source/Runtime/MovieScene/Public/EntitySystem/MovieSceneEntityMutations.h
--- a/Engine/Source/Runtime/MovieScene/Public/EntitySystem/MovieSceneEntityMutations.h
+++ b/Engine/Source/Runtime/MovieScene/Public/EntitySystem/MovieSceneEntityMutations.h
@@ -15,6 +15,7 @@ struct IMovieSceneEntityMutation
 {
 	virtual ~IMovieSceneEntityMutation() {}
 	virtual void CreateMutation(FEntityManager* EntityManager, FComponentMask* InOutEntityComponentTypes) const = 0;
+	virtual void InitializeAllocation(FEntityAllocation* Allocation, const FComponentMask& AllocationType) const {}
 };
 
 struct FAddSingleMutation : IMovieSceneEntityMutation
