Sequencer: Added the ability to register extensions with a linker

diff --git a/Engine/Source/Runtime/MovieScene/Private/EntitySystem/MovieSceneEntitySystemLinker.cpp b/Engine/Source/Runtime/MovieScene/Private/EntitySystem/MovieSceneEntitySystemLinker.cpp
--- a/Engine/Source/Runtime/MovieScene/Private/EntitySystem/MovieSceneEntitySystemLinker.cpp
+++ b/Engine/Source/Runtime/MovieScene/Private/EntitySystem/MovieSceneEntitySystemLinker.cpp
@@ -61,6 +61,12 @@ UMovieSceneEntitySystemLinker::UMovieSceneEntitySystemLinker(const FObjectInitia
 	}
 }
 
+UE::MovieScene::FEntitySystemLinkerExtensionID UMovieSceneEntitySystemLinker::RegisterExtension()
+{
+	static int32 StaticID = 0;
+	return UE::MovieScene::FEntitySystemLinkerExtensionID{ StaticID++ };
+}
+
 void UMovieSceneEntitySystemLinker::Reset()
 {
 	Events.AbandonLinker.Broadcast(this);
diff --git a/Engine/Source/Runtime/MovieScene/Public/EntitySystem/MovieSceneEntitySystemLinker.h b/Engine/Source/Runtime/MovieScene/Public/EntitySystem/MovieSceneEntitySystemLinker.h
--- a/Engine/Source/Runtime/MovieScene/Public/EntitySystem/MovieSceneEntitySystemLinker.h
+++ b/Engine/Source/Runtime/MovieScene/Public/EntitySystem/MovieSceneEntitySystemLinker.h
@@ -13,6 +13,7 @@
 #include "EntitySystem/MovieSceneEntitySystemTypes.h"
 #include "EntitySystem/MovieSceneEntitySystemGraphs.h"
 #include "EntitySystem/MovieSceneSequenceInstance.h"
+#include "EntitySystem/MovieSceneEntitySystemLinkerExtension.h"
 
 #include "MovieSceneEntitySystemLinker.generated.h"
 
@@ -134,6 +135,86 @@ public:
 	 */
 	void Reset();
 
+public:
+
+	/**
+	 * Register a new extension type for use with any instance of a UMovieSceneEntitySystemLinker
+	 */
+	template<typename ExtensionType>
+	static UE::MovieScene::TEntitySystemLinkerExtensionID<ExtensionType> RegisterExtension()
+	{
+		return UE::MovieScene::TEntitySystemLinkerExtensionID<ExtensionType>(RegisterExtension().ID);
+	}
+
+	/**
+	 * Add an extension to this linker.
+	 *
+	 * @param InID        The unique identifier for the type of extension (retrieved from RegisterExtension)
+	 * @param InExtension Pointer to the extension to register - must be kept alive externally - only a raw ptr is kept in this class
+	 */
+	template<typename ExtensionType>
+	void AddExtension(UE::MovieScene::TEntitySystemLinkerExtensionID<ExtensionType> InID, ExtensionType* InExtension)
+	{
+		const int32 Index = InID.ID;
+		if (!ExtensionsByID.IsValidIndex(Index))
+		{
+			ExtensionsByID.Insert(Index, InExtension);
+		}
+		else
+		{
+			check(ExtensionsByID[Index] == InExtension);
+		}
+	}
+
+	/**
+	 * Add an extension to this linker.
+	 *
+	 * @param InID        The unique identifier for the type of extension (retrieved from RegisterExtension)
+	 * @param InExtension Pointer to the extension to register - must be kept alive externally - only a raw ptr is kept in this class
+	 */
+	template<typename ExtensionType>
+	void AddExtension(ExtensionType* InExtension)
+	{
+		AddExtension(ExtensionType::GetExtensionID(), InExtension);
+	}
+
+	/**
+	 * Attempt to find an extension to this linker by its ID
+	 *
+	 * @param InID        The unique identifier for the type of extension (retrieved from RegisterExtension)
+	 * @return A pointer to the extension, or nullptr if it is not active.
+	 */
+	template<typename ExtensionType>
+	ExtensionType* FindExtension(UE::MovieScene::TEntitySystemLinkerExtensionID<ExtensionType> InID)
+	{
+		const int32 Index = InID.ID;
+		if (ExtensionsByID.IsValidIndex(Index))
+		{
+			return static_cast<ExtensionType*>(ExtensionsByID[Index]);
+		}
+
+		return nullptr;
+	}
+
+
+	/**
+	 * Attempt to find an extension to this linker by its ID
+	 *
+	 * @param InID        The unique identifier for the type of extension (retrieved from RegisterExtension)
+	 * @return A pointer to the extension, or nullptr if it is not active.
+	 */
+	template<typename ExtensionType>
+	ExtensionType* FindExtension()
+	{
+		const int32 Index = ExtensionType::GetExtensionID().ID;
+		if (ExtensionsByID.IsValidIndex(Index))
+		{
+			return static_cast<ExtensionType*>(ExtensionsByID[Index]);
+		}
+
+		return nullptr;
+	}
+
 public:
 
 	// Internal API
@@ -164,6 +245,8 @@ private:
 
 	virtual void BeginDestroy() override;
 
+	static UE::MovieScene::FEntitySystemLinkerExtensionID RegisterExtension();
+
 private:
 
 	TUniquePtr<FInstanceRegistry> InstanceRegistry;
@@ -177,6 +260,8 @@ private:
 	};
 	TArray<FActiveRunnerInfo> ActiveRunners;
 
+	TSparseArray<void*> ExtensionsByID;
+
 	friend struct FMovieSceneEntitySystemEvaluationReentrancyWindow;
 
 public:
diff --git a/Engine/Source/Runtime/MovieScene/Public/EntitySystem/MovieSceneEntitySystemLinkerExtension.h b/Engine/Source/Runtime/MovieScene/Public/EntitySystem/MovieSceneEntitySystemLinkerExtension.h
new file mode 100644
--- /dev/null
+++ b/Engine/Source/Runtime/MovieScene/Public/EntitySystem/MovieSceneEntitySystemLinkerExtension.h
@@ -0,0 +1,29 @@
+// Copyright Epic Games, Inc. All Rights Reserved.
+
+#pragma once
+
+#include "CoreTypes.h"
+
+namespace UE
+{
+namespace MovieScene
+{
+
+/** Base extension identifier for a UMovieSceneEntitySystemLinker */
+struct FEntitySystemLinkerExtensionID
+{
+	int32 ID;
+};
+
+/** Typed extension identifier for a UMovieSceneEntitySystemLinker */
+template<typename ExtensionType>
+struct TEntitySystemLinkerExtensionID : FEntitySystemLinkerExtensionID
+{
+	explicit TEntitySystemLinkerExtensionID(int32 InID)
+		: FEntitySystemLinkerExtensionID{ InID }
+	{}
+};
+
+
+} // namespace MovieScene
+} // namespace UE
