Sequencer: Initial Values are now shared between FSequencer instances and Interrogators

  - This ensures that exports and trail rendering share the same initial value data despite objects potentially being in animated states when an interrogation is performed.
  - All initial value processing is now handled by a new system, UMovieSceneInitialValueSystem, rather than the property instantiator. This means the same system can initialize the correct initial values for interrogators and runtime properties.
  - This change introduces 2 new linker extensions (one for the interrogator, and one for the shared initial value storage) which allows the initial value system to perform optimally when the extensions are not present, but also use the initial value cache storage when it present in-editor.

diff --git a/Engine/Source/Editor/Sequencer/Private/Sequencer.cpp b/Engine/Source/Editor/Sequencer/Private/Sequencer.cpp
--- a/Engine/Source/Editor/Sequencer/Private/Sequencer.cpp
+++ b/Engine/Source/Editor/Sequencer/Private/Sequencer.cpp
@@ -139,6 +139,7 @@
 #include "EngineAnalytics.h"
 #include "Interfaces/IAnalyticsProvider.h"
 #include "Components/SkeletalMeshComponent.h"
+#include "EntitySystem/MovieSceneInitialValueCache.h"
 #include "SequencerCustomizationManager.h"
 #include "SSequencerGroupManager.h"
 #include "ActorTreeItem.h"
@@ -426,6 +427,9 @@ void FSequencer::InitSequencer(const FSequencerInitParams& InitParams, const TSh
 	ActiveTemplateStates.Add(true);
 	RootTemplateInstance.Initialize(*InitParams.RootSequence, *this, CompiledDataManager);
 
+	InitialValueCache = UE::MovieScene::FInitialValueCache::GetGlobalInitialValues();
+	RootTemplateInstance.GetEntitySystemLinker()->AddExtension(InitialValueCache.Get());
+
 	ResetTimeController();
 
 	// Make internal widgets
diff --git a/Engine/Source/Editor/Sequencer/Private/Sequencer.h b/Engine/Source/Editor/Sequencer/Private/Sequencer.h
--- a/Engine/Source/Editor/Sequencer/Private/Sequencer.h
+++ b/Engine/Source/Editor/Sequencer/Private/Sequencer.h
@@ -76,6 +76,17 @@ struct FNotificationInfo;
 
 enum class EMapChangeType : uint8;
 
+namespace UE
+{
+namespace MovieScene
+{
+
+	struct FInitialValueCache;
+
+} // namespace MovieScene
+} // namespace UE
+
+
 /**
  * Sequencer is the editing tool for MovieScene assets.
  */
@@ -1418,6 +1429,8 @@ private:
 
 	TSharedPtr<FCurveEditor> CurveEditorModel;
 
+	TSharedPtr<UE::MovieScene::FInitialValueCache> InitialValueCache;
+
 	/** A signature that will suppress auto evaluation when it is the only change dirtying the template. */
 	TOptional<TTuple<TWeakObjectPtr<UMovieSceneSequence>, FGuid>> SuppressAutoEvalSignature;
 
diff --git a/Engine/Source/Runtime/MovieScene/Private/EntitySystem/BuiltInComponentTypes.cpp b/Engine/Source/Runtime/MovieScene/Private/EntitySystem/BuiltInComponentTypes.cpp
--- a/Engine/Source/Runtime/MovieScene/Private/EntitySystem/BuiltInComponentTypes.cpp
+++ b/Engine/Source/Runtime/MovieScene/Private/EntitySystem/BuiltInComponentTypes.cpp
@@ -64,6 +64,7 @@ FBuiltInComponentTypes::FBuiltInComponentTypes()
 	ComponentRegistry->NewComponentType(&BlendChannelInput,     TEXT("Blend Channel Input"));
 	ComponentRegistry->NewComponentType(&HierarchicalBias,      TEXT("Hierarchical Bias"));
 	ComponentRegistry->NewComponentType(&BlendChannelOutput,    TEXT("Blend Channel Output"));
+	ComponentRegistry->NewComponentType(&InitialValueIndex,     TEXT("Initial Value Index"));
 
 	ComponentRegistry->NewComponentType(&CustomPropertyIndex,   TEXT("Custom Property Index"));			// Not EComponentTypeFlags::Preserved because the system property manager will always ensure that the component is added to the correct entity
 	ComponentRegistry->NewComponentType(&FastPropertyOffset,    TEXT("Fast Property Offset"));			// Not EComponentTypeFlags::Preserved because the system property manager will always ensure that the component is added to the correct entity
diff --git a/Engine/Source/Runtime/MovieScene/Private/EntitySystem/MovieSceneInitialValueCache.cpp b/Engine/Source/Runtime/MovieScene/Private/EntitySystem/MovieSceneInitialValueCache.cpp
new file mode 100644
--- /dev/null
+++ b/Engine/Source/Runtime/MovieScene/Private/EntitySystem/MovieSceneInitialValueCache.cpp
@@ -0,0 +1,33 @@
+// Copyright Epic Games, Inc. All Rights Reserved.
+
+#include "EntitySystem/MovieSceneInitialValueCache.h"
+#include "EntitySystem/MovieSceneEntitySystemLinker.h"
+
+namespace UE
+{
+namespace MovieScene
+{
+
+
+TEntitySystemLinkerExtensionID<FInitialValueCache> FInitialValueCache::GetExtensionID()
+{
+	static TEntitySystemLinkerExtensionID<FInitialValueCache> ID = UMovieSceneEntitySystemLinker::RegisterExtension<FInitialValueCache>();
+	return ID;
+}
+
+TSharedPtr<FInitialValueCache> FInitialValueCache::GetGlobalInitialValues()
+{
+	static TWeakPtr<FInitialValueCache> GSharedInitialValues = MakeShared<FInitialValueCache>();
+
+	TSharedPtr<FInitialValueCache> Ptr = GSharedInitialValues.Pin();
+	if (!Ptr)
+	{
+		Ptr = MakeShared<FInitialValueCache>();
+		GSharedInitialValues = Ptr;
+	}
+	return Ptr;
+}
+
+} // namespace MovieScene
+} // namespace UE
+
diff --git a/Engine/Source/Runtime/MovieScene/Public/EntitySystem/BuiltInComponentTypes.h b/Engine/Source/Runtime/MovieScene/Public/EntitySystem/BuiltInComponentTypes.h
--- a/Engine/Source/Runtime/MovieScene/Public/EntitySystem/BuiltInComponentTypes.h
+++ b/Engine/Source/Runtime/MovieScene/Public/EntitySystem/BuiltInComponentTypes.h
@@ -11,7 +11,7 @@
 #include "Templates/SubclassOf.h"
 
 #include "EntitySystem/MovieScenePropertyRegistry.h"
-#include "EntitySystem/MovieScenePropertySystemTypes.h"
+#include "EntitySystem/MovieSceneInitialValueCache.h"
 
 #include "BuiltInComponentTypes.generated.h"
 
@@ -145,6 +145,7 @@ public:
 
 	TComponentTypeID<uint16>              BlendChannelOutput;
 
+	TComponentTypeID<FInitialValueIndex>  InitialValueIndex;
 public:
 
 	// An FMovieScenePropertyBinding structure
diff --git a/Engine/Source/Runtime/MovieScene/Public/EntitySystem/IMovieScenePropertyComponentHandler.h b/Engine/Source/Runtime/MovieScene/Public/EntitySystem/IMovieScenePropertyComponentHandler.h
--- a/Engine/Source/Runtime/MovieScene/Public/EntitySystem/IMovieScenePropertyComponentHandler.h
+++ b/Engine/Source/Runtime/MovieScene/Public/EntitySystem/IMovieScenePropertyComponentHandler.h
@@ -14,6 +14,7 @@ namespace MovieScene
 
 struct FPropertyStats;
 struct FPropertyDefinition;
+struct FInitialValueCache;
 struct FSystemSubsequentTasks;
 struct FSystemTaskPrerequisites;
 struct FFloatDecompositionParams;
@@ -123,6 +124,18 @@ private:
 	int32 ArrayNum;
 };
 
+/** Interface required for initializing initial values on entities */
+struct IInitialValueProcessor
+{
+	virtual ~IInitialValueProcessor(){}
+
+	/** Initialize this processor before any allocations are visited */
+	virtual void Initialize(UMovieSceneEntitySystemLinker* Linker, const FPropertyDefinition* Definition, FInitialValueCache* InitialValueCache) = 0;
+	/** Process all initial values for the specified allocation */
+	virtual void Process(const FEntityAllocation* Allocation, const FComponentMask& AllocationType) = 0;
+	/** Finish processing */
+	virtual void Finalize() = 0;
+};
 
 /** Interface for a property type handler that is able to interact with properties in sequencer */
 struct IPropertyComponentHandler
@@ -161,16 +174,6 @@ struct IPropertyComponentHandler
 	 */
 	virtual void DispatchRestorePreAnimatedStateTasks(const FPropertyDefinition& Definition, FSystemTaskPrerequisites& InPrerequisites, FSystemSubsequentTasks& Subsequents, UMovieSceneEntitySystemLinker* Linker) = 0;
 
-	/**
-	 * Dispatch tasks that cache an initial unblended value for any entities that have the NeedsLink tag
-	 *
-	 * @param Definition       The property definition this handler was registered for
-	 * @param InPrerequisites  Task prerequisites for any entity system tasks that are dispatched
-	 * @param Subsequents      Subsequents to add any dispatched tasks to
-	 * @param Linker           The linker that owns the entity manager to dispatch tasks for
-	 */
-	virtual void DispatchCacheInitialValueTasks(const FPropertyDefinition& Definition, FSystemTaskPrerequisites& InPrerequisites, FSystemSubsequentTasks& Subsequents, UMovieSceneEntitySystemLinker* Linker) = 0;
-
 	/**
 	 * Run a recomposition using the specified params and values. The current value and result views must be of type PropertyType
 	 *
@@ -214,6 +217,11 @@ struct IPropertyComponentHandler
 	 * @param Linker           The linker that owns the entity manager to dispatch tasks for
 	 */
 	virtual void SaveGlobalPreAnimatedState(const FPropertyDefinition& Definition, UMovieSceneEntitySystemLinker* Linker) = 0;
+
+	/**
+	 * Retrieve an initial value processor interface for this property type
+	 */
+	virtual IInitialValueProcessor* GetInitialValueProcessor() = 0;
 };
 
 
diff --git a/Engine/Source/Runtime/MovieScene/Public/EntitySystem/MovieSceneInitialValueCache.h b/Engine/Source/Runtime/MovieScene/Public/EntitySystem/MovieSceneInitialValueCache.h
new file mode 100644
--- /dev/null
+++ b/Engine/Source/Runtime/MovieScene/Public/EntitySystem/MovieSceneInitialValueCache.h
@@ -0,0 +1,365 @@
+// Copyright Epic Games, Inc. All Rights Reserved.
+
+#pragma once
+
+#include "Templates/Tuple.h"
+#include "Containers/Map.h"
+#include "Misc/TVariant.h"
+#include "MovieSceneCommonHelpers.h"
+#include "EntitySystem/MovieSceneEntitySystemLinkerExtension.h"
+#include "EntitySystem/MovieScenePropertySystemTypes.h"
+
+namespace UE
+{
+namespace MovieScene
+{
+
+template<typename PropertyType> struct TPropertyValueStorage;
+
+/** Type-safe index that defines a unique index of an initial value within a TPropertyValueStorage instance
+ *  This struct is added as a component to all entities with initial values inside non-interrogated linkers */
+struct FInitialValueIndex
+{
+	uint16 Index;
+};
+
+/** Base interface for all initial value storage types */
+struct IPropertyValueStorage
+{
+	virtual ~IPropertyValueStorage(){}
+
+	/** Reset all the initial values for the specified indices */
+	virtual void Reset(TArrayView<const FInitialValueIndex> Indices) = 0;
+};
+
+/**
+ * Container that houses initial values for any properties animated through FEntityManager
+ * Each type is stored in its own container, organized by the initial value component ID
+ * The cache is stored as a singleton (through GetSharedInitialValues()) and added as an
+ * extension to UMovieSceneEntitySystemLinker where it is active.
+ */
+struct FInitialValueCache
+{
+	/**
+	 * Retrieve the extension ID for this structure when added to a UMovieSceneEntitySystemLinker
+	 */
+	MOVIESCENE_API static TEntitySystemLinkerExtensionID<FInitialValueCache> GetExtensionID();
+
+
+	/**
+	 * Retrieve a container used for sharing initial values between linkers.
+	 * @note: The instance is not referenced internally, so will be destroyed when all
+	 * externally owned shared pointers are destroyed.
+	 */
+	MOVIESCENE_API static TSharedPtr<FInitialValueCache> GetGlobalInitialValues();
+
+public:
+
+	/**
+	 * Reset all the initial values that relate to the specified type and indices
+	 *
+	 * @param InitialValueType      The ComponentTypeID for the initial value that the indices relate to
+	 * @param InitialValueIndices   An array containing all the indices to remove
+	 */
+	void Reset(FComponentTypeID InitialValueType, TArrayView<const FInitialValueIndex> InitialValueIndices)
+	{
+		if (TUniquePtr<IPropertyValueStorage>* Storage = StorageByComponent.Find(InitialValueType))
+		{
+			Storage->Get()->Reset(InitialValueIndices);
+		}
+	}
+
+	/**
+	 * Retrieve the initial value storage for a given initial value type, creating it if necessary.
+	 * @note: Care should be taken to ensure that the template parameter matches TPropertyComponents::PropertyType,
+	 * not TPropertyComponents::OperationalType.
+	 *
+	 * @param InitialValueType      The ComponentTypeID for the initial value
+	 * @return Property storage for the initial values
+	 */
+	template<typename PropertyType>
+	TPropertyValueStorage<PropertyType>* GetStorage(FComponentTypeID InitialValueType)
+	{
+		if (TPropertyValueStorage<PropertyType>* Storage = FindStorage<PropertyType>(InitialValueType))
+		{
+			return Storage;
+		}
+
+		TPropertyValueStorage<PropertyType>* NewStorage = new TPropertyValueStorage<PropertyType>();
+		StorageByComponent.Add(InitialValueType, TUniquePtr<IPropertyValueStorage>(NewStorage));
+		return NewStorage;
+	}
+
+
+	/**
+	 * Retrieve the initial value storage for a given initial value type.
+	 * @note: Care should be taken to ensure that the template parameter matches TPropertyComponents::PropertyType,
+	 * not TPropertyComponents::OperationalType.
+	 *
+	 * @param InitialValueType      The ComponentTypeID for the initial value
+	 * @return Property storage for the initial values or nullptr if none exists
+	 */
+	template<typename PropertyType>
+	TPropertyValueStorage<PropertyType>* FindStorage(FComponentTypeID InitialValueType)
+	{
+		if (TUniquePtr<IPropertyValueStorage>* Existing = StorageByComponent.Find(InitialValueType))
+		{
+			// If the ptr exists, it cannot be null
+			check(Existing->IsValid());
+			return static_cast<TPropertyValueStorage<PropertyType>*>(Existing->Get());
+		}
+
+		return nullptr;
+	}
+
+private:
+
+	TMap<FComponentTypeID, TUniquePtr<IPropertyValueStorage>> StorageByComponent;
+};
+
+/**
+ * Templated storage for any initial value type, templated on TPropertyComponents::PropertyType for correct retrieval of resolved properties
+ * Initial values are stored as a sparse array, with stable indices that uniquely identify the value.
+ * A look-up-table exists for finding indices based on an object and resolved property.
+ */
+template<typename PropertyType>
+struct TPropertyValueStorage : IPropertyValueStorage
+{
+	/**
+	 * Reset all the initial values that relate to the specified indices
+	 *
+	 * @param Indices    A array specifying all of the indices to remove. Such indices will be invalid after this function call returns.
+	 */
+	virtual void Reset(TArrayView<const FInitialValueIndex> Indices) override
+	{
+		for (FInitialValueIndex Index : Indices)
+		{
+			PropertyValues.RemoveAt(Index.Index);
+		}
+		bLUTContainsInvalidEntries = true;
+	}
+
+
+	/**
+	 * Retrieve the cached value for the specified object and fast property ptr offset, returning both the value and its corresponding index.
+	 * @note: if the value is already cached, the existing value is returned
+	 *
+	 * @param BoundObject            The object instance to cache the property from
+	 * @param ResolvedPropertyOffset The byte offset from BoundObject that defines the address of the property
+	 * @return A tuple containing the cached value and its index
+	 */
+	TPair<FInitialValueIndex, PropertyType> CacheInitialValue(UObject* BoundObject, uint16 ResolvedPropertyOffset)
+	{
+		FKeyType Key{ FObjectKey(BoundObject), FPropertyKey(TInPlaceType<uint16>(), ResolvedPropertyOffset) };
+
+		if (TOptional<FInitialValueIndex> ExistingIndex = FindPropertyIndex(Key))
+		{
+			FInitialValueIndex ExistingValue = ExistingIndex.GetValue();
+			return MakeTuple(ExistingValue, PropertyValues[ExistingValue.Index]);
+		}
+
+		PropertyType CachedValue = *reinterpret_cast<const PropertyType*>(reinterpret_cast<const uint8*>(static_cast<const void*>(BoundObject)) + ResolvedPropertyOffset);
+
+		const int32 NewIndex = PropertyValues.Add(CachedValue);
+		check(NewIndex < int32(uint16(0xFFFF)));
+
+		const uint16 NarrowIndex = static_cast<uint16>(NewIndex);
+		KeyToPropertyIndex.Add(Key, NarrowIndex);
+
+		return MakeTuple(FInitialValueIndex{NarrowIndex}, CachedValue);
+	}
+
+
+	/**
+	 * Retrieve the cached value for the specified object and a custom property accessor, returning both the value and its corresponding index.
+	 * @note: if the value is already cached, the existing value is returned
+	 *
+	 * @param BoundObject            The object instance to cache the property from.
+	 * @param Accessors              A type-erased view retrieved from the necessary ICustomPropertyRegistration::GetAccessors for the property type this cache relates to.
+	 * @param AccessorIndex          The index into Accessors to use for resolving the property
+	 * @return A tuple containing the cached value and its index
+	 */
+	TPair<FInitialValueIndex, PropertyType> CacheInitialValue(UObject* BoundObject, FCustomAccessorView Accessors, FCustomPropertyIndex AccessorIndex)
+	{
+		FKeyType Key{ FObjectKey(BoundObject), FPropertyKey(TInPlaceType<FCustomPropertyIndex>(), AccessorIndex) };
+
+		if (TOptional<FInitialValueIndex> ExistingIndex = FindPropertyIndex(Key))
+		{
+			FInitialValueIndex ExistingValue = ExistingIndex.GetValue();
+			return MakeTuple(ExistingValue, PropertyValues[ExistingValue.Index]);
+		}
+
+		const TCustomPropertyAccessor<PropertyType>& CustomAccessor = static_cast<const TCustomPropertyAccessor<PropertyType>&>(Accessors[AccessorIndex.Value]);
+		PropertyType CachedValue = CustomAccessor.Functions.Getter(BoundObject);
+
+		const int32 NewIndex = PropertyValues.Add(CachedValue);
+		check(NewIndex < int32(uint16(0xFFFF)));
+
+		const uint16 NarrowIndex = static_cast<uint16>(NewIndex);
+		KeyToPropertyIndex.Add(Key, NarrowIndex);
+
+		return MakeTuple(FInitialValueIndex{NarrowIndex}, CachedValue);
+	}
+
+
+	/**
+	 * Retrieve the cached value for the specified object and a slow bindings instance, returning both the value and its corresponding index.
+	 * @note: if the value is already cached, the existing value is returned
+	 *
+	 * @param BoundObject            The object instance to cache the property from.
+	 * @param SlowBindings           Pointer to the track instance property bindings object used for retrieving the property value
+	 * @return A tuple containing the cached value and its index
+	 */
+	TPair<FInitialValueIndex, PropertyType> CacheInitialValue(UObject* BoundObject, FTrackInstancePropertyBindings* SlowBindings)
+	{
+		FKeyType Key{ FObjectKey(BoundObject), FPropertyKey(TInPlaceType<FName>(), SlowBindings->GetPropertyPath())};
+
+		if (TOptional<FInitialValueIndex> ExistingIndex = FindPropertyIndex(Key))
+		{
+			FInitialValueIndex ExistingValue = ExistingIndex.GetValue();
+			return MakeTuple(ExistingValue, PropertyValues[ExistingValue.Index]);
+		}
+
+		PropertyType CachedValue = SlowBindings->GetCurrentValue<PropertyType>(*BoundObject);
+
+		const int32 NewIndex = PropertyValues.Add(CachedValue);
+		check(NewIndex < int32(uint16(0xFFFF)));
+
+		const uint16 NarrowIndex = static_cast<uint16>(NewIndex);
+		KeyToPropertyIndex.Add(Key, NarrowIndex);
+
+		return MakeTuple(FInitialValueIndex{NarrowIndex}, CachedValue);
+	}
+
+
+	/**
+	 * Find an initial value index given its object and fast ptr offset
+	 */
+	TOptional<FInitialValueIndex> FindPropertyIndex(UObject* BoundObject, uint16 ResolvedPropertyOffset)
+	{
+		CleanupStaleEntries();
+		return FindPropertyIndex(FKeyType{ FObjectKey(BoundObject), FPropertyKey(TInPlaceType<uint16>(), ResolvedPropertyOffset) });
+	}
+
+
+	/**
+	 * Find an initial value index given its object and custom accessor index
+	 */
+	TOptional<FInitialValueIndex> FindPropertyIndex(UObject* BoundObject, FCustomPropertyIndex AccessorIndex)
+	{
+		CleanupStaleEntries();
+		return FindPropertyIndex(FKeyType{ FObjectKey(BoundObject), FPropertyKey(TInPlaceType<FCustomPropertyIndex>(), AccessorIndex) });
+	}
+
+
+	/**
+	 * Find an initial value index given its object and property name.
+	 * @note: Only properties cached using a FTrackInstancePropertyBindings instance will be retrieved using this method.
+	 */
+	TOptional<FInitialValueIndex> FindPropertyIndex(UObject* BoundObject, const FName& PropertyPath)
+	{
+		CleanupStaleEntries();
+		return FindPropertyIndex(FKeyType{ FObjectKey(BoundObject), FPropertyKey(TInPlaceType<FName>(), PropertyPath) });
+	}
+
+
+	/**
+	 * Find an initial value given its object and property name.
+	 */
+	const PropertyType* FindCachedValue(UObject* BoundObject, uint16 ResolvedPropertyOffset)
+	{
+		TOptional<FInitialValueIndex> Index = FindPropertyIndex(BoundObject, ResolvedPropertyOffset);
+		return Index.IsSet() ? &PropertyValues[Index.GetValue().Index] : nullptr;
+	}
+
+
+	/**
+	 * Find an initial value given its object and custom accessor index
+	 */
+	const PropertyType* FindCachedValue(UObject* BoundObject, FCustomPropertyIndex CustomIndex)
+	{
+		TOptional<FInitialValueIndex> Index = FindPropertyIndex(BoundObject, CustomIndex);
+		return Index.IsSet() ? &PropertyValues[Index.GetValue().Index] : nullptr;
+	}
+
+
+	/**
+	 * Find an initial value given its object and property name.
+	 * @note: Only properties cached using a FTrackInstancePropertyBindings instance will be retrieved using this method.
+	 */
+	const PropertyType* FindCachedValue(UObject* BoundObject, const FName& PropertyPath)
+	{
+		TOptional<FInitialValueIndex> Index = FindPropertyIndex(BoundObject, PropertyPath);
+		return Index.IsSet() ? &PropertyValues[Index.GetValue().Index] : nullptr;
+	}
+
+private:
+
+	using FPropertyKey = TVariant<uint16, FCustomPropertyIndex, FName>;
+
+	struct FKeyType
+	{
+		FObjectKey Object;
+		FPropertyKey Property;
+
+		friend uint32 GetTypeHash(const FKeyType& InKey)
+		{
+			// Hash only considers the _type_ of the resolved property
+			// which defers the final comparison to the equality operator
+			uint32 Hash = GetTypeHash(InKey.Object);
+			Hash = HashCombine(Hash, InKey.Property.GetIndex());
+			return Hash;
+		}
+		friend bool operator==(const FKeyType& A, const FKeyType& B)
+		{
+			if (A.Object != B.Object || A.Property.GetIndex() != B.Property.GetIndex())
+			{
+				return false;
+			}
+			switch (A.Property.GetIndex())
+			{
+				case 0: return A.Property.template Get<uint16>() == B.Property.template Get<uint16>();
+				case 1: return A.Property.template Get<FCustomPropertyIndex>().Value == B.Property.template Get<FCustomPropertyIndex>().Value;
+				case 2: return A.Property.template Get<FName>() == B.Property.template Get<FName>();
+			}
+			return true;
+		}
+	};
+
+	FORCEINLINE void CleanupStaleEntries()
+	{
+		if (!bLUTContainsInvalidEntries)
+		{
+			return;
+		}
+		for (auto It = KeyToPropertyIndex.CreateIterator(); It; ++It)
+		{
+			if (!PropertyValues.IsValidIndex(It.Value()))
+			{
+				It.RemoveCurrent();
+			}
+		}
+		PropertyValues.Shrink();
+		bLUTContainsInvalidEntries = false;
+	}
+
+	TOptional<FInitialValueIndex> FindPropertyIndex(const FKeyType& InKey)
+	{
+		CleanupStaleEntries();
+
+		const uint16* Index = KeyToPropertyIndex.Find(InKey);
+		return Index ? TOptional<FInitialValueIndex>(FInitialValueIndex{*Index}) : TOptional<FInitialValueIndex>();
+	}
+
+	/** Sparse array containing all cached property values */
+	TSparseArray<PropertyType> PropertyValues;
+	/** LUT from object+property to its index. May contain stale values if bLUTContainsInvalidEntries is true */
+	TMap<FKeyType, uint16> KeyToPropertyIndex;
+	/** When true, KeyToPropertyIndex contains invalid entries which must be purged before use */
+	bool bLUTContainsInvalidEntries = false;
+};
+
+
+} // namespace MovieScene
+} // namespace UE
+
diff --git a/Engine/Source/Runtime/MovieScene/Public/EntitySystem/MovieScenePropertyComponentHandler.h b/Engine/Source/Runtime/MovieScene/Public/EntitySystem/MovieScenePropertyComponentHandler.h
--- a/Engine/Source/Runtime/MovieScene/Public/EntitySystem/MovieScenePropertyComponentHandler.h
+++ b/Engine/Source/Runtime/MovieScene/Public/EntitySystem/MovieScenePropertyComponentHandler.h
@@ -9,9 +9,11 @@
 #include "EntitySystem/MovieSceneDecompositionQuery.h"
 #include "EntitySystem/MovieSceneBlenderSystem.h"
 #include "EntitySystem/MovieScenePreAnimatedPropertyHelper.h"
+#include "EntitySystem/MovieSceneInitialValueCache.h"
 #include "EntitySystem/MovieScenePropertySystemTypes.inl"
-
 #include "EntitySystem/MovieSceneOperationalTypeConversions.h"
+#include "EntitySystem/Interrogation/MovieSceneInterrogationExtension.h"
+#include "EntitySystem/Interrogation/MovieSceneInterrogationLinker.h"
 
 
 
@@ -148,19 +150,237 @@ struct TPropertyComponentHandlerImpl<TIntegerSequence<int, Indices...>, Property
 		.RunInline_PerAllocation(&Linker->EntityManager, SetProperties);
 	}
 
-	virtual void DispatchCacheInitialValueTasks(const FPropertyDefinition& Definition, FSystemTaskPrerequisites& InPrerequisites, FSystemSubsequentTasks& Subsequents, UMovieSceneEntitySystemLinker* Linker) override
+
+	struct FInitialValueProcessor : IInitialValueProcessor
 	{
-		FBuiltInComponentTypes* BuiltInComponents = FBuiltInComponentTypes::Get();
+		TSortedMap<FInterrogationChannel, OperationalType> ValuesByChannel;
 
-		TGetPropertyValues<PropertyType, OperationalType> GetProperties(Definition.CustomPropertyRegistration);
+		FBuiltInComponentTypes* BuiltInComponents;
+		IInterrogationExtension* Interrogation;
+		const FPropertyDefinition* PropertyDefinition;
+		FCustomAccessorView CustomAccessors;
 
-		FEntityTaskBuilder()
-		.Read(BuiltInComponents->BoundObject)
-		.ReadOneOf(BuiltInComponents->CustomPropertyIndex, BuiltInComponents->FastPropertyOffset, BuiltInComponents->SlowProperty)
-		.Write(Definition.InitialValueType.ReinterpretCast<OperationalType>())
-		.FilterAll({ BuiltInComponents->Tags.NeedsLink, Definition.PropertyType })
-		.SetDesiredThread(Linker->EntityManager.GetGatherThread())
-		.RunInline_PerAllocation(&Linker->EntityManager, GetProperties);
+		FEntityAllocationWriteContext WriteContext;
+		TPropertyValueStorage<PropertyType>* CacheStorage;
+
+		FInitialValueProcessor()
+			: WriteContext(FEntityAllocationWriteContext::NewAllocation())
+		{
+			BuiltInComponents = FBuiltInComponentTypes::Get();
+
+			Interrogation = nullptr;
+			CacheStorage = nullptr;
+		}
+
+		virtual void Initialize(UMovieSceneEntitySystemLinker* Linker, const FPropertyDefinition* Definition, FInitialValueCache* InitialValueCache) override
+		{
+			PropertyDefinition = Definition;
+			Interrogation = Linker->FindExtension<IInterrogationExtension>();
+			WriteContext  = FEntityAllocationWriteContext(Linker->EntityManager);
+
+			CustomAccessors = PropertyDefinition->CustomPropertyRegistration->GetAccessors();
+
+			if (InitialValueCache)
+			{
+				CacheStorage = InitialValueCache->GetStorage<PropertyType>(Definition->InitialValueType);
+			}
+		}
+
+		virtual void Process(const FEntityAllocation* Allocation, const FComponentMask& AllocationType) override
+		{
+			if (Interrogation && AllocationType.Contains(BuiltInComponents->Interrogation.OutputKey))
+			{
+				VisitInterrogationAllocation(Allocation);
+			}
+			else if (CacheStorage)
+			{
+				VisitAllocationCached(Allocation);
+			}
+			else
+			{
+				VisitAllocation(Allocation);
+			}
+		}
+
+		virtual void Finalize() override
+		{
+			ValuesByChannel.Empty();
+			Interrogation = nullptr;
+			CacheStorage = nullptr;
+			CustomAccessors = FCustomAccessorView();
+		}
+
+		void VisitAllocation(const FEntityAllocation* Allocation)
+		{
+			const int32 Num = Allocation->Num();
+
+			TComponentWriter<OperationalType> InitialValues = Allocation->WriteComponents(PropertyDefinition->InitialValueType.ReinterpretCast<OperationalType>(), WriteContext);
+			TComponentReader<UObject*>        BoundObjects  = Allocation->ReadComponents(BuiltInComponents->BoundObject);
+
+			if (TOptionalComponentReader<FCustomPropertyIndex> CustomIndices = Allocation->TryReadComponents(BuiltInComponents->CustomPropertyIndex))
+			{
+				const FCustomPropertyIndex* RawIndices = CustomIndices.AsPtr();
+				for (int32 Index = 0; Index < Num; ++Index)
+				{
+					const TCustomPropertyAccessor<PropertyType>& CustomAccessor = static_cast<const TCustomPropertyAccessor<PropertyType>&>(CustomAccessors[RawIndices[Index].Value]);
+
+					PropertyType CurrentValue = CustomAccessor.Functions.Getter(BoundObjects[Index]);
+					ConvertOperationalProperty(CurrentValue, InitialValues[Index]);
+				}
+			}
+
+			else if (TOptionalComponentReader<uint16> FastOffsets = Allocation->TryReadComponents(BuiltInComponents->FastPropertyOffset))
+			{
+				const uint16* RawOffsets = FastOffsets.AsPtr();
+				for (int32 Index = 0; Index < Num; ++Index)
+				{
+					PropertyType CurrentValue = *reinterpret_cast<const PropertyType*>(reinterpret_cast<const uint8*>(static_cast<const void*>(BoundObjects[Index])) + RawOffsets[Index]);
+					ConvertOperationalProperty(CurrentValue, InitialValues[Index]);
+				}
+			}
+
+			else if (TOptionalComponentReader<TSharedPtr<FTrackInstancePropertyBindings>> SlowProperties = Allocation->TryReadComponents(BuiltInComponents->SlowProperty))
+			{
+				const TSharedPtr<FTrackInstancePropertyBindings>* RawProperties = SlowProperties.AsPtr();
+				for (int32 Index = 0; Index < Num; ++Index)
+				{
+					PropertyType CurrentValue = RawProperties[Index]->GetCurrentValue<PropertyType>(*BoundObjects[Index]);
+					ConvertOperationalProperty(CurrentValue, InitialValues[Index]);
+				}
+			}
+		}
+
+		void VisitAllocationCached(const FEntityAllocation* Allocation)
+		{
+			const int32 Num = Allocation->Num();
+
+			TComponentWriter<FInitialValueIndex> InitialValueIndices = Allocation->WriteComponents(BuiltInComponents->InitialValueIndex, WriteContext);
+			TComponentWriter<OperationalType>    InitialValues       = Allocation->WriteComponents(PropertyDefinition->InitialValueType.ReinterpretCast<OperationalType>(), WriteContext);
+			TComponentReader<UObject*>           BoundObjects        = Allocation->ReadComponents(BuiltInComponents->BoundObject);
+
+			if (TOptionalComponentReader<FCustomPropertyIndex> CustomIndices = Allocation->TryReadComponents(BuiltInComponents->CustomPropertyIndex))
+			{
+				const FCustomPropertyIndex* RawIndices = CustomIndices.AsPtr();
+				for (int32 Index = 0; Index < Num; ++Index)
+				{
+					TPair<FInitialValueIndex, PropertyType> CachedValue = CacheStorage->CacheInitialValue(BoundObjects[Index], CustomAccessors, RawIndices[Index]);
+
+					ConvertOperationalProperty(CachedValue.Value, InitialValues[Index]);
+					InitialValueIndices[Index] = CachedValue.Key;
+				}
+			}
+
+			else if (TOptionalComponentReader<uint16> FastOffsets = Allocation->TryReadComponents(BuiltInComponents->FastPropertyOffset))
+			{
+				const uint16* RawOffsets = FastOffsets.AsPtr();
+				for (int32 Index = 0; Index < Num; ++Index)
+				{
+					TPair<FInitialValueIndex, PropertyType> CachedValue = CacheStorage->CacheInitialValue(BoundObjects[Index], RawOffsets[Index]);
+
+					ConvertOperationalProperty(CachedValue.Value, InitialValues[Index]);
+					InitialValueIndices[Index] = CachedValue.Key;
+				}
+			}
+
+			else if (TOptionalComponentReader<TSharedPtr<FTrackInstancePropertyBindings>> SlowProperties = Allocation->TryReadComponents(BuiltInComponents->SlowProperty))
+			{
+				const TSharedPtr<FTrackInstancePropertyBindings>* RawProperties = SlowProperties.AsPtr();
+				for (int32 Index = 0; Index < Num; ++Index)
+				{
+					TPair<FInitialValueIndex, PropertyType> CachedValue = CacheStorage->CacheInitialValue(BoundObjects[Index], RawProperties[Index].Get());
+
+					ConvertOperationalProperty(CachedValue.Value, InitialValues[Index]);
+					InitialValueIndices[Index] = CachedValue.Key;
+				}
+			}
+		}
+
+		void VisitInterrogationAllocation(const FEntityAllocation* Allocation)
+		{
+			const int32 Num = Allocation->Num();
+
+			TComponentWriter<OperationalType>   InitialValues = Allocation->WriteComponents(PropertyDefinition->InitialValueType.ReinterpretCast<OperationalType>(), WriteContext);
+			TComponentReader<FInterrogationKey> OutputKeys    = Allocation->ReadComponents(BuiltInComponents->Interrogation.OutputKey);
+
+			const FSparseInterrogationChannelInfo& SparseChannelInfo = Interrogation->GetSparseChannelInfo();
+
+			for (int32 Index = 0; Index < Num; ++Index)
+			{
+				FInterrogationChannel Channel = OutputKeys[Index].Channel;
+
+				// Did we already cache this value?
+				if (const OperationalType* CachedValue = ValuesByChannel.Find(Channel))
+				{
+					InitialValues[Index] = *CachedValue;
+					continue;
+				}
+
+				const FInterrogationChannelInfo* ChannelInfo = SparseChannelInfo.Find(Channel);
+				UObject* Object = ChannelInfo ? ChannelInfo->WeakObject.Get() : nullptr;
+				if (!ChannelInfo || !Object || ChannelInfo->PropertyBinding.PropertyName.IsNone())
+				{
+					continue;
+				}
+
+				TOptional< FResolvedFastProperty > Property = FPropertyRegistry::ResolveFastProperty(Object, ChannelInfo->PropertyBinding, CustomAccessors);
+
+				// Retrieve a cached value if possible
+				if (CacheStorage)
+				{
+					const PropertyType* CachedValue = nullptr;
+					if (!Property.IsSet())
+					{
+						CachedValue = CacheStorage->FindCachedValue(Object, ChannelInfo->PropertyBinding.PropertyPath);
+					}
+					else if (const FCustomPropertyIndex* CustomIndex = Property->TryGet<FCustomPropertyIndex>())
+					{
+						CachedValue = CacheStorage->FindCachedValue(Object, *CustomIndex);
+					}
+					else
+					{
+						CachedValue = CacheStorage->FindCachedValue(Object, Property->Get<uint16>());
+					}
+					if (CachedValue)
+					{
+						OperationalType ConvertedProperty;
+						ConvertOperationalProperty(*CachedValue, ConvertedProperty);
+
+						InitialValues[Index] = ConvertedProperty;
+						ValuesByChannel.Add(Channel, ConvertedProperty);
+						continue;
+					}
+				}
+
+				// No cached value available, must retrieve it now
+				TOptional<PropertyType> CurrentValue;
+
+				if (!Property.IsSet())
+				{
+					CurrentValue = FTrackInstancePropertyBindings::StaticValue<PropertyType>(Object, ChannelInfo->PropertyBinding.PropertyPath.ToString());
+				}
+				else if (const FCustomPropertyIndex* Custom = Property->TryGet<FCustomPropertyIndex>())
+				{
+					CurrentValue = static_cast<const TCustomPropertyAccessor<PropertyType>&>(CustomAccessors[Custom->Value]).Functions.Getter(Object);
+				}
+				else
+				{
+					const uint16 FastPtrOffset = Property->Get<uint16>();
+					CurrentValue = *reinterpret_cast<const PropertyType*>(reinterpret_cast<const uint8*>(static_cast<const void*>(Object)) + FastPtrOffset);
+				}
+
+				OperationalType NewValue;
+				ConvertOperationalProperty(CurrentValue.GetValue(), NewValue);
+
+				InitialValues[Index] = NewValue;
+				ValuesByChannel.Add(Channel, NewValue);
+			};
+		}
+	};
+
+	virtual IInitialValueProcessor* GetInitialValueProcessor() override
+	{
+		static FInitialValueProcessor Processor;
+		return &Processor;
 	}
 
 	virtual void SaveGlobalPreAnimatedState(const FPropertyDefinition& Definition, UMovieSceneEntitySystemLinker* Linker) override
diff --git a/Engine/Source/Runtime/MovieSceneTracks/Private/EntitySystem/Interrogation/MovieSceneInterrogationLinker.cpp b/Engine/Source/Runtime/MovieSceneTracks/Private/EntitySystem/Interrogation/MovieSceneInterrogationLinker.cpp
--- a/Engine/Source/Runtime/MovieSceneTracks/Private/EntitySystem/Interrogation/MovieSceneInterrogationLinker.cpp
+++ b/Engine/Source/Runtime/MovieSceneTracks/Private/EntitySystem/Interrogation/MovieSceneInterrogationLinker.cpp
@@ -6,6 +6,7 @@
 #include "EntitySystem/MovieSceneEntitySystemRunner.h"
 #include "EntitySystem/MovieSceneEntitySystemTask.h"
 #include "EntitySystem/MovieSceneEntitySystemLinker.h"
+#include "EntitySystem/MovieSceneInitialValueCache.h"
 
 #include "MovieSceneTimeHelpers.h"
 #include "MovieSceneSequence.h"
@@ -64,11 +65,22 @@ struct FSystemInterrogatorEntityTracker
 };
 
 
+TEntitySystemLinkerExtensionID<IInterrogationExtension> IInterrogationExtension::GetExtensionID()
+{
+	static TEntitySystemLinkerExtensionID<IInterrogationExtension> ID = UMovieSceneEntitySystemLinker::RegisterExtension<IInterrogationExtension>();
+	return ID;
+}
+
 FSystemInterrogator::FSystemInterrogator()
 {
+	InitialValueCache = FInitialValueCache::GetGlobalInitialValues();
+
 	Linker = NewObject<UMovieSceneEntitySystemLinker>(GetTransientPackage());
 	Linker->SetSystemContext(EEntitySystemContext::Interrogation);
 
+	Linker->AddExtension(IInterrogationExtension::GetExtensionID(), static_cast<IInterrogationExtension*>(this));
+	Linker->AddExtension(InitialValueCache.Get());
+
 	// Always add a bit for the default channel
 	ImportedChannelBits.Add(false);
 }
@@ -87,7 +99,7 @@ FString FSystemInterrogator::GetReferencerName() const
 	return TEXT("FSystemInterrogator");
 }
 
-FInterrogationChannel FSystemInterrogator::AllocateChannel(FInterrogationChannel ParentChannel)
+FInterrogationChannel FSystemInterrogator::AllocateChannel(FInterrogationChannel ParentChannel, const FMovieScenePropertyBinding& PropertyBinding)
 {
 	if (!ensureMsgf(ImportedChannelBits.Num() < MAX_int32, TEXT("Reached the maximum available number of interrogation channels")))
 	{
@@ -97,20 +109,22 @@ FInterrogationChannel FSystemInterrogator::AllocateChannel(FInterrogationChannel
 	FInterrogationChannel Channel = FInterrogationChannel::FromIndex(ImportedChannelBits.Num());
 	ImportedChannelBits.Add(false);
 
-	if (ParentChannel)
+	if (ParentChannel || !PropertyBinding.PropertyPath.IsNone())
 	{
-		SparseChannelInfo.FindOrAdd(Channel).ParentChannel = ParentChannel;
+		FInterrogationChannelInfo& ChannelInfo = SparseChannelInfo.Get(Channel);
+		ChannelInfo.ParentChannel   = ParentChannel;
+		ChannelInfo.PropertyBinding = PropertyBinding;
 	}
 	return Channel;
 }
 
-FInterrogationChannel FSystemInterrogator::AllocateChannel(UObject* Object, FInterrogationChannel ParentChannel)
+FInterrogationChannel FSystemInterrogator::AllocateChannel(UObject* Object, FInterrogationChannel ParentChannel, const FMovieScenePropertyBinding& PropertyBinding)
 {
-	FInterrogationChannel NewChannel = AllocateChannel(ParentChannel);
+	FInterrogationChannel NewChannel = AllocateChannel(ParentChannel, PropertyBinding);
 	if (NewChannel)
 	{
 		ObjectToChannel.Add(Object, NewChannel);
-		SparseChannelInfo.FindOrAdd(NewChannel).WeakObject = Object;
+		SparseChannelInfo.Get(NewChannel).WeakObject = Object;
 	}
 	return NewChannel;
 }
@@ -267,6 +281,7 @@ void FSystemInterrogator::Update()
 	.Write(FBuiltInComponentTypes::Get()->EvalTime)
 	.FilterNone({ FBuiltInComponentTypes::Get()->Tags.FixedTime })
 	.Iterate_PerEntity(&Linker->EntityManager, [this](const FInterrogationKey& InterrogationKey, FFrameTime& OutEvalTime) { OutEvalTime = this->Interrogations[InterrogationKey.InterrogationIndex].Time; });
+
 	FMovieSceneEntitySystemRunner Runner;
 	Runner.AttachToLinker(Linker);
 	Runner.Flush();
@@ -315,7 +330,7 @@ FInterrogationChannel FSystemInterrogator::ImportLocalTransforms(USceneComponent
 	FInterrogationChannel Channel = ObjectToChannel.FindRef(SceneComponent);
 	if (!Channel.IsValid())
 	{
-		Channel = AllocateChannel(SceneComponent, ParentChannel);
+		Channel = AllocateChannel(SceneComponent, ParentChannel, FMovieScenePropertyBinding("Transform", TEXT("Transform")));
 	}
 
 	if (!Channel.IsValid())
@@ -482,7 +497,7 @@ void FSystemInterrogator::QueryWorldSpaceTransforms(const TBitArray<>& ChannelsT
 		int32 NumChildren = 0;
 
 		FInterrogationChannel Channel = FInterrogationChannel::FromIndex(ChannelBit.GetIndex());
-		if (ImportedChannelBits[Channel.AsIndex()] == false && FindObjectFromChannel(Channel) == nullptr)
+		if (ImportedChannelBits[Channel.AsIndex()] == false && SparseChannelInfo.FindObject(Channel) == nullptr)
 		{
 			continue;
 		}
@@ -500,8 +515,8 @@ void FSystemInterrogator::QueryWorldSpaceTransforms(const TBitArray<>& ChannelsT
 			// Update the component's current value
 			FTransform BaseValue = FTransform::Identity;
 
-			FInterrogationChannelInfo ChannelInfo = SparseChannelInfo.FindRef(Channel);
-			if (USceneComponent* SceneComponent = Cast<USceneComponent>(ChannelInfo.WeakObject.Get()))
+			const FInterrogationChannelInfo* ChannelInfo = SparseChannelInfo.Find(Channel);
+			if (USceneComponent* SceneComponent = Cast<USceneComponent>(ChannelInfo ? ChannelInfo->WeakObject.Get() : nullptr))
 			{
 				BaseValue = SceneComponent->GetRelativeTransform();
 			}
@@ -519,7 +534,7 @@ void FSystemInterrogator::QueryWorldSpaceTransforms(const TBitArray<>& ChannelsT
 				VariableTransformsByChannel[ThisChannelIndex].SetNum(Interrogations.Num());
 			}
 
-			Channel = ChannelInfo.ParentChannel;
+			Channel = ChannelInfo ? ChannelInfo->ParentChannel : FInterrogationChannel::Invalid();
 			++NumChildren;
 		}
 	}
@@ -634,7 +649,8 @@ void FSystemInterrogator::QueryWorldSpaceTransforms(const TBitArray<>& ChannelsT
 
 			const int32 ChannelIndex = Channel.AsIndex();
 
-			FInterrogationChannel Parent = SparseChannelInfo.FindRef(Channel).ParentChannel;
+
+			FInterrogationChannel Parent = SparseChannelInfo.FindParent(Channel);
 			FTransform OffsetFromAnimatedParent = FTransform::Identity;
 
 			bool bProcessedTransforms = false;
@@ -670,7 +686,7 @@ void FSystemInterrogator::QueryWorldSpaceTransforms(const TBitArray<>& ChannelsT
 				}
 
 				OffsetFromAnimatedParent = BaseValues[Parent.AsIndex()] * OffsetFromAnimatedParent;
-				Parent = SparseChannelInfo.FindRef(Parent).ParentChannel;
+				Parent = SparseChannelInfo.FindParent(Parent);
 			}
 
 			if (!bProcessedTransforms)
@@ -708,15 +724,14 @@ void FSystemInterrogator::QueryLocalSpaceTransforms(const TBitArray<>& ChannelsT
 	for (TConstSetBitIterator<> ChannelBit(ChannelsToQuery); ChannelBit; ++ChannelBit)
 	{
 		FInterrogationChannel Channel = FInterrogationChannel::FromIndex(ChannelBit.GetIndex());
-		if (ImportedChannelBits[Channel.AsIndex()] == false && FindObjectFromChannel(Channel) == nullptr)
+		if (ImportedChannelBits[Channel.AsIndex()] == false && SparseChannelInfo.FindObject(Channel) == nullptr)
 		{
 			continue;
 		}
 
 		FIntermediate3DTransform BaseValue;
 
-		FInterrogationChannelInfo ChannelInfo = SparseChannelInfo.FindRef(Channel);
-		if (USceneComponent* SceneComponent = Cast<USceneComponent>(ChannelInfo.WeakObject.Get()))
+		if (USceneComponent* SceneComponent = Cast<USceneComponent>(SparseChannelInfo.FindObject(Channel)))
 		{
 			BaseValue = FIntermediate3DTransform(SceneComponent->GetRelativeLocation(), SceneComponent->GetRelativeRotation(), SceneComponent->GetRelativeScale3D());
 		}
diff --git a/Engine/Source/Runtime/MovieSceneTracks/Private/Systems/MovieSceneInitialValueSystem.cpp b/Engine/Source/Runtime/MovieSceneTracks/Private/Systems/MovieSceneInitialValueSystem.cpp
new file mode 100644
--- /dev/null
+++ b/Engine/Source/Runtime/MovieSceneTracks/Private/Systems/MovieSceneInitialValueSystem.cpp
@@ -0,0 +1,172 @@
+// Copyright Epic Games, Inc. All Rights Reserved.
+
+#include "Systems/MovieSceneInitialValueSystem.h"
+#include "EntitySystem/MovieSceneEntitySystemLinker.h"
+#include "EntitySystem/MovieScenePropertyRegistry.h"
+#include "EntitySystem/MovieSceneEntityMutations.h"
+#include "EntitySystem/BuiltInComponentTypes.h"
+#include "EntitySystem/MovieSceneInitialValueCache.h"
+#include "EntitySystem/Interrogation/MovieSceneInterrogationExtension.h"
+
+#include "Systems/MovieScenePropertyInstantiator.h"
+#include "EntitySystem/Interrogation/MovieSceneInterrogatedPropertyInstantiator.h"
+
+namespace UE
+{
+namespace MovieScene
+{
+
+struct FInitialValueMutation : IMovieSceneEntityMutation
+{
+	TSortedMap<FComponentTypeID, IInitialValueProcessor*> PropertyTypeToProcessor;
+	FInitialValueCache* InitialValueCache;
+	FBuiltInComponentTypes* BuiltInComponents;
+	FComponentMask AnyInitialValue;
+
+	FInitialValueMutation(UMovieSceneEntitySystemLinker* Linker)
+	{
+		BuiltInComponents = FBuiltInComponentTypes::Get();
+
+		InitialValueCache = Linker->FindExtension<FInitialValueCache>();
+		for (const FPropertyDefinition& Definition : BuiltInComponents->PropertyRegistry.GetProperties())
+		{
+			IInitialValueProcessor* Processor = Definition.Handler->GetInitialValueProcessor();
+			if (Processor && Definition.InitialValueType && Linker->EntityManager.ContainsComponent(Definition.PropertyType))
+			{
+				Processor->Initialize(Linker, &Definition, InitialValueCache);
+
+				AnyInitialValue.Set(Definition.InitialValueType);
+				PropertyTypeToProcessor.Add(Definition.InitialValueType, Processor);
+			}
+		}
+	}
+
+	~FInitialValueMutation()
+	{
+		for (TPair<FComponentTypeID, IInitialValueProcessor*> Pair : PropertyTypeToProcessor)
+		{
+			Pair.Value->Finalize();
+		}
+	}
+
+	bool IsCached() const
+	{
+		return InitialValueCache != nullptr;
+	}
+
+	virtual void CreateMutation(FEntityManager* EntityManager, FComponentMask* InOutEntityComponentTypes) const override
+	{
+		InOutEntityComponentTypes->Set(BuiltInComponents->InitialValueIndex);
+	}
+
+	virtual void InitializeAllocation(FEntityAllocation* Allocation, const FComponentMask& AllocationType) const
+	{
+		FComponentTypeID InitialValueType = FComponentMask::BitwiseAND(AllocationType, AnyInitialValue, EBitwiseOperatorFlags::MinSize).First();
+
+		IInitialValueProcessor* Processor = PropertyTypeToProcessor.FindRef(InitialValueType);
+		if (ensure(Processor))
+		{
+			Processor->Process(Allocation, AllocationType);
+		}
+	}
+};
+
+} // namespace MovieScene
+} // namespace UE
+
+UMovieSceneInitialValueSystem::UMovieSceneInitialValueSystem(const FObjectInitializer& ObjInit)
+	: Super(ObjInit)
+{
+	if (HasAnyFlags(RF_ClassDefaultObject))
+	{
+		DefineImplicitPrerequisite(UMovieScenePropertyInstantiatorSystem::StaticClass(), StaticClass());
+		DefineImplicitPrerequisite(UMovieSceneInterrogatedPropertyInstantiatorSystem::StaticClass(), StaticClass());
+	}
+}
+
+bool UMovieSceneInitialValueSystem::IsRelevantImpl(UMovieSceneEntitySystemLinker* InLinker) const
+{
+	using namespace UE::MovieScene;
+
+	FBuiltInComponentTypes* BuiltInComponents = FBuiltInComponentTypes::Get();
+	for (const FPropertyDefinition& Definition : BuiltInComponents->PropertyRegistry.GetProperties())
+	{
+		if (InLinker->EntityManager.ContainsComponent(Definition.InitialValueType))
+		{
+			return true;
+		}
+	}
+	return false;
+}
+
+void UMovieSceneInitialValueSystem::OnLink()
+{
+
+}
+
+void UMovieSceneInitialValueSystem::OnUnlink()
+{
+
+}
+
+void UMovieSceneInitialValueSystem::OnRun(FSystemTaskPrerequisites& InPrerequisites, FSystemSubsequentTasks& Subsequents)
+{
+	using namespace UE::MovieScene;
+
+	FBuiltInComponentTypes* BuiltInComponents = FBuiltInComponentTypes::Get();
+
+	FInitialValueMutation Mutation(Linker);
+	if (Mutation.IsCached() && Linker->FindExtension<IInterrogationExtension>() == nullptr)
+	{
+		// When there is an initial value cache extension, we mutate anything with an initial value component on it by
+		// adding an additional index that refers to its cached position. This ensures we are able to clean up the cache
+		// easily.
+		{
+			FEntityComponentFilter Filter;
+			Filter.Any(Mutation.AnyInitialValue);
+			Filter.All({ BuiltInComponents->Tags.NeedsLink });
+			Filter.None({ BuiltInComponents->InitialValueIndex });
+
+			Linker->EntityManager.MutateAll(Filter, Mutation);
+		}
+
+		// Clean up any stale cache entries
+		{
+			FEntityComponentFilter Filter;
+			Filter.Any(Mutation.AnyInitialValue);
+			Filter.All({ BuiltInComponents->InitialValueIndex, BuiltInComponents->Tags.NeedsUnlink });
+
+			for (FEntityAllocationIteratorItem Item : Linker->EntityManager.Iterate(&Filter))
+			{
+				const FEntityAllocation* Allocation     = Item.GetAllocation();
+				FComponentMask           AllocationType = Item.GetAllocationType();
+
+				FComponentTypeID InitialValueType = FComponentMask::BitwiseAND(AllocationType, Mutation.AnyInitialValue, EBitwiseOperatorFlags::MinSize).First();
+
+				TComponentReader<FInitialValueIndex> Indices = Allocation->ReadComponents(BuiltInComponents->InitialValueIndex);
+				Mutation.InitialValueCache->Reset(InitialValueType, Indices.AsArray(Allocation->Num()));
+			}
+		}
+	}
+	else
+	{
+		// When there is no caching extension, or we are interrogating we simply initialize any initial values directly without going through the cache
+		FEntityComponentFilter Filter;
+		Filter.Any(Mutation.AnyInitialValue);
+		Filter.All({ BuiltInComponents->Tags.NeedsLink });
+
+		for (FEntityAllocationIteratorItem Item : Linker->EntityManager.Iterate(&Filter))
+		{
+			const FEntityAllocation* Allocation     = Item.GetAllocation();
+			FComponentMask           AllocationType = Item.GetAllocationType();
+
+			FComponentTypeID InitialValueType = FComponentMask::BitwiseAND(AllocationType, Mutation.AnyInitialValue, EBitwiseOperatorFlags::MinSize).First();
+
+			IInitialValueProcessor* Processor = Mutation.PropertyTypeToProcessor.FindRef(InitialValueType);
+			if (ensure(Processor))
+			{
+				Processor->Process(Allocation, AllocationType);
+			}
+		}
+	}
+}
\ No newline at end of file
diff --git a/Engine/Source/Runtime/MovieSceneTracks/Private/Systems/MovieScenePropertyInstantiator.cpp b/Engine/Source/Runtime/MovieSceneTracks/Private/Systems/MovieScenePropertyInstantiator.cpp
--- a/Engine/Source/Runtime/MovieSceneTracks/Private/Systems/MovieScenePropertyInstantiator.cpp
+++ b/Engine/Source/Runtime/MovieSceneTracks/Private/Systems/MovieScenePropertyInstantiator.cpp
@@ -82,12 +82,6 @@ void UMovieScenePropertyInstantiatorSystem::OnRun(FSystemTaskPrerequisites& InPr
 		ProcessInvalidatedProperties(InvalidatedProperties);
 	}
 
-	// Kick off initial value gather task immediately
-	if (InitialValueStateTasks.Find(true) != INDEX_NONE)
-	{
-		AssignInitialValues(InPrerequisites, Subsequents);
-	}
-
 	if (CachePreAnimatedStateTasks.Find(true) != INDEX_NONE)
 	{
 		UMovieSceneCachePreAnimatedStateSystem* PreAnimatedState = Linker->LinkSystem<UMovieSceneCachePreAnimatedStateSystem>();
@@ -442,8 +436,6 @@ void UMovieScenePropertyInstantiatorSystem::InitializeBlendPath(const FPropertyP
 	if (!bWasAlreadyBlended)
 	{
 		NewMask.Set(Params.PropertyDefinition->InitialValueType);
-		InitialValueStateTasks.PadToNum(Params.PropertyInfo->PropertyDefinitionIndex+1, false);
-		InitialValueStateTasks[Params.PropertyInfo->PropertyDefinitionIndex] = true;
 
 		for (int32 Index = 0; Index < Composites.Num(); ++Index)
 		{
@@ -583,21 +575,6 @@ UE::MovieScene::FPropertyRecomposerPropertyInfo UMovieScenePropertyInstantiatorS
 	return FPropertyRecomposerPropertyInfo::Invalid();
 }
 
-void UMovieScenePropertyInstantiatorSystem::AssignInitialValues(FSystemTaskPrerequisites& InPrerequisites, FSystemSubsequentTasks& Subsequents)
-{
-	using namespace UE::MovieScene;
-
-	for (TConstSetBitIterator<> TypesToCache(InitialValueStateTasks); TypesToCache; ++TypesToCache)
-	{
-		FCompositePropertyTypeID PropertyID = FCompositePropertyTypeID::FromIndex(TypesToCache.GetIndex());
-
-		const FPropertyDefinition& Definition = BuiltInComponents->PropertyRegistry.GetDefinition(PropertyID);
-		Definition.Handler->DispatchCacheInitialValueTasks(Definition, InPrerequisites, Subsequents, Linker);
-	}
-
-	InitialValueStateTasks.Empty();
-}
-
 void UMovieScenePropertyInstantiatorSystem::SavePreAnimatedState(FSystemTaskPrerequisites& InPrerequisites, FSystemSubsequentTasks& Subsequents)
 {
 	using namespace UE::MovieScene;
diff --git a/Engine/Source/Runtime/MovieSceneTracks/Public/EntitySystem/Interrogation/MovieSceneInterrogationExtension.h b/Engine/Source/Runtime/MovieSceneTracks/Public/EntitySystem/Interrogation/MovieSceneInterrogationExtension.h
new file mode 100644
--- /dev/null
+++ b/Engine/Source/Runtime/MovieSceneTracks/Public/EntitySystem/Interrogation/MovieSceneInterrogationExtension.h
@@ -0,0 +1,84 @@
+// Copyright Epic Games, Inc. All Rights Reserved.
+
+#pragma once
+
+#include "UObject/WeakObjectPtrTemplates.h"
+#include "UObject/NameTypes.h"
+#include "Containers/SortedMap.h"
+#include "EntitySystem/MovieSceneEntitySystemTypes.h"
+#include "EntitySystem/MovieSceneEntitySystemLinkerExtension.h"
+#include "EntitySystem/MovieScenePropertyBinding.h"
+
+class UObject;
+
+namespace UE
+{
+namespace MovieScene
+{
+
+struct FInterrogationChannel;
+struct FInterrogationChannelInfo;
+
+struct FInterrogationChannelInfo
+{
+	/** The object that relates to the channel */
+	TWeakObjectPtr<UObject> WeakObject;
+
+	/** (Optional) property path for this channel */
+	FMovieScenePropertyBinding PropertyBinding;
+
+	/** The channel's hierarchical 'parent' - generally only used for transforms */
+	FInterrogationChannel ParentChannel;
+};
+
+
+struct FInterrogationUpdateParams
+{
+
+	/** The channel's hierarchical 'parent' - generally only used for transforms */
+	FInterrogationChannel ParentChannel;
+};
+
+struct FSparseInterrogationChannelInfo
+{
+	FInterrogationChannelInfo& Get(FInterrogationChannel Channel) { return ChannelInfo.FindOrAdd(Channel); }
+	const FInterrogationChannelInfo& Get(FInterrogationChannel Channel) const { return ChannelInfo.FindChecked(Channel); }
+
+	FInterrogationChannelInfo* Find(FInterrogationChannel Channel) { return ChannelInfo.Find(Channel); }
+	const FInterrogationChannelInfo* Find(FInterrogationChannel Channel) const { return ChannelInfo.Find(Channel); }
+
+	UObject* FindObject(FInterrogationChannel Channel) const
+	{
+		return ChannelInfo.FindRef(Channel).WeakObject.Get();
+	}
+
+	FInterrogationChannel FindParent(FInterrogationChannel Channel) const
+	{
+		return ChannelInfo.FindRef(Channel).ParentChannel;
+	}
+
+	void Empty()
+	{
+		ChannelInfo.Empty();
+	}
+
+private:
+
+	TSortedMap<FInterrogationChannel, FInterrogationChannelInfo> ChannelInfo;
+};
+
+struct IInterrogationExtension
+{
+public:
+	MOVIESCENETRACKS_API static TEntitySystemLinkerExtensionID<IInterrogationExtension> GetExtensionID();
+
+	virtual const FSparseInterrogationChannelInfo& GetSparseChannelInfo() const = 0;
+
+protected:
+
+	virtual ~IInterrogationExtension() {}
+};
+
+
+} // namespace MovieScene
+} // namespace UE
diff --git a/Engine/Source/Runtime/MovieSceneTracks/Public/EntitySystem/Interrogation/MovieSceneInterrogationLinker.h b/Engine/Source/Runtime/MovieSceneTracks/Public/EntitySystem/Interrogation/MovieSceneInterrogationLinker.h
--- a/Engine/Source/Runtime/MovieSceneTracks/Public/EntitySystem/Interrogation/MovieSceneInterrogationLinker.h
+++ b/Engine/Source/Runtime/MovieSceneTracks/Public/EntitySystem/Interrogation/MovieSceneInterrogationLinker.h
@@ -5,10 +5,14 @@
 #include "CoreTypes.h"
 #include "UObject/GCObject.h"
 #include "Containers/SortedMap.h"
-#include "EntitySystem/MovieSceneEntitySystemTypes.h"
 #include "MovieSceneTracksComponentTypes.h"
 #include "Evaluation/MovieSceneEvaluationField.h"
 
+#include "EntitySystem/MovieSceneEntitySystemTypes.h"
+#include "EntitySystem/MovieSceneEntitySystemLinkerExtension.h"
+
+#include "EntitySystem/Interrogation/MovieSceneInterrogationExtension.h"
+
 struct FGuid;
 struct FMovieSceneBinding;
 
@@ -21,6 +25,7 @@ namespace UE
 namespace MovieScene
 {
 
+struct FInitialValueCache;
 struct FSystemInterrogatorEntityTracker;
 
 struct FInterrogationParams
@@ -35,18 +40,6 @@ struct FInterrogationParams
 	{}
 };
 
-
-
-struct FInterrogationChannelInfo
-{
-	/** The object that relates to the channel */
-	TWeakObjectPtr<UObject> WeakObject;
-
-	/** The channel's hierarchical 'parent' - generally only used for transforms */
-	FInterrogationChannel ParentChannel;
-};
-
-
 /**
  * A class specialized for interrogating Sequencer entity data without applying any state to objects.
  * Currently only tracks within the same time-base are supported.
@@ -72,7 +65,7 @@ struct FInterrogationChannelInfo
  *        MySystem->Interrogate(FInterogationKey::Default(), OutData);
  *    }
  */
-class MOVIESCENETRACKS_API FSystemInterrogator : FGCObject
+class MOVIESCENETRACKS_API FSystemInterrogator : FGCObject, IInterrogationExtension
 {
 public:
 	FSystemInterrogator();
@@ -84,13 +77,13 @@ public:
 	 * @param ParentChannel     The channel that should be considered this channel's parent, or FInterrogationChannel::Invalid if there is none
 	 * @return A new interrogation channel
 	 */
-	FInterrogationChannel AllocateChannel(FInterrogationChannel ParentChannel);
+	FInterrogationChannel AllocateChannel(FInterrogationChannel ParentChannel, const FMovieScenePropertyBinding& PropertyBinding);
 
 
 	/**
 	 * Allocate a new interrogation channel that relates to a specific object
 	 */
-	FInterrogationChannel AllocateChannel(UObject* Object, FInterrogationChannel ParentChannel);
+	FInterrogationChannel AllocateChannel(UObject* Object, FInterrogationChannel ParentChannel, const FMovieScenePropertyBinding& PropertyBinding);
 
 
 	/**
@@ -290,12 +283,9 @@ public:
 
 public:
 
-	/**
-	 * Find the object that is allocated within a given channel
-	 */
-	UObject* FindObjectFromChannel(FInterrogationChannel Channel) const
+	const FSparseInterrogationChannelInfo& GetSparseChannelInfo() const override
 	{
-		return SparseChannelInfo.FindRef(Channel).WeakObject.Get();
+		return SparseChannelInfo;
 	}
 
 private:
@@ -335,7 +325,7 @@ protected:
 	TMap<UObject*, FInterrogationChannel> ObjectToChannel;
 
 	/** Array of information pertaining to a given channel */
-	TSortedMap<FInterrogationChannel, FInterrogationChannelInfo> SparseChannelInfo;
+	FSparseInterrogationChannelInfo SparseChannelInfo;
 
 	/** BitArray containing set bits for any channel that has data associated with it. The number of bits (0 or 1) in this array defines how many channels are allocated */
 	TBitArray<> ImportedChannelBits;
@@ -348,6 +338,8 @@ protected:
 
 	/** The linker we own */
 	UMovieSceneEntitySystemLinker* Linker;
+
+	TSharedPtr<FInitialValueCache> InitialValueCache;
 };
 
 
diff --git a/Engine/Source/Runtime/MovieSceneTracks/Public/Systems/MovieSceneInitialValueSystem.h b/Engine/Source/Runtime/MovieSceneTracks/Public/Systems/MovieSceneInitialValueSystem.h
new file mode 100644
--- /dev/null
+++ b/Engine/Source/Runtime/MovieSceneTracks/Public/Systems/MovieSceneInitialValueSystem.h
@@ -0,0 +1,27 @@
+// Copyright Epic Games, Inc. All Rights Reserved.
+
+#include "EntitySystem/MovieSceneEntityInstantiatorSystem.h"
+#include "MovieSceneInitialValueSystem.generated.h"
+
+
+/**
+ * System responsible for initializing initial values for all property types
+ * Will handle the presence of an FInitialValueCache extension on the linker
+ */
+UCLASS()
+class MOVIESCENETRACKS_API UMovieSceneInitialValueSystem
+	: public UMovieSceneEntityInstantiatorSystem
+{
+public:
+
+	GENERATED_BODY()
+
+	UMovieSceneInitialValueSystem(const FObjectInitializer& ObjInit);
+
+private:
+
+	virtual void OnRun(FSystemTaskPrerequisites& InPrerequisites, FSystemSubsequentTasks& Subsequents) override;
+	virtual bool IsRelevantImpl(UMovieSceneEntitySystemLinker* InLinker) const override;
+	virtual void OnLink() override;
+	virtual void OnUnlink() override;
+};
\ No newline at end of file
diff --git a/Engine/Source/Runtime/MovieSceneTracks/Public/Systems/MovieScenePropertyInstantiator.h b/Engine/Source/Runtime/MovieSceneTracks/Public/Systems/MovieScenePropertyInstantiator.h
--- a/Engine/Source/Runtime/MovieSceneTracks/Public/Systems/MovieScenePropertyInstantiator.h
+++ b/Engine/Source/Runtime/MovieSceneTracks/Public/Systems/MovieScenePropertyInstantiator.h
@@ -146,7 +146,6 @@ private:
 
 	void AssignPreAnimatedValues(FSystemTaskPrerequisites& InPrerequisites, FSystemSubsequentTasks& Subsequents);
 	void RestorePreAnimatedValues(FSystemTaskPrerequisites& InPrerequisites, FSystemSubsequentTasks& Subsequents);
-	void AssignInitialValues(FSystemTaskPrerequisites& InPrerequisites, FSystemSubsequentTasks& Subsequents);
 
 	void CleanTaggedGarbage(UMovieSceneEntitySystemLinker*);
 
@@ -170,7 +169,6 @@ private:
 	TBitArray<> SaveGlobalStateTasks;
 	TBitArray<> CachePreAnimatedStateTasks;
 	TBitArray<> RestorePreAnimatedStateTasks;
-	TBitArray<> InitialValueStateTasks;
 
 	UE::MovieScene::FBuiltInComponentTypes* BuiltInComponents;
 	 
