Sequencer: The actor class is now updated when replacing a possessable binding

diff --git a/Engine/Source/Editor/Sequencer/Private/Sequencer.cpp b/Engine/Source/Editor/Sequencer/Private/Sequencer.cpp
--- a/Engine/Source/Editor/Sequencer/Private/Sequencer.cpp
+++ b/Engine/Source/Editor/Sequencer/Private/Sequencer.cpp
@@ -7855,6 +7855,8 @@ void FSequencer::AddActorsToBinding(FGuid InObjectBinding, const TArray<AActor*>
 			{
 				Possessable->SetName(InActors[0]->GetActorLabel());
 			}
+
+			Possessable->SetPossessedObjectClass(ActorClass);
 		}
 	}
 
diff --git a/Engine/Source/Runtime/MovieScene/Public/MovieScenePossessable.h b/Engine/Source/Runtime/MovieScene/Public/MovieScenePossessable.h
--- a/Engine/Source/Runtime/MovieScene/Public/MovieScenePossessable.h
+++ b/Engine/Source/Runtime/MovieScene/Public/MovieScenePossessable.h
@@ -88,6 +88,17 @@ public:
 		return PossessedObjectClass;
 	}
 
+	/**
+	 * Set the class of the possessed object.
+	 *
+	 * @return Object class.
+	 * @see GetGuid, GetName
+	 */
+	void SetPossessedObjectClass(UClass* InClass)
+	{
+		PossessedObjectClass = InClass;
+	}
+
 	/**
 	 * Get the guid of this possessable's parent, if applicable
 	 *
