Handle nullptr being passed to FMovieSceneSpawnable::SetObjectTemplate

[Copied from 15761629]

diff --git a/Engine/Source/Runtime/MovieScene/Public/MovieSceneSpawnable.h b/Engine/Source/Runtime/MovieScene/Public/MovieSceneSpawnable.h
--- a/Engine/Source/Runtime/MovieScene/Public/MovieSceneSpawnable.h
+++ b/Engine/Source/Runtime/MovieScene/Public/MovieSceneSpawnable.h
@@ -104,7 +104,7 @@ public:
 	 */
 	void SetObjectTemplate(UObject* InObjectTemplate)
 	{
-		checkf(!InObjectTemplate->HasAnyFlags(RF_ClassDefaultObject), TEXT("Setting CDOs as object templates is not supported. Please use the class directly."));
+		checkf(InObjectTemplate == nullptr || !InObjectTemplate->HasAnyFlags(RF_ClassDefaultObject), TEXT("Setting CDOs as object templates is not supported. Please use the class directly."));
 		ObjectTemplate = InObjectTemplate;
 	}
 
