Sequencer: Put in UE::MovieScene::Test namespace

diff --git a/Engine/Source/Runtime/MovieScene/Private/Tests/MovieSceneBlendingTests.cpp b/Engine/Source/Runtime/MovieScene/Private/Tests/MovieSceneBlendingTests.cpp
--- a/Engine/Source/Runtime/MovieScene/Private/Tests/MovieSceneBlendingTests.cpp
+++ b/Engine/Source/Runtime/MovieScene/Private/Tests/MovieSceneBlendingTests.cpp
@@ -8,8 +8,19 @@
 
 #define LOCTEXT_NAMESPACE "MovieSceneBlendingTests"
 
-static const int32 StartingValue = 0xefefefef;
-static int32 TestValue = StartingValue;
+namespace UE
+{
+namespace MovieScene
+{
+namespace Test
+{
+
+static const int32 GBlendingStartingValue = 0xefefefef;
+static int32 GBlendingTestValue = GBlendingStartingValue;
+
+} // namespace Test
+} // namespace MovieScene
+} // namespace UE
 
 struct FInt32Actuator : TMovieSceneBlendingActuator<int32>
 {
@@ -23,13 +34,17 @@ struct FInt32Actuator : TMovieSceneBlendingActuator<int32>
 
 	virtual void Actuate(UObject* InObject, TCallTraits<int32>::ParamType InValue, const TBlendableTokenStack<int32>& OriginalStack, const FMovieSceneContext& Context, FPersistentEvaluationData& PersistentData, IMovieScenePlayer& Player)
 	{
+		using namespace UE::MovieScene::Test;
+		
 		ensure(!InObject);
-		TestValue = InValue;
+		GBlendingTestValue = InValue;
 	}
 
 	virtual int32 RetrieveCurrentValue(UObject* InObject, IMovieScenePlayer* Player) const
 	{
-		return TestValue;
+		using namespace UE::MovieScene::Test;
+
+		return GBlendingTestValue;
 	}
 };
 
@@ -51,6 +66,8 @@ struct FNullPlayer : IMovieScenePlayer
 IMPLEMENT_SIMPLE_AUTOMATION_TEST(FMovieSceneBlendingTest, "System.Engine.Sequencer.Blending.Basic", EAutomationTestFlags::EditorContext | EAutomationTestFlags::EngineFilter)
 bool FMovieSceneBlendingTest::RunTest(const FString& Parameters)
 {
+	using namespace UE::MovieScene::Test;
+
 	const FMovieSceneBlendingActuatorID ID = FInt32Actuator::GetID();
 	const FMovieSceneEvaluationScope Scope(FMovieSceneEvaluationKey(), EMovieSceneCompletionMode::KeepState);
 	const FMovieSceneContext Context(FMovieSceneEvaluationRange(0, FFrameRate()));
@@ -70,29 +87,29 @@ bool FMovieSceneBlendingTest::RunTest(const FString& Parameters)
 		Accumulator.Apply(Context, PersistentDataProxy, Player);
 
 		int32 Expected = 7;
-		if (TestValue != Expected)
+		if (GBlendingTestValue != Expected)
 		{
-			AddError(FString::Printf(TEXT("Expected result 1 to be %d, actual %d."), TestValue, Expected));
+			AddError(FString::Printf(TEXT("Expected result 1 to be %d, actual %d."), GBlendingTestValue, Expected));
 		}
 	}
 
-	TestValue = StartingValue;
+	GBlendingTestValue = GBlendingStartingValue;
 
 	{
-		// Result should be StartingValue + 500 + 10
+		// Result should be GBlendingStartingValue + 500 + 10
 		Accumulator.BlendToken<int32>(ID, Scope, Context, 10, EMovieSceneBlendType::Additive, 1.f);
 		Accumulator.BlendToken<int32>(ID, Scope, Context, 500, EMovieSceneBlendType::Relative, 1.f);
 
 		Accumulator.Apply(Context, PersistentDataProxy, Player);
 
-		int32 Expected = StartingValue + 510;
-		if (TestValue != Expected)
+		int32 Expected = GBlendingStartingValue + 510;
+		if (GBlendingTestValue != Expected)
 		{
-			AddError(FString::Printf(TEXT("Expected result 2 to be %d, actual %d."), TestValue, Expected));
+			AddError(FString::Printf(TEXT("Expected result 2 to be %d, actual %d."), GBlendingTestValue, Expected));
 		}
 	}
 
-	TestValue = StartingValue;
+	GBlendingTestValue = GBlendingStartingValue;
 
 	{
 		// Result should be 85
@@ -104,13 +121,13 @@ bool FMovieSceneBlendingTest::RunTest(const FString& Parameters)
 		Accumulator.Apply(Context, PersistentDataProxy, Player);
 
 		int32 Expected = 85 / 4;
-		if (TestValue != Expected)
+		if (GBlendingTestValue != Expected)
 		{
-			AddError(FString::Printf(TEXT("Expected result 3 to be %d, actual %d."), TestValue, Expected));
+			AddError(FString::Printf(TEXT("Expected result 3 to be %d, actual %d."), GBlendingTestValue, Expected));
 		}
 	}
 
-	TestValue = StartingValue;
+	GBlendingTestValue = GBlendingStartingValue;
 
 	
 	{
@@ -122,9 +139,9 @@ bool FMovieSceneBlendingTest::RunTest(const FString& Parameters)
 		Accumulator.Apply(Context, PersistentDataProxy, Player);
 
 		int32 Expected = 10770075;
-		if (TestValue != Expected)
+		if (GBlendingTestValue != Expected)
 		{
-			AddError(FString::Printf(TEXT("Expected result 4 to be %d, actual %d."), TestValue, Expected));
+			AddError(FString::Printf(TEXT("Expected result 4 to be %d, actual %d."), GBlendingTestValue, Expected));
 		}
 	}
 
