Sequencer: Bypass static analysis false positives for GBlendingTestValue. The static analyzer can't tell that the actuator is being run on GBlendingTestValue.

diff --git a/Engine/Source/Runtime/MovieScene/Private/Tests/MovieSceneBlendingTests.cpp b/Engine/Source/Runtime/MovieScene/Private/Tests/MovieSceneBlendingTests.cpp
--- a/Engine/Source/Runtime/MovieScene/Private/Tests/MovieSceneBlendingTests.cpp
+++ b/Engine/Source/Runtime/MovieScene/Private/Tests/MovieSceneBlendingTests.cpp
@@ -103,7 +103,7 @@ bool FMovieSceneBlendingTest::RunTest(const FString& Parameters)
 		Accumulator.Apply(Context, PersistentDataProxy, Player);
 
 		int32 Expected = GBlendingStartingValue + 510;
-		if (GBlendingTestValue != Expected)
+		if (GBlendingTestValue != Expected) //-V547
 		{
 			AddError(FString::Printf(TEXT("Expected result 2 to be %d, actual %d."), GBlendingTestValue, Expected));
 		}
@@ -121,7 +121,7 @@ bool FMovieSceneBlendingTest::RunTest(const FString& Parameters)
 		Accumulator.Apply(Context, PersistentDataProxy, Player);
 
 		int32 Expected = 85 / 4;
-		if (GBlendingTestValue != Expected)
+		if (GBlendingTestValue != Expected) //-V547
 		{
 			AddError(FString::Printf(TEXT("Expected result 3 to be %d, actual %d."), GBlendingTestValue, Expected));
 		}
@@ -139,7 +139,7 @@ bool FMovieSceneBlendingTest::RunTest(const FString& Parameters)
 		Accumulator.Apply(Context, PersistentDataProxy, Player);
 
 		int32 Expected = 10770075;
-		if (GBlendingTestValue != Expected)
+		if (GBlendingTestValue != Expected) //-V547
 		{
 			AddError(FString::Printf(TEXT("Expected result 4 to be %d, actual %d."), GBlendingTestValue, Expected));
 		}
