Sequencer: Fix utility method to get sequence instance parentage.

It was somehow adding an invalid handle ID at the end of the output array instead of adding the root instance handle (which is the last parent always). This is because the root instance handle itself isn't in the list of "sub instances" of the updater, so the FindSubInstance method would return an invalid handle.

#[robomerge] 4.27

diff --git a/Engine/Source/Runtime/MovieScene/Private/Evaluation/MovieSceneEvaluationTemplateInstance.cpp b/Engine/Source/Runtime/MovieScene/Private/Evaluation/MovieSceneEvaluationTemplateInstance.cpp
--- a/Engine/Source/Runtime/MovieScene/Private/Evaluation/MovieSceneEvaluationTemplateInstance.cpp
+++ b/Engine/Source/Runtime/MovieScene/Private/Evaluation/MovieSceneEvaluationTemplateInstance.cpp
@@ -228,8 +228,15 @@ void FMovieSceneRootEvaluationTemplateInstance::GetSequenceParentage(const UE::M
 	const FMovieSceneSequenceHierarchyNode* HierarchyNode = Hierarchy->FindNode(Instance.GetSequenceID());
 	while (HierarchyNode && HierarchyNode->ParentID.IsValid())
 	{
-		const FInstanceHandle ParentHandle = RootInstance.FindSubInstance(HierarchyNode->ParentID);
-		OutParentHandles.Add(ParentHandle);
+		if (HierarchyNode->ParentID != MovieSceneSequenceID::Root)
+		{
+			const FInstanceHandle ParentHandle = RootInstance.FindSubInstance(HierarchyNode->ParentID);
+			OutParentHandles.Add(ParentHandle);
+		}
+		else
+		{
+			OutParentHandles.Add(RootInstanceHandle);
+		}
 		HierarchyNode = Hierarchy->FindNode(HierarchyNode->ParentID);
 	}
 }
