Sequencer: Fixed crash caused by bindings becoming invalidated while being resolved

diff --git a/Engine/Source/Runtime/MovieScene/Private/Evaluation/MovieSceneEvaluationState.cpp b/Engine/Source/Runtime/MovieScene/Private/Evaluation/MovieSceneEvaluationState.cpp
--- a/Engine/Source/Runtime/MovieScene/Private/Evaluation/MovieSceneEvaluationState.cpp
+++ b/Engine/Source/Runtime/MovieScene/Private/Evaluation/MovieSceneEvaluationState.cpp
@@ -73,9 +73,12 @@ FGuid FMovieSceneObjectCache::FindObjectId(UObject& InObject, IMovieScenePlayer&
 		return FGuid();
 	}
 
-	// @todo: Currently we nuke the entire object cache when attempting to find an object's ID to ensure that we do a 
-	// complete lookup from scratch. This is required for UMG as it interchanges content slots without notifying sequencer.
-	Clear(Player);
+	if (!bReentrantUpdate)
+	{
+		// @todo: Currently we nuke the entire object cache when attempting to find an object's ID to ensure that we do a 
+		// complete lookup from scratch. This is required for UMG as it interchanges content slots without notifying sequencer.
+		Clear(Player);
+	}
 
 	return FindCachedObjectId(InObject, Player);
 }
@@ -278,6 +281,8 @@ void FMovieSceneObjectCache::SetSequence(UMovieSceneSequence& InSequence, FMovie
 
 void FMovieSceneObjectCache::UpdateBindings(const FGuid& InGuid, IMovieScenePlayer& Player)
 {
+	TGuardValue<bool> ReentrancyGuard(bReentrantUpdate, true);
+
 	// Invalidate existing bindings, we're going to rebuild them.
 	FBoundObjects* Bindings = &BoundObjects.FindOrAdd(InGuid);
 	Bindings->Objects.Reset();
diff --git a/Engine/Source/Runtime/MovieScene/Public/Evaluation/MovieSceneEvaluationState.h b/Engine/Source/Runtime/MovieScene/Public/Evaluation/MovieSceneEvaluationState.h
--- a/Engine/Source/Runtime/MovieScene/Public/Evaluation/MovieSceneEvaluationState.h
+++ b/Engine/Source/Runtime/MovieScene/Public/Evaluation/MovieSceneEvaluationState.h
@@ -182,6 +182,8 @@ private:
 
 	/** Serial number for this cache */
 	uint32 SerialNumber = 0;
+
+	bool bReentrantUpdate = false;
 };
 
 /**
