Sequencer: Fixed crash caused by bindings becoming invalidated while being resolved

diff --git a/Engine/Source/Runtime/MovieScene/Private/Evaluation/MovieSceneEvaluationState.cpp b/Engine/Source/Runtime/MovieScene/Private/Evaluation/MovieSceneEvaluationState.cpp
--- a/Engine/Source/Runtime/MovieScene/Private/Evaluation/MovieSceneEvaluationState.cpp
+++ b/Engine/Source/Runtime/MovieScene/Private/Evaluation/MovieSceneEvaluationState.cpp
@@ -347,6 +347,7 @@ void FMovieSceneObjectCache::UpdateBindings(const FGuid& InGuid, IMovieScenePlay
 
 					TArray<UObject*, TInlineAllocator<1>> FoundObjects;
 					Player.ResolveBoundObjects(InGuid, SequenceID, *Sequence, ResolutionContext, FoundObjects);
+					Bindings = BoundObjects.Find(InGuid);
 					for (UObject* Object : FoundObjects)
 					{
 						Bindings->Objects.Add(Object);
@@ -357,6 +358,7 @@ void FMovieSceneObjectCache::UpdateBindings(const FGuid& InGuid, IMovieScenePlay
 			{
 				TArray<UObject*, TInlineAllocator<1>> FoundObjects;
 				Player.ResolveBoundObjects(InGuid, SequenceID, *Sequence, ResolutionContext, FoundObjects);
+				Bindings = BoundObjects.Find(InGuid);
 				for (UObject* Object : FoundObjects)
 				{
 					Bindings->Objects.Add(Object);
