remove PLATFORM_XBOXONE from MovieSceneFwd.h

#rnx

diff --git a/Engine/Source/Runtime/MovieScene/Public/MovieSceneFwd.h b/Engine/Source/Runtime/MovieScene/Public/MovieSceneFwd.h
--- a/Engine/Source/Runtime/MovieScene/Public/MovieSceneFwd.h
+++ b/Engine/Source/Runtime/MovieScene/Public/MovieSceneFwd.h
@@ -120,7 +120,7 @@ MOVIESCENE_API FFrameNumber UpgradeLegacyMovieSceneTime(UObject* ErrorContext, F
 #endif
 
 
-#if PLATFORM_WINDOWS || PLATFORM_XBOXONE
+#if defined(_MSC_VER)
 	#define UE_MOVIESCENE_TODO_IMPL(x) __pragma (x)
 #else
 	#define UE_MOVIESCENE_TODO_IMPL(x) _Pragma (#x)
