Sequencer: Fixed System.Engine.Sequencer.Track Types test

The final evaluation needs to occur outside the range of the section (and thus the loop)

Copied from 4.27 at 16275085

diff --git a/Engine/Source/Runtime/MovieScene/Private/Tests/MovieScenePreAnimatedStateTests.cpp b/Engine/Source/Runtime/MovieScene/Private/Tests/MovieScenePreAnimatedStateTests.cpp
--- a/Engine/Source/Runtime/MovieScene/Private/Tests/MovieScenePreAnimatedStateTests.cpp
+++ b/Engine/Source/Runtime/MovieScene/Private/Tests/MovieScenePreAnimatedStateTests.cpp
@@ -406,19 +406,15 @@ bool FMovieScenePreAnimatedStateTrackTypesTest::RunTest(const FString& Parameter
 			FMovieSceneEvaluationRange EvaluatedRange(TRange<FFrameTime>(i*(SectionLength/NumEvaluations), (i+1)*(SectionLength/NumEvaluations)), TickResolution, EPlayDirection::Forwards);
 			TestPlayer.Template.Evaluate(EvaluatedRange, TestPlayer);
 
-			TestPlayer.Template.GetEntitySystemRunner().Flush();
-			if (i == NumEvaluations-1)
-			{
-				Assert(this, TestValue1, EndValue, TEXT("EvaluationHook did not End correctly."));
-			}
-			else
-			{
-				Assert(this, TestValue1, StartValue + i, TEXT("EvaluationHook did not Begin or Update correctly."));
-			}
+			Assert(this, TestValue1, StartValue + i, TEXT("Keep-State EvaluationHook did not Begin or Update correctly."));
 		}
 
+		FMovieSceneEvaluationRange EvaluatedRange(TRange<FFrameTime>(SectionLength, SectionLength+100), TickResolution, EPlayDirection::Forwards);
+		TestPlayer.Template.Evaluate(EvaluatedRange, TestPlayer);
+		Assert(this, TestValue1, EndValue, TEXT("Keep-State EvaluationHook did not End correctly."));
+
 		TestPlayer.RestorePreAnimatedState();
-		Assert(this, TestValue1, TestMagicNumber, TEXT("Global pre-animated state did not restore correctly."));
+		Assert(this, TestValue1, TestMagicNumber, TEXT("Keep-State Global pre-animated state did not restore correctly."));
 	}
 
 	// Test the Restore state section
@@ -429,19 +425,15 @@ bool FMovieScenePreAnimatedStateTrackTypesTest::RunTest(const FString& Parameter
 			FMovieSceneEvaluationRange EvaluatedRange(TRange<FFrameTime>(2000 + i*(SectionLength/NumEvaluations), 2000 + (i+1)*(SectionLength/NumEvaluations)), TickResolution, EPlayDirection::Forwards);
 			TestPlayer.Template.Evaluate(EvaluatedRange, TestPlayer);
 
-			TestPlayer.Template.GetEntitySystemRunner().Flush();
-			if (i == NumEvaluations-1)
-			{
-				Assert(this, TestValue1, TestMagicNumber, TEXT("EvaluationHook did not End correctly."));
-			}
-			else
-			{
-				Assert(this, TestValue1, StartValue + i, TEXT("EvaluationHook did not Begin or Update correctly."));
-			}
+			Assert(this, TestValue1, StartValue + i, TEXT("Restore-State EvaluationHook did not Begin or Update correctly."));
 		}
 
+		FMovieSceneEvaluationRange FinalRange(TRange<FFrameTime>(2000 + SectionLength, 2000 + SectionLength + 100), TickResolution, EPlayDirection::Forwards);
+		TestPlayer.Template.Evaluate(FinalRange, TestPlayer);
+		Assert(this, TestValue1, TestMagicNumber, TEXT("Restore-State EvaluationHook did not End correctly."));
+
 		TestPlayer.RestorePreAnimatedState();
-		Assert(this, TestValue1, TestMagicNumber, TEXT("Global pre-animated state did not restore correctly."));
+		Assert(this, TestValue1, TestMagicNumber, TEXT("Global pre-animated state did not restore correctly after Restore-State Hook."));
 	}
 
 	return true;
