Sequencer: Pass through correct context when dissecting update requests

This file is a continuation of CL 16253157 which was missed off the check-in since it was in the wrong CL.

[Cherry-pick from 16267616]

diff --git a/Engine/Source/Runtime/MovieScene/Private/EntitySystem/MovieSceneEntitySystemRunner.cpp b/Engine/Source/Runtime/MovieScene/Private/EntitySystem/MovieSceneEntitySystemRunner.cpp
--- a/Engine/Source/Runtime/MovieScene/Private/EntitySystem/MovieSceneEntitySystemRunner.cpp
+++ b/Engine/Source/Runtime/MovieScene/Private/EntitySystem/MovieSceneEntitySystemRunner.cpp
@@ -216,7 +216,7 @@ void FMovieSceneEntitySystemRunner::GameThread_ProcessQueue()
 				for (int32 Index = 0; Index < Dissections.Num()-1; ++Index)
 				{
 					FDissectedUpdate Dissection {
-						FMovieSceneContext(FMovieSceneEvaluationRange(Dissections[Index], Request.Context.GetFrameRate(), Request.Context.GetDirection())),
+						FMovieSceneContext(FMovieSceneEvaluationRange(Dissections[Index], Request.Context.GetFrameRate(), Request.Context.GetDirection()), Request.Context.GetStatus()),
 						Request.InstanceHandle,
 						Index
 					};
@@ -225,7 +225,7 @@ void FMovieSceneEntitySystemRunner::GameThread_ProcessQueue()
 
 				// Add the last one with MAX_int32 so it gets evaluated with all the others in this flush
 				FDissectedUpdate Dissection {
-					FMovieSceneContext(FMovieSceneEvaluationRange(Dissections.Last(), Request.Context.GetFrameRate(), Request.Context.GetDirection())),
+					FMovieSceneContext(FMovieSceneEvaluationRange(Dissections.Last(), Request.Context.GetFrameRate(), Request.Context.GetDirection()), Request.Context.GetStatus()),
 					Request.InstanceHandle,
 					MAX_int32
 				};
