Sequencer: Fixed cache pre-animated state system not running when it had a valid extension ptr

  - There was one branch where the strong ref could be initialized without also initializing the weak ref, which is what was used for the return value

Copied from 16278713

diff --git a/Engine/Source/Runtime/MovieScene/Private/EntitySystem/MovieScenePreAnimatedStateSystem.cpp b/Engine/Source/Runtime/MovieScene/Private/EntitySystem/MovieScenePreAnimatedStateSystem.cpp
--- a/Engine/Source/Runtime/MovieScene/Private/EntitySystem/MovieScenePreAnimatedStateSystem.cpp
+++ b/Engine/Source/Runtime/MovieScene/Private/EntitySystem/MovieScenePreAnimatedStateSystem.cpp
@@ -36,6 +36,7 @@ TSharedPtr<FPreAnimatedStateExtension> FPreAnimatedStateExtensionReference::Upda
 		if (ExistingExtension)
 		{
 			PreAnimatedStateExtensionRef = ExistingExtension->AsShared();
+			WeakPreAnimatedStateExtension = PreAnimatedStateExtensionRef;
 		}
 		else
 		{
@@ -155,6 +156,9 @@ void UMovieSceneRestorePreAnimatedStateSystem::OnLink()
 		// FPreAnimatedStateExtension automatically adds itself to the linker
 		PreAnimatedStateRef = MakeShared<FPreAnimatedStateExtension>(Linker);
 	}
+
+	UMovieSceneCachePreAnimatedStateSystem* CacheSystem = Linker->LinkSystem<UMovieSceneCachePreAnimatedStateSystem>();
+	Linker->SystemGraph.AddReference(this, CacheSystem);
 }
 
 void UMovieSceneRestorePreAnimatedStateSystem::OnUnlink()
