Sequencer: Fixed Pre-Animated State not being correctly handled when PlaybackContextChanged is called

  - Added a new automation test to verify this codepath's correctness

diff --git a/Engine/Source/Runtime/MovieScene/Private/EntitySystem/MovieSceneSequenceInstance.cpp b/Engine/Source/Runtime/MovieScene/Private/EntitySystem/MovieSceneSequenceInstance.cpp
--- a/Engine/Source/Runtime/MovieScene/Private/EntitySystem/MovieSceneSequenceInstance.cpp
+++ b/Engine/Source/Runtime/MovieScene/Private/EntitySystem/MovieSceneSequenceInstance.cpp
@@ -179,6 +179,11 @@ void FSequenceInstance::EnableGlobalPreAnimatedStateCapture(UMovieSceneEntitySys
 	}
 }
 
+bool FSequenceInstance::IsCapturingGlobalPreAnimatedState() const
+{
+	return GlobalPreAnimatedState != nullptr;
+}
+
 void FSequenceInstance::Start(UMovieSceneEntitySystemLinker* Linker, const FMovieSceneContext& InContext)
 {
 	check(SequenceID == MovieSceneSequenceID::Root);
diff --git a/Engine/Source/Runtime/MovieScene/Private/Evaluation/MovieSceneEvaluationTemplateInstance.cpp b/Engine/Source/Runtime/MovieScene/Private/Evaluation/MovieSceneEvaluationTemplateInstance.cpp
--- a/Engine/Source/Runtime/MovieScene/Private/Evaluation/MovieSceneEvaluationTemplateInstance.cpp
+++ b/Engine/Source/Runtime/MovieScene/Private/Evaluation/MovieSceneEvaluationTemplateInstance.cpp
@@ -342,11 +342,21 @@ UObject* FMovieSceneRootEvaluationTemplateInstance::GetOrCreateDirectorInstance(
 
 void FMovieSceneRootEvaluationTemplateInstance::PlaybackContextChanged(IMovieScenePlayer& Player)
 {
+	using namespace UE::MovieScene;
+
+	const bool bGlobalCapture = EntitySystemLinker
+		&& RootInstanceHandle.IsValid()
+		&& EntitySystemLinker->GetInstanceRegistry()->GetInstance(RootInstanceHandle).IsCapturingGlobalPreAnimatedState();
+
 	if (EntitySystemLinker && !EntitySystemLinker->IsPendingKillOrUnreachable() && !EntitySystemLinker->HasAnyFlags(RF_BeginDestroyed))
 	{
 		EntitySystemLinker->CleanupInvalidBoundObjects();
 
 		Finish(Player);
+		if (bGlobalCapture)
+		{
+			Player.RestorePreAnimatedState();
+		}
 		EntitySystemLinker->GetInstanceRegistry()->DestroyInstance(RootInstanceHandle);
 	}
 
@@ -359,6 +369,12 @@ void FMovieSceneRootEvaluationTemplateInstance::PlaybackContextChanged(IMovieSce
 
 	RootInstanceHandle = EntitySystemLinker->GetInstanceRegistry()->AllocateRootInstance(&Player);
 	DirectorInstances.Reset();
+
+	Player.PreAnimatedState.Initialize(EntitySystemLinker, RootInstanceHandle);
+	if (bGlobalCapture)
+	{
+		EnableGlobalPreAnimatedStateCapture();
+	}
 }
 
 
diff --git a/Engine/Source/Runtime/MovieScene/Private/Evaluation/MovieScenePreAnimatedState.cpp b/Engine/Source/Runtime/MovieScene/Private/Evaluation/MovieScenePreAnimatedState.cpp
--- a/Engine/Source/Runtime/MovieScene/Private/Evaluation/MovieScenePreAnimatedState.cpp
+++ b/Engine/Source/Runtime/MovieScene/Private/Evaluation/MovieScenePreAnimatedState.cpp
@@ -24,6 +24,7 @@ IMovieScenePlayer* FRestoreStateParams::GetTerminalPlayer() const
 		return Linker->GetInstanceRegistry()->GetInstance(TerminalInstanceHandle).GetPlayer();
 	}
 
+	ensureAlways(false);
 	return nullptr;
 }
 
@@ -66,6 +67,15 @@ FMovieScenePreAnimatedState::~FMovieScenePreAnimatedState()
 
 void FMovieScenePreAnimatedState::Initialize(UMovieSceneEntitySystemLinker* Linker, UE::MovieScene::FInstanceHandle InInstanceHandle)
 {
+	using namespace UE::MovieScene;
+
+	WeakExtension = nullptr;
+	EntityExtensionRef = nullptr;
+	WeakObjectStorage = nullptr;
+	WeakMasterStorage = nullptr;
+	TemplateMetaData = nullptr;
+	EvaluationHookMetaData = nullptr;
+
 	WeakLinker = Linker;
 	InstanceHandle = InInstanceHandle;
 }
@@ -427,4 +437,12 @@ void FMovieScenePreAnimatedState::OnObjectsReplaced(const TMap<UObject*, UObject
 	{
 		ObjectGroupManager->OnObjectsReplaced(ReplacementMap);
 	}
+}
+
+bool FMovieScenePreAnimatedState::ContainsAnyStateForSequence() const
+{
+	using namespace UE::MovieScene;
+
+	TSharedPtr<FPreAnimatedStateExtension> Extension = WeakExtension.Pin();
+	return Extension && InstanceHandle.IsValid() && Extension->ContainsAnyStateForInstanceHandle(InstanceHandle);
 }
\ No newline at end of file
diff --git a/Engine/Source/Runtime/MovieScene/Private/Evaluation/PreAnimatedState/MovieScenePreAnimatedCaptureSources.inl b/Engine/Source/Runtime/MovieScene/Private/Evaluation/PreAnimatedState/MovieScenePreAnimatedCaptureSources.inl
--- a/Engine/Source/Runtime/MovieScene/Private/Evaluation/PreAnimatedState/MovieScenePreAnimatedCaptureSources.inl
+++ b/Engine/Source/Runtime/MovieScene/Private/Evaluation/PreAnimatedState/MovieScenePreAnimatedCaptureSources.inl
@@ -113,6 +113,21 @@ void TPreAnimatedCaptureSources<KeyType>::GatherAndRemoveMetaDataForGroup(FPreAn
 	}
 }
 
+template<typename KeyType>
+bool TPreAnimatedCaptureSources<KeyType>::ContainsInstanceHandle(FInstanceHandle RootInstanceHandle) const
+{
+	for (const TPair<KeyType, FPreAnimatedStateMetaDataArray>& Pair : KeyToMetaData)
+	{
+		for (const FPreAnimatedStateMetaData& MetaData : Pair.Value)
+		{
+			if (MetaData.RootInstanceHandle == RootInstanceHandle)
+			{
+				return true;
+			}
+		}
+	}
+	return false;
+}
 
 } // namespace MovieScene
 } // namespace UE
diff --git a/Engine/Source/Runtime/MovieScene/Private/Evaluation/PreAnimatedState/MovieScenePreAnimatedEntityCaptureSource.cpp b/Engine/Source/Runtime/MovieScene/Private/Evaluation/PreAnimatedState/MovieScenePreAnimatedEntityCaptureSource.cpp
--- a/Engine/Source/Runtime/MovieScene/Private/Evaluation/PreAnimatedState/MovieScenePreAnimatedEntityCaptureSource.cpp
+++ b/Engine/Source/Runtime/MovieScene/Private/Evaluation/PreAnimatedState/MovieScenePreAnimatedEntityCaptureSource.cpp
@@ -115,6 +115,20 @@ void FPreAnimatedEntityCaptureSource::GatherAndRemoveMetaDataForGroup(FPreAnimat
 	}
 }
 
+bool FPreAnimatedEntityCaptureSource::ContainsInstanceHandle(FInstanceHandle RootInstanceHandle) const
+{
+	for (const TPair<FMovieSceneEntityID, FPreAnimatedStateMetaDataArray>& Pair : KeyToMetaData)
+	{
+		for (const FPreAnimatedStateMetaData& MetaData : Pair.Value)
+		{
+			if (MetaData.RootInstanceHandle == RootInstanceHandle)
+			{
+				return true;
+			}
+		}
+	}
+	return false;
+}
 
 } // namespace MovieScene
 } // namespace UE
diff --git a/Engine/Source/Runtime/MovieScene/Private/Evaluation/PreAnimatedState/MovieScenePreAnimatedStateExtension.cpp b/Engine/Source/Runtime/MovieScene/Private/Evaluation/PreAnimatedState/MovieScenePreAnimatedStateExtension.cpp
--- a/Engine/Source/Runtime/MovieScene/Private/Evaluation/PreAnimatedState/MovieScenePreAnimatedStateExtension.cpp
+++ b/Engine/Source/Runtime/MovieScene/Private/Evaluation/PreAnimatedState/MovieScenePreAnimatedStateExtension.cpp
@@ -186,7 +186,7 @@ void FPreAnimatedStateExtension::UpdateMetaData(const FPreAnimatedStateMetaData&
 	}
 }
 
-FPreAnimatedEntityCaptureSource* FPreAnimatedStateExtension::GetEntityMetaData()
+FPreAnimatedEntityCaptureSource* FPreAnimatedStateExtension::GetEntityMetaData() const
 {
 	return EntityCaptureSource.Get();
 }
@@ -200,7 +200,7 @@ FPreAnimatedEntityCaptureSource* FPreAnimatedStateExtension::GetOrCreateEntityMe
 	return EntityCaptureSource.Get();
 }
 
-FPreAnimatedTrackInstanceCaptureSources* FPreAnimatedStateExtension::GetTrackInstanceMetaData()
+FPreAnimatedTrackInstanceCaptureSources* FPreAnimatedStateExtension::GetTrackInstanceMetaData() const
 {
 	return TrackInstanceCaptureSource.Get();
 }
@@ -379,5 +379,27 @@ void FPreAnimatedStateExtension::DiscardStateForGroup(FPreAnimatedStorageGroupHa
 	bEntriesInvalidated = true;
 }
 
+bool FPreAnimatedStateExtension::ContainsAnyStateForInstanceHandle(FInstanceHandle RootInstanceHandle) const
+{
+	if (FPreAnimatedEntityCaptureSource* EntityMetaData = GetEntityMetaData())
+	{
+		if (EntityMetaData->ContainsInstanceHandle(RootInstanceHandle))
+		{
+			return true;
+		}
+	}
+
+	for (int32 Index = WeakExternalCaptureSources.Num()-1; Index >= 0; --Index)
+	{
+		TSharedPtr<IPreAnimatedCaptureSource> MetaData = WeakExternalCaptureSources[Index].Pin();
+		if (MetaData && MetaData->ContainsInstanceHandle(RootInstanceHandle))
+		{
+			return true;
+		}
+	}
+
+	return false;
+}
+
 } // namespace MovieScene
 } // namespace UE
diff --git a/Engine/Source/Runtime/MovieScene/Private/Tests/MovieScenePreAnimatedStateTests.cpp b/Engine/Source/Runtime/MovieScene/Private/Tests/MovieScenePreAnimatedStateTests.cpp
--- a/Engine/Source/Runtime/MovieScene/Private/Tests/MovieScenePreAnimatedStateTests.cpp
+++ b/Engine/Source/Runtime/MovieScene/Private/Tests/MovieScenePreAnimatedStateTests.cpp
@@ -439,6 +439,83 @@ bool FMovieScenePreAnimatedStateTrackTypesTest::RunTest(const FString& Parameter
 	return true;
 }
 
+/** Simulates the flow of FSequencer transitioning into/out of PIE */
+IMPLEMENT_SIMPLE_AUTOMATION_TEST(FMovieScenePreAnimatedStateContextChangedTest, "System.Engine.Sequencer.Pre-Animated State.Context Changed", EAutomationTestFlags::EditorContext | EAutomationTestFlags::EngineFilter)
+bool FMovieScenePreAnimatedStateContextChangedTest::RunTest(const FString& Parameters)
+{
+	using namespace Impl;
+	using namespace UE::MovieScene;
+
+	ResetValues();
+
+	static const int32 StartValue    = 1000;
+	static const int32 EndValue      = 2000;
+	static const int32 SectionLength = 1000;
+
+	FFrameRate TickResolution(1000, 1);
+
+	UTestMovieSceneSequence* Sequence = NewObject<UTestMovieSceneSequence>(GetTransientPackage());
+	Sequence->MovieScene->SetTickResolutionDirectly(TickResolution);
+
+	// Make a keep-state eval hook section
+	{
+		UTestMovieSceneEvalHookTrack*   Track = Sequence->MovieScene->AddMasterTrack<UTestMovieSceneEvalHookTrack>();
+		UTestMovieSceneEvalHookSection* Section = NewObject<UTestMovieSceneEvalHookSection>(Track);
+
+		Section->EvalOptions.CompletionMode = EMovieSceneCompletionMode::KeepState;
+		Section->StartValue   = StartValue;
+		Section->EndValue     = EndValue;
+
+		Section->SetRange(TRange<FFrameNumber>(0, SectionLength));
+		Track->SectionArray.Add(Section);
+	}
+
+	FTestMovieScenePlayer TestPlayer;
+
+	UMovieSceneCompiledDataManager* CompiledDataManager = UMovieSceneCompiledDataManager::GetPrecompiledData();
+
+	FMovieSceneCompiledDataID DataID = CompiledDataManager->Compile(Sequence);
+	TestPlayer.Template.Initialize(*Sequence, TestPlayer, CompiledDataManager);
+	TestPlayer.Template.EnableGlobalPreAnimatedStateCapture();
+
+	// ---------------------------------------------
+	// Simulate in-editor
+	{
+		static const int32 NumEvaluations = 5;
+		for (int32 i = 0; i < NumEvaluations; ++i)
+		{
+			FMovieSceneEvaluationRange EvaluatedRange(TRange<FFrameTime>(i*(SectionLength/NumEvaluations), (i+1)*(SectionLength/NumEvaluations)), TickResolution, EPlayDirection::Forwards);
+			TestPlayer.Template.Evaluate(EvaluatedRange, TestPlayer);
+
+			Assert(this, TestValue1, StartValue + i, TEXT("In-Editor 1: EvaluationHook did not Begin or Update correctly."));
+		}
+
+		TestPlayer.Template.PlaybackContextChanged(TestPlayer);
+		TestEqual(TEXT("Pre-Animated State still exists for In-Editor 1."), TestPlayer.PreAnimatedState.ContainsAnyStateForSequence(), false);
+		Assert(this, TestValue1, TestMagicNumber, TEXT("In-Editor 1: Global pre-animated state did not restore correctly."));
+	}
+
+	// ---------------------------------------------
+	// Simulate PIE 
+	{
+		static const int32 NumEvaluations = 5;
+		for (int32 i = 0; i < NumEvaluations; ++i)
+		{
+			FMovieSceneEvaluationRange EvaluatedRange(TRange<FFrameTime>(i*(SectionLength/NumEvaluations), (i+1)*(SectionLength/NumEvaluations)), TickResolution, EPlayDirection::Forwards);
+			TestPlayer.Template.Evaluate(EvaluatedRange, TestPlayer);
+
+			Assert(this, TestValue1, StartValue + i, TEXT("PIE: EvaluationHook did not Begin or Update correctly."));
+		}
+
+		TestPlayer.Template.PlaybackContextChanged(TestPlayer);
+		TestEqual(TEXT("Pre-Animated State still exists for PIE."), TestPlayer.PreAnimatedState.ContainsAnyStateForSequence(), false);
+		Assert(this, TestValue1, TestMagicNumber, TEXT("PIE: Global pre-animated state did not restore correctly."));
+	}
+
+	Assert(this, TestValue1, TestMagicNumber, TEXT("Post PIE: Global pre-animated state did not restore correctly."));
+	return true;
+}
+
 /** Tests an edge case where one section keeps state, while another subsequent section restores state. the second section should restore to its starting value, not the original state before the first section. */
 IMPLEMENT_SIMPLE_AUTOMATION_TEST(FMovieScenePreAnimatedStatePerformanceTest, "System.Engine.Sequencer.Pre-Animated State.Performance", EAutomationTestFlags::EditorContext | EAutomationTestFlags::EngineFilter | EAutomationTestFlags::Disabled)
 bool FMovieScenePreAnimatedStatePerformanceTest::RunTest(const FString& Parameters)
diff --git a/Engine/Source/Runtime/MovieScene/Public/EntitySystem/MovieSceneSequenceInstance.h b/Engine/Source/Runtime/MovieScene/Public/EntitySystem/MovieSceneSequenceInstance.h
--- a/Engine/Source/Runtime/MovieScene/Public/EntitySystem/MovieSceneSequenceInstance.h
+++ b/Engine/Source/Runtime/MovieScene/Public/EntitySystem/MovieSceneSequenceInstance.h
@@ -95,6 +95,11 @@ public:
 	 */
 	void EnableGlobalPreAnimatedStateCapture(UMovieSceneEntitySystemLinker* Linker);
 
+	/**
+	 * Check whether this sequence instance is capturing any and all changes of state so they can be restored later
+	 */
+	bool IsCapturingGlobalPreAnimatedState() const;
+
 public:
 
 	/**
diff --git a/Engine/Source/Runtime/MovieScene/Public/Evaluation/MovieScenePreAnimatedState.h b/Engine/Source/Runtime/MovieScene/Public/Evaluation/MovieScenePreAnimatedState.h
--- a/Engine/Source/Runtime/MovieScene/Public/Evaluation/MovieScenePreAnimatedState.h
+++ b/Engine/Source/Runtime/MovieScene/Public/Evaluation/MovieScenePreAnimatedState.h
@@ -166,6 +166,14 @@ public:
 	 */
 	MOVIESCENE_API void OnObjectsReplaced(const TMap<UObject*, UObject*>& ReplacementMap);
 
+public:
+
+	/**
+	 * Search the global pre-animated state extension for any captured state that originated from this sequence
+	 * WARNING: This is a linear search across all state, and so is potentially very slow
+	 */
+	MOVIESCENE_API bool ContainsAnyStateForSequence() const;
+
 private:
 
 	void ConditionalInitializeEntityStorage(bool bOverrideWantsRestoreState);
diff --git a/Engine/Source/Runtime/MovieScene/Public/Evaluation/PreAnimatedState/IMovieScenePreAnimatedCaptureSource.h b/Engine/Source/Runtime/MovieScene/Public/Evaluation/PreAnimatedState/IMovieScenePreAnimatedCaptureSource.h
--- a/Engine/Source/Runtime/MovieScene/Public/Evaluation/PreAnimatedState/IMovieScenePreAnimatedCaptureSource.h
+++ b/Engine/Source/Runtime/MovieScene/Public/Evaluation/PreAnimatedState/IMovieScenePreAnimatedCaptureSource.h
@@ -11,6 +11,7 @@ namespace UE
 namespace MovieScene
 {
 
+struct FInstanceHandle;
 struct FRestoreStateParams;
 struct FPreAnimatedStateMetaData;
 struct FPreAnimatedStorageGroupHandle;
@@ -20,6 +21,7 @@ struct IPreAnimatedCaptureSource
 	virtual ~IPreAnimatedCaptureSource(){}
 
 	virtual void Reset() = 0;
+	virtual bool ContainsInstanceHandle(FInstanceHandle RootInstanceHandle) const = 0;
 	virtual void GatherAndRemoveExpiredMetaData(const FRestoreStateParams& Params, TArray<FPreAnimatedStateMetaData>& OutExpiredMetaData) = 0;
 	virtual void GatherAndRemoveMetaDataForGroup(FPreAnimatedStorageGroupHandle Group, TArray<FPreAnimatedStateMetaData>& OutExpiredMetaData) = 0;
 };
diff --git a/Engine/Source/Runtime/MovieScene/Public/Evaluation/PreAnimatedState/MovieScenePreAnimatedCaptureSources.h b/Engine/Source/Runtime/MovieScene/Public/Evaluation/PreAnimatedState/MovieScenePreAnimatedCaptureSources.h
--- a/Engine/Source/Runtime/MovieScene/Public/Evaluation/PreAnimatedState/MovieScenePreAnimatedCaptureSources.h
+++ b/Engine/Source/Runtime/MovieScene/Public/Evaluation/PreAnimatedState/MovieScenePreAnimatedCaptureSources.h
@@ -62,10 +62,8 @@ struct TPreAnimatedCaptureSources : IPreAnimatedCaptureSource
 	 */
 	void StopTrackingCaptureSource(const KeyType& InKey);
 
-	/**
-	 * Gather and remove all meta data that matches the specified restoration parameters
-	 */
 	void Reset() override;
+	bool ContainsInstanceHandle(FInstanceHandle RootInstanceHandle) const override;
 	void GatherAndRemoveExpiredMetaData(const FRestoreStateParams& Params, TArray<FPreAnimatedStateMetaData>& OutExpiredMetaData) override;
 	void GatherAndRemoveMetaDataForGroup(FPreAnimatedStorageGroupHandle Group, TArray<FPreAnimatedStateMetaData>& OutExpiredMetaData) override;
 
diff --git a/Engine/Source/Runtime/MovieScene/Public/Evaluation/PreAnimatedState/MovieScenePreAnimatedEntityCaptureSource.h b/Engine/Source/Runtime/MovieScene/Public/Evaluation/PreAnimatedState/MovieScenePreAnimatedEntityCaptureSource.h
--- a/Engine/Source/Runtime/MovieScene/Public/Evaluation/PreAnimatedState/MovieScenePreAnimatedEntityCaptureSource.h
+++ b/Engine/Source/Runtime/MovieScene/Public/Evaluation/PreAnimatedState/MovieScenePreAnimatedEntityCaptureSource.h
@@ -32,6 +32,7 @@ struct FPreAnimatedEntityCaptureSource
 	MOVIESCENE_API void BeginTrackingEntity(const FPreAnimatedStateEntry& Entry, FMovieSceneEntityID EntityID, FInstanceHandle RootInstanceHandle, bool bWantsRestoreState);
 	MOVIESCENE_API void StopTrackingEntity(FMovieSceneEntityID EntityID, FPreAnimatedStorageID StorageID);
 
+	bool ContainsInstanceHandle(FInstanceHandle RootInstanceHandle) const;
 	void GatherAndRemoveExpiredMetaData(const FRestoreStateParams& Params, TArray<FPreAnimatedStateMetaData>& OutExpiredMetaData);
 	void GatherAndRemoveMetaDataForGroup(FPreAnimatedStorageGroupHandle Group, TArray<FPreAnimatedStateMetaData>& OutExpiredMetaData);
 
diff --git a/Engine/Source/Runtime/MovieScene/Public/Evaluation/PreAnimatedState/MovieScenePreAnimatedStateExtension.h b/Engine/Source/Runtime/MovieScene/Public/Evaluation/PreAnimatedState/MovieScenePreAnimatedStateExtension.h
--- a/Engine/Source/Runtime/MovieScene/Public/Evaluation/PreAnimatedState/MovieScenePreAnimatedStateExtension.h
+++ b/Engine/Source/Runtime/MovieScene/Public/Evaluation/PreAnimatedState/MovieScenePreAnimatedStateExtension.h
@@ -243,12 +243,18 @@ public:
 	 */
 	void DiscardStateForGroup(FPreAnimatedStorageGroupHandle GroupHandle);
 
+	/**
+	 * Search for any captured state that originated from the specified root instance handle
+	 * WARNING: This is a linear search across all state, and so is potentially very slow
+	 */
+	bool ContainsAnyStateForInstanceHandle(FInstanceHandle RootInstanceHandle) const;
+
 public:
 
-	FPreAnimatedEntityCaptureSource* GetEntityMetaData();
+	FPreAnimatedEntityCaptureSource* GetEntityMetaData() const;
 	FPreAnimatedEntityCaptureSource* GetOrCreateEntityMetaData();
 
-	FPreAnimatedTrackInstanceCaptureSources* GetTrackInstanceMetaData();
+	FPreAnimatedTrackInstanceCaptureSources* GetTrackInstanceMetaData() const;
 	FPreAnimatedTrackInstanceCaptureSources* GetOrCreateTrackInstanceMetaData();
 
 	void AddWeakCaptureSource(TWeakPtr<IPreAnimatedCaptureSource> InWeakMetaData);
