Sequencer: Always restore entities in the restore state system automatically

  - Previously each system was restoring its own sets of entities, whereas this can happen in a single place now (and removes the need for each system to manually restore its own entities)

diff --git a/Engine/Source/Runtime/MovieScene/Private/EntitySystem/MovieScenePreAnimatedStateSystem.cpp b/Engine/Source/Runtime/MovieScene/Private/EntitySystem/MovieScenePreAnimatedStateSystem.cpp
--- a/Engine/Source/Runtime/MovieScene/Private/EntitySystem/MovieScenePreAnimatedStateSystem.cpp
+++ b/Engine/Source/Runtime/MovieScene/Private/EntitySystem/MovieScenePreAnimatedStateSystem.cpp
@@ -3,10 +3,13 @@
 #include "EntitySystem/MovieScenePreAnimatedStateSystem.h"
 #include "EntitySystem/MovieSceneEntitySystemLinker.h"
 #include "EntitySystem/BuiltInComponentTypes.h"
+#include "EntitySystem/MovieSceneEntitySystemTask.h"
 #include "IMovieScenePlayer.h"
 
 #include "Evaluation/PreAnimatedState/MovieScenePreAnimatedStateExtension.h"
 #include "Evaluation/PreAnimatedState/MovieScenePreAnimatedStateStorage.h"
+#include "Evaluation/PreAnimatedState/MovieScenePreAnimatedEntityCaptureSource.h"
+
 
 
 namespace UE
@@ -192,5 +195,19 @@ void UMovieSceneRestorePreAnimatedStateSystem::OnRun(FSystemTaskPrerequisites& I
 		Interfaces[Index]->RestorePreAnimatedState(Params);
 	}
 
+	FPreAnimatedEntityCaptureSource* EntityMetaData = PreAnimatedStateRef->GetEntityMetaData();
+	if (EntityMetaData)
+	{
+		auto CleanupExpiredObjects = [EntityMetaData](FMovieSceneEntityID EntityID)
+		{
+			EntityMetaData->StopTrackingEntity(EntityID);
+		};
+
+		FEntityTaskBuilder()
+		.ReadEntityIDs()
+		.FilterAll({ FBuiltInComponentTypes::Get()->Tags.NeedsUnlink })
+		.Iterate_PerEntity(&Linker->EntityManager, CleanupExpiredObjects);
+	}
+
 	Params.CacheExtension->ResetEntryInvalidation();
 }
diff --git a/Engine/Source/Runtime/MovieScene/Private/Evaluation/PreAnimatedState/MovieScenePreAnimatedEntityCaptureSource.cpp b/Engine/Source/Runtime/MovieScene/Private/Evaluation/PreAnimatedState/MovieScenePreAnimatedEntityCaptureSource.cpp
--- a/Engine/Source/Runtime/MovieScene/Private/Evaluation/PreAnimatedState/MovieScenePreAnimatedEntityCaptureSource.cpp
+++ b/Engine/Source/Runtime/MovieScene/Private/Evaluation/PreAnimatedState/MovieScenePreAnimatedEntityCaptureSource.cpp
@@ -62,6 +62,21 @@ void FPreAnimatedEntityCaptureSource::StopTrackingEntity(FMovieSceneEntityID Ent
 	}
 }
 
+void FPreAnimatedEntityCaptureSource::StopTrackingEntity(FMovieSceneEntityID EntityID)
+{
+	FPreAnimatedStateMetaDataArray* Array = KeyToMetaData.Find(EntityID);
+	if (Array)
+	{
+		for (int32 Index = Array->Num()-1; Index >= 0; --Index)
+		{
+			const FPreAnimatedStateMetaData& MetaData = (*Array)[Index];
+			Owner->RemoveMetaData(MetaData);
+		}
+
+		KeyToMetaData.Remove(EntityID);
+	}
+}
+
 void FPreAnimatedEntityCaptureSource::GatherAndRemoveExpiredMetaData(const FRestoreStateParams& Params, TArray<FPreAnimatedStateMetaData>& OutExpiredMetaData)
 {
 	FInstanceHandle InstanceHandle = Params.TerminalInstanceHandle;
diff --git a/Engine/Source/Runtime/MovieScene/Public/Evaluation/PreAnimatedState/IMovieScenePreAnimatedStorage.h b/Engine/Source/Runtime/MovieScene/Public/Evaluation/PreAnimatedState/IMovieScenePreAnimatedStorage.h
--- a/Engine/Source/Runtime/MovieScene/Public/Evaluation/PreAnimatedState/IMovieScenePreAnimatedStorage.h
+++ b/Engine/Source/Runtime/MovieScene/Public/Evaluation/PreAnimatedState/IMovieScenePreAnimatedStorage.h
@@ -24,7 +24,6 @@ struct FCustomPropertyIndex;
 struct FPreAnimatedStateExtension;
 
 struct IPreAnimatedObjectEntityStorage;
-struct IPreAnimatedMasterStorage;
 struct IPreAnimatedObjectPropertyStorage;
 
 struct IPreAnimatedStorage : TSharedFromThis<IPreAnimatedStorage>
@@ -41,7 +40,6 @@ struct IPreAnimatedStorage : TSharedFromThis<IPreAnimatedStorage>
 
 	virtual IPreAnimatedObjectPropertyStorage* AsPropertyStorage() { return nullptr; }
 	virtual IPreAnimatedObjectEntityStorage*   AsObjectStorage()   { return nullptr; }
-	virtual IPreAnimatedMasterStorage*         AsMasterStorage()   { return nullptr; }
 };
 
 
diff --git a/Engine/Source/Runtime/MovieScene/Public/Evaluation/PreAnimatedState/MovieScenePreAnimatedEntityCaptureSource.h b/Engine/Source/Runtime/MovieScene/Public/Evaluation/PreAnimatedState/MovieScenePreAnimatedEntityCaptureSource.h
--- a/Engine/Source/Runtime/MovieScene/Public/Evaluation/PreAnimatedState/MovieScenePreAnimatedEntityCaptureSource.h
+++ b/Engine/Source/Runtime/MovieScene/Public/Evaluation/PreAnimatedState/MovieScenePreAnimatedEntityCaptureSource.h
@@ -31,6 +31,7 @@ struct FPreAnimatedEntityCaptureSource
 
 	MOVIESCENE_API void BeginTrackingEntity(const FPreAnimatedStateEntry& Entry, FMovieSceneEntityID EntityID, FInstanceHandle RootInstanceHandle, bool bWantsRestoreState);
 	MOVIESCENE_API void StopTrackingEntity(FMovieSceneEntityID EntityID, FPreAnimatedStorageID StorageID);
+	MOVIESCENE_API void StopTrackingEntity(FMovieSceneEntityID EntityID);
 
 	bool ContainsInstanceHandle(FInstanceHandle RootInstanceHandle) const;
 	void GatherAndRemoveExpiredMetaData(const FRestoreStateParams& Params, TArray<FPreAnimatedStateMetaData>& OutExpiredMetaData);
diff --git a/Engine/Source/Runtime/MovieSceneTracks/Private/Systems/MovieSceneComponentMobilitySystem.cpp b/Engine/Source/Runtime/MovieSceneTracks/Private/Systems/MovieSceneComponentMobilitySystem.cpp
--- a/Engine/Source/Runtime/MovieSceneTracks/Private/Systems/MovieSceneComponentMobilitySystem.cpp
+++ b/Engine/Source/Runtime/MovieSceneTracks/Private/Systems/MovieSceneComponentMobilitySystem.cpp
@@ -231,25 +231,5 @@ void UMovieSceneComponentMobilitySystem::SavePreAnimatedState(const FPreAnimatio
 
 void UMovieSceneComponentMobilitySystem::RestorePreAnimatedState(const FPreAnimationParameters& InParameters)
 {
-	using namespace UE::MovieScene;
-
-	FPreAnimatedEntityCaptureSource* EntityMetaData = InParameters.CacheExtension->GetOrCreateEntityMetaData();
-	if (!EntityMetaData)
-	{
-		return;
-	}
-
-	auto CleanupExpiredObjects = [EntityMetaData](FMovieSceneEntityID EntityID)
-	{
-		EntityMetaData->StopTrackingEntity(EntityID, FPreAnimatedComponentMobilityStorage::StorageID);
-	};
 
-	FBuiltInComponentTypes*          BuiltInComponents = FBuiltInComponentTypes::Get();
-	FMovieSceneTracksComponentTypes* TrackComponents   = FMovieSceneTracksComponentTypes::Get();
-
-	FEntityTaskBuilder()
-	.ReadEntityIDs()
-	.FilterAll({ BuiltInComponents->BoundObject, BuiltInComponents->Tags.NeedsUnlink })
-	.FilterAny({ TrackComponents->ComponentTransform.PropertyTag, TrackComponents->AttachParent })
-	.Iterate_PerEntity(&Linker->EntityManager, CleanupExpiredObjects);
 }
diff --git a/Engine/Source/Runtime/MovieSceneTracks/Private/Systems/MovieScenePropertySystem.cpp b/Engine/Source/Runtime/MovieSceneTracks/Private/Systems/MovieScenePropertySystem.cpp
--- a/Engine/Source/Runtime/MovieSceneTracks/Private/Systems/MovieScenePropertySystem.cpp
+++ b/Engine/Source/Runtime/MovieSceneTracks/Private/Systems/MovieScenePropertySystem.cpp
@@ -125,18 +125,5 @@ void UMovieScenePropertySystem::SavePreAnimatedState(const FPreAnimationParamete
 
 void UMovieScenePropertySystem::RestorePreAnimatedState(const FPreAnimationParameters& InParameters)
 {
-	using namespace UE::MovieScene;
-
-	FPreAnimatedEntityCaptureSource* EntityMetaData = InParameters.CacheExtension->GetEntityMetaData();
-
-	if (EntityMetaData)
-	{
-		FBuiltInComponentTypes* BuiltInComponents = FBuiltInComponentTypes::Get();
-		const FPropertyDefinition& Definition = BuiltInComponents->PropertyRegistry.GetDefinition(CompositePropertyID);
 
-		FEntityTaskBuilder()
-		.ReadEntityIDs()
-		.FilterAll({ BuiltInComponents->BoundObject, Definition.PropertyType, BuiltInComponents->Tags.NeedsUnlink })
-		.Iterate_PerEntity(&Linker->EntityManager, [this, EntityMetaData](FMovieSceneEntityID EntityID) { EntityMetaData->StopTrackingEntity(EntityID, this->PreAnimatedStorageID); });
-	}
 }
\ No newline at end of file
