Sequencer: Add the preroll tag manually to anything with section preroll to ensure that it works with non-instantiated master tracks

diff --git a/Engine/Source/Runtime/MovieScene/Private/MovieSceneSection.cpp b/Engine/Source/Runtime/MovieScene/Private/MovieSceneSection.cpp
--- a/Engine/Source/Runtime/MovieScene/Private/MovieSceneSection.cpp
+++ b/Engine/Source/Runtime/MovieScene/Private/MovieSceneSection.cpp
@@ -322,7 +322,7 @@ void UMovieSceneSection::BuildDefaultComponents(UMovieSceneEntitySystemLinker* E
 		.AddTagConditional(Components->Tags.RestoreState,       bShouldRestoreState)
 		.AddTagConditional(Components->Tags.FixedTime,          bHasForcedTime)
 		.AddTagConditional(Components->Tags.SectionPreRoll,     bHasSectionPreRoll)
-		.AddTagConditional(Components->Tags.PreRoll,            bHasSequencePreRoll)
+		.AddTagConditional(Components->Tags.PreRoll,            bHasSequencePreRoll || bHasSectionPreRoll)
 		.AddTagConditional(BlendTag,                            BlendTag != FComponentTypeID::Invalid())
 	);
 
