Sequencer: Expose GetSequence on the MovieSceneSequencePlayer and the SequenceDirector. GetSequence() on the SequenceDirector will return the sub sequence that it's being played on

diff --git a/Engine/Source/Runtime/LevelSequence/Private/LevelSequenceDirector.cpp b/Engine/Source/Runtime/LevelSequence/Private/LevelSequenceDirector.cpp
--- a/Engine/Source/Runtime/LevelSequence/Private/LevelSequenceDirector.cpp
+++ b/Engine/Source/Runtime/LevelSequence/Private/LevelSequenceDirector.cpp
@@ -4,7 +4,9 @@
 #include "Engine/World.h"
 #include "UObject/Stack.h"
 #include "MovieSceneObjectBindingID.h"
+#include "LevelSequencePlayer.h"
 #include "IMovieScenePlayer.h"
+#include "Evaluation/MovieSceneEvaluationTemplateInstance.h"
 
 UWorld* ULevelSequenceDirector::GetWorld() const
 {
@@ -99,3 +101,17 @@ AActor* ULevelSequenceDirector::GetBoundActor(FMovieSceneObjectBindingID ObjectB
 	return nullptr;
 }
 
+UMovieSceneSequence* ULevelSequenceDirector::GetSequence()
+{
+	if (IMovieScenePlayer* PlayerInterface = IMovieScenePlayer::Get(static_cast<uint16>(MovieScenePlayerIndex)))
+	{
+		return PlayerInterface->GetEvaluationTemplate().GetSequence(FMovieSceneSequenceID(SubSequenceID));
+	}
+	else
+	{		
+		FFrame::KismetExecutionMessage(TEXT("No sequence player."), ELogVerbosity::Error);
+
+		return nullptr;
+	}
+}
+
diff --git a/Engine/Source/Runtime/LevelSequence/Public/LevelSequenceDirector.h b/Engine/Source/Runtime/LevelSequence/Public/LevelSequenceDirector.h
--- a/Engine/Source/Runtime/LevelSequence/Public/LevelSequenceDirector.h
+++ b/Engine/Source/Runtime/LevelSequence/Public/LevelSequenceDirector.h
@@ -62,6 +62,12 @@ public:
 	UFUNCTION(BlueprintCallable, Category="Sequencer")
 	AActor* GetBoundActor(FMovieSceneObjectBindingID ObjectBinding);
 
+	/*
+	 * Get the current sequence that this director is playing back within 
+	 */
+	UFUNCTION(BlueprintCallable, Category="Sequencer")
+	UMovieSceneSequence* GetSequence();
+
 public:
 
 	virtual UWorld* GetWorld() const override;
diff --git a/Engine/Source/Runtime/MovieScene/Public/MovieSceneSequencePlayer.h b/Engine/Source/Runtime/MovieScene/Public/MovieSceneSequencePlayer.h
--- a/Engine/Source/Runtime/MovieScene/Public/MovieSceneSequencePlayer.h
+++ b/Engine/Source/Runtime/MovieScene/Public/MovieSceneSequencePlayer.h
@@ -490,6 +490,7 @@ public:
 	 * Access the sequence this player is playing
 	 * @return the sequence currently assigned to this player
 	 */
+	UFUNCTION(BlueprintCallable, Category = "Sequencer|Player")
 	UMovieSceneSequence* GetSequence() const { return Sequence; }
 
 	/**
@@ -533,10 +534,12 @@ protected:
 	void QueueLatentAction(FMovieSceneSequenceLatentActionDelegate Delegate);
 	void RunLatentActions();
 
-protected:
-
+public:
 	//~ IMovieScenePlayer interface
 	virtual FMovieSceneRootEvaluationTemplateInstance& GetEvaluationTemplate() override { return RootTemplateInstance; }
+
+protected:
+	//~ IMovieScenePlayer interface
 	virtual UMovieSceneEntitySystemLinker* ConstructEntitySystemLinker() override;
 	virtual EMovieScenePlayerStatus::Type GetPlaybackStatus() const override;
 	virtual FMovieSceneSpawnRegister& GetSpawnRegister() override;
