Sequencer: Fixed an isssue where the PlayTo method would play past the desired frame before jumping backwards and pausing during low framerate situations. This would cause events and other setup from frames beyond the desired frame to get called early.

diff --git a/Engine/Source/Runtime/MovieScene/Private/MovieSceneSequencePlayer.cpp b/Engine/Source/Runtime/MovieScene/Private/MovieSceneSequencePlayer.cpp
--- a/Engine/Source/Runtime/MovieScene/Private/MovieSceneSequencePlayer.cpp
+++ b/Engine/Source/Runtime/MovieScene/Private/MovieSceneSequencePlayer.cpp
@@ -635,11 +635,11 @@ bool UMovieSceneSequencePlayer::ShouldPause(FFrameTime NewPosition) const
 	{
 		if (!bReversePlayback)
 		{
-			bShouldPause = PauseOnFrame.GetValue() <= PlayPosition.GetCurrentPosition();
+			bShouldPause = PauseOnFrame.GetValue() <= NewPosition;
 		}
 		else
 		{
-			bShouldPause = PauseOnFrame.GetValue() >= PlayPosition.GetCurrentPosition();
+			bShouldPause = PauseOnFrame.GetValue() >= NewPosition;
 		}
 	}
 
