#sequencer - Back out Sequencer change that appeared to stop frame 0 events from firing.

diff --git a/Engine/Source/Runtime/MovieScene/Private/MovieSceneSequencePlayer.cpp b/Engine/Source/Runtime/MovieScene/Private/MovieSceneSequencePlayer.cpp
--- a/Engine/Source/Runtime/MovieScene/Private/MovieSceneSequencePlayer.cpp
+++ b/Engine/Source/Runtime/MovieScene/Private/MovieSceneSequencePlayer.cpp
@@ -635,11 +635,11 @@ bool UMovieSceneSequencePlayer::ShouldPause(FFrameTime NewPosition) const
 	{
 		if (!bReversePlayback)
 		{
-			bShouldPause = PauseOnFrame.GetValue() <= NewPosition;
+			bShouldPause = PauseOnFrame.GetValue() <= PlayPosition.GetCurrentPosition();
 		}
 		else
 		{
-			bShouldPause = PauseOnFrame.GetValue() >= NewPosition;
+			bShouldPause = PauseOnFrame.GetValue() >= PlayPosition.GetCurrentPosition();
 		}
 	}
 
