Sequencer: Replace BindingOverrides TWeakObjectPtr by a TSoftObjectPtr to fix World Partition incompatibility (Serialization is compatible between the 2 types through FSoftObjectPath::SerializeFromMismatchedTag)

diff --git a/Engine/Source/Runtime/MovieScene/Public/MovieSceneBindingOverrides.h b/Engine/Source/Runtime/MovieScene/Public/MovieSceneBindingOverrides.h
--- a/Engine/Source/Runtime/MovieScene/Public/MovieSceneBindingOverrides.h
+++ b/Engine/Source/Runtime/MovieScene/Public/MovieSceneBindingOverrides.h
@@ -27,7 +27,7 @@ struct FMovieSceneBindingOverrideData
 
 	/** Specifies the object to override the binding with. */
 	UPROPERTY(EditAnywhere, Category="Binding")
-	TWeakObjectPtr<UObject> Object;
+	TSoftObjectPtr<UObject> Object;
 
 	/** Specifies whether the default assignment should remain bound (false) or if this should completely override the default binding (true). */
 	UPROPERTY(EditAnywhere, Category="Binding")
