Sequencer: Redesigned the PlayTo API to allow users to specify if the provided time is inclusive or exclusive. This makes the API much more natural to use as you can specify it play to 45 (exclusive) then later jump to 45 and play from there without overlapping evaluations.

Additionally, this fixes some issues with how PlayTo handles the last frame. Previously it would issue a 'Jump' command which skips calling events on the frame being jumped to. But the Jump state also skips sweeping over time ranges. This means that event keys on the last frame of a section (or more depending on framerate) could get skipped, ie:

Playhead is at frame 43. Event key on 44. PlayTo(45). If the delta time was large enough to go from 43 => 45.5 the system would historically issue a Jump command to jump to 45, skipping the event key on 44. It now issues a Play command to cover the gap (which will evaluate the event keys).

The reason a different params struct was added to PlayTo (instead of extending the existing one) is that inclusive/exclusive doesn't make sense for Jump/Seek (which are directionless), and is a params struct so that if future parameters are needed it is easier to add one without breaking blueprints.

diff --git a/Engine/Plugins/MovieScene/TemplateSequence/Source/TemplateSequence/Private/TemplateSequence.cpp b/Engine/Plugins/MovieScene/TemplateSequence/Source/TemplateSequence/Private/TemplateSequence.cpp
--- a/Engine/Plugins/MovieScene/TemplateSequence/Source/TemplateSequence/Private/TemplateSequence.cpp
+++ b/Engine/Plugins/MovieScene/TemplateSequence/Source/TemplateSequence/Private/TemplateSequence.cpp
@@ -149,7 +149,7 @@ FGuid UTemplateSequence::FindOrAddBinding(UObject* InObject)
 
 		Player->Initialize(this, World, FMovieSceneSequencePlaybackSettings());
 		Player->State.AssignSequence(MovieSceneSequenceID::Root, *this, *Player);
-		Player->PlayTo(FMovieSceneSequencePlaybackParams(MovieScene->GetPlaybackRange().GetLowerBoundValue(), EUpdatePositionMethod::Play));
+		Player->PlayTo(FMovieSceneSequencePlaybackParams(MovieScene->GetPlaybackRange().GetLowerBoundValue(), EUpdatePositionMethod::Play), FMovieSceneSequencePlayToParams());
 
 		FGuid FoundBinding;
 		for (int32 BindingIndex = 0; !FoundBinding.IsValid() && BindingIndex < GetMovieScene()->GetBindings().Num(); ++BindingIndex)
diff --git a/Engine/Source/Runtime/MovieScene/Private/MovieSceneSequencePlayer.cpp b/Engine/Source/Runtime/MovieScene/Private/MovieSceneSequencePlayer.cpp
--- a/Engine/Source/Runtime/MovieScene/Private/MovieSceneSequencePlayer.cpp
+++ b/Engine/Source/Runtime/MovieScene/Private/MovieSceneSequencePlayer.cpp
@@ -523,11 +523,14 @@ void UMovieSceneSequencePlayer::SetTimeRange( float StartTimeSeconds, float Dura
 	SetFrameRange(StartFrame.Value, Duration.Value);
 }
 
-void UMovieSceneSequencePlayer::PlayTo(FMovieSceneSequencePlaybackParams InPlaybackParams)
+void UMovieSceneSequencePlayer::PlayTo(FMovieSceneSequencePlaybackParams InPlaybackParams, FMovieSceneSequencePlayToParams PlayToParams)
 {
-	PauseOnFrame = InPlaybackParams.GetPlaybackPosition(this);
+	FPauseOnArgs Args;
+	Args.Time = InPlaybackParams.GetPlaybackPosition(this);
+	Args.bExclusive = PlayToParams.bExclusive;
+	PauseOnFrame = Args;
 
-	if (GetCurrentTime().Time < PauseOnFrame.GetValue())
+	if (GetCurrentTime().Time < Args.Time)
 	{
 		Play();
 	}
@@ -628,22 +631,47 @@ bool UMovieSceneSequencePlayer::ShouldStopOrLoop(FFrameTime NewPosition) const
 	return bShouldStopOrLoop;
 }
 
-bool UMovieSceneSequencePlayer::ShouldPause(FFrameTime NewPosition) const
+TOptional<TRange<FFrameTime>> UMovieSceneSequencePlayer::GetPauseRange(const FFrameTime& NewPosition) const
 {
-	bool bShouldPause = false;
 	if (IsPlaying() && PauseOnFrame.IsSet())
 	{
-		if (!bReversePlayback)
+		TRange<FFrameTime> OutRange;
+
+		// NewPosition and PauseOnFrame are stored in Display Rate, but we need to convert to tick resolution for the actual ranges.
+		FFrameTime PauseTime = FFrameRate::TransformTime(PauseOnFrame->Time, PlayPosition.GetInputRate(), PlayPosition.GetOutputRate());
+		FFrameTime CurrentTime = FFrameRate::TransformTime(PlayPosition.GetCurrentPosition(), PlayPosition.GetInputRate(), PlayPosition.GetOutputRate());
+		FFrameTime NewTime = FFrameRate::TransformTime(NewPosition, PlayPosition.GetInputRate(), PlayPosition.GetOutputRate());
+
+		// This logic is a little complicated due to inclusive/exclusive range times for pausing. Inclusive/Exclusive only have any meaning
+		// with direction (as we need direction to decide which bound to make inclusive/exclusive).
+		if (bReversePlayback)
 		{
-			bShouldPause = PauseOnFrame.GetValue() <= NewPosition;
+
+			// If we're going in reverse then the start frame is the exclusive one
+			OutRange = TRange<FFrameTime>(PauseOnFrame->bExclusive
+				? TRangeBound<FFrameTime>::Inclusive(PauseTime + FFrameTime(FFrameNumber(1)))
+				: TRangeBound<FFrameTime>::Inclusive(PauseTime),
+				TRangeBound<FFrameTime>::Inclusive(CurrentTime));
+			
 		}
 		else
 		{
-			bShouldPause = PauseOnFrame.GetValue() >= NewPosition;
+			// If we're playing forward then the upper bound is (potentially) exclusive
+			OutRange = TRange<FFrameTime>(TRangeBound<FFrameTime>::Inclusive(CurrentTime),
+				PauseOnFrame->bExclusive
+				? TRangeBound<FFrameTime>::Exclusive(PauseTime)
+				: TRangeBound<FFrameTime>::Exclusive(PauseTime + FFrameTime(FFrameNumber(1))));
+
+		}
+
+		// If the new time is outside of bounds, then we should pause.
+		if (!OutRange.Contains(NewTime))
+		{
+			return OutRange;
 		}
 	}
 
-	return bShouldPause;
+	return TOptional<TRange<FFrameTime>>();
 }
 
 UMovieSceneEntitySystemLinker* UMovieSceneSequencePlayer::ConstructEntitySystemLinker()
@@ -856,15 +884,9 @@ void UMovieSceneSequencePlayer::UpdateTimeCursorPosition_Internal(FFrameTime New
 		bPendingOnStartedPlaying = false;
 	}
 
-	if (Method == EUpdatePositionMethod::Play && ShouldPause(NewPosition))
-	{
-		if (PauseOnFrame.GetValue() != PlayPosition.GetCurrentPosition())
-		{
-			UpdateTimeCursorPosition(PauseOnFrame.GetValue(), EUpdatePositionMethod::Jump);
-		}
-		Pause();
-	}
-	else if (Method == EUpdatePositionMethod::Play && ShouldStopOrLoop(NewPosition))
+	// If we should pause during this evaluation, we'll handle that below.
+	const TOptional<TRange<FFrameTime>> PauseRange = GetPauseRange(NewPosition);
+	if (!PauseRange.IsSet() && Method == EUpdatePositionMethod::Play && ShouldStopOrLoop(NewPosition))
 	{
 		// The actual start time taking into account reverse playback
 		FFrameNumber StartTimeWithReversed = bReversePlayback ? GetLastValidTime().FrameNumber : StartTime;
@@ -960,19 +982,38 @@ void UMovieSceneSequencePlayer::UpdateTimeCursorPosition_Internal(FFrameTime New
 	}
 	else
 	{
+		// If the desired evaluation will take us past where we want to go we need to use a clipped range provided by the PauseRange, otherwise use the normal one.
+		FMovieSceneEvaluationRange Range = PauseRange.IsSet() 
+			? FMovieSceneEvaluationRange(PauseRange.GetValue(), PlayPosition.GetOutputRate(), bReversePlayback ? EPlayDirection::Backwards : EPlayDirection::Forwards)
+			: UpdatePlayPosition(PlayPosition, NewPosition, Method);
+
+		UMovieSceneSequence* MovieSceneSequence = RootTemplateInstance.GetSequence(MovieSceneSequenceID::Root);
+		const bool bIsSequenceBlocking = EnumHasAnyFlags(MovieSceneSequence->GetFlags(), EMovieSceneSequenceFlags::BlockingEvaluation);
+		
 		// Just update the time and sequence... if we are in the main level update we want, if possible,
 		// to only queue this sequence's update, so everything updates in parallel. If not possible, or if
 		// not in the main level update, we run the evaluation synchronously.
-		
-		UMovieSceneSequence* MovieSceneSequence = RootTemplateInstance.GetSequence(MovieSceneSequenceID::Root);
-		const bool bIsSequenceBlocking = EnumHasAnyFlags(MovieSceneSequence->GetFlags(), EMovieSceneSequenceFlags::BlockingEvaluation);
-		FMovieSceneEvaluationRange Range = UpdatePlayPosition(PlayPosition, NewPosition, Method);
 		FMovieSceneUpdateArgs Args;
 		Args.bIsAsync = (bIsMainLevelUpdate && !bIsSequenceBlocking);
 
 		PostEvaluationCallbacks.Add(FOnEvaluationCallback::CreateUObject(this, &UMovieSceneSequencePlayer::UpdateNetworkSyncProperties));
 
 		UpdateMovieSceneInstance(Range, StatusOverride, Args);
+
+		// Now that the evaluation has taken place we call Pause (to trigger delegates, etc.), however if we're running the evaluation
+		// as async we actually need to queue a latent action to do it so that it happens after the data is actually updated. We can't
+		// use the QueueLatentAction inside of Pause() because bIsEvaluating is only set during actual evaluation so it won't be set yet.
+		if (PauseRange.IsSet())
+		{
+			if (Args.bIsAsync)
+			{
+				QueueLatentAction(FMovieSceneSequenceLatentActionDelegate::CreateUObject(this, &UMovieSceneSequencePlayer::Pause));
+			}
+			else
+			{
+				Pause();
+			}
+		}
 	}
 
 	// WARNING: DO NOT CHANGE PLAYER STATE ANYMORE HERE!
diff --git a/Engine/Source/Runtime/MovieScene/Public/MovieSceneSequencePlayer.h b/Engine/Source/Runtime/MovieScene/Public/MovieSceneSequencePlayer.h
--- a/Engine/Source/Runtime/MovieScene/Public/MovieSceneSequencePlayer.h
+++ b/Engine/Source/Runtime/MovieScene/Public/MovieSceneSequencePlayer.h
@@ -217,6 +217,16 @@ struct FMovieSceneSequencePlaybackParams
 	EUpdatePositionMethod UpdateMethod;
 };
 
+USTRUCT(BlueprintType)
+struct FMovieSceneSequencePlayToParams
+{
+	GENERATED_BODY()
+
+	/** Should the PlayTo time be considered exclusive? Defaults to true as end frames in Sequencer are exclusive by default. */
+	UPROPERTY(EditAnywhere, BlueprintReadWrite, Category="Cinematic")
+	bool bExclusive = true;
+};
+
 template<> struct TStructOpsTypeTraits<FMovieSceneSequencePlaybackSettings> : public TStructOpsTypeTraitsBase2<FMovieSceneSequencePlaybackSettings>
 {
 	enum { WithCopy = true, WithStructuredSerializeFromMismatchedTag = true };
@@ -350,7 +360,7 @@ public:
 	 * @param PlaybackParams The position settings (ie. the position to play to)
 	 */
 	UFUNCTION(BlueprintCallable, Category = "Sequencer|Player")
-	void PlayTo(FMovieSceneSequencePlaybackParams PlaybackParams);
+	void PlayTo(FMovieSceneSequencePlaybackParams PlaybackParams, FMovieSceneSequencePlayToParams PlayToParams);
 
 	/**
 	 * Set the current time of the player by evaluating from the current time to the specified time, as if the sequence is playing. 
@@ -524,7 +534,12 @@ protected:
 
 	void UpdateTimeCursorPosition(FFrameTime NewPosition, EUpdatePositionMethod Method);
 	bool ShouldStopOrLoop(FFrameTime NewPosition) const;
-	bool ShouldPause(FFrameTime NewPosition) const;
+	/** 
+	* If the current sequence should pause (due to NewPosition overshooting a previously set ShouldPause) 
+	* then a range of time that should be evaluated to reach there will be returned. If we should not pause
+	* then the TOptional will be unset.
+	* */
+	TOptional<TRange<FFrameTime>> GetPauseRange(const FFrameTime& NewPosition) const;
 
 	UWorld* GetPlaybackWorld() const;
 
@@ -688,8 +703,14 @@ private:
 	*/
 	TOptional<float> LastTickGameTimeSeconds;
 
+	struct FPauseOnArgs
+	{
+		FFrameTime Time;
+		bool bExclusive;
+	};
+
 	/** If set, pause playback on this frame */
-	TOptional<FFrameTime> PauseOnFrame;
+	TOptional<FPauseOnArgs> PauseOnFrame;
 
 	/** Pre and post evaluation callbacks, for async evaluations */
 	DECLARE_DELEGATE(FOnEvaluationCallback);
