Sequencer: Add option to remove tracks on cook. This is enabled by a cvar: MovieScene.RemoveMutedTracksOnCook

Take care when using this in the case that a blueprint is enabling a muted track. You do not want this to remove a muted track.

diff --git a/Engine/Source/Runtime/MovieScene/Private/MovieScene.cpp b/Engine/Source/Runtime/MovieScene/Private/MovieScene.cpp
--- a/Engine/Source/Runtime/MovieScene/Private/MovieScene.cpp
+++ b/Engine/Source/Runtime/MovieScene/Private/MovieScene.cpp
@@ -90,6 +90,43 @@ void UMovieScene::Serialize( FArchive& Ar )
 	Ar.UsingCustomVersion(FMovieSceneEvaluationCustomVersion::GUID);
 	Ar.UsingCustomVersion(FSequencerObjectVersion::GUID);
 
+#if WITH_EDITOR
+
+	// Perform optimizations for cooking
+	if (Ar.IsCooking())
+	{
+		for (int32 MasterTrackIndex = 0; MasterTrackIndex < MasterTracks.Num(); )
+		{
+			UMovieSceneTrack* MasterTrack = MasterTracks[MasterTrackIndex];
+			if (MasterTrack && MasterTrack->GetCookOptimizationFlags() == ECookOptimizationFlags::RemoveTrack)
+			{				
+				MasterTrack->RemoveForCook();
+				RemoveMasterTrack(*MasterTrack);
+				UE_LOG(LogMovieScene, Display, TEXT("Removing muted track: %s from: %s"), *MasterTrack->GetDisplayName().ToString(), *GetPathName());
+				continue;
+			}
+			++MasterTrackIndex;
+		}
+
+		for (FMovieSceneBinding& ObjectBinding : ObjectBindings)
+		{
+			for (int32 TrackIndex = 0; TrackIndex < ObjectBinding.GetTracks().Num(); )
+			{
+				UMovieSceneTrack* Track = ObjectBinding.GetTracks()[TrackIndex];
+				if (Track && Track->GetCookOptimizationFlags() == ECookOptimizationFlags::RemoveTrack)
+				{
+					Track->RemoveForCook();
+					ObjectBinding.RemoveTrack(*Track);
+					UE_LOG(LogMovieScene, Display, TEXT("Removing muted track: %s from: %s"), *Track->GetDisplayName().ToString(), *GetPathName());
+					continue;
+				}
+				++TrackIndex;
+			}
+		}
+	}
+
+#endif // WITH_EDITOR
+
 	// Serialize the MovieScene
 	Super::Serialize(Ar);
 
diff --git a/Engine/Source/Runtime/MovieScene/Private/MovieSceneTrack.cpp b/Engine/Source/Runtime/MovieScene/Private/MovieSceneTrack.cpp
--- a/Engine/Source/Runtime/MovieScene/Private/MovieSceneTrack.cpp
+++ b/Engine/Source/Runtime/MovieScene/Private/MovieSceneTrack.cpp
@@ -17,6 +17,16 @@
 
 #include "Evaluation/MovieSceneEvaluationCustomVersion.h"
 
+#include "Channels/MovieSceneChannelProxy.h"
+#include "Channels/MovieSceneChannel.h"
+
+int32 GMovieSceneRemoveMutedTracksOnCook = 0;
+static FAutoConsoleVariableRef CVarMovieSceneRemoveMutedTracksOnCook(
+	TEXT("MovieScene.RemoveMutedTracksOnCook"),
+	GMovieSceneRemoveMutedTracksOnCook,
+	TEXT("If 1 remove muted tracks on cook, otherwise leave as is."),
+	ECVF_Default);
+
 UMovieSceneTrack::UMovieSceneTrack(const FObjectInitializer& InInitializer)
 	: Super(InInitializer)
 {
@@ -271,6 +281,42 @@ bool UMovieSceneTrack::FixRowIndices()
 	return bFixesMade;
 }
 
+#if WITH_EDITOR
+
+ECookOptimizationFlags UMovieSceneTrack::GetCookOptimizationFlags() const
+{
+	const int RemoveMutedTracksOnCook = CVarMovieSceneRemoveMutedTracksOnCook->GetInt();
+
+	if (RemoveMutedTracksOnCook && IsEvalDisabled())
+	{
+		return ECookOptimizationFlags::RemoveTrack;
+	}
+	return ECookOptimizationFlags::None; 
+}
+
+void UMovieSceneTrack::RemoveForCook()
+{
+	for (UMovieSceneSection* Section : GetAllSections())
+	{
+		if (Section)
+		{
+			for (const FMovieSceneChannelEntry& Entry : Section->GetChannelProxy().GetAllEntries())
+			{
+				for (FMovieSceneChannel* Channel : Entry.GetChannels())
+				{
+					if (Channel)
+					{
+						Channel->Reset();
+					}
+				}
+			}
+		}
+	}
+	RemoveAllAnimationData();
+}
+
+#endif
+
 bool UMovieSceneTrack::IsRowEvalDisabled(int32 RowIndex) const
 {
 	return RowsDisabled.Contains(RowIndex);
diff --git a/Engine/Source/Runtime/MovieScene/Public/MovieSceneTrack.h b/Engine/Source/Runtime/MovieScene/Public/MovieSceneTrack.h
--- a/Engine/Source/Runtime/MovieScene/Public/MovieSceneTrack.h
+++ b/Engine/Source/Runtime/MovieScene/Public/MovieSceneTrack.h
@@ -22,6 +22,18 @@ struct IMovieSceneTemplateGenerator;
 
 template<typename> struct TMovieSceneEvaluationTree;
 
+/** Flags used to perform cook-time optimization of movie scene data */
+enum class ECookOptimizationFlags
+{
+	/** Perform no cook optimization */
+	None 			= 0,
+	/** Remove this track since its of no consequence to runtime */
+	RemoveTrack		= 1 << 0,
+	/** Remove this track's object since its of no consequence to runtime */
+	RemoveObject	= 1 << 1,
+};
+ENUM_CLASS_FLAGS(ECookOptimizationFlags)
+
 /** Generic evaluation options for any track */
 USTRUCT()
 struct FMovieSceneTrackEvalOptions
@@ -384,6 +396,21 @@ public:
 	 */
 	virtual void RemoveSectionAt(int32 SectionIndex) PURE_VIRTUAL(UMovieSceneSection::RemoveSectionAt, );
 
+#if WITH_EDITOR
+
+	/**
+	 * Called when this track's movie scene is being cooked to determine if/how this track should be cooked.
+	 * @return ECookOptimizationFlags detailing how to optimize this track
+	 */
+	MOVIESCENE_API virtual ECookOptimizationFlags GetCookOptimizationFlags() const;
+
+	/**
+	 * Called when this track should be removed for cooking
+	 */
+	MOVIESCENE_API virtual void RemoveForCook();
+
+#endif
+
 #if WITH_EDITORONLY_DATA
 
 	/**
