Sequencer: TryLoad subsequence if it doesn't resolve. Added a few checks for valid subsequence movie scene.

diff --git a/Engine/Source/Runtime/MovieScene/Private/Compilation/MovieSceneCompiledDataManager.cpp b/Engine/Source/Runtime/MovieScene/Private/Compilation/MovieSceneCompiledDataManager.cpp
--- a/Engine/Source/Runtime/MovieScene/Private/Compilation/MovieSceneCompiledDataManager.cpp
+++ b/Engine/Source/Runtime/MovieScene/Private/Compilation/MovieSceneCompiledDataManager.cpp
@@ -696,29 +696,32 @@ void UMovieSceneCompiledDataManager::Compile(FMovieSceneCompiledDataID DataID, U
 	{
 		UMovieScene* MovieScene = Sequence->GetMovieScene();
 
-		for (const FMovieSceneMarkedFrame& Mark : MovieScene->GetMarkedFrames())
+		if (ensure(MovieScene))
 		{
-			if (Mark.bIsDeterminismFence)
+			for (const FMovieSceneMarkedFrame& Mark : MovieScene->GetMarkedFrames())
 			{
-				GatheredData.DeterminismData.Fences.Add(Mark.FrameNumber);
+				if (Mark.bIsDeterminismFence)
+				{
+					GatheredData.DeterminismData.Fences.Add(Mark.FrameNumber);
+				}
 			}
-		}
 
-		if (UMovieSceneTrack* Track = MovieScene->GetCameraCutTrack())
-		{
-			CompileTrack(&Entry, nullptr, Track, Params, &GatheredSignatures, &GatheredData);
-		}
+			if (UMovieSceneTrack* Track = MovieScene->GetCameraCutTrack())
+			{
+				CompileTrack(&Entry, nullptr, Track, Params, &GatheredSignatures, &GatheredData);
+			}
 
-		for (UMovieSceneTrack* Track : MovieScene->GetMasterTracks())
-		{
-			CompileTrack(&Entry, nullptr, Track, Params, &GatheredSignatures, &GatheredData);
-		}
+			for (UMovieSceneTrack* Track : MovieScene->GetMasterTracks())
+			{
+				CompileTrack(&Entry, nullptr, Track, Params, &GatheredSignatures, &GatheredData);
+			}
 
-		for (const FMovieSceneBinding& ObjectBinding : MovieScene->GetBindings())
-		{
-			for (UMovieSceneTrack* Track : ObjectBinding.GetTracks())
+			for (const FMovieSceneBinding& ObjectBinding : MovieScene->GetBindings())
 			{
-				CompileTrack(&Entry, &ObjectBinding, Track, Params, &GatheredSignatures, &GatheredData);
+				for (UMovieSceneTrack* Track : ObjectBinding.GetTracks())
+				{
+					CompileTrack(&Entry, &ObjectBinding, Track, Params, &GatheredSignatures, &GatheredData);
+				}
 			}
 		}
 	}
@@ -822,21 +825,24 @@ void UMovieSceneCompiledDataManager::Gather(const FMovieSceneCompiledDataEntry&
 
 	UMovieScene* MovieScene = Sequence->GetMovieScene();
 
-	if (UMovieSceneTrack* Track = MovieScene->GetCameraCutTrack())
+	if (ensure(MovieScene))
 	{
-		GatherTrack(nullptr, Track, Params, TrackTemplate, OutCompilerData);
-	}
+		if (UMovieSceneTrack* Track = MovieScene->GetCameraCutTrack())
+		{
+			GatherTrack(nullptr, Track, Params, TrackTemplate, OutCompilerData);
+		}
 
-	for (UMovieSceneTrack* Track : MovieScene->GetMasterTracks())
-	{
-		GatherTrack(nullptr, Track, Params, TrackTemplate, OutCompilerData);
-	}
+		for (UMovieSceneTrack* Track : MovieScene->GetMasterTracks())
+		{
+			GatherTrack(nullptr, Track, Params, TrackTemplate, OutCompilerData);
+		}
 
-	for (const FMovieSceneBinding& ObjectBinding : MovieScene->GetBindings())
-	{
-		for (UMovieSceneTrack* Track : ObjectBinding.GetTracks())
+		for (const FMovieSceneBinding& ObjectBinding : MovieScene->GetBindings())
 		{
-			GatherTrack(&ObjectBinding, Track, Params, TrackTemplate, OutCompilerData);
+			for (UMovieSceneTrack* Track : ObjectBinding.GetTracks())
+			{
+				GatherTrack(&ObjectBinding, Track, Params, TrackTemplate, OutCompilerData);
+			}
 		}
 	}
 }
@@ -1385,6 +1391,12 @@ bool UMovieSceneCompiledDataManager::GenerateSubSequenceData(UMovieSceneSubTrack
 			continue;
 		}
 
+		UMovieScene* MovieScene = SubSequence->GetMovieScene();
+		if (!MovieScene)
+		{
+			continue;
+		}
+
 		const FMovieSceneSequenceID InnerSequenceID = RootPath->ResolveChildSequenceID(SubSection->GetSequenceID());
 
 		FSubSequenceInstanceDataParams InstanceParams{ InnerSequenceID, Operand };
diff --git a/Engine/Source/Runtime/MovieScene/Private/Compilation/MovieSceneCompiledVolatilityManager.cpp b/Engine/Source/Runtime/MovieScene/Private/Compilation/MovieSceneCompiledVolatilityManager.cpp
--- a/Engine/Source/Runtime/MovieScene/Private/Compilation/MovieSceneCompiledVolatilityManager.cpp
+++ b/Engine/Source/Runtime/MovieScene/Private/Compilation/MovieSceneCompiledVolatilityManager.cpp
@@ -64,7 +64,7 @@ bool FCompiledDataVolatilityManager::HasBeenRecompiled(FMovieSceneCompiledDataID
 		for (const TTuple<FMovieSceneSequenceID, FMovieSceneSubSequenceData>& Pair : Hierarchy->AllSubSequenceData())
 		{
 			FMovieSceneCompiledDataID SubDataID = CompiledDataManager->GetSubDataID(RootDataID, Pair.Key);
-			if (HasSequenceBeenRecompiled(SubDataID, Pair.Key, CompiledDataManager))
+			if (SubDataID.IsValid() && HasSequenceBeenRecompiled(SubDataID, Pair.Key, CompiledDataManager))
 			{
 				return true;
 			}
@@ -124,6 +124,12 @@ void FCompiledDataVolatilityManager::UpdateCachedSignatures(IMovieScenePlayer& P
 		for (const TTuple<FMovieSceneSequenceID, FMovieSceneSubSequenceData>& SubData : Hierarchy->AllSubSequenceData())
 		{
 			const FMovieSceneCompiledDataID     SubDataID = CompiledDataManager->GetSubDataID(RootDataID, SubData.Key);
+			if (!SubDataID.IsValid())
+			{
+				UE_LOG(LogMovieScene, Error, TEXT("Invalid SubDataID for: %s"), *SubData.Value.Sequence.ToString());
+				continue;
+			}
+
 			const FMovieSceneCompiledDataEntry& SubEntry  = CompiledDataManager->GetEntryRef(SubDataID);
 
 			CachedCompilationSignatures.Add(SubData.Key, SubEntry.CompiledSignature);
diff --git a/Engine/Source/Runtime/MovieScene/Private/Evaluation/MovieSceneSequenceHierarchy.cpp b/Engine/Source/Runtime/MovieScene/Private/Evaluation/MovieSceneSequenceHierarchy.cpp
--- a/Engine/Source/Runtime/MovieScene/Private/Evaluation/MovieSceneSequenceHierarchy.cpp
+++ b/Engine/Source/Runtime/MovieScene/Private/Evaluation/MovieSceneSequenceHierarchy.cpp
@@ -87,6 +87,12 @@ UMovieSceneSequence* FMovieSceneSubSequenceData::GetSequence() const
 	if (!ResolvedSequence)
 	{
 		ResolvedSequence = Cast<UMovieSceneSequence>(Sequence.ResolveObject());
+
+		if (!ResolvedSequence)
+		{
+			ResolvedSequence = Cast<UMovieSceneSequence>(Sequence.TryLoad());
+		}
+
 		CachedSequence = ResolvedSequence;
 	}
 	return ResolvedSequence;
