Movie Scene: Log updates

- Added a LogMovieSceneECS and moved entity system logs to it - this helps distinguish between ECS notifications vs frame/timing information
- Added client name to Evaluating sequence log
- Remove outdaded "Clock Tick" words from TimeController log

LogMovieScene: VeryVerbose: Evaluating sequence SimpleSequence at frame 133, subframe [REDACTED_PHONE] fps).
LogMovieScene: VeryVerbose: Evaluating sequence SimpleSequence (client 0) at frame 109, subframe [REDACTED_PHONE] fps).
LogMovieScene: VeryVerbose: Evaluating sequence SimpleSequence (client 1) at frame 106, subframe [REDACTED_PHONE] fps).

diff --git a/Engine/Source/Runtime/MovieScene/Private/EntitySystem/MovieSceneEntitySystem.cpp b/Engine/Source/Runtime/MovieScene/Private/EntitySystem/MovieSceneEntitySystem.cpp
--- a/Engine/Source/Runtime/MovieScene/Private/EntitySystem/MovieSceneEntitySystem.cpp
+++ b/Engine/Source/Runtime/MovieScene/Private/EntitySystem/MovieSceneEntitySystem.cpp
@@ -357,7 +357,7 @@ void UMovieSceneEntitySystem::Run(FSystemTaskPrerequisites& InPrerequisites, FSy
 
 	Linker->EntityManager.IncrementSystemSerial();
 
-	UE_LOG(LogMovieScene, Verbose, TEXT("Running moviescene system for phase %d: %s"), (int32)Phase, *GetName());
+	UE_LOG(LogMovieSceneECS, Verbose, TEXT("Running moviescene system for phase %d: %s"), (int32)Phase, *GetName());
 	OnRun(InPrerequisites, Subsequents);
 }
 
diff --git a/Engine/Source/Runtime/MovieScene/Private/EntitySystem/MovieSceneEntitySystemLinker.cpp b/Engine/Source/Runtime/MovieScene/Private/EntitySystem/MovieSceneEntitySystemLinker.cpp
--- a/Engine/Source/Runtime/MovieScene/Private/EntitySystem/MovieSceneEntitySystemLinker.cpp
+++ b/Engine/Source/Runtime/MovieScene/Private/EntitySystem/MovieSceneEntitySystemLinker.cpp
@@ -202,7 +202,7 @@ bool UMovieSceneEntitySystemLinker::StartEvaluation(FMovieSceneEntitySystemRunne
 		return true;
 	}
 		
-	UE_LOG(LogMovieScene, Warning, TEXT("Can't start a new evaluation: the active runner is not in a re-entrancy window."));
+	UE_LOG(LogMovieSceneECS, Warning, TEXT("Can't start a new evaluation: the active runner is not in a re-entrancy window."));
 	return false;
 }
 
diff --git a/Engine/Source/Runtime/MovieScene/Private/EntitySystem/MovieSceneEntitySystemRunner.cpp b/Engine/Source/Runtime/MovieScene/Private/EntitySystem/MovieSceneEntitySystemRunner.cpp
--- a/Engine/Source/Runtime/MovieScene/Private/EntitySystem/MovieSceneEntitySystemRunner.cpp
+++ b/Engine/Source/Runtime/MovieScene/Private/EntitySystem/MovieSceneEntitySystemRunner.cpp
@@ -122,7 +122,7 @@ void FMovieSceneEntitySystemRunner::Update(const FMovieSceneContext& Context, FI
 {
 	if (UpdateQueue.Num() > 0)
 	{
-		UE_LOG(LogMovieScene, Warning, TEXT("Updates are already queued! This will run those updates as well, which might not be what's intended."));
+		UE_LOG(LogMovieSceneECS, Warning, TEXT("Updates are already queued! This will run those updates as well, which might not be what's intended."));
 	}
 
 	// Queue our one update and flush immediately.
@@ -144,7 +144,7 @@ void FMovieSceneEntitySystemRunner::Flush()
 
 	if (!Linker->StartEvaluation(*this))
 	{
-		UE_LOG(LogMovieScene, Error, TEXT("An evaluation is already running on this linker, and re-entrancy isn't allow at this point"));
+		UE_LOG(LogMovieSceneECS, Error, TEXT("An evaluation is already running on this linker, and re-entrancy isn't allow at this point"));
 		return;
 	}
 
diff --git a/Engine/Source/Runtime/MovieScene/Private/EntitySystem/MovieSceneInstanceRegistry.cpp b/Engine/Source/Runtime/MovieScene/Private/EntitySystem/MovieSceneInstanceRegistry.cpp
--- a/Engine/Source/Runtime/MovieScene/Private/EntitySystem/MovieSceneInstanceRegistry.cpp
+++ b/Engine/Source/Runtime/MovieScene/Private/EntitySystem/MovieSceneInstanceRegistry.cpp
@@ -100,7 +100,7 @@ void FInstanceRegistry::DestroyInstance(FInstanceHandle InstanceHandle)
 		const bool bHasFinished = (GExitPurge || Instance.HasFinished());
 		if (!bHasFinished)
 		{
-			UE_LOG(LogMovieScene, Verbose, TEXT("Instance being destroyed without finishing evaluation."));
+			UE_LOG(LogMovieSceneECS, Verbose, TEXT("Instance being destroyed without finishing evaluation."));
 		}
 		Instance.DestroyImmediately(Linker);
 		Instances.RemoveAt(InstanceHandle.InstanceID);
diff --git a/Engine/Source/Runtime/MovieScene/Private/EntitySystem/MovieScenePropertyRegistry.cpp b/Engine/Source/Runtime/MovieScene/Private/EntitySystem/MovieScenePropertyRegistry.cpp
--- a/Engine/Source/Runtime/MovieScene/Private/EntitySystem/MovieScenePropertyRegistry.cpp
+++ b/Engine/Source/Runtime/MovieScene/Private/EntitySystem/MovieScenePropertyRegistry.cpp
@@ -89,7 +89,7 @@ TOptional<FResolvedProperty> FPropertyRegistry::ResolveProperty(UObject* Object,
 	TSharedPtr<FTrackInstancePropertyBindings> SlowBindings = MakeShared<FTrackInstancePropertyBindings>(PropertyBinding.PropertyName, PropertyBinding.PropertyPath.ToString());
 	if (SlowBindings->GetProperty(*Object) == nullptr)
 	{
-		UE_LOG(LogMovieScene, Warning, TEXT("Unable to resolve property '%s' from '%s' instance '%s'"), *PropertyBinding.PropertyPath.ToString(), *Object->GetClass()->GetName(), *Object->GetName());
+		UE_LOG(LogMovieSceneECS, Warning, TEXT("Unable to resolve property '%s' from '%s' instance '%s'"), *PropertyBinding.PropertyPath.ToString(), *Object->GetClass()->GetName(), *Object->GetName());
 		return TOptional<FResolvedProperty>();
 	}
 
diff --git a/Engine/Source/Runtime/MovieScene/Private/EntitySystem/MovieSceneSequenceInstance.cpp b/Engine/Source/Runtime/MovieScene/Private/EntitySystem/MovieSceneSequenceInstance.cpp
--- a/Engine/Source/Runtime/MovieScene/Private/EntitySystem/MovieSceneSequenceInstance.cpp
+++ b/Engine/Source/Runtime/MovieScene/Private/EntitySystem/MovieSceneSequenceInstance.cpp
@@ -308,7 +308,7 @@ void FSequenceInstance::DestroyImmediately(UMovieSceneEntitySystemLinker* Linker
 {
 	if (!Ledger.IsEmpty())
 	{
-		UE_LOG(LogMovieScene, Verbose, TEXT("Instance being destroyed without first having been finished by calling Finish()"));
+		UE_LOG(LogMovieSceneECS, Verbose, TEXT("Instance being destroyed without first having been finished by calling Finish()"));
 		Ledger.UnlinkEverything(Linker);
 	}
 
diff --git a/Engine/Source/Runtime/MovieScene/Private/Evaluation/Instances/MovieSceneTrackEvaluator.cpp b/Engine/Source/Runtime/MovieScene/Private/Evaluation/Instances/MovieSceneTrackEvaluator.cpp
--- a/Engine/Source/Runtime/MovieScene/Private/Evaluation/Instances/MovieSceneTrackEvaluator.cpp
+++ b/Engine/Source/Runtime/MovieScene/Private/Evaluation/Instances/MovieSceneTrackEvaluator.cpp
@@ -67,7 +67,7 @@ FMovieSceneTrackEvaluator::~FMovieSceneTrackEvaluator()
 	const bool bHasFinished = (GExitPurge || ThisFrameMetaData.ActiveEntities.Num() == 0);
 	if (!bHasFinished)
 	{
-		UE_LOG(LogMovieScene, Verbose, TEXT("Evaluator instance being torn down without calling Finish (ThisFrameMetaData has data)"));
+		UE_LOG(LogMovieSceneECS, Verbose, TEXT("Evaluator instance being torn down without calling Finish (ThisFrameMetaData has data)"));
 	}
 }
 
diff --git a/Engine/Source/Runtime/MovieScene/Private/MovieSceneModule.cpp b/Engine/Source/Runtime/MovieScene/Private/MovieSceneModule.cpp
--- a/Engine/Source/Runtime/MovieScene/Private/MovieSceneModule.cpp
+++ b/Engine/Source/Runtime/MovieScene/Private/MovieSceneModule.cpp
@@ -11,6 +11,7 @@
 #include "EntitySystem/BuiltInComponentTypes.h"
 
 DEFINE_LOG_CATEGORY(LogMovieScene);
+DEFINE_LOG_CATEGORY(LogMovieSceneECS);
 
 TAutoConsoleVariable<FString> CVarLegacyConversionFrameRate(
 	TEXT("MovieScene.LegacyConversionFrameRate"),
diff --git a/Engine/Source/Runtime/MovieScene/Private/MovieSceneSequencePlayer.cpp b/Engine/Source/Runtime/MovieScene/Private/MovieSceneSequencePlayer.cpp
--- a/Engine/Source/Runtime/MovieScene/Private/MovieSceneSequencePlayer.cpp
+++ b/Engine/Source/Runtime/MovieScene/Private/MovieSceneSequencePlayer.cpp
@@ -1066,7 +1066,14 @@ void UMovieSceneSequencePlayer::UpdateMovieSceneInstance(FMovieSceneEvaluationRa
 
 #if !NO_LOGGING
 	FQualifiedFrameTime CurrentTime = GetCurrentTime();
-	UE_LOG(LogMovieScene, VeryVerbose, TEXT("Evaluating sequence %s at frame %d, subframe %f (%f fps)."), *MovieSceneSequence->GetName(), CurrentTime.Time.FrameNumber.Value, CurrentTime.Time.GetSubFrame(), CurrentTime.Rate.AsDecimal());
+	FString    SequenceName = MovieSceneSequence->GetName();
+
+	AActor* Actor = GetTypedOuter<AActor>();
+	if (Actor && Actor->GetWorld()->GetNetMode() == NM_Client)
+	{
+		SequenceName += FString::Printf(TEXT(" (client %d)"), GPlayInEditorID - 1);
+	}
+	UE_LOG(LogMovieScene, VeryVerbose, TEXT("Evaluating sequence %s at frame %d, subframe %f (%f fps)."), *SequenceName, CurrentTime.Time.FrameNumber.Value, CurrentTime.Time.GetSubFrame(), CurrentTime.Rate.AsDecimal());
 #endif
 
 	// Once we have updated we must no longer skip updates
diff --git a/Engine/Source/Runtime/MovieScene/Private/MovieSceneTimeController.cpp b/Engine/Source/Runtime/MovieScene/Private/MovieSceneTimeController.cpp
--- a/Engine/Source/Runtime/MovieScene/Private/MovieSceneTimeController.cpp
+++ b/Engine/Source/Runtime/MovieScene/Private/MovieSceneTimeController.cpp
@@ -95,10 +95,10 @@ FFrameTime FMovieSceneTimeController_ExternalClock::OnRequestCurrentTime(const F
 	FFrameTime StartTime = StartTimeIfPlaying->ConvertTo(InCurrentTime.Rate);
 	FFrameTime NewTime   = StartTime + CurrentSequenceTimeSeconds * InCurrentTime.Rate;
 
-
-	UE_LOG(LogMovieScene, VeryVerbose, TEXT("TimeController Clock tick: Clock Start Time: %f, Clock Now: %f, Dilation Offset: %f, Sequence Start Time: frame %d, subframe %f, Sequence Offset Seconds: %f, Sequence Now: frame %d, subframe %f"),
+#if !NO_LOGGING
+	UE_LOG(LogMovieScene, VeryVerbose, TEXT("TimeController Clock Start Time: %f, Clock Now: %f, Dilation Offset: %f, Sequence Start Time: frame %d, subframe %f, Sequence Offset Seconds: %f, Sequence Now: frame %d, subframe %f"),
 		StartedTime, CurrentTime, AccumulatedDilation, StartTime.FrameNumber.Value, StartTime.GetSubFrame(), CurrentSequenceTimeSeconds, NewTime.FrameNumber.Value, NewTime.GetSubFrame());
-
+#endif
 
 	return NewTime;
 }
diff --git a/Engine/Source/Runtime/MovieScene/Public/EntitySystem/MovieSceneEntitySystemTask.h b/Engine/Source/Runtime/MovieScene/Public/EntitySystem/MovieSceneEntitySystemTask.h
--- a/Engine/Source/Runtime/MovieScene/Public/EntitySystem/MovieSceneEntitySystemTask.h
+++ b/Engine/Source/Runtime/MovieScene/Public/EntitySystem/MovieSceneEntitySystemTask.h
@@ -1097,7 +1097,7 @@ struct TEntityTaskBase
 
 	void Run(TaskImpl& TaskImplInstance)
 	{
-		UE_LOG(LogMovieScene, VeryVerbose, TEXT("Running entity task the following components: %s"), *FilteredTask.GetComponents().ToString(EntityManager));
+		UE_LOG(LogMovieSceneECS, VeryVerbose, TEXT("Running entity task the following components: %s"), *FilteredTask.GetComponents().ToString(EntityManager));
 
 		PreTask(&TaskImplInstance);
 
@@ -1211,7 +1211,7 @@ struct TEntityAllocationTaskBase
 
 	void Run(TaskImpl& TaskImplInstance)
 	{
-		UE_LOG(LogMovieScene, VeryVerbose, TEXT("Running entity task the following components: %s"), *ComponentFilter.GetComponents().ToString(EntityManager));
+		UE_LOG(LogMovieSceneECS, VeryVerbose, TEXT("Running entity task the following components: %s"), *ComponentFilter.GetComponents().ToString(EntityManager));
 
 		PreTask(&TaskImplInstance);
 
diff --git a/Engine/Source/Runtime/MovieScene/Public/MovieSceneFwd.h b/Engine/Source/Runtime/MovieScene/Public/MovieSceneFwd.h
--- a/Engine/Source/Runtime/MovieScene/Public/MovieSceneFwd.h
+++ b/Engine/Source/Runtime/MovieScene/Public/MovieSceneFwd.h
@@ -102,6 +102,7 @@ enum class EMovieSceneServerClientMask : uint8
 ENUM_CLASS_FLAGS(EMovieSceneServerClientMask)
 
 MOVIESCENE_API DECLARE_LOG_CATEGORY_EXTERN(LogMovieScene, Log, All);
+MOVIESCENE_API DECLARE_LOG_CATEGORY_EXTERN(LogMovieSceneECS, Log, All);
 DECLARE_STATS_GROUP(TEXT("Movie Scene Evaluation"), STATGROUP_MovieSceneEval, STATCAT_Advanced);
 
 MOVIESCENE_API FFrameRate GetLegacyConversionFrameRate();
