Sequencer: Large world coordinates support

diff --git a/Engine/Source/Editor/SerializedRecorderInterface/Private/MovieSceneSectionSerialization.cpp b/Engine/Source/Editor/SerializedRecorderInterface/Private/MovieSceneSectionSerialization.cpp
--- a/Engine/Source/Editor/SerializedRecorderInterface/Private/MovieSceneSectionSerialization.cpp
+++ b/Engine/Source/Editor/SerializedRecorderInterface/Private/MovieSceneSectionSerialization.cpp
@@ -67,9 +67,19 @@ template<> SERIALIZEDRECORDERINTERFACE_API
 FRunnableThread* TMovieSceneSerializer<FPropertyFileHeader, FSerializedProperty<float>>::Thread = nullptr;
 
 template<> SERIALIZEDRECORDERINTERFACE_API
-FMovieSceneSerializerRunnable<FPropertyFileHeader, FSerializedProperty<FVector>>* TMovieSceneSerializer<FPropertyFileHeader, FSerializedProperty<FVector>>::Runnable = nullptr;
+FMovieSceneSerializerRunnable<FPropertyFileHeader, FSerializedProperty<double>>* TMovieSceneSerializer<FPropertyFileHeader, FSerializedProperty<double>>::Runnable = nullptr;
 template<> SERIALIZEDRECORDERINTERFACE_API
-FRunnableThread* TMovieSceneSerializer<FPropertyFileHeader, FSerializedProperty<FVector>>::Thread = nullptr;
+FRunnableThread* TMovieSceneSerializer<FPropertyFileHeader, FSerializedProperty<double>>::Thread = nullptr;
+
+template<> SERIALIZEDRECORDERINTERFACE_API
+FMovieSceneSerializerRunnable<FPropertyFileHeader, FSerializedProperty<FVector3f>>* TMovieSceneSerializer<FPropertyFileHeader, FSerializedProperty<FVector3f>>::Runnable = nullptr;
+template<> SERIALIZEDRECORDERINTERFACE_API
+FRunnableThread* TMovieSceneSerializer<FPropertyFileHeader, FSerializedProperty<FVector3f>>::Thread = nullptr;
+
+template<> SERIALIZEDRECORDERINTERFACE_API
+FMovieSceneSerializerRunnable<FPropertyFileHeader, FSerializedProperty<FVector3d>>* TMovieSceneSerializer<FPropertyFileHeader, FSerializedProperty<FVector3d>>::Runnable = nullptr;
+template<> SERIALIZEDRECORDERINTERFACE_API
+FRunnableThread* TMovieSceneSerializer<FPropertyFileHeader, FSerializedProperty<FVector3d>>::Thread = nullptr;
 
 template<> SERIALIZEDRECORDERINTERFACE_API
 FMovieSceneSerializerRunnable<FPropertyFileHeader, FSerializedProperty<FColor>>* TMovieSceneSerializer<FPropertyFileHeader, FSerializedProperty<FColor>>::Runnable = nullptr;
diff --git a/Engine/Source/Editor/SerializedRecorderInterface/Public/Serializers/MovieScenePropertySerialization.h b/Engine/Source/Editor/SerializedRecorderInterface/Public/Serializers/MovieScenePropertySerialization.h
--- a/Engine/Source/Editor/SerializedRecorderInterface/Public/Serializers/MovieScenePropertySerialization.h
+++ b/Engine/Source/Editor/SerializedRecorderInterface/Public/Serializers/MovieScenePropertySerialization.h
@@ -13,7 +13,9 @@ enum class ESerializedPropertyType : uint8 {
 	ByteType,
 	EnumType,
 	FloatType,
-	VectorType,
+	DoubleType,
+	Vector3fType,
+	Vector3dType,
 	ColorType,
 	IntegerType,
 	StringType
@@ -74,11 +76,28 @@ struct FPropertyFileHeader
 			{
 				PropertyType = (ESerializedPropertyType::FloatType);
 			}
+			else if (Property->IsA<FDoubleProperty>())
+			{
+				PropertyType = (ESerializedPropertyType::DoubleType);
+			}
 			else if (const FStructProperty* StructProperty = CastField<const FStructProperty>(Property))
 			{
-				if (StructProperty->Struct->GetFName() == NAME_Vector)
+				// LWC_TODO: vector types
+				if (StructProperty->Struct->GetFName() == NAME_Vector3f
+#if UE_LARGE_WORLD_COORDINATES_DISABLED
+						|| StructProperty->Struct->GetFName() == NAME_Vector
+#endif
+						)
+				{
+					PropertyType = (ESerializedPropertyType::Vector3fType);
+				}
+				else if (StructProperty->Struct->GetFName() == NAME_Vector3d
+#if !UE_LARGE_WORLD_COORDINATES_DISABLED
+						|| StructProperty->Struct->GetFName() == NAME_Vector
+#endif
+						)
 				{
-					PropertyType = (ESerializedPropertyType::VectorType);
+					PropertyType = (ESerializedPropertyType::Vector3dType);
 				}
 				else if (StructProperty->Struct->GetFName() == NAME_Color)
 				{
@@ -136,11 +155,17 @@ using FPropertySerializedEnum = FSerializedProperty<int64>;
 using FPropertySerializedEnumFrame = TMovieSceneSerializedFrame<FSerializedProperty<int64>>;
 using FPropertySerializerEnum = TMovieSceneSerializer<FPropertyFileHeader, FSerializedProperty<int64>>;
 
+using FPropertySerializedDouble = FSerializedProperty<double>;
+using FPropertySerializedDoubleFrame = TMovieSceneSerializedFrame<FSerializedProperty<double>>;
+
 using FPropertySerializedFloat = FSerializedProperty<float>;
 using FPropertySerializedFloatFrame = TMovieSceneSerializedFrame<FSerializedProperty<float>>;
 
-using FPropertySerializedVector = FSerializedProperty<FVector>;
-using FPropertySerializedVectorFrame = TMovieSceneSerializedFrame<FSerializedProperty<FVector>>;
+using FPropertySerializedVector3f = FSerializedProperty<FVector3f>;
+using FPropertySerializedVector3fFrame = TMovieSceneSerializedFrame<FSerializedProperty<FVector3f>>;
+
+using FPropertySerializedVector3d = FSerializedProperty<FVector3d>;
+using FPropertySerializedVector3dFrame = TMovieSceneSerializedFrame<FSerializedProperty<FVector3d>>;
 
 using FPropertySerializedColor = FSerializedProperty<FColor>;
 using FPropertySerializedColorFrame = TMovieSceneSerializedFrame<FSerializedProperty<FColor>>;
diff --git a/Engine/Source/Editor/UnrealEd/Private/Fbx/FbxMainExport.cpp b/Engine/Source/Editor/UnrealEd/Private/Fbx/FbxMainExport.cpp
--- a/Engine/Source/Editor/UnrealEd/Private/Fbx/FbxMainExport.cpp
+++ b/Engine/Source/Editor/UnrealEd/Private/Fbx/FbxMainExport.cpp
@@ -18,6 +18,7 @@
 #include "Components/StaticMeshComponent.h"
 #include "Components/LightComponent.h"
 #include "Model.h"
+#include "Channels/MovieSceneDoubleChannel.h"
 #include "Channels/MovieSceneFloatChannel.h"
 #include "Channels/MovieSceneIntegerChannel.h"
 #include "Channels/MovieSceneStringChannel.h"
@@ -82,10 +83,12 @@
 #include "IMovieScenePlayer.h"
 #include "MovieScene.h"
 #include "Tracks/MovieScene3DTransformTrack.h"
+#include "Tracks/MovieSceneDoubleTrack.h"
 #include "Tracks/MovieSceneFloatTrack.h"
 #include "Tracks/MovieSceneSkeletalAnimationTrack.h"
 #include "Sections/MovieSceneSkeletalAnimationSection.h"
 #include "Sections/MovieScene3DTransformSection.h"
+#include "Sections/MovieSceneDoubleSection.h"
 #include "Sections/MovieSceneFloatSection.h"
 #include "Evaluation/MovieScenePlayback.h"
 #include "Evaluation/MovieSceneEvaluationTemplateInstance.h"
@@ -2884,21 +2887,25 @@ void RichCurveInterpolationToFbxInterpolation(ERichCurveInterpMode InInterpolati
 	}
 }
 
-void FFbxExporter::ExportChannelToFbxCurve(FbxAnimCurve& InFbxCurve, const FMovieSceneFloatChannel& InChannel, FFrameRate TickResolution, ERichCurveValueMode ValueMode, bool bNegative, const FMovieSceneSequenceTransform& RootToLocalTransform)
+template<typename ChannelType>
+void FFbxExporter::ExportBezierChannelToFbxCurve(FbxAnimCurve& InFbxCurve, const ChannelType& InChannel, FFrameRate TickResolution, ERichCurveValueMode ValueMode, bool bNegative, const FMovieSceneSequenceTransform& RootToLocalTransform)
 {
+	using ChannelValueType = typename ChannelType::ChannelValueType;
+	using CurveValueType = typename ChannelType::CurveValueType;
+
 	const float NegateFactor = bNegative ? -1.f : 1.f;
 	const float kOneThird = 1.f / 3.f;
 	InFbxCurve.KeyModifyBegin();
 
 	TArrayView<const FFrameNumber>          Times  = InChannel.GetTimes();
-	TArrayView<const FMovieSceneFloatValue> Values = InChannel.GetValues();
+	TArrayView<const ChannelValueType> Values = InChannel.GetValues();
 
 	for (int32 Index = 0; Index < Times.Num(); ++Index)
 	{
-		const FFrameNumber          KeyTime  = Times[Index];
-		const FMovieSceneFloatValue KeyValue = Values[Index];
+		const FFrameNumber       KeyTime  = Times[Index];
+		const ChannelValueType   KeyValue = Values[Index];
 
-		const float Value = (ValueMode == ERichCurveValueMode::Fov ? DefaultCamera->ComputeFocalLength( KeyValue.Value ) : KeyValue.Value) * NegateFactor;
+		const CurveValueType Value = (ValueMode == ERichCurveValueMode::Fov ? DefaultCamera->ComputeFocalLength( KeyValue.Value ) : KeyValue.Value) * NegateFactor;
 
 		FbxTime FbxTime;
 		FbxAnimCurveKey FbxKey;
@@ -2928,14 +2935,14 @@ void FFbxExporter::ExportChannelToFbxCurve(FbxAnimCurve& InFbxCurve, const FMovi
 				//Need to convert UE tangent weight which is the length of the hypotenuse to FBX normalized X(time) weight
 				if (WeightedMode == FbxAnimCurveDef::eWeightedAll || WeightedMode == FbxAnimCurveDef::eWeightedRight)
 				{
-					const float XVal = FMath::Sqrt((KeyValue.Tangent.LeaveTangentWeight * KeyValue.Tangent.LeaveTangentWeight) / (1.0f + LeaveTangent * LeaveTangent));
+					const CurveValueType XVal = FMath::Sqrt((KeyValue.Tangent.LeaveTangentWeight * KeyValue.Tangent.LeaveTangentWeight) / (1.0f + LeaveTangent * LeaveTangent));
 					LeaveTangentWeight = XVal / TimeDiff;
 				}
 
 				//make sure next tangent is weighted else use default weight
 				if ((Values[Index + 1].Tangent.TangentWeightMode == ERichCurveTangentWeightMode::RCTWM_WeightedBoth || Values[Index + 1].Tangent.TangentWeightMode == ERichCurveTangentWeightMode::RCTWM_WeightedArrive) )
 				{
-					const float XVal = FMath::Sqrt((Values[Index + 1].Tangent.ArriveTangentWeight * Values[Index + 1].Tangent.ArriveTangentWeight) / (1.0f + NextArriveTangent * NextArriveTangent));
+					const CurveValueType XVal = FMath::Sqrt((Values[Index + 1].Tangent.ArriveTangentWeight * Values[Index + 1].Tangent.ArriveTangentWeight) / (1.0f + NextArriveTangent * NextArriveTangent));
 					NextArriveTangentWeight = XVal / TimeDiff;
 				}
 				if (bNegative)
@@ -2959,6 +2966,16 @@ void FFbxExporter::ExportChannelToFbxCurve(FbxAnimCurve& InFbxCurve, const FMovi
 	InFbxCurve.KeyModifyEnd();
 }
 
+void FFbxExporter::ExportChannelToFbxCurve(FbxAnimCurve& InFbxCurve, const FMovieSceneFloatChannel& InChannel, FFrameRate TickResolution, ERichCurveValueMode ValueMode, bool bNegative, const FMovieSceneSequenceTransform& RootToLocalTransform)
+{
+	ExportBezierChannelToFbxCurve(InFbxCurve, InChannel, TickResolution, ValueMode, bNegative, RootToLocalTransform);
+}
+
+void FFbxExporter::ExportChannelToFbxCurve(FbxAnimCurve& InFbxCurve, const FMovieSceneDoubleChannel& InChannel, FFrameRate TickResolution, ERichCurveValueMode ValueMode, bool bNegative, const FMovieSceneSequenceTransform& RootToLocalTransform)
+{
+	ExportBezierChannelToFbxCurve(InFbxCurve, InChannel, TickResolution, ValueMode, bNegative, RootToLocalTransform);
+}
+
 void FFbxExporter::ExportChannelToFbxCurve(FbxAnimCurve& InFbxCurve, const FMovieSceneIntegerChannel& InChannel, FFrameRate TickResolution, const FMovieSceneSequenceTransform& RootToLocalTransform)
 {
 	InFbxCurve.KeyModifyBegin();
@@ -3029,27 +3046,27 @@ void FFbxExporter::ExportLevelSequence3DTransformTrack(FbxNode* FbxNode, IMovieS
 	FbxAnimCurve* FbxCurveScaleY = FbxNode->LclScaling.GetCurve(BaseLayer, FBXSDK_CURVENODE_COMPONENT_Y, true);
 	FbxAnimCurve* FbxCurveScaleZ = FbxNode->LclScaling.GetCurve(BaseLayer, FBXSDK_CURVENODE_COMPONENT_Z, true);
 
-	TArrayView<FMovieSceneFloatChannel*> FloatChannels = TransformSection->GetChannelProxy().GetChannels<FMovieSceneFloatChannel>();
+	TArrayView<FMovieSceneDoubleChannel*> DoubleChannels = TransformSection->GetChannelProxy().GetChannels<FMovieSceneDoubleChannel>();
 
 	// Translation
-	ExportChannelToFbxCurve(*FbxCurveTransX, *FloatChannels[0], TickResolution, ERichCurveValueMode::Default, false, RootToLocalTransform);
-	ExportChannelToFbxCurve(*FbxCurveTransY, *FloatChannels[1], TickResolution, ERichCurveValueMode::Default, true, RootToLocalTransform);
-	ExportChannelToFbxCurve(*FbxCurveTransZ, *FloatChannels[2], TickResolution, ERichCurveValueMode::Default, false, RootToLocalTransform);
+	ExportChannelToFbxCurve(*FbxCurveTransX, *DoubleChannels[0], TickResolution, ERichCurveValueMode::Default, false, RootToLocalTransform);
+	ExportChannelToFbxCurve(*FbxCurveTransY, *DoubleChannels[1], TickResolution, ERichCurveValueMode::Default, true, RootToLocalTransform);
+	ExportChannelToFbxCurve(*FbxCurveTransZ, *DoubleChannels[2], TickResolution, ERichCurveValueMode::Default, false, RootToLocalTransform);
 
 	// Scale - don't generate scale keys for cameras
 	if (!bIsCameraActor)
 	{
-		ExportChannelToFbxCurve(*FbxCurveScaleX, *FloatChannels[6], TickResolution, ERichCurveValueMode::Default, false, RootToLocalTransform);
-		ExportChannelToFbxCurve(*FbxCurveScaleY, *FloatChannels[7], TickResolution, ERichCurveValueMode::Default, false, RootToLocalTransform);
-		ExportChannelToFbxCurve(*FbxCurveScaleZ, *FloatChannels[8], TickResolution, ERichCurveValueMode::Default, false, RootToLocalTransform);
+		ExportChannelToFbxCurve(*FbxCurveScaleX, *DoubleChannels[6], TickResolution, ERichCurveValueMode::Default, false, RootToLocalTransform);
+		ExportChannelToFbxCurve(*FbxCurveScaleY, *DoubleChannels[7], TickResolution, ERichCurveValueMode::Default, false, RootToLocalTransform);
+		ExportChannelToFbxCurve(*FbxCurveScaleZ, *DoubleChannels[8], TickResolution, ERichCurveValueMode::Default, false, RootToLocalTransform);
 	}
 
 	// Rotation - bake rotation for cameras and lights
 	if (!bBakeRotations)
 	{
-		ExportChannelToFbxCurve(*FbxCurveRotX, *FloatChannels[3], TickResolution, ERichCurveValueMode::Default, false, RootToLocalTransform);
-		ExportChannelToFbxCurve(*FbxCurveRotY, *FloatChannels[4], TickResolution, ERichCurveValueMode::Default, true, RootToLocalTransform);
-		ExportChannelToFbxCurve(*FbxCurveRotZ, *FloatChannels[5], TickResolution, ERichCurveValueMode::Default, true, RootToLocalTransform);
+		ExportChannelToFbxCurve(*FbxCurveRotX, *DoubleChannels[3], TickResolution, ERichCurveValueMode::Default, false, RootToLocalTransform);
+		ExportChannelToFbxCurve(*FbxCurveRotY, *DoubleChannels[4], TickResolution, ERichCurveValueMode::Default, true, RootToLocalTransform);
+		ExportChannelToFbxCurve(*FbxCurveRotZ, *DoubleChannels[5], TickResolution, ERichCurveValueMode::Default, true, RootToLocalTransform);
 	}
 	else
 	{
@@ -3080,19 +3097,19 @@ void FFbxExporter::ExportLevelSequence3DTransformTrack(FbxNode* FbxNode, IMovieS
 			FFrameTime LocalTime = FFrameRate::TransformTime(FFrameTime(LocalFrame), DisplayRate, TickResolution);
 
 			FVector3f Trans = FVector3f::ZeroVector;
-			FloatChannels[0]->Evaluate(LocalTime, Trans.X);
-			FloatChannels[1]->Evaluate(LocalTime, Trans.Y);
-			FloatChannels[2]->Evaluate(LocalTime, Trans.Z);
+			DoubleChannels[0]->Evaluate(LocalTime, Trans.X);
+			DoubleChannels[1]->Evaluate(LocalTime, Trans.Y);
+			DoubleChannels[2]->Evaluate(LocalTime, Trans.Z);
 
 			FRotator Rotator;
-			FloatChannels[3]->Evaluate(LocalTime, Rotator.Roll);
-			FloatChannels[4]->Evaluate(LocalTime, Rotator.Pitch);
-			FloatChannels[5]->Evaluate(LocalTime, Rotator.Yaw);
+			DoubleChannels[3]->Evaluate(LocalTime, Rotator.Roll);
+			DoubleChannels[4]->Evaluate(LocalTime, Rotator.Pitch);
+			DoubleChannels[5]->Evaluate(LocalTime, Rotator.Yaw);
 
 			FVector3f Scale;
-			FloatChannels[6]->Evaluate(LocalTime, Scale.X);
-			FloatChannels[7]->Evaluate(LocalTime, Scale.Y);
-			FloatChannels[8]->Evaluate(LocalTime, Scale.Z);
+			DoubleChannels[6]->Evaluate(LocalTime, Scale.X);
+			DoubleChannels[7]->Evaluate(LocalTime, Scale.Y);
+			DoubleChannels[8]->Evaluate(LocalTime, Scale.Z);
 
 			FTransform RelativeTransform;
 			RelativeTransform.SetTranslation(Trans);
@@ -3334,6 +3351,7 @@ void FFbxExporter::ExportLevelSequenceTrackChannels( FbxNode* FbxNode, UMovieSce
 	FbxCamera* FbxCamera = FbxNode->GetCamera();
 	FFrameRate TickResolution = Track.GetTypedOuter<UMovieScene>()->GetTickResolution();
 
+	const FName DoubleChannelTypeName = FMovieSceneDoubleChannel::StaticStruct()->GetFName();
 	const FName FloatChannelTypeName = FMovieSceneFloatChannel::StaticStruct()->GetFName();
 	const FName IntegerChannelTypeName = FMovieSceneIntegerChannel::StaticStruct()->GetFName();
 	const FName StringChannelTypeName = FMovieSceneStringChannel::StaticStruct()->GetFName();
@@ -3341,7 +3359,8 @@ void FFbxExporter::ExportLevelSequenceTrackChannels( FbxNode* FbxNode, UMovieSce
 	for (const FMovieSceneChannelEntry& Entry : Section->GetChannelProxy().GetAllEntries())
 	{
 		const FName ChannelTypeName = Entry.GetChannelTypeName();
-		if (ChannelTypeName != FloatChannelTypeName && 
+		if (ChannelTypeName != DoubleChannelTypeName && 
+			ChannelTypeName != FloatChannelTypeName && 
 			ChannelTypeName != IntegerChannelTypeName && 
 			ChannelTypeName != StringChannelTypeName)
 		{
@@ -3355,11 +3374,12 @@ void FFbxExporter::ExportLevelSequenceTrackChannels( FbxNode* FbxNode, UMovieSce
 		{
 			FMovieSceneChannelHandle Channel = ChannelProxy.MakeHandle(ChannelTypeName, Index);
 
+			FMovieSceneDoubleChannel* DoubleChannel = Entry.GetChannelTypeName() == DoubleChannelTypeName ? Channel.Cast<FMovieSceneDoubleChannel>().Get() : nullptr;
 			FMovieSceneFloatChannel* FloatChannel = Entry.GetChannelTypeName() == FloatChannelTypeName ? Channel.Cast<FMovieSceneFloatChannel>().Get() : nullptr;
 			FMovieSceneIntegerChannel* IntegerChannel = Entry.GetChannelTypeName() == IntegerChannelTypeName ? Channel.Cast<FMovieSceneIntegerChannel>().Get() : nullptr;
 			FMovieSceneStringChannel* StringChannel = Entry.GetChannelTypeName() == StringChannelTypeName ? Channel.Cast<FMovieSceneStringChannel>().Get() : nullptr;
 
-			if (!FloatChannel && !IntegerChannel && !StringChannel)
+			if (!DoubleChannel && !FloatChannel && !IntegerChannel && !StringChannel)
 			{
 				continue;
 			}
@@ -3412,7 +3432,11 @@ void FFbxExporter::ExportLevelSequenceTrackChannels( FbxNode* FbxNode, UMovieSce
 
 			if (Property == 0)
 			{
-				if (FloatChannel)
+				if (DoubleChannel)
+				{
+					CreateAnimatableUserProperty(FbxNode, DoubleChannel->GetDefault().Get(MAX_flt), TCHAR_TO_UTF8(*PropertyName), TCHAR_TO_UTF8(*PropertyName));
+				}
+				else if (FloatChannel)
 				{
 					CreateAnimatableUserProperty(FbxNode, FloatChannel->GetDefault().Get(MAX_flt), TCHAR_TO_UTF8(*PropertyName), TCHAR_TO_UTF8(*PropertyName));
 				}
@@ -3445,7 +3469,14 @@ void FFbxExporter::ExportLevelSequenceTrackChannels( FbxNode* FbxNode, UMovieSce
 				continue;
 			}
 
-			if (FloatChannel)
+			if (DoubleChannel)
+			{
+				CurveNode->SetChannelValue<double>(0U, DoubleChannel->GetDefault().Get(MAX_dbl));
+				CurveNode->ConnectToChannel(AnimCurve, 0U);
+
+				ExportChannelToFbxCurve(*AnimCurve, *DoubleChannel, TickResolution, IsFoV ? ERichCurveValueMode::Fov : ERichCurveValueMode::Default, false, RootToLocalTransform);
+			}
+			else if (FloatChannel)
 			{
 				CurveNode->SetChannelValue<double>(0U, FloatChannel->GetDefault().Get(MAX_flt));
 				CurveNode->ConnectToChannel(AnimCurve, 0U);
diff --git a/Engine/Source/Editor/UnrealEd/Private/FbxExporter.h b/Engine/Source/Editor/UnrealEd/Private/FbxExporter.h
--- a/Engine/Source/Editor/UnrealEd/Private/FbxExporter.h
+++ b/Engine/Source/Editor/UnrealEd/Private/FbxExporter.h
@@ -44,6 +44,7 @@ class FColorVertexBuffer;
 class UFbxExportOption;
 struct FAnimControlTrackKey;
 struct FExpressionInput;
+struct FMovieSceneDoubleChannel;
 struct FMovieSceneFloatChannel;
 struct FMovieSceneIntegerChannel;
 struct FMovieSceneSequenceTransform;
@@ -489,9 +490,16 @@ private:
 		Fov
 	};
 
+	/** Generic implementation of exporting a movie scene bezier curve channel to an fbx animation curve */
+	template<typename ChannelType>
+	void ExportBezierChannelToFbxCurve(FbxAnimCurve& InFbxCurve, const ChannelType& InChannel, FFrameRate TickResolution, ERichCurveValueMode ValueMode, bool bNegative, const FMovieSceneSequenceTransform& RootToLocalTransform);
+
 	/** Exports a movie scene float channel to an fbx animation curve. */
 	void ExportChannelToFbxCurve(FbxAnimCurve& InFbxCurve, const FMovieSceneFloatChannel& InChannel, FFrameRate TickResolution, ERichCurveValueMode ValueMode = ERichCurveValueMode::Default, bool bNegative = false, const FMovieSceneSequenceTransform& RootToLocalTransform = FMovieSceneSequenceTransform());
 
+	/** Exports a movie scene double channel to an fbx animation curve. */
+	void ExportChannelToFbxCurve(FbxAnimCurve& InFbxCurve, const FMovieSceneDoubleChannel& InChannel, FFrameRate TickResolution, ERichCurveValueMode ValueMode = ERichCurveValueMode::Default, bool bNegative = false, const FMovieSceneSequenceTransform& RootToLocalTransform = FMovieSceneSequenceTransform());
+
 	/** Exports a movie scene integer channel to an fbx animation curve. */
 	void ExportChannelToFbxCurve(FbxAnimCurve& InFbxCurve, const FMovieSceneIntegerChannel& InChannel, FFrameRate TickResolution, const FMovieSceneSequenceTransform& RootToLocalTransform = FMovieSceneSequenceTransform());
 
diff --git a/Engine/Source/Runtime/MovieScene/Private/Channels/MovieSceneCurveChannelImpl.cpp b/Engine/Source/Runtime/MovieScene/Private/Channels/MovieSceneCurveChannelImpl.cpp
new file mode 100644
--- /dev/null
+++ b/Engine/Source/Runtime/MovieScene/Private/Channels/MovieSceneCurveChannelImpl.cpp
@@ -0,0 +1,1161 @@
+// Copyright Epic Games, Inc. All Rights Reserved.
+
+#include "Channels/MovieSceneCurveChannelImpl.h"
+#include "Channels/MovieSceneDoubleChannel.h"
+#include "Channels/MovieSceneFloatChannel.h"
+#include "HAL/ConsoleManager.h"
+#include "MovieSceneFrameMigration.h"
+#include "UObject/FortniteMainBranchObjectVersion.h"
+#include "UObject/SequencerObjectVersion.h"
+
+int32 GSequencerLinearCubicInterpolation = 1;
+static FAutoConsoleVariableRef CVarSequencerLinearCubicInterpolation(
+	TEXT("Sequencer.LinearCubicInterpolation"),
+	GSequencerLinearCubicInterpolation,
+	TEXT("If 1 Linear Keys Act As Cubic Interpolation with Linear Tangents, if 0 Linear Key Forces Linear Interpolation to Next Key."),
+	ECVF_Default);
+
+static TAutoConsoleVariable<int32> CVarSequencerAutoTangentInterpolation(
+	TEXT("Sequencer.AutoTangentNew"),
+	1,
+	TEXT("If 1 Auto Tangent will use new algorithm to gradually flatten maximum/minimum keys, if 0 Auto Tangent will average all keys (pre 4.23 behavior)."),
+	ECVF_Default);
+
+
+namespace UE
+{
+namespace MovieScene
+{
+
+/** Util to find value on bezier defined by 4 control points */
+template<typename CurveValueType>
+CurveValueType BezierInterp(CurveValueType P0, CurveValueType P1, CurveValueType P2, CurveValueType P3, float Alpha)
+{
+	const CurveValueType P01   = FMath::Lerp(P0,   P1,   Alpha);
+	const CurveValueType P12   = FMath::Lerp(P1,   P2,   Alpha);
+	const CurveValueType P23   = FMath::Lerp(P2,   P3,   Alpha);
+	const CurveValueType P012  = FMath::Lerp(P01,  P12,  Alpha);
+	const CurveValueType P123  = FMath::Lerp(P12,  P23,  Alpha);
+	const CurveValueType P0123 = FMath::Lerp(P012, P123, Alpha);
+
+	return P0123;
+}
+
+template<typename ChannelType>
+static typename ChannelType::CurveValueType
+EvalForTwoKeys(
+		const typename ChannelType::ChannelValueType& Key1, FFrameNumber Key1Time,
+		const typename ChannelType::ChannelValueType& Key2, FFrameNumber Key2Time,
+		FFrameNumber InTime,
+		FFrameRate DisplayRate)
+{
+	using CurveValueType = typename ChannelType::CurveValueType;
+
+	double DecimalRate = DisplayRate.AsDecimal();
+
+	float Diff = (float)(Key2Time - Key1Time).Value;
+	Diff /= DecimalRate;
+	const int CheckBothLinear = GSequencerLinearCubicInterpolation;
+
+	if (Diff > 0 && Key1.InterpMode != RCIM_Constant)
+	{
+		const float Alpha = ((float)(InTime - Key1Time).Value / DecimalRate) / Diff;
+		const CurveValueType P0 = Key1.Value;
+		const CurveValueType P3 = Key2.Value;
+
+		if (Key1.InterpMode == RCIM_Linear && (!CheckBothLinear || Key2.InterpMode != RCIM_Cubic))
+		{
+			return FMath::Lerp(P0, P3, Alpha);
+		}
+		else
+		{
+			float LeaveTangent = Key1.Tangent.LeaveTangent * DecimalRate;
+			float ArriveTangent = Key2.Tangent.ArriveTangent * DecimalRate;
+
+			const float OneThird = 1.0f / 3.0f;
+			const CurveValueType P1 = P0 + (LeaveTangent * Diff*OneThird);
+			const CurveValueType P2 = P3 - (ArriveTangent * Diff*OneThird);
+
+			return BezierInterp(P0, P1, P2, P3, Alpha);
+		}
+	}
+	else
+	{
+		return Key1.Value;
+	}
+}
+
+struct FCycleParams
+{
+	FFrameTime Time;
+	int32 CycleCount;
+	float ValueOffset;
+
+	FCycleParams(FFrameTime InTime)
+		: Time(InTime)
+		, CycleCount(0)
+		, ValueOffset(0.f)
+	{}
+
+	FORCEINLINE void ComputePreValueOffset(float FirstValue, float LastValue)
+	{
+		ValueOffset = (FirstValue-LastValue) * CycleCount;
+	}
+	FORCEINLINE void ComputePostValueOffset(float FirstValue, float LastValue)
+	{
+		ValueOffset = (LastValue-FirstValue) * CycleCount;
+	}
+	FORCEINLINE void Oscillate(int32 MinFrame, int32 MaxFrame)
+	{
+		if (CycleCount % 2 == 1)
+		{
+			Time = MinFrame + (FFrameTime(MaxFrame) - Time);
+		}
+	}
+};
+
+FCycleParams CycleTime(FFrameNumber MinFrame, FFrameNumber MaxFrame, FFrameTime InTime)
+{
+	FCycleParams Params(InTime);
+	
+	const int32 Duration = MaxFrame.Value - MinFrame.Value;
+	if (Duration == 0)
+	{
+		Params.Time = MaxFrame;
+		Params.CycleCount = 0;
+	}
+	else if (InTime < MinFrame)
+	{
+		const int32 CycleCount = ((MaxFrame - InTime) / Duration).FloorToFrame().Value;
+
+		Params.Time = InTime + FFrameTime(Duration)*CycleCount;
+		Params.CycleCount = CycleCount;
+	}
+	else if (InTime > MaxFrame)
+	{
+		const int32 CycleCount = ((InTime - MinFrame) / Duration).FloorToFrame().Value;
+
+		Params.Time = InTime - FFrameTime(Duration)*CycleCount;
+		Params.CycleCount = CycleCount;
+	}
+
+	return Params;
+}
+
+/* Solve Cubic Equation using Cardano's forumla
+* Adopted from Graphic Gems 1
+* [REDACTED_PROFILE_URL]/GraphicsGems/blob/master/gems/Roots3And4.c
+*  Solve cubic of form
+*
+* @param Coeff Coefficient parameters of form  Coeff[0] + Coeff[1]*x + Coeff[2]*x^2 + Coeff[3]*x^3 + Coeff[4]*x^4 = 0
+* @param Solution Up to 3 real solutions. We don't include imaginary solutions, would need a complex number objecct
+* @return Returns the number of real solutions returned in the Solution array.
+*/
+static int SolveCubic(double Coeff[4], double Solution[3])
+{
+	auto cbrt = [](double x) -> double
+	{
+		return ((x) > 0.0 ? pow((x), 1.0 / 3.0) : ((x) < 0.0 ? -pow((double)-(x), 1.0 / 3.0) : 0.0));
+	};
+	int     NumSolutions = 0;
+
+	/* normal form: x^3 + Ax^2 + Bx + C = 0 */
+
+	double A = Coeff[2] / Coeff[3];
+	double B = Coeff[1] / Coeff[3];
+	double C = Coeff[0] / Coeff[3];
+
+	/*  substitute x = y - A/3 to eliminate quadric term:
+	x^3 +px + q = 0 */
+
+	double SqOfA = A * A;
+	double P = 1.0 / 3 * (-1.0 / 3 * SqOfA + B);
+	double Q = 1.0 / 2 * (2.0 / 27 * A * SqOfA - 1.0 / 3 * A * B + C);
+
+	/* use Cardano's formula */
+
+	double CubeOfP = P * P * P;
+	double D = Q * Q + CubeOfP;
+
+	if (FMath::IsNearlyZero(D))
+	{
+		if (FMath::IsNearlyZero(Q)) /* one triple solution */
+		{
+			Solution[0] = 0;
+			NumSolutions = 1;
+		}
+		else /* one single and one double solution */
+		{
+			double u = cbrt(-Q);
+			Solution[0] = 2 * u;
+			Solution[1] = -u;
+			NumSolutions = 2;
+		}
+	}
+	else if (D < 0) /* Casus irreducibilis: three real solutions */
+	{
+		double phi = 1.0 / 3 * acos(-Q / sqrt(-CubeOfP));
+		double t = 2 * sqrt(-P);
+
+		Solution[0] = t * cos(phi);
+		Solution[1] = -t * cos(phi + PI / 3);
+		Solution[2] = -t * cos(phi - PI / 3);
+		NumSolutions = 3;
+	}
+	else /* one real solution */
+	{
+		double sqrt_D = sqrt(D);
+		double u = cbrt(sqrt_D - Q);
+		double v = -cbrt(sqrt_D + Q);
+
+		Solution[0] = u + v;
+		NumSolutions = 1;
+	}
+
+	/* resubstitute */
+
+	double Sub = 1.0 / 3 * A;
+
+	for (int i = 0; i < NumSolutions; ++i)
+		Solution[i] -= Sub;
+
+	return NumSolutions;
+}
+
+/*
+*   Convert the control values for a polynomial defined in the Bezier
+*		basis to a polynomial defined in the power basis (t^3 t^2 t 1).
+*/
+static void BezierToPower(	double A1, double B1, double C1, double D1,
+	double *A2, double *B2, double *C2, double *D2)
+{
+	double A = B1 - A1;
+	double B = C1 - B1;
+	double C = D1 - C1;
+	double D = B - A;
+	*A2 = C- B - D;
+	*B2 = 3.0 * D;
+	*C2 = 3.0 * A;
+	*D2 = A1;
+}
+
+
+} // namespace MovieScene
+} // namespace UE
+
+
+template<typename ChannelType>
+void TMovieSceneCurveChannelImpl<ChannelType>::Set(ChannelType* InChannel, TArray<FFrameNumber> InTimes, TArray<ChannelValueType> InValues)
+{
+	check(InTimes.Num() == InValues.Num());
+
+	InChannel->Times = MoveTemp(InTimes);
+	InChannel->Values = MoveTemp(InValues);
+
+	InChannel->KeyHandles.Reset();
+	for (int32 Index = 0; Index < InChannel->Times.Num(); ++Index)
+	{
+		InChannel->KeyHandles.AllocateHandle(Index);
+	}
+}
+
+template<typename ChannelType>
+int32 TMovieSceneCurveChannelImpl<ChannelType>::InsertKeyInternal(ChannelType* InChannel, FFrameNumber InTime)
+{
+	const int32 InsertIndex = Algo::UpperBound(InChannel->Times, InTime);
+
+	InChannel->Times.Insert(InTime, InsertIndex);
+	InChannel->Values.Insert(ChannelValueType(), InsertIndex);
+
+	InChannel->KeyHandles.AllocateHandle(InsertIndex);
+
+	return InsertIndex;
+}
+
+template<typename ChannelType>
+int32 TMovieSceneCurveChannelImpl<ChannelType>::AddConstantKey(ChannelType* InChannel, FFrameNumber InTime, CurveValueType InValue)
+{
+	const int32 Index = InsertKeyInternal(InChannel, InTime);
+
+	ChannelValueType& Value = InChannel->Values[Index];
+	Value.Value = InValue;
+	Value.InterpMode = RCIM_Constant;
+
+	AutoSetTangents(InChannel);
+
+	return Index;
+}
+
+template<typename ChannelType>
+int32 TMovieSceneCurveChannelImpl<ChannelType>::AddLinearKey(ChannelType* InChannel, FFrameNumber InTime, CurveValueType InValue)
+{
+	const int32 Index = InsertKeyInternal(InChannel, InTime);
+
+	ChannelValueType& Value = InChannel->Values[Index];
+	Value.Value = InValue;
+	Value.InterpMode = RCIM_Linear;
+
+	AutoSetTangents(InChannel);
+
+	return Index;
+}
+
+template<typename ChannelType>
+int32 TMovieSceneCurveChannelImpl<ChannelType>::AddCubicKey(ChannelType* InChannel, FFrameNumber InTime, CurveValueType InValue, ERichCurveTangentMode TangentMode, const FMovieSceneTangentData& Tangent)
+{
+	const int32 Index = InsertKeyInternal(InChannel, InTime);
+
+	ChannelValueType& Value = InChannel->Values[Index];
+	Value.Value = InValue;
+	Value.InterpMode = RCIM_Cubic;
+	Value.TangentMode = TangentMode;
+	Value.Tangent = Tangent;
+
+	AutoSetTangents(InChannel);
+
+	return Index;
+}
+
+template<typename ChannelType>
+bool TMovieSceneCurveChannelImpl<ChannelType>::EvaluateExtrapolation(const ChannelType* InChannel, FFrameTime InTime, CurveValueType& OutValue)
+{
+	// If the time is outside of the curve, deal with extrapolation
+	if (InTime < InChannel->Times[0])
+	{
+		if (InChannel->PreInfinityExtrap == RCCE_None)
+		{
+			return false;
+		}
+
+		if (InChannel->PreInfinityExtrap == RCCE_Constant)
+		{
+			OutValue = InChannel->Values[0].Value;
+			return true;
+		}
+
+		if (InChannel->PreInfinityExtrap == RCCE_Linear)
+		{
+			const ChannelValueType FirstValue = InChannel->Values[0];
+
+			if (FirstValue.InterpMode == RCIM_Constant)
+			{
+				OutValue = FirstValue.Value;
+			}
+			else if(FirstValue.InterpMode == RCIM_Cubic)
+			{
+				FFrameTime Delta = FFrameTime(InChannel->Times[0]) - InTime;
+				OutValue = FirstValue.Value - Delta.AsDecimal() * FirstValue.Tangent.ArriveTangent;
+			}
+			else if(FirstValue.InterpMode == RCIM_Linear)
+			{
+				const int32 InterpStartFrame = InChannel->Times[1].Value;
+				const int32 DeltaFrame       = InterpStartFrame - InChannel->Times[0].Value;
+				if (DeltaFrame == 0)
+				{
+					OutValue = FirstValue.Value;
+				}
+				else
+				{
+					OutValue = FMath::Lerp(InChannel->Values[1].Value, FirstValue.Value, (InterpStartFrame - InTime.AsDecimal())/DeltaFrame);
+				}
+			}
+			return true;
+		}
+	}
+	else if (InTime > InChannel->Times.Last())
+	{
+		if (InChannel->PostInfinityExtrap == RCCE_None)
+		{
+			return false;
+		}
+
+		if (InChannel->PostInfinityExtrap == RCCE_Constant)
+		{
+			OutValue = InChannel->Values.Last().Value;
+			return true;
+		}
+
+		if (InChannel->PostInfinityExtrap == RCCE_Linear)
+		{
+			const ChannelValueType LastValue = InChannel->Values.Last();
+
+			if (LastValue.InterpMode == RCIM_Constant)
+			{
+				OutValue = LastValue.Value;
+			}
+			else if(LastValue.InterpMode == RCIM_Cubic)
+			{
+				FFrameTime Delta = InTime - InChannel->Times.Last();
+				OutValue = LastValue.Value + Delta.AsDecimal() * LastValue.Tangent.LeaveTangent;
+			}
+			else if(LastValue.InterpMode == RCIM_Linear)
+			{
+				const int32 NumKeys          = InChannel->Times.Num();
+				const int32 InterpStartFrame = InChannel->Times[NumKeys-2].Value;
+				const int32 DeltaFrame       = InChannel->Times.Last().Value-InterpStartFrame;
+
+				if (DeltaFrame == 0)
+				{
+					OutValue = LastValue.Value;
+				}
+				else
+				{
+					OutValue = FMath::Lerp(InChannel->Values[NumKeys-2].Value, LastValue.Value, (InTime.AsDecimal() - InterpStartFrame)/DeltaFrame);
+				}
+			}
+			return true;
+		}
+	}
+
+	return false;
+}
+
+template<typename ChannelType>
+bool TMovieSceneCurveChannelImpl<ChannelType>::Evaluate(const ChannelType* InChannel, FFrameTime InTime, CurveValueType& OutValue) 
+{
+	using namespace UE::MovieScene;
+
+	const int32 NumKeys = InChannel->Times.Num();
+
+	// No keys means default value, or nothing
+	if (NumKeys == 0)
+	{
+		if (InChannel->bHasDefaultValue)
+		{
+			OutValue = InChannel->DefaultValue;
+			return true;
+		}
+		return false;
+	}
+
+	// For single keys, we can only ever return that value
+	if (NumKeys == 1)
+	{
+		OutValue = InChannel->Values[0].Value;
+		return true;
+	}
+
+	// Evaluate with extrapolation if we're outside the bounds of the curve
+	if (EvaluateExtrapolation(InChannel, InTime, OutValue))
+	{
+		return true;
+	}
+
+	const FFrameNumber MinFrame = InChannel->Times[0];
+	const FFrameNumber MaxFrame = InChannel->Times.Last();
+
+	// Compute the cycled time
+	FCycleParams Params = CycleTime(MinFrame, MaxFrame, InTime);
+
+	// Deal with offset cycles and oscillation
+	if (InTime < FFrameTime(MinFrame))
+	{
+		switch (InChannel->PreInfinityExtrap)
+		{
+		case RCCE_CycleWithOffset: Params.ComputePreValueOffset(InChannel->Values[0].Value, InChannel->Values[NumKeys-1].Value); break;
+		case RCCE_Oscillate:       Params.Oscillate(MinFrame.Value, MaxFrame.Value);                       break;
+		}
+	}
+	else if (InTime > FFrameTime(MaxFrame))
+	{
+		switch (InChannel->PostInfinityExtrap)
+		{
+		case RCCE_CycleWithOffset: Params.ComputePostValueOffset(InChannel->Values[0].Value, InChannel->Values[NumKeys-1].Value); break;
+		case RCCE_Oscillate:       Params.Oscillate(MinFrame.Value, MaxFrame.Value);                        break;
+		}
+	}
+
+	if (!ensureMsgf(Params.Time.FrameNumber >= MinFrame && Params.Time.FrameNumber <= MaxFrame, TEXT("Invalid time computed for float channel evaluation")))
+	{
+		return false;
+	}
+
+	// Evaluate the curve data
+	float Interp = 0.f;
+	int32 Index1 = INDEX_NONE, Index2 = INDEX_NONE;
+	UE::MovieScene::EvaluateTime(InChannel->Times, Params.Time, Index1, Index2, Interp);
+	const int CheckBothLinear = GSequencerLinearCubicInterpolation;
+
+	if (Index1 == INDEX_NONE)
+	{
+		OutValue = Params.ValueOffset + InChannel->Values[Index2].Value;
+	}
+	else if (Index2 == INDEX_NONE)
+	{
+		OutValue = Params.ValueOffset + InChannel->Values[Index1].Value;
+	}
+	else
+	{
+		ChannelValueType Key1 = InChannel->Values[Index1];
+		ChannelValueType Key2 = InChannel->Values[Index2];
+		TEnumAsByte<ERichCurveInterpMode> InterpMode = Key1.InterpMode;
+	    if(InterpMode == RCIM_Linear && (CheckBothLinear  && Key2.InterpMode == RCIM_Cubic))
+		{
+			InterpMode = RCIM_Cubic;
+		}
+		
+		switch (InterpMode)
+		{
+		case RCIM_Cubic:
+		{
+			const float OneThird = 1.0f / 3.0f;
+			if ((Key1.Tangent.TangentWeightMode == RCTWM_WeightedNone || Key1.Tangent.TangentWeightMode == RCTWM_WeightedArrive)
+				&& (Key2.Tangent.TangentWeightMode == RCTWM_WeightedNone || Key2.Tangent.TangentWeightMode == RCTWM_WeightedLeave))
+			{
+				const int32 Diff = InChannel->Times[Index2].Value - InChannel->Times[Index1].Value;
+				const float P0 = Key1.Value;
+				const float P1 = P0 + (Key1.Tangent.LeaveTangent * Diff * OneThird);
+				const float P3 = Key2.Value;
+				const float P2 = P3 - (Key2.Tangent.ArriveTangent * Diff * OneThird);
+
+				OutValue = Params.ValueOffset + BezierInterp(P0, P1, P2, P3, Interp);
+				break;
+			}
+			else //its weighted
+			{
+				const float TimeInterval = InChannel->TickResolution.AsInterval();
+				const float ToSeconds = 1.0f / TimeInterval;
+
+				const double Time1 = InChannel->TickResolution.AsSeconds(InChannel->Times[Index1].Value);
+				const double Time2 = InChannel->TickResolution.AsSeconds(InChannel->Times[Index2].Value);
+				const float X = Time2 - Time1;
+				float CosAngle, SinAngle;
+				float Angle = FMath::Atan(Key1.Tangent.LeaveTangent * ToSeconds);
+				FMath::SinCos(&SinAngle, &CosAngle, Angle);
+				float LeaveWeight;
+				if (Key1.Tangent.TangentWeightMode == RCTWM_WeightedNone || Key1.Tangent.TangentWeightMode == RCTWM_WeightedArrive)
+				{
+					const float LeaveTangentNormalized = Key1.Tangent.LeaveTangent / (TimeInterval);
+					const float Y = LeaveTangentNormalized * X;
+					LeaveWeight = FMath::Sqrt(X*X + Y * Y) * OneThird;
+				}
+				else
+				{
+					LeaveWeight = Key1.Tangent.LeaveTangentWeight;
+				}
+				const float Key1TanX = CosAngle * LeaveWeight + Time1;
+				const float Key1TanY = SinAngle * LeaveWeight + Key1.Value;
+
+				Angle = FMath::Atan(Key2.Tangent.ArriveTangent * ToSeconds);
+				FMath::SinCos(&SinAngle, &CosAngle, Angle);
+				float ArriveWeight;
+				if (Key2.Tangent.TangentWeightMode == RCTWM_WeightedNone || Key2.Tangent.TangentWeightMode == RCTWM_WeightedLeave)
+				{
+					const float ArriveTangentNormalized = Key2.Tangent.ArriveTangent / (TimeInterval);
+					const float Y = ArriveTangentNormalized * X;
+					ArriveWeight = FMath::Sqrt(X*X + Y * Y) * OneThird;
+				}
+				else
+				{
+					ArriveWeight =  Key2.Tangent.ArriveTangentWeight;
+				}
+				const float Key2TanX = -CosAngle * ArriveWeight + Time2;
+				const float Key2TanY = -SinAngle * ArriveWeight + Key2.Value;
+
+				//Normalize the Time Range
+				const float RangeX = Time2 - Time1;
+
+				const float Dx1 = Key1TanX - Time1;
+				const float Dx2 = Key2TanX - Time1;
+
+				// Normalize values
+				const float NormalizedX1 = Dx1 / RangeX;
+				const float NormalizedX2 = Dx2 / RangeX;
+				
+				double Coeff[4];
+				double Results[3];
+
+				//Convert Bezier to Power basis, also float to double for precision for root finding.
+				BezierToPower(
+					0.0, NormalizedX1, NormalizedX2, 1.0,
+					&(Coeff[3]), &(Coeff[2]), &(Coeff[1]), &(Coeff[0])
+				);
+
+				Coeff[0] = Coeff[0] - Interp;
+				
+				int NumResults = SolveCubic(Coeff, Results);
+				float NewInterp = Interp;
+				if (NumResults == 1)
+				{
+					NewInterp = Results[0];
+				}
+				else
+				{
+					NewInterp = TNumericLimits<float>::Lowest(); //just need to be out of range
+					for (double Result : Results)
+					{
+						if ((Result >= 0.0f) && (Result <= 1.0f))
+						{
+							if (NewInterp < 0.0f || Result > NewInterp)
+							{
+								NewInterp = Result;
+							}
+						}
+					}
+
+					if (NewInterp == TNumericLimits<float>::Lowest())
+					{
+						NewInterp = 0.f;
+					}
+
+				}
+				//now use NewInterp and adjusted tangents plugged into the Y (Value) part of the graph.
+				const float P0 = Key1.Value;
+				const float P1 = Key1TanY;
+				const float P3 = Key2.Value;
+				const float P2 = Key2TanY;
+
+				OutValue = Params.ValueOffset + BezierInterp(P0, P1, P2, P3,  NewInterp);
+			}
+			break;
+		}
+
+		case RCIM_Linear:
+			OutValue = Params.ValueOffset + FMath::Lerp(Key1.Value, Key2.Value, Interp);
+			break;
+
+		default:
+			OutValue = Params.ValueOffset + Key1.Value;
+			break;
+		}
+	}
+
+	return true;
+}
+
+template<typename ChannelType>
+void TMovieSceneCurveChannelImpl<ChannelType>::AutoSetTangents(ChannelType* InChannel, float Tension)
+{
+	if (InChannel->Values.Num() < 2)
+	{
+		return;
+	}
+
+	const int UseNewAutoTangent = CVarSequencerAutoTangentInterpolation->GetInt();
+
+	{
+		ChannelValueType& FirstValue = InChannel->Values[0];
+		if (FirstValue.InterpMode == RCIM_Linear)
+		{
+			FirstValue.Tangent.TangentWeightMode = RCTWM_WeightedNone;
+			ChannelValueType& NextKey = InChannel->Values[1];
+			const float NextTimeDiff = FMath::Max<double>(KINDA_SMALL_NUMBER, InChannel->Times[1].Value - InChannel->Times[0].Value);
+			const float NewTangent = (NextKey.Value - FirstValue.Value) / NextTimeDiff;
+			FirstValue.Tangent.LeaveTangent = NewTangent;
+		}
+		else if (FirstValue.InterpMode == RCIM_Cubic && FirstValue.TangentMode == RCTM_Auto)
+		{
+			FirstValue.Tangent.LeaveTangent = FirstValue.Tangent.ArriveTangent = 0.0f;
+			FirstValue.Tangent.TangentWeightMode = RCTWM_WeightedNone; 
+		}
+	}
+
+	{
+		ChannelValueType& LastValue = InChannel->Values.Last();
+		if (LastValue.InterpMode == RCIM_Linear)
+		{
+			LastValue.Tangent.TangentWeightMode = RCTWM_WeightedNone;
+			int32 Index = InChannel->Values.Num() - 1;
+			ChannelValueType& PrevKey = InChannel->Values[Index-1];
+			const float PrevTimeDiff = FMath::Max<double>(KINDA_SMALL_NUMBER, InChannel->Times[Index].Value - InChannel->Times[Index - 1].Value);
+			const float NewTangent = (LastValue.Value - PrevKey.Value) / PrevTimeDiff;
+			LastValue.Tangent.ArriveTangent = NewTangent;
+		}
+		else if (LastValue.InterpMode == RCIM_Cubic && LastValue.TangentMode == RCTM_Auto)
+		{
+			LastValue.Tangent.LeaveTangent = LastValue.Tangent.ArriveTangent = 0.0f;
+			LastValue.Tangent.TangentWeightMode = RCTWM_WeightedNone;
+		}
+	}
+
+	for (int32 Index = 1; Index < InChannel->Values.Num() - 1; ++Index)
+	{
+		ChannelValueType  PrevKey = InChannel->Values[Index-1];
+		ChannelValueType& ThisKey = InChannel->Values[Index  ];
+
+		if (ThisKey.InterpMode == RCIM_Cubic && ThisKey.TangentMode == RCTM_Auto)
+		{
+			ChannelValueType& NextKey = InChannel->Values[Index+1];
+			CurveValueType NewTangent = 0.f;
+			const double PrevToNextTimeDiff = FMath::Max<double>(KINDA_SMALL_NUMBER, InChannel->Times[Index + 1].Value - InChannel->Times[Index - 1].Value);
+
+			if (!UseNewAutoTangent)
+			{
+				AutoCalcTangent(PrevKey.Value, ThisKey.Value, NextKey.Value, Tension, NewTangent);
+				NewTangent /= PrevToNextTimeDiff;
+			}
+			else
+			{
+				// if key doesn't lie between we keep it flat(0.0).
+				if ( (ThisKey.Value > PrevKey.Value && ThisKey.Value < NextKey.Value) ||
+					(ThisKey.Value < PrevKey.Value && ThisKey.Value > NextKey.Value))
+				{
+					AutoCalcTangent(PrevKey.Value, ThisKey.Value, NextKey.Value, Tension, NewTangent);
+					NewTangent /= PrevToNextTimeDiff;
+					//if within 0 to 15% or 85% to 100% range we gradually weight tangent to zero
+					const float AverageToZeroRange = 0.85f;
+					const float ValDiff = FMath::Abs<float>(NextKey.Value - PrevKey.Value);
+					const float OurDiff = FMath::Abs<float>(ThisKey.Value - PrevKey.Value);
+					//ValDiff won't be zero ever due to previous check
+					float PercDiff = OurDiff / ValDiff;
+					if (PercDiff > AverageToZeroRange)
+					{
+						PercDiff = (PercDiff - AverageToZeroRange) / (1.0f - AverageToZeroRange);
+						NewTangent = NewTangent * (1.0f - PercDiff);
+					}
+					else if (PercDiff < (1.0f - AverageToZeroRange))
+					{
+						PercDiff = PercDiff  / (1.0f - AverageToZeroRange);
+						NewTangent = NewTangent * PercDiff;
+					}
+				}
+			}
+
+			// In 'auto' mode, arrive and leave tangents are always the same
+			ThisKey.Tangent.LeaveTangent = ThisKey.Tangent.ArriveTangent = (float)NewTangent;
+			ThisKey.Tangent.TangentWeightMode = RCTWM_WeightedNone;
+		}
+		else if (ThisKey.InterpMode == RCIM_Linear)
+		{
+			ThisKey.Tangent.TangentWeightMode = RCTWM_WeightedNone; 
+			ChannelValueType& NextKey = InChannel->Values[Index + 1];
+
+			const double PrevTimeDiff = FMath::Max<double>(KINDA_SMALL_NUMBER, InChannel->Times[Index].Value - InChannel->Times[Index - 1].Value);
+			float NewTangent  = (ThisKey.Value - PrevKey.Value) / PrevTimeDiff;
+			ThisKey.Tangent.ArriveTangent = NewTangent;
+			
+			const double NextTimeDiff = FMath::Max<double>(KINDA_SMALL_NUMBER, InChannel->Times[Index + 1].Value - InChannel->Times[Index].Value);
+			NewTangent = (NextKey.Value - ThisKey.Value) / NextTimeDiff;
+			ThisKey.Tangent.LeaveTangent = NewTangent;
+		}
+	}
+}
+
+template<typename ChannelType>
+void TMovieSceneCurveChannelImpl<ChannelType>::DeleteKeysFrom(ChannelType* InChannel, FFrameNumber InTime, bool bDeleteKeysBefore)
+{
+	// Insert a key at the current time to maintain evaluation
+	TMovieSceneChannelData<ChannelValueType> ChannelData(InChannel->GetData());
+	if (ChannelData.GetTimes().Num() > 0)
+	{
+		int32 KeyHandleIndex = ChannelData.FindKey(InTime);
+		if (KeyHandleIndex == INDEX_NONE)
+		{
+			CurveValueType Value = 0;
+			if (Evaluate(InChannel, InTime, Value))
+			{
+				AddCubicKey(InChannel, InTime, Value);
+			}
+		}
+	}
+
+	ChannelData.DeleteKeysFrom(InTime, bDeleteKeysBefore);
+}
+
+template<typename ChannelType>
+void TMovieSceneCurveChannelImpl<ChannelType>::ChangeFrameResolution(ChannelType* InChannel, FFrameRate SourceRate, FFrameRate DestinationRate)
+{
+	check(InChannel->Times.Num() == InChannel->Values.Num());
+
+	float IntervalFactor = DestinationRate.AsInterval() / SourceRate.AsInterval();
+	for (int32 Index = 0; Index < InChannel->Times.Num(); ++Index)
+	{
+		InChannel->Times[Index] = ConvertFrameTime(InChannel->Times[Index], SourceRate, DestinationRate).RoundToFrame();
+
+		ChannelValueType& Value = InChannel->Values[Index];
+		Value.Tangent.ArriveTangent *= IntervalFactor;
+		Value.Tangent.LeaveTangent  *= IntervalFactor;
+	}
+}
+
+template<typename ChannelType>
+void TMovieSceneCurveChannelImpl<ChannelType>::Optimize(ChannelType* InChannel, const FKeyDataOptimizationParams& InParameters)
+{
+	TMovieSceneChannelData<ChannelValueType> ChannelData = InChannel->GetData();
+	TArray<FFrameNumber> OutKeyTimes;
+	TArray<FKeyHandle> OutKeyHandles;
+
+	InChannel->GetKeys(InParameters.Range, &OutKeyTimes, &OutKeyHandles);
+
+	if (OutKeyHandles.Num() > 2)
+	{
+		int32 MostRecentKeepKeyIndex = 0;
+		TArray<FKeyHandle> KeysToRemove;
+
+		for (int32 TestIndex = 1; TestIndex < OutKeyHandles.Num() - 1; ++TestIndex)
+		{
+			int32 Index = ChannelData.GetIndex(OutKeyHandles[TestIndex]);
+			int32 NextIndex = ChannelData.GetIndex(OutKeyHandles[TestIndex+1]);
+
+			const CurveValueType KeyValue = ChannelData.GetValues()[Index].Value;
+			const CurveValueType ValueWithoutKey = UE::MovieScene::EvalForTwoKeys<ChannelType>(
+				ChannelData.GetValues()[MostRecentKeepKeyIndex], ChannelData.GetTimes()[MostRecentKeepKeyIndex].Value,
+				ChannelData.GetValues()[NextIndex], ChannelData.GetTimes()[NextIndex].Value,
+				ChannelData.GetTimes()[Index].Value,
+				InParameters.DisplayRate);
+				
+			if (FMath::Abs(ValueWithoutKey - KeyValue) > InParameters.Tolerance) // Is this key needed
+			{
+				MostRecentKeepKeyIndex = Index;
+			}
+			else
+			{
+				KeysToRemove.Add(OutKeyHandles[TestIndex]);
+			}
+		}
+
+		ChannelData.DeleteKeys(KeysToRemove);
+
+		if (InParameters.bAutoSetInterpolation)
+		{
+			AutoSetTangents(InChannel);
+		}
+	}
+}
+
+template<typename ChannelType>
+FKeyHandle TMovieSceneCurveChannelImpl<ChannelType>::AddKeyToChannel(ChannelType* InChannel, FFrameNumber InFrameNumber, float InValue, EMovieSceneKeyInterpolation Interpolation)
+{
+	TMovieSceneChannelData<ChannelValueType> ChannelData = InChannel->GetData();
+	int32 ExistingIndex = ChannelData.FindKey(InFrameNumber);
+	if (ExistingIndex != INDEX_NONE)
+	{
+		ChannelValueType& Value = ChannelData.GetValues()[ExistingIndex]; //-V758
+		Value.Value = InValue;
+		AutoSetTangents(InChannel);
+	}
+	else switch (Interpolation)
+	{
+		case EMovieSceneKeyInterpolation::Auto:     ExistingIndex = InChannel->AddCubicKey(InFrameNumber, InValue, RCTM_Auto);  break;
+		case EMovieSceneKeyInterpolation::User:     ExistingIndex = InChannel->AddCubicKey(InFrameNumber, InValue, RCTM_User);  break;
+		case EMovieSceneKeyInterpolation::Break:    ExistingIndex = InChannel->AddCubicKey(InFrameNumber, InValue, RCTM_Break); break;
+		case EMovieSceneKeyInterpolation::Linear:   ExistingIndex = InChannel->AddLinearKey(InFrameNumber, InValue);            break;
+		case EMovieSceneKeyInterpolation::Constant: ExistingIndex = InChannel->AddConstantKey(InFrameNumber, InValue);          break;
+	}
+
+	return InChannel->GetData().GetHandle(ExistingIndex);
+}
+
+template<typename ChannelType>
+void TMovieSceneCurveChannelImpl<ChannelType>::Dilate(ChannelType* InChannel, FFrameNumber Origin, float DilationFactor)
+{
+	TArrayView<FFrameNumber> Times = InChannel->GetData().GetTimes();
+	for (FFrameNumber& Time : Times)
+	{
+		Time = Origin + FFrameNumber(FMath::FloorToInt((Time - Origin).Value * DilationFactor));
+	}
+	AutoSetTangents(InChannel);
+}
+
+template<typename ChannelType>
+void TMovieSceneCurveChannelImpl<ChannelType>::AssignValue(ChannelType* InChannel, FKeyHandle InKeyHandle, typename ChannelType::CurveValueType InValue)
+{
+	TMovieSceneChannelData<ChannelValueType> ChannelData = InChannel->GetData();
+	int32 ValueIndex = ChannelData.GetIndex(InKeyHandle);
+
+	if (ValueIndex != INDEX_NONE)
+	{
+		ChannelData.GetValues()[ValueIndex].Value = InValue;
+	}
+}
+
+template<typename ChannelType>
+void TMovieSceneCurveChannelImpl<ChannelType>::PopulateCurvePoints(const ChannelType* InChannel, double StartTimeSeconds, double EndTimeSeconds, double TimeThreshold, CurveValueType ValueThreshold, FFrameRate InTickResolution, TArray<TTuple<double, double>>& InOutPoints)
+{
+	const FFrameNumber StartFrame = (StartTimeSeconds * InTickResolution).FloorToFrame();
+	const FFrameNumber EndFrame   = (EndTimeSeconds   * InTickResolution).CeilToFrame();
+
+	const int32 StartingIndex = Algo::UpperBound(InChannel->Times, StartFrame);
+	const int32 EndingIndex   = Algo::LowerBound(InChannel->Times, EndFrame);
+
+	// Add the lower bound of the visible space
+	CurveValueType EvaluatedValue = 0;
+	if (Evaluate(InChannel, StartFrame, EvaluatedValue))
+	{
+		InOutPoints.Add(MakeTuple(StartFrame / InTickResolution, double(EvaluatedValue)));
+	}
+
+	// Add all keys in-between
+	for (int32 KeyIndex = StartingIndex; KeyIndex < EndingIndex; ++KeyIndex)
+	{
+		InOutPoints.Add(MakeTuple(InChannel->Times[KeyIndex] / InTickResolution, double(InChannel->Values[KeyIndex].Value)));
+	}
+
+	// Add the upper bound of the visible space
+	if (Evaluate(InChannel, EndFrame, EvaluatedValue))
+	{
+		InOutPoints.Add(MakeTuple(EndFrame / InTickResolution, double(EvaluatedValue)));
+	}
+
+	int32 OldSize = InOutPoints.Num();
+	do
+	{
+		OldSize = InOutPoints.Num();
+		RefineCurvePoints(InChannel, InTickResolution, TimeThreshold, ValueThreshold, InOutPoints);
+	}
+	while(OldSize != InOutPoints.Num());
+}
+
+template<typename ChannelType>
+void TMovieSceneCurveChannelImpl<ChannelType>::RefineCurvePoints(const ChannelType* InChannel, FFrameRate InTickResolution, double TimeThreshold, CurveValueType ValueThreshold, TArray<TTuple<double, double>>& InOutPoints)
+{
+	const float InterpTimes[] = { 0.25f, 0.5f, 0.6f };
+
+	for (int32 Index = 0; Index < InOutPoints.Num() - 1; ++Index)
+	{
+		TTuple<double, double> Lower = InOutPoints[Index];
+		TTuple<double, double> Upper = InOutPoints[Index + 1];
+
+		if ((Upper.Get<0>() - Lower.Get<0>()) >= TimeThreshold)
+		{
+			bool bSegmentIsLinear = true;
+
+			TTuple<double, double> Evaluated[UE_ARRAY_COUNT(InterpTimes)] = { TTuple<double, double>(0, 0) };
+
+			for (int32 InterpIndex = 0; InterpIndex < UE_ARRAY_COUNT(InterpTimes); ++InterpIndex)
+			{
+				double& EvalTime  = Evaluated[InterpIndex].Get<0>();
+
+				EvalTime = FMath::Lerp(Lower.Get<0>(), Upper.Get<0>(), InterpTimes[InterpIndex]);
+
+				CurveValueType Value = 0.f;
+				Evaluate(InChannel, EvalTime * InTickResolution, Value);
+
+				const CurveValueType LinearValue = FMath::Lerp(Lower.Get<1>(), Upper.Get<1>(), InterpTimes[InterpIndex]);
+				if (bSegmentIsLinear)
+				{
+					bSegmentIsLinear = FMath::IsNearlyEqual(Value, LinearValue, ValueThreshold);
+				}
+
+				Evaluated[InterpIndex].Get<1>() = Value;
+			}
+
+			if (!bSegmentIsLinear)
+			{
+				// Add the point
+				InOutPoints.Insert(Evaluated, UE_ARRAY_COUNT(Evaluated), Index+1);
+				--Index;
+			}
+		}
+	}
+}
+
+template<typename ChannelType>
+bool TMovieSceneCurveChannelImpl<ChannelType>::ValueExistsAtTime(const ChannelType* Channel, FFrameNumber InFrameNumber, typename ChannelType::CurveValueType Value)
+{
+	const FFrameTime FrameTime(InFrameNumber);
+
+	CurveValueType ExistingValue = 0.f;
+	return Channel->Evaluate(FrameTime, ExistingValue) && FMath::IsNearlyEqual(ExistingValue, Value, (CurveValueType)KINDA_SMALL_NUMBER);
+}
+
+template<typename ChannelType>
+bool TMovieSceneCurveChannelImpl<ChannelType>::ValueExistsAtTime(const ChannelType* Channel, FFrameNumber InFrameNumber, const typename ChannelType::ChannelValueType& InValue)
+{
+	return ValueExistsAtTime(Channel, InFrameNumber, InValue.Value);
+}
+
+template<typename ChannelType>
+bool TMovieSceneCurveChannelImpl<ChannelType>::Serialize(ChannelType* InChannel, FArchive& Ar)
+{
+	Ar.UsingCustomVersion(FSequencerObjectVersion::GUID);
+	Ar.UsingCustomVersion(FFortniteMainBranchObjectVersion::GUID);
+	if (Ar.CustomVer(FSequencerObjectVersion::GUID) < FSequencerObjectVersion::SerializeFloatChannelCompletely &&
+		Ar.CustomVer(FFortniteMainBranchObjectVersion::GUID) < FFortniteMainBranchObjectVersion::SerializeFloatChannelShowCurve)
+	{
+		return false;
+	}
+
+	const bool bSerializeShowCurve = (
+			Ar.CustomVer(FFortniteMainBranchObjectVersion::GUID) >= FFortniteMainBranchObjectVersion::SerializeFloatChannelShowCurve);
+
+	Ar << InChannel->PreInfinityExtrap;
+	Ar << InChannel->PostInfinityExtrap;
+
+	// Save FFrameNumber(int32) and channel value arrays.
+	// We try to save and load the full array data, unless we are
+	// ByteSwapping or the Size has a mismatch on load, then we do normal save/load
+	if (Ar.IsLoading())
+	{
+		int32 CurrentSerializedElementSize = sizeof(FFrameNumber);
+		int32 SerializedElementSize = 0;
+		Ar << SerializedElementSize;
+		if (SerializedElementSize != CurrentSerializedElementSize || Ar.IsByteSwapping())
+		{
+			Ar << InChannel->Times;
+		}
+		else
+		{
+			InChannel->Times.CountBytes(Ar);
+			int32 NewArrayNum = 0;
+			Ar << NewArrayNum;
+			InChannel->Times.Empty(NewArrayNum);
+			if (NewArrayNum > 0)
+			{
+				InChannel->Times.AddUninitialized(NewArrayNum);
+				Ar.Serialize(InChannel->Times.GetData(), NewArrayNum * SerializedElementSize);
+			}
+		}
+		CurrentSerializedElementSize = sizeof(ChannelValueType);
+		Ar << SerializedElementSize;
+
+		if (SerializedElementSize != CurrentSerializedElementSize || Ar.IsByteSwapping())
+		{
+			Ar << InChannel->Values;
+		}
+		else
+		{
+			InChannel->Values.CountBytes(Ar);
+			int32 NewArrayNum = 0;
+			Ar << NewArrayNum;
+			InChannel->Values.Empty(NewArrayNum);
+			if (NewArrayNum > 0)
+			{
+				InChannel->Values.AddUninitialized(NewArrayNum);
+				Ar.Serialize(InChannel->Values.GetData(), NewArrayNum * SerializedElementSize);
+			}
+		}
+	}
+	else if (Ar.IsSaving())
+	{
+		int32 SerializedElementSize = sizeof(FFrameNumber);
+		Ar << SerializedElementSize;
+		InChannel->Times.CountBytes(Ar);
+		int32 ArrayCount = InChannel->Times.Num();
+		Ar << ArrayCount;
+		if (ArrayCount > 0)
+		{
+			Ar.Serialize(InChannel->Times.GetData(), ArrayCount * SerializedElementSize);
+		}
+		InChannel->Values.CountBytes(Ar);
+		SerializedElementSize = sizeof(ChannelValueType);
+		Ar << SerializedElementSize;
+		ArrayCount = InChannel->Values.Num();
+		Ar << ArrayCount;
+		if (ArrayCount > 0)
+		{
+			Ar.Serialize(InChannel->Values.GetData(), ArrayCount * SerializedElementSize);
+		}
+	}
+
+	Ar << InChannel->DefaultValue;
+	Ar << InChannel->bHasDefaultValue;
+	Ar << InChannel->TickResolution.Numerator;
+	Ar << InChannel->TickResolution.Denominator;
+	if (Ar.IsTransacting())
+	{
+		Ar << InChannel->KeyHandles;
+	}
+
+	if (bSerializeShowCurve)
+	{
+#if WITH_EDITOR
+		Ar << InChannel->bShowCurve;
+#else
+		bool bUnused = false;
+		Ar << bUnused;
+#endif
+	}
+	return true;
+}
+
+template<typename ChannelType>
+bool TMovieSceneCurveChannelImpl<ChannelType>::SerializeFromRichCurve(ChannelType* InChannel, const FPropertyTag& Tag, FStructuredArchive::FSlot Slot)
+{
+	static const FName RichCurveName("RichCurve");
+
+	check(InChannel);
+
+	if (Tag.Type == NAME_StructProperty && Tag.StructName == RichCurveName)
+	{
+		FRichCurve RichCurve;
+		FRichCurve::StaticStruct()->SerializeItem(Slot, &RichCurve, nullptr);
+
+		if (RichCurve.GetDefaultValue() != MAX_flt)
+		{
+			InChannel->bHasDefaultValue = true;
+			InChannel->DefaultValue = RichCurve.GetDefaultValue();
+		}
+
+		InChannel->PreInfinityExtrap = RichCurve.PreInfinityExtrap;
+		InChannel->PostInfinityExtrap = RichCurve.PostInfinityExtrap;
+
+		InChannel->Times.Reserve(RichCurve.GetNumKeys());
+		InChannel->Values.Reserve(RichCurve.GetNumKeys());
+
+		const FFrameRate LegacyFrameRate = GetLegacyConversionFrameRate();
+		const float      Interval        = LegacyFrameRate.AsInterval();
+
+		int32 Index = 0;
+		for (auto It = RichCurve.GetKeyIterator(); It; ++It)
+		{
+			const FRichCurveKey& Key = *It;
+
+			FFrameNumber KeyTime = UpgradeLegacyMovieSceneTime(nullptr, LegacyFrameRate, It->Time);
+
+			ChannelValueType NewValue;
+			NewValue.Value = Key.Value;
+			NewValue.InterpMode  = Key.InterpMode;
+			NewValue.TangentMode = Key.TangentMode;
+			NewValue.Tangent.ArriveTangent = Key.ArriveTangent * Interval;
+			NewValue.Tangent.LeaveTangent  = Key.LeaveTangent  * Interval;
+			ConvertInsertAndSort<ChannelValueType>(Index++, KeyTime, NewValue, InChannel->Times, InChannel->Values);
+		}
+
+		return true;
+	}
+
+	return false;
+}
+
+template<typename ChannelType>
+bool TMovieSceneCurveChannelImpl<ChannelType>::SerializeChannelValue(ChannelValueType& InValue, FArchive& Ar)
+{
+	Ar.UsingCustomVersion(FSequencerObjectVersion::GUID);
+	if (Ar.CustomVer(FSequencerObjectVersion::GUID) < FSequencerObjectVersion::SerializeFloatChannel)
+	{
+		return false;
+	}
+
+	if constexpr(TIsSame<CurveValueType, double>::Value)
+	{
+		// LWC_TODO: Serializer
+		// We serialize all values as float to allow the same data to be loaded with LWC enabled or disabled.
+		float TempValue = (float)InValue.Value;
+		Ar << TempValue;
+		if (Ar.IsLoading())
+		{
+			InValue.Value = (double)TempValue;
+		}
+	}
+	else
+	{
+		Ar << InValue.Value;
+	}
+
+	if (Ar.CustomVer(FSequencerObjectVersion::GUID) < FSequencerObjectVersion::SerializeFloatChannelCompletely)
+	{
+		// Serialization is handled manually to avoid the extra size overhead of FProperty tagging.
+		// Otherwise with many keys in a FMovieSceneFloatValue the size can become quite large.
+		Ar << InValue.InterpMode;
+		Ar << InValue.TangentMode;
+		Ar << InValue.Tangent;
+	}
+	else
+	{
+		Ar << InValue.Tangent.ArriveTangent;
+		Ar << InValue.Tangent.LeaveTangent;
+		Ar << InValue.Tangent.ArriveTangentWeight;
+		Ar << InValue.Tangent.LeaveTangentWeight;
+		Ar << InValue.Tangent.TangentWeightMode;
+		Ar << InValue.InterpMode;
+		Ar << InValue.TangentMode;
+		Ar << InValue.PaddingByte;
+	}
+
+	return true;
+}
+
+template struct MOVIESCENE_API TMovieSceneCurveChannelImpl<FMovieSceneFloatChannel>;
+template struct MOVIESCENE_API TMovieSceneCurveChannelImpl<FMovieSceneDoubleChannel>;
+
diff --git a/Engine/Source/Runtime/MovieScene/Private/Channels/MovieSceneFloatChannel.cpp b/Engine/Source/Runtime/MovieScene/Private/Channels/MovieSceneFloatChannel.cpp
--- a/Engine/Source/Runtime/MovieScene/Private/Channels/MovieSceneFloatChannel.cpp
+++ b/Engine/Source/Runtime/MovieScene/Private/Channels/MovieSceneFloatChannel.cpp
@@ -1,96 +1,21 @@
 // Copyright Epic Games, Inc. All Rights Reserved.
 
 #include "Channels/MovieSceneFloatChannel.h"
-#include "MovieSceneFwd.h"
 #include "Channels/MovieSceneChannelProxy.h"
+#include "Channels/MovieSceneCurveChannelImpl.h"
 #include "MovieSceneFrameMigration.h"
-#include "UObject/SequencerObjectVersion.h"
+#include "MovieSceneFwd.h"
 #include "UObject/FortniteMainBranchObjectVersion.h"
-#include "HAL/ConsoleManager.h"
-
-int32 GSequencerLinearCubicInterpolation = 1;
-static FAutoConsoleVariableRef CVarSequencerLinearCubicInterpolation(
-	TEXT("Sequencer.LinearCubicInterpolation"),
-	GSequencerLinearCubicInterpolation,
-	TEXT("If 1 Linear Keys Act As Cubic Interpolation with Linear Tangents, if 0 Linear Key Forces Linear Interpolation to Next Key."),
-	ECVF_Default);
-
-static TAutoConsoleVariable<int32> CVarSequencerAutoTangentInterpolation(
-	TEXT("Sequencer.AutoTangentNew"),
-	1,
-	TEXT("If 1 Auto Tangent will use new algorithm to gradually flatten maximum/minimum keys, if 0 Auto Tangent will average all keys (pre 4.23 behavior)."),
-	ECVF_Default);
-
-bool FMovieSceneTangentData::Serialize(FArchive& Ar)
-{
-	Ar.UsingCustomVersion(FSequencerObjectVersion::GUID);
-	if (Ar.CustomVer(FSequencerObjectVersion::GUID) < FSequencerObjectVersion::SerializeFloatChannel)
-	{
-		return false;
-	}
-
-	if (Ar.CustomVer(FSequencerObjectVersion::GUID) < FSequencerObjectVersion::SerializeFloatChannelCompletely)
-	{
-		Ar << ArriveTangent;
-		Ar << LeaveTangent;
-		Ar << TangentWeightMode;
-		Ar << ArriveTangentWeight;
-		Ar << LeaveTangentWeight;
-	}
-	else
-	{
-		Ar << ArriveTangent;
-		Ar << LeaveTangent;
-		Ar << ArriveTangentWeight;
-		Ar << LeaveTangentWeight;
-		Ar << TangentWeightMode;
-	}
+#include "UObject/SequencerObjectVersion.h"
 
-	return true;
-}
-
-bool FMovieSceneTangentData::operator==(const FMovieSceneTangentData& TangentData) const
-{
-	return (ArriveTangent == TangentData.ArriveTangent) && (LeaveTangent == TangentData.LeaveTangent) && (TangentWeightMode == TangentData.TangentWeightMode) && (ArriveTangentWeight == TangentData.ArriveTangentWeight) && (LeaveTangentWeight == TangentData.LeaveTangentWeight);
-}
-
-bool FMovieSceneTangentData::operator!=(const FMovieSceneTangentData& Other) const
-{
-	return !(*this == Other);
-}
+static_assert(
+		sizeof(FMovieSceneFloatValue) == 28,
+		"The size of the float channel value has changed. You need to update the padding byte at the end of the structure. "
+		"You also need to update the layout in FMovieSceneDoubleValue so that they match!");
 
 bool FMovieSceneFloatValue::Serialize(FArchive& Ar)
 {
-	Ar.UsingCustomVersion(FSequencerObjectVersion::GUID);
-	if (Ar.CustomVer(FSequencerObjectVersion::GUID) < FSequencerObjectVersion::SerializeFloatChannel)
-	{
-		return false;
-	}
-	
-	if (Ar.CustomVer(FSequencerObjectVersion::GUID) < FSequencerObjectVersion::SerializeFloatChannelCompletely)
-	{
-		// Serialization is handled manually to avoid the extra size overhead of FProperty tagging.
-		// Otherwise with many keys in a FMovieSceneFloatValue the size can become quite large.
-		Ar << Value;
-		Ar << InterpMode;
-		Ar << TangentMode;
-		Ar << Tangent;
-	}
-	else
-	{
-		Ar << Value;
-		Ar << Tangent.ArriveTangent;
-		Ar << Tangent.LeaveTangent;
-		Ar << Tangent.ArriveTangentWeight;
-		Ar << Tangent.LeaveTangentWeight;
-		Ar << Tangent.TangentWeightMode;
-		Ar << InterpMode;
-		Ar << TangentMode;
-		Ar << PaddingByte;
-	}
-
-
-	return true;
+	return TMovieSceneCurveChannelImpl<FMovieSceneFloatChannel>::SerializeChannelValue(*this, Ar);
 }
 
 bool FMovieSceneFloatValue::operator==(const FMovieSceneFloatValue& FloatValue) const
@@ -103,799 +28,39 @@ bool FMovieSceneFloatValue::operator!=(const FMovieSceneFloatValue& Other) const
 	return !(*this == Other);
 }
 
-bool FMovieSceneFloatChannel::SerializeFromMismatchedTag(const FPropertyTag& Tag, FStructuredArchive::FSlot Slot)
-{
-	static const FName RichCurveName("RichCurve");
-	if (Tag.Type == NAME_StructProperty && Tag.StructName == RichCurveName)
-	{
-		FRichCurve RichCurve;
-		FRichCurve::StaticStruct()->SerializeItem(Slot, &RichCurve, nullptr);
-
-		if (RichCurve.GetDefaultValue() != MAX_flt)
-		{
-			bHasDefaultValue = true;
-			DefaultValue = RichCurve.GetDefaultValue();
-		}
-
-		PreInfinityExtrap = RichCurve.PreInfinityExtrap;
-		PostInfinityExtrap = RichCurve.PostInfinityExtrap;
-
-		Times.Reserve(RichCurve.GetNumKeys());
-		Values.Reserve(RichCurve.GetNumKeys());
-
-		const FFrameRate LegacyFrameRate = GetLegacyConversionFrameRate();
-		const float      Interval        = LegacyFrameRate.AsInterval();
-
-		int32 Index = 0;
-		for (auto It = RichCurve.GetKeyIterator(); It; ++It)
-		{
-			const FRichCurveKey& Key = *It;
-
-			FFrameNumber KeyTime = UpgradeLegacyMovieSceneTime(nullptr, LegacyFrameRate, It->Time);
-
-			FMovieSceneFloatValue NewValue;
-			NewValue.Value = Key.Value;
-			NewValue.InterpMode  = Key.InterpMode;
-			NewValue.TangentMode = Key.TangentMode;
-			NewValue.Tangent.ArriveTangent = Key.ArriveTangent * Interval;
-			NewValue.Tangent.LeaveTangent  = Key.LeaveTangent  * Interval;
-			ConvertInsertAndSort<FMovieSceneFloatValue>(Index++, KeyTime, NewValue, Times, Values);
-		}
-		return true;
-	}
-
-	return false;
-}
-
-int32 FMovieSceneFloatChannel::InsertKeyInternal(FFrameNumber InTime)
-{
-	const int32 InsertIndex = Algo::UpperBound(Times, InTime);
-
-	Times.Insert(InTime, InsertIndex);
-	Values.Insert(FMovieSceneFloatValue(), InsertIndex);
-
-	KeyHandles.AllocateHandle(InsertIndex);
-
-	return InsertIndex;
-}
-
 int32 FMovieSceneFloatChannel::AddConstantKey(FFrameNumber InTime, float InValue)
 {
-	const int32 Index = InsertKeyInternal(InTime);
-
-	FMovieSceneFloatValue& Value = Values[Index];
-	Value.Value = InValue;
-	Value.InterpMode = RCIM_Constant;
-
-	AutoSetTangents();
-
-	return Index;
+	return FMovieSceneFloatChannelImpl::AddConstantKey(this, InTime, InValue);
 }
 
 int32 FMovieSceneFloatChannel::AddLinearKey(FFrameNumber InTime, float InValue)
 {
-	const int32 Index = InsertKeyInternal(InTime);
-
-	FMovieSceneFloatValue& Value = Values[Index];
-	Value.Value = InValue;
-	Value.InterpMode = RCIM_Linear;
-
-	AutoSetTangents();
-
-	return Index;
+	return FMovieSceneFloatChannelImpl::AddLinearKey(this, InTime, InValue);
 }
 
 int32 FMovieSceneFloatChannel::AddCubicKey(FFrameNumber InTime, float InValue, ERichCurveTangentMode TangentMode, const FMovieSceneTangentData& Tangent)
 {
-	const int32 Index = InsertKeyInternal(InTime);
-
-	FMovieSceneFloatValue& Value = Values[Index];
-	Value.Value = InValue;
-	Value.InterpMode = RCIM_Cubic;
-	Value.TangentMode = TangentMode;
-	Value.Tangent = Tangent;
-
-	AutoSetTangents();
-
-	return Index;
-}
-
-/** Util to find float value on bezier defined by 4 control points */
-float BezierInterp(float P0, float P1, float P2, float P3, float Alpha)
-{
-	const float P01   = FMath::Lerp(P0,   P1,   Alpha);
-	const float P12   = FMath::Lerp(P1,   P2,   Alpha);
-	const float P23   = FMath::Lerp(P2,   P3,   Alpha);
-	const float P012  = FMath::Lerp(P01,  P12,  Alpha);
-	const float P123  = FMath::Lerp(P12,  P23,  Alpha);
-	const float P0123 = FMath::Lerp(P012, P123, Alpha);
-
-	return P0123;
-}
-
-
-static float EvalForTwoKeys(const FMovieSceneFloatValue& Key1, FFrameNumber Key1Time,
-							const FMovieSceneFloatValue& Key2, FFrameNumber Key2Time,
-							FFrameNumber InTime,
-							FFrameRate DisplayRate)
-{
-	double DecimalRate = DisplayRate.AsDecimal();
-
-	float Diff = (float)(Key2Time - Key1Time).Value;
-	Diff /= DecimalRate;
-	const int CheckBothLinear = GSequencerLinearCubicInterpolation;
-
-	if (Diff > 0 && Key1.InterpMode != RCIM_Constant)
-	{
-		const float Alpha = ((float)(InTime - Key1Time).Value / DecimalRate) / Diff;
-		const float P0 = Key1.Value;
-		const float P3 = Key2.Value;
-
-		if (Key1.InterpMode == RCIM_Linear && (!CheckBothLinear || Key2.InterpMode != RCIM_Cubic))
-		{
-			return FMath::Lerp(P0, P3, Alpha);
-		}
-		else
-		{
-			float LeaveTangent = Key1.Tangent.LeaveTangent * DecimalRate;
-			float ArriveTangent = Key2.Tangent.ArriveTangent * DecimalRate;
-
-			const float OneThird = 1.0f / 3.0f;
-			const float P1 = P0 + (LeaveTangent * Diff*OneThird);
-			const float P2 = P3 - (ArriveTangent * Diff*OneThird);
-
-			return BezierInterp(P0, P1, P2, P3, Alpha);
-		}
-	}
-	else
-	{
-		return Key1.Value;
-	}
-}
-
-struct FCycleParams
-{
-	FFrameTime Time;
-	int32 CycleCount;
-	float ValueOffset;
-
-	FCycleParams(FFrameTime InTime)
-		: Time(InTime)
-		, CycleCount(0)
-		, ValueOffset(0.f)
-	{}
-
-	FORCEINLINE void ComputePreValueOffset(float FirstValue, float LastValue)
-	{
-		ValueOffset = (FirstValue-LastValue) * CycleCount;
-	}
-	FORCEINLINE void ComputePostValueOffset(float FirstValue, float LastValue)
-	{
-		ValueOffset = (LastValue-FirstValue) * CycleCount;
-	}
-	FORCEINLINE void Oscillate(int32 MinFrame, int32 MaxFrame)
-	{
-		if (CycleCount % 2 == 1)
-		{
-			Time = MinFrame + (FFrameTime(MaxFrame) - Time);
-		}
-	}
-};
-
-FCycleParams CycleTime(FFrameNumber MinFrame, FFrameNumber MaxFrame, FFrameTime InTime)
-{
-	FCycleParams Params(InTime);
-	
-	const int32 Duration = MaxFrame.Value - MinFrame.Value;
-	if (Duration == 0)
-	{
-		Params.Time = MaxFrame;
-		Params.CycleCount = 0;
-	}
-	else if (InTime < MinFrame)
-	{
-		const int32 CycleCount = ((MaxFrame - InTime) / Duration).FloorToFrame().Value;
-
-		Params.Time = InTime + FFrameTime(Duration)*CycleCount;
-		Params.CycleCount = CycleCount;
-	}
-	else if (InTime > MaxFrame)
-	{
-		const int32 CycleCount = ((InTime - MinFrame) / Duration).FloorToFrame().Value;
-
-		Params.Time = InTime - FFrameTime(Duration)*CycleCount;
-		Params.CycleCount = CycleCount;
-	}
-
-	return Params;
-}
-
-bool FMovieSceneFloatChannel::EvaluateExtrapolation(FFrameTime InTime, float& OutValue) const
-{
-	// If the time is outside of the curve, deal with extrapolation
-	if (InTime < Times[0])
-	{
-		if (PreInfinityExtrap == RCCE_None)
-		{
-			return false;
-		}
-
-		if (PreInfinityExtrap == RCCE_Constant)
-		{
-			OutValue = Values[0].Value;
-			return true;
-		}
-
-		if (PreInfinityExtrap == RCCE_Linear)
-		{
-			const FMovieSceneFloatValue FirstValue = Values[0];
-
-			if (FirstValue.InterpMode == RCIM_Constant)
-			{
-				OutValue = FirstValue.Value;
-			}
-			else if(FirstValue.InterpMode == RCIM_Cubic)
-			{
-				FFrameTime Delta = FFrameTime(Times[0]) - InTime;
-				OutValue = FirstValue.Value - Delta.AsDecimal() * FirstValue.Tangent.ArriveTangent;
-			}
-			else if(FirstValue.InterpMode == RCIM_Linear)
-			{
-				const int32 InterpStartFrame = Times[1].Value;
-				const int32 DeltaFrame       = InterpStartFrame - Times[0].Value;
-				if (DeltaFrame == 0)
-				{
-					OutValue = FirstValue.Value;
-				}
-				else
-				{
-					OutValue = FMath::Lerp(Values[1].Value, FirstValue.Value, (InterpStartFrame - InTime.AsDecimal())/DeltaFrame);
-				}
-			}
-			return true;
-		}
-	}
-	else if (InTime > Times.Last())
-	{
-		if (PostInfinityExtrap == RCCE_None)
-		{
-			return false;
-		}
-
-		if (PostInfinityExtrap == RCCE_Constant)
-		{
-			OutValue = Values.Last().Value;
-			return true;
-		}
-
-		if (PostInfinityExtrap == RCCE_Linear)
-		{
-			const FMovieSceneFloatValue LastValue = Values.Last();
-
-			if (LastValue.InterpMode == RCIM_Constant)
-			{
-				OutValue = LastValue.Value;
-			}
-			else if(LastValue.InterpMode == RCIM_Cubic)
-			{
-				FFrameTime Delta = InTime - Times.Last();
-				OutValue = LastValue.Value + Delta.AsDecimal() * LastValue.Tangent.LeaveTangent;
-			}
-			else if(LastValue.InterpMode == RCIM_Linear)
-			{
-				const int32 NumKeys          = Times.Num();
-				const int32 InterpStartFrame = Times[NumKeys-2].Value;
-				const int32 DeltaFrame       = Times.Last().Value-InterpStartFrame;
-
-				if (DeltaFrame == 0)
-				{
-					OutValue = LastValue.Value;
-				}
-				else
-				{
-					OutValue = FMath::Lerp(Values[NumKeys-2].Value, LastValue.Value, (InTime.AsDecimal() - InterpStartFrame)/DeltaFrame);
-				}
-			}
-			return true;
-		}
-	}
-
-	return false;
-}
-
-
-/* Solve Cubic Euqation using Cardano's forumla
-* Adopted from Graphic Gems 1
-* [REDACTED_PROFILE_URL]/GraphicsGems/blob/master/gems/Roots3And4.c
-*  Solve cubic of form
-*
-* @param Coeff Coefficient parameters of form  Coeff[0] + Coeff[1]*x + Coeff[2]*x^2 + Coeff[3]*x^3 + Coeff[4]*x^4 = 0
-* @param Solution Up to 3 real solutions. We don't include imaginary solutions, would need a complex number objecct
-* @return Returns the number of real solutions returned in the Solution array.
-*/
-static int SolveCubic(double Coeff[4], double Solution[3])
-{
-	auto cbrt = [](double x) -> double
-	{
-		return ((x) > 0.0 ? pow((x), 1.0 / 3.0) : ((x) < 0.0 ? -pow((double)-(x), 1.0 / 3.0) : 0.0));
-	};
-	int     NumSolutions = 0;
-
-	/* normal form: x^3 + Ax^2 + Bx + C = 0 */
-
-	double A = Coeff[2] / Coeff[3];
-	double B = Coeff[1] / Coeff[3];
-	double C = Coeff[0] / Coeff[3];
-
-	/*  substitute x = y - A/3 to eliminate quadric term:
-	x^3 +px + q = 0 */
-
-	double SqOfA = A * A;
-	double P = 1.0 / 3 * (-1.0 / 3 * SqOfA + B);
-	double Q = 1.0 / 2 * (2.0 / 27 * A * SqOfA - 1.0 / 3 * A * B + C);
-
-	/* use Cardano's formula */
-
-	double CubeOfP = P * P * P;
-	double D = Q * Q + CubeOfP;
-
-	if (FMath::IsNearlyZero(D))
-	{
-		if (FMath::IsNearlyZero(Q)) /* one triple solution */
-		{
-			Solution[0] = 0;
-			NumSolutions = 1;
-		}
-		else /* one single and one double solution */
-		{
-			double u = cbrt(-Q);
-			Solution[0] = 2 * u;
-			Solution[1] = -u;
-			NumSolutions = 2;
-		}
-	}
-	else if (D < 0) /* Casus irreducibilis: three real solutions */
-	{
-		double phi = 1.0 / 3 * acos(-Q / sqrt(-CubeOfP));
-		double t = 2 * sqrt(-P);
-
-		Solution[0] = t * cos(phi);
-		Solution[1] = -t * cos(phi + PI / 3);
-		Solution[2] = -t * cos(phi - PI / 3);
-		NumSolutions = 3;
-	}
-	else /* one real solution */
-	{
-		double sqrt_D = sqrt(D);
-		double u = cbrt(sqrt_D - Q);
-		double v = -cbrt(sqrt_D + Q);
-
-		Solution[0] = u + v;
-		NumSolutions = 1;
-	}
-
-	/* resubstitute */
-
-	double Sub = 1.0 / 3 * A;
-
-	for (int i = 0; i < NumSolutions; ++i)
-		Solution[i] -= Sub;
-
-	return NumSolutions;
+	return FMovieSceneFloatChannelImpl::AddCubicKey(this, InTime, InValue, TangentMode, Tangent);
 }
 
-/*
-*   Convert the control values for a polynomial defined in the Bezier
-*		basis to a polynomial defined in the power basis (t^3 t^2 t 1).
-*/
-static void BezierToPower(	double A1, double B1, double C1, double D1,
-	double *A2, double *B2, double *C2, double *D2)
+bool FMovieSceneFloatChannel::Evaluate(FFrameTime InTime,  float& OutValue) const
 {
-	double A = B1 - A1;
-	double B = C1 - B1;
-	double C = D1 - C1;
-	double D = B - A;
-	*A2 = C- B - D;
-	*B2 = 3.0 * D;
-	*C2 = 3.0 * A;
-	*D2 = A1;
+	return FMovieSceneFloatChannelImpl::Evaluate(this, InTime, OutValue);
 }
 
-bool FMovieSceneFloatChannel::Evaluate(FFrameTime InTime,  float& OutValue) const
+void FMovieSceneFloatChannel::Set(TArray<FFrameNumber> InTimes, TArray<FMovieSceneFloatValue> InValues)
 {
-	const int32 NumKeys = Times.Num();
-
-	// No keys means default value, or nothing
-	if (NumKeys == 0)
-	{
-		if (bHasDefaultValue)
-		{
-			OutValue = DefaultValue;
-			return true;
-		}
-		return false;
-	}
-
-	// For single keys, we can only ever return that value
-	if (NumKeys == 1)
-	{
-		OutValue = Values[0].Value;
-		return true;
-	}
-
-	// Evaluate with extrapolation if we're outside the bounds of the curve
-	if (EvaluateExtrapolation(InTime, OutValue))
-	{
-		return true;
-	}
-
-	const FFrameNumber MinFrame = Times[0];
-	const FFrameNumber MaxFrame = Times.Last();
-
-	// Compute the cycled time
-	FCycleParams Params = CycleTime(MinFrame, MaxFrame, InTime);
-
-	// Deal with offset cycles and oscillation
-	if (InTime < FFrameTime(MinFrame))
-	{
-		switch (PreInfinityExtrap)
-		{
-		case RCCE_CycleWithOffset: Params.ComputePreValueOffset(Values[0].Value, Values[NumKeys-1].Value); break;
-		case RCCE_Oscillate:       Params.Oscillate(MinFrame.Value, MaxFrame.Value);                       break;
-		}
-	}
-	else if (InTime > FFrameTime(MaxFrame))
-	{
-		switch (PostInfinityExtrap)
-		{
-		case RCCE_CycleWithOffset: Params.ComputePostValueOffset(Values[0].Value, Values[NumKeys-1].Value); break;
-		case RCCE_Oscillate:       Params.Oscillate(MinFrame.Value, MaxFrame.Value);                        break;
-		}
-	}
-
-	if (!ensureMsgf(Params.Time.FrameNumber >= MinFrame && Params.Time.FrameNumber <= MaxFrame, TEXT("Invalid time computed for float channel evaluation")))
-	{
-		return false;
-	}
-
-	// Evaluate the curve data
-	float Interp = 0.f;
-	int32 Index1 = INDEX_NONE, Index2 = INDEX_NONE;
-	UE::MovieScene::EvaluateTime(Times, Params.Time, Index1, Index2, Interp);
-	const int CheckBothLinear = GSequencerLinearCubicInterpolation;
-
-	if (Index1 == INDEX_NONE)
-	{
-		OutValue = Params.ValueOffset + Values[Index2].Value;
-	}
-	else if (Index2 == INDEX_NONE)
-	{
-		OutValue = Params.ValueOffset + Values[Index1].Value;
-	}
-	else
-	{
-		FMovieSceneFloatValue Key1 = Values[Index1];
-		FMovieSceneFloatValue Key2 = Values[Index2];
-		TEnumAsByte<ERichCurveInterpMode> InterpMode = Key1.InterpMode;
-	    if(InterpMode == RCIM_Linear && (CheckBothLinear  && Key2.InterpMode == RCIM_Cubic))
-		{
-			InterpMode = RCIM_Cubic;
-		}
-		
-		switch (InterpMode)
-		{
-		case RCIM_Cubic:
-		{
-			const float OneThird = 1.0f / 3.0f;
-			if ((Key1.Tangent.TangentWeightMode == RCTWM_WeightedNone || Key1.Tangent.TangentWeightMode == RCTWM_WeightedArrive)
-				&& (Key2.Tangent.TangentWeightMode == RCTWM_WeightedNone || Key2.Tangent.TangentWeightMode == RCTWM_WeightedLeave))
-			{
-				const int32 Diff = Times[Index2].Value - Times[Index1].Value;
-				const float P0 = Key1.Value;
-				const float P1 = P0 + (Key1.Tangent.LeaveTangent * Diff * OneThird);
-				const float P3 = Key2.Value;
-				const float P2 = P3 - (Key2.Tangent.ArriveTangent * Diff * OneThird);
-
-				OutValue = Params.ValueOffset + BezierInterp(P0, P1, P2, P3, Interp);
-				break;
-			}
-			else //its weighted
-			{
-				const float TimeInterval = TickResolution.AsInterval();
-				const float ToSeconds = 1.0f / TimeInterval;
-
-				const double Time1 = TickResolution.AsSeconds(Times[Index1].Value);
-				const double Time2 = TickResolution.AsSeconds(Times[Index2].Value);
-				const float X = Time2 - Time1;
-				float CosAngle, SinAngle;
-				float Angle = FMath::Atan(Key1.Tangent.LeaveTangent * ToSeconds);
-				FMath::SinCos(&SinAngle, &CosAngle, Angle);
-				float LeaveWeight;
-				if (Key1.Tangent.TangentWeightMode == RCTWM_WeightedNone || Key1.Tangent.TangentWeightMode == RCTWM_WeightedArrive)
-				{
-					const float LeaveTangentNormalized = Key1.Tangent.LeaveTangent / (TimeInterval);
-					const float Y = LeaveTangentNormalized * X;
-					LeaveWeight = FMath::Sqrt(X*X + Y * Y) * OneThird;
-				}
-				else
-				{
-					LeaveWeight = Key1.Tangent.LeaveTangentWeight;
-				}
-				const float Key1TanX = CosAngle * LeaveWeight + Time1;
-				const float Key1TanY = SinAngle * LeaveWeight + Key1.Value;
-
-				Angle = FMath::Atan(Key2.Tangent.ArriveTangent * ToSeconds);
-				FMath::SinCos(&SinAngle, &CosAngle, Angle);
-				float ArriveWeight;
-				if (Key2.Tangent.TangentWeightMode == RCTWM_WeightedNone || Key2.Tangent.TangentWeightMode == RCTWM_WeightedLeave)
-				{
-					const float ArriveTangentNormalized = Key2.Tangent.ArriveTangent / (TimeInterval);
-					const float Y = ArriveTangentNormalized * X;
-					ArriveWeight = FMath::Sqrt(X*X + Y * Y) * OneThird;
-				}
-				else
-				{
-					ArriveWeight =  Key2.Tangent.ArriveTangentWeight;
-				}
-				const float Key2TanX = -CosAngle * ArriveWeight + Time2;
-				const float Key2TanY = -SinAngle * ArriveWeight + Key2.Value;
-
-				//Normalize the Time Range
-				const float RangeX = Time2 - Time1;
-
-				const float Dx1 = Key1TanX - Time1;
-				const float Dx2 = Key2TanX - Time1;
-
-				// Normalize values
-				const float NormalizedX1 = Dx1 / RangeX;
-				const float NormalizedX2 = Dx2 / RangeX;
-				
-				double Coeff[4];
-				double Results[3];
-
-				//Convert Bezier to Power basis, also float to double for precision for root finding.
-				BezierToPower(
-					0.0, NormalizedX1, NormalizedX2, 1.0,
-					&(Coeff[3]), &(Coeff[2]), &(Coeff[1]), &(Coeff[0])
-				);
-
-				Coeff[0] = Coeff[0] - Interp;
-				
-				int NumResults = SolveCubic(Coeff, Results);
-				float NewInterp = Interp;
-				if (NumResults == 1)
-				{
-					NewInterp = Results[0];
-				}
-				else
-				{
-					NewInterp = TNumericLimits<float>::Lowest(); //just need to be out of range
-					for (double Result : Results)
-					{
-						if ((Result >= 0.0f) && (Result <= 1.0f))
-						{
-							if (NewInterp < 0.0f || Result > NewInterp)
-							{
-								NewInterp = Result;
-							}
-						}
-					}
-
-					if (NewInterp == TNumericLimits<float>::Lowest())
-					{
-						NewInterp = 0.f;
-					}
-
-				}
-				//now use NewInterp and adjusted tangents plugged into the Y (Value) part of the graph.
-				const float P0 = Key1.Value;
-				const float P1 = Key1TanY;
-				const float P3 = Key2.Value;
-				const float P2 = Key2TanY;
-
-				OutValue = Params.ValueOffset + BezierInterp(P0, P1, P2, P3,  NewInterp);
-			}
-			break;
-		}
-
-		case RCIM_Linear:
-			OutValue = Params.ValueOffset + FMath::Lerp(Key1.Value, Key2.Value, Interp);
-			break;
-
-		default:
-			OutValue = Params.ValueOffset + Key1.Value;
-			break;
-		}
-	}
-
-	return true;
+	FMovieSceneFloatChannelImpl::Set(this, InTimes, InValues);
 }
 
 void FMovieSceneFloatChannel::AutoSetTangents(float Tension)
 {
-	if (Values.Num() < 2)
-	{
-		return;
-	}
-	const int UseNewAutoTangent = CVarSequencerAutoTangentInterpolation->GetInt();
-
-	{
-		FMovieSceneFloatValue& FirstValue = Values[0];
-		if (FirstValue.InterpMode == RCIM_Linear)
-		{
-			FirstValue.Tangent.TangentWeightMode = RCTWM_WeightedNone;
-			FMovieSceneFloatValue& NextKey = Values[1];
-			const float NextTimeDiff = FMath::Max<double>(KINDA_SMALL_NUMBER, Times[1].Value - Times[0].Value);
-			float NewTangent = (NextKey.Value - FirstValue.Value) / NextTimeDiff;
-			FirstValue.Tangent.LeaveTangent = NewTangent;
-		}
-		else if (FirstValue.InterpMode == RCIM_Cubic && FirstValue.TangentMode == RCTM_Auto)
-		{
-			FirstValue.Tangent.LeaveTangent = FirstValue.Tangent.ArriveTangent = 0.0f;
-			FirstValue.Tangent.TangentWeightMode = RCTWM_WeightedNone; 
-		}
-	}
-
-	{
-		FMovieSceneFloatValue& LastValue = Values.Last();
-		if (LastValue.InterpMode == RCIM_Linear)
-		{
-			LastValue.Tangent.TangentWeightMode = RCTWM_WeightedNone;
-			int32 Index = Values.Num() - 1;
-			FMovieSceneFloatValue& PrevKey = Values[Index-1];
-			const float PrevTimeDiff = FMath::Max<double>(KINDA_SMALL_NUMBER, Times[Index].Value - Times[Index - 1].Value);
-			float NewTangent = (LastValue.Value - PrevKey.Value) / PrevTimeDiff;
-			LastValue.Tangent.ArriveTangent = NewTangent;
-		}
-		else if (LastValue.InterpMode == RCIM_Cubic && LastValue.TangentMode == RCTM_Auto)
-		{
-			LastValue.Tangent.LeaveTangent = LastValue.Tangent.ArriveTangent = 0.0f;
-			LastValue.Tangent.TangentWeightMode = RCTWM_WeightedNone;
-		}
-	}
-
-	for (int32 Index = 1; Index < Values.Num() - 1; ++Index)
-	{
-		FMovieSceneFloatValue  PrevKey = Values[Index-1];
-		FMovieSceneFloatValue& ThisKey = Values[Index  ];
-
-		if (ThisKey.InterpMode == RCIM_Cubic && ThisKey.TangentMode == RCTM_Auto)
-		{
-			FMovieSceneFloatValue& NextKey = Values[Index+1];
-			float NewTangent = 0.f;
-			const float PrevToNextTimeDiff = FMath::Max<double>(KINDA_SMALL_NUMBER, Times[Index + 1].Value - Times[Index - 1].Value);
-
-			if (!UseNewAutoTangent)
-			{
-				AutoCalcTangent(PrevKey.Value, ThisKey.Value, NextKey.Value, Tension, NewTangent);
-				NewTangent /= PrevToNextTimeDiff;
-			}
-			else
-			{
-				// if key doesn't lie between we keep it flat(0.0).
-				if ( (ThisKey.Value > PrevKey.Value && ThisKey.Value < NextKey.Value) ||
-					(ThisKey.Value < PrevKey.Value && ThisKey.Value > NextKey.Value))
-				{
-					AutoCalcTangent(PrevKey.Value, ThisKey.Value, NextKey.Value, Tension, NewTangent);
-					NewTangent /= PrevToNextTimeDiff;
-					//if within 0 to 15% or 85% to 100% range we gradually weight tangent to zero
-					const float AverageToZeroRange = 0.85f;
-					const float ValDiff = FMath::Abs<float>(NextKey.Value - PrevKey.Value);
-					const float OurDiff = FMath::Abs<float>(ThisKey.Value - PrevKey.Value);
-					//ValDiff won't be zero ever due to previous check
-					float PercDiff = OurDiff / ValDiff;
-					if (PercDiff > AverageToZeroRange)
-					{
-						PercDiff = (PercDiff - AverageToZeroRange) / (1.0f - AverageToZeroRange);
-						NewTangent = NewTangent * (1.0f - PercDiff);
-					}
-					else if (PercDiff < (1.0f - AverageToZeroRange))
-					{
-						PercDiff = PercDiff  / (1.0f - AverageToZeroRange);
-						NewTangent = NewTangent * PercDiff;
-					}
-				}
-			}
-
-			// In 'auto' mode, arrive and leave tangents are always the same
-			ThisKey.Tangent.LeaveTangent = ThisKey.Tangent.ArriveTangent = NewTangent;
-			ThisKey.Tangent.TangentWeightMode = RCTWM_WeightedNone;
-		}
-		else if (ThisKey.InterpMode == RCIM_Linear)
-		{
-			ThisKey.Tangent.TangentWeightMode = RCTWM_WeightedNone; 
-			FMovieSceneFloatValue& NextKey = Values[Index + 1];
-
-			const float PrevTimeDiff = FMath::Max<double>(KINDA_SMALL_NUMBER, Times[Index].Value - Times[Index - 1].Value);
-			float NewTangent  = (ThisKey.Value - PrevKey.Value) / PrevTimeDiff;
-			ThisKey.Tangent.ArriveTangent = NewTangent;
-			
-			const float NextTimeDiff = FMath::Max<double>(KINDA_SMALL_NUMBER, Times[Index + 1].Value - Times[Index].Value);
-			NewTangent = (NextKey.Value - ThisKey.Value) / NextTimeDiff;
-			ThisKey.Tangent.LeaveTangent = NewTangent;
-		}
-
-		
-	}
+	FMovieSceneFloatChannelImpl::AutoSetTangents(this, Tension);
 }
 
 void FMovieSceneFloatChannel::PopulateCurvePoints(double StartTimeSeconds, double EndTimeSeconds, double TimeThreshold, float ValueThreshold, FFrameRate InTickResolution, TArray<TTuple<double, double>>& InOutPoints) const
 {
-	const FFrameNumber StartFrame = (StartTimeSeconds * InTickResolution).FloorToFrame();
-	const FFrameNumber EndFrame   = (EndTimeSeconds   * InTickResolution).CeilToFrame();
-
-	const int32 StartingIndex = Algo::UpperBound(Times, StartFrame);
-	const int32 EndingIndex   = Algo::LowerBound(Times, EndFrame);
-
-	// Add the lower bound of the visible space
-	float EvaluatedValue = 0.f;
-	if (Evaluate(StartFrame, EvaluatedValue))
-	{
-		InOutPoints.Add(MakeTuple(StartFrame / InTickResolution, double(EvaluatedValue)));
-	}
-
-	// Add all keys in-between
-	for (int32 KeyIndex = StartingIndex; KeyIndex < EndingIndex; ++KeyIndex)
-	{
-		InOutPoints.Add(MakeTuple(Times[KeyIndex] / InTickResolution, double(Values[KeyIndex].Value)));
-	}
-
-	// Add the upper bound of the visible space
-	if (Evaluate(EndFrame, EvaluatedValue))
-	{
-		InOutPoints.Add(MakeTuple(EndFrame / InTickResolution, double(EvaluatedValue)));
-	}
-
-	int32 OldSize = InOutPoints.Num();
-	do
-	{
-		OldSize = InOutPoints.Num();
-		RefineCurvePoints(InTickResolution, TimeThreshold, ValueThreshold, InOutPoints);
-	}
-	while(OldSize != InOutPoints.Num());
-}
-
-void FMovieSceneFloatChannel::RefineCurvePoints(FFrameRate InTickResolution, double TimeThreshold, float ValueThreshold, TArray<TTuple<double, double>>& InOutPoints) const
-{
-	const float InterpTimes[] = { 0.25f, 0.5f, 0.6f };
-
-	for (int32 Index = 0; Index < InOutPoints.Num() - 1; ++Index)
-	{
-		TTuple<double, double> Lower = InOutPoints[Index];
-		TTuple<double, double> Upper = InOutPoints[Index + 1];
-
-		if ((Upper.Get<0>() - Lower.Get<0>()) >= TimeThreshold)
-		{
-			bool bSegmentIsLinear = true;
-
-			TTuple<double, double> Evaluated[UE_ARRAY_COUNT(InterpTimes)] = { TTuple<double, double>(0, 0) };
-
-			for (int32 InterpIndex = 0; InterpIndex < UE_ARRAY_COUNT(InterpTimes); ++InterpIndex)
-			{
-				double& EvalTime  = Evaluated[InterpIndex].Get<0>();
-
-				EvalTime = FMath::Lerp(Lower.Get<0>(), Upper.Get<0>(), InterpTimes[InterpIndex]);
-
-				float Value = 0.f;
-				Evaluate(EvalTime * InTickResolution, Value);
-
-				const float LinearValue = FMath::Lerp(Lower.Get<1>(), Upper.Get<1>(), InterpTimes[InterpIndex]);
-				if (bSegmentIsLinear)
-				{
-					bSegmentIsLinear = FMath::IsNearlyEqual(Value, LinearValue, ValueThreshold);
-				}
-
-				Evaluated[InterpIndex].Get<1>() = Value;
-			}
-
-			if (!bSegmentIsLinear)
-			{
-				// Add the point
-				InOutPoints.Insert(Evaluated, UE_ARRAY_COUNT(Evaluated), Index+1);
-				--Index;
-			}
-		}
-	}
+	FMovieSceneFloatChannelImpl::PopulateCurvePoints(this, StartTimeSeconds, EndTimeSeconds, TimeThreshold, ValueThreshold, InTickResolution, InOutPoints);
 }
 
 void FMovieSceneFloatChannel::GetKeys(const TRange<FFrameNumber>& WithinRange, TArray<FFrameNumber>* OutKeyTimes, TArray<FKeyHandle>* OutKeyHandles)
@@ -927,37 +92,13 @@ void FMovieSceneFloatChannel::DeleteKeys(TArrayView<const FKeyHandle> InHandles)
 
 void FMovieSceneFloatChannel::DeleteKeysFrom(FFrameNumber InTime, bool bDeleteKeysBefore)
 {
-	// Insert a key at the current time to maintain evaluation
-	if (GetData().GetTimes().Num() > 0)
-	{
-		int32 KeyHandleIndex = GetData().FindKey(InTime);
-		if (KeyHandleIndex == INDEX_NONE)
-		{
-			float Value = 0.f;
-			if (Evaluate(InTime, Value))
-			{
-				AddCubicKey(InTime, Value);
-			}
-		}
-	}
-
-	GetData().DeleteKeysFrom(InTime, bDeleteKeysBefore);
+	FMovieSceneFloatChannelImpl::DeleteKeysFrom(this, InTime, bDeleteKeysBefore);
 	AutoSetTangents();
 }
 
 void FMovieSceneFloatChannel::ChangeFrameResolution(FFrameRate SourceRate, FFrameRate DestinationRate)
 {
-	check(Times.Num() == Values.Num());
-
-	float IntervalFactor = DestinationRate.AsInterval() / SourceRate.AsInterval();
-	for (int32 Index = 0; Index < Times.Num(); ++Index)
-	{
-		Times[Index] = ConvertFrameTime(Times[Index], SourceRate, DestinationRate).RoundToFrame();
-
-		FMovieSceneFloatValue& Value = Values[Index];
-		Value.Tangent.ArriveTangent *= IntervalFactor;
-		Value.Tangent.LeaveTangent  *= IntervalFactor;
-	}
+	FMovieSceneFloatChannelImpl::ChangeFrameResolution(this, SourceRate, DestinationRate);
 }
 
 TRange<FFrameNumber> FMovieSceneFloatChannel::ComputeEffectiveRange() const
@@ -991,46 +132,7 @@ void FMovieSceneFloatChannel::Offset(FFrameNumber DeltaPosition)
 
 void FMovieSceneFloatChannel::Optimize(const FKeyDataOptimizationParams& Params)
 {
-	TMovieSceneChannelData<FMovieSceneFloatValue> ChannelData = GetData();
-	TArray<FFrameNumber> OutKeyTimes;
-	TArray<FKeyHandle> OutKeyHandles;
-
-	GetKeys(Params.Range, &OutKeyTimes, &OutKeyHandles);
-
-	if (OutKeyHandles.Num() > 2)
-	{
-		int32 MostRecentKeepKeyIndex = 0;
-		TArray<FKeyHandle> KeysToRemove;
-
-		for (int32 TestIndex = 1; TestIndex < OutKeyHandles.Num() - 1; ++TestIndex)
-		{
-			int32 Index = ChannelData.GetIndex(OutKeyHandles[TestIndex]);
-			int32 NextIndex = ChannelData.GetIndex(OutKeyHandles[TestIndex+1]);
-
-			const float KeyValue = ChannelData.GetValues()[Index].Value;
-			const float ValueWithoutKey = EvalForTwoKeys(
-				ChannelData.GetValues()[MostRecentKeepKeyIndex], ChannelData.GetTimes()[MostRecentKeepKeyIndex].Value,
-				ChannelData.GetValues()[NextIndex], ChannelData.GetTimes()[NextIndex].Value,
-				ChannelData.GetTimes()[Index].Value,
-				Params.DisplayRate);
-				
-			if (FMath::Abs(ValueWithoutKey - KeyValue) > Params.Tolerance) // Is this key needed
-			{
-				MostRecentKeepKeyIndex = Index;
-			}
-			else
-			{
-				KeysToRemove.Add(OutKeyHandles[TestIndex]);
-			}
-		}
-
-		ChannelData.DeleteKeys(KeysToRemove);
-
-		if (Params.bAutoSetInterpolation)
-		{
-			AutoSetTangents();
-		}
-	}
+	FMovieSceneFloatChannelImpl::Optimize(this, Params);
 }
 
 void FMovieSceneFloatChannel::ClearDefault()
@@ -1040,34 +142,27 @@ void FMovieSceneFloatChannel::ClearDefault()
 
 FKeyHandle AddKeyToChannel(FMovieSceneFloatChannel* Channel, FFrameNumber InFrameNumber, float InValue, EMovieSceneKeyInterpolation Interpolation)
 {
-	TMovieSceneChannelData<FMovieSceneFloatValue> ChannelData = Channel->GetData();
-	int32 ExistingIndex = ChannelData.FindKey(InFrameNumber);
-	if (ExistingIndex != INDEX_NONE)
-	{
-		FMovieSceneFloatValue& Value = ChannelData.GetValues()[ExistingIndex]; //-V758
-		Value.Value = InValue;
-		Channel->AutoSetTangents();
-	}
-	else switch (Interpolation)
-	{
-		case EMovieSceneKeyInterpolation::Auto:     ExistingIndex = Channel->AddCubicKey(InFrameNumber, InValue, RCTM_Auto);  break;
-		case EMovieSceneKeyInterpolation::User:     ExistingIndex = Channel->AddCubicKey(InFrameNumber, InValue, RCTM_User);  break;
-		case EMovieSceneKeyInterpolation::Break:    ExistingIndex = Channel->AddCubicKey(InFrameNumber, InValue, RCTM_Break); break;
-		case EMovieSceneKeyInterpolation::Linear:   ExistingIndex = Channel->AddLinearKey(InFrameNumber, InValue);            break;
-		case EMovieSceneKeyInterpolation::Constant: ExistingIndex = Channel->AddConstantKey(InFrameNumber, InValue);          break;
-	}
-
-	return Channel->GetData().GetHandle(ExistingIndex);
+	return TMovieSceneCurveChannelImpl<FMovieSceneFloatChannel>::AddKeyToChannel(Channel, InFrameNumber, InValue, Interpolation);
 }
 
 void Dilate(FMovieSceneFloatChannel* InChannel, FFrameNumber Origin, float DilationFactor)
 {
-	TArrayView<FFrameNumber> Times = InChannel->GetData().GetTimes();
-	for (FFrameNumber& Time : Times)
-	{
-		Time = Origin + FFrameNumber(FMath::FloorToInt((Time - Origin).Value * DilationFactor));
-	}
-	InChannel->AutoSetTangents();
+	return TMovieSceneCurveChannelImpl<FMovieSceneFloatChannel>::Dilate(InChannel, Origin, DilationFactor);
+}
+
+bool ValueExistsAtTime(const FMovieSceneFloatChannel* InChannel, FFrameNumber InFrameNumber, float InValue)
+{
+	return TMovieSceneCurveChannelImpl<FMovieSceneFloatChannel>::ValueExistsAtTime(InChannel, InFrameNumber, InValue);
+}
+
+bool ValueExistsAtTime(const FMovieSceneFloatChannel* InChannel, FFrameNumber InFrameNumber, const FMovieSceneFloatValue& InValue)
+{
+	return TMovieSceneCurveChannelImpl<FMovieSceneFloatChannel>::ValueExistsAtTime(InChannel, InFrameNumber, InValue);
+}
+
+void AssignValue(FMovieSceneFloatChannel* InChannel, FKeyHandle InKeyHandle, float InValue)
+{
+	return TMovieSceneCurveChannelImpl<FMovieSceneFloatChannel>::AssignValue(InChannel, InKeyHandle, InValue);
 }
 
 void FMovieSceneFloatChannel::AddKeys(const TArray<FFrameNumber>& InTimes, const TArray<FMovieSceneFloatValue>& InValues)
@@ -1099,101 +194,7 @@ void FMovieSceneFloatChannel::SetShowCurve(bool bInShowCurve)
 
 bool FMovieSceneFloatChannel::Serialize(FArchive& Ar)
 {
-	Ar.UsingCustomVersion(FSequencerObjectVersion::GUID);
-	Ar.UsingCustomVersion(FFortniteMainBranchObjectVersion::GUID);
-	if (Ar.CustomVer(FSequencerObjectVersion::GUID) < FSequencerObjectVersion::SerializeFloatChannelCompletely &&
-		Ar.CustomVer(FFortniteMainBranchObjectVersion::GUID) < FFortniteMainBranchObjectVersion::SerializeFloatChannelShowCurve)
-	{
-		return false;
-	}
-
-	Ar << PreInfinityExtrap;
-	Ar << PostInfinityExtrap;
-
-	//Save FFrameNumber(int32) and  FMovieSceneFloatValue Arrays.
-	//We try to save and load the full array data, unless we are
-	//ByteSwapping or the Size has a mismatch on load, then we do normal save/load
-	if (Ar.IsLoading())
-	{
-		int32 CurrentSerializedElementSize = sizeof(FFrameNumber);
-		int32 SerializedElementSize = 0;
-		Ar << SerializedElementSize;
-		if (SerializedElementSize != CurrentSerializedElementSize || Ar.IsByteSwapping())
-		{
-			Ar << Times;
-		}
-		else
-		{
-			Times.CountBytes(Ar);
-			int32 NewArrayNum = 0;
-			Ar << NewArrayNum;
-			Times.Empty(NewArrayNum);
-			if (NewArrayNum > 0)
-			{
-				Times.AddUninitialized(NewArrayNum);
-				Ar.Serialize(Times.GetData(), NewArrayNum * SerializedElementSize);
-			}
-		}
-		CurrentSerializedElementSize = sizeof(FMovieSceneFloatValue);
-		Ar << SerializedElementSize;
-		if (SerializedElementSize != CurrentSerializedElementSize || Ar.IsByteSwapping())
-		{
-			Ar << Values;
-		}
-		else
-		{
-			Values.CountBytes(Ar);
-			int32 NewArrayNum = 0;
-			Ar << NewArrayNum;
-			Values.Empty(NewArrayNum);
-			if (NewArrayNum > 0)
-			{
-				Values.AddUninitialized(NewArrayNum);
-				Ar.Serialize(Values.GetData(), NewArrayNum * SerializedElementSize);
-			}
-		}
-	}
-	else if (Ar.IsSaving())
-	{
-		int32 SerializedElementSize = sizeof(FFrameNumber);
-		Ar << SerializedElementSize;
-		Times.CountBytes(Ar);
-		int32 ArrayCount = Times.Num();
-		Ar << ArrayCount;
-		if (ArrayCount > 0)
-		{
-			Ar.Serialize(Times.GetData(), ArrayCount * SerializedElementSize);
-		}
-		Values.CountBytes(Ar);
-		SerializedElementSize = sizeof(FMovieSceneFloatValue);
-		Ar << SerializedElementSize;
-		ArrayCount = Values.Num();
-		Ar << ArrayCount;
-		if (ArrayCount > 0)
-		{
-			Ar.Serialize(Values.GetData(), ArrayCount * SerializedElementSize);
-		}
-	}
-
-	Ar << DefaultValue;
-	Ar << bHasDefaultValue;
-	Ar << TickResolution.Numerator;
-	Ar << TickResolution.Denominator;
-	if (Ar.IsTransacting())
-	{
-		Ar << KeyHandles;
-	}
-
-	if (Ar.CustomVer(FFortniteMainBranchObjectVersion::GUID) >= FFortniteMainBranchObjectVersion::SerializeFloatChannelShowCurve)
-	{
-#if WITH_EDITOR
-		Ar << bShowCurve;
-#else
-		bool bUnused = false;
-		Ar << bUnused;
-#endif
-	}
-	return true;
+	return FMovieSceneFloatChannelImpl::Serialize(this, Ar);
 }
 
 #if WITH_EDITORONLY_DATA
@@ -1228,4 +229,10 @@ void FMovieSceneFloatChannel::PostSerialize(const FArchive& Ar)
 		}
 	}
 }
-#endif
\ No newline at end of file
+#endif
+
+bool FMovieSceneFloatChannel::SerializeFromMismatchedTag(const FPropertyTag& Tag, FStructuredArchive::FSlot Slot)
+{
+	return FMovieSceneFloatChannelImpl::SerializeFromRichCurve(this, Tag, Slot);
+}
+
