Sequencer: Remove muted spawnables on cook

Follow up to 17070294 which removes muted tracks on cook

diff --git a/Engine/Source/Runtime/MovieScene/Private/MovieScene.cpp b/Engine/Source/Runtime/MovieScene/Private/MovieScene.cpp
--- a/Engine/Source/Runtime/MovieScene/Private/MovieScene.cpp
+++ b/Engine/Source/Runtime/MovieScene/Private/MovieScene.cpp
@@ -108,19 +108,43 @@ void UMovieScene::Serialize( FArchive& Ar )
 			++MasterTrackIndex;
 		}
 
-		for (FMovieSceneBinding& ObjectBinding : ObjectBindings)
+		for (int32 ObjectBindingIndex = 0; ObjectBindingIndex < ObjectBindings.Num(); )
 		{
-			for (int32 TrackIndex = 0; TrackIndex < ObjectBinding.GetTracks().Num(); )
+			bool bRemoveObject = false;
+
+			// First, look to remove the object
+			for (int32 TrackIndex = 0; TrackIndex < ObjectBindings[ObjectBindingIndex].GetTracks().Num(); ++TrackIndex)
 			{
-				UMovieSceneTrack* Track = ObjectBinding.GetTracks()[TrackIndex];
-				if (Track && Track->GetCookOptimizationFlags() == ECookOptimizationFlags::RemoveTrack)
+				UMovieSceneTrack* Track = ObjectBindings[ObjectBindingIndex].GetTracks()[TrackIndex];
+				if (Track && Track->GetCookOptimizationFlags() == ECookOptimizationFlags::RemoveObject)
 				{
-					Track->RemoveForCook();
-					ObjectBinding.RemoveTrack(*Track);
-					UE_LOG(LogMovieScene, Display, TEXT("Removing muted track: %s from: %s"), *Track->GetDisplayName().ToString(), *GetPathName());
-					continue;
+					bRemoveObject = true;
+					break;
 				}
-				++TrackIndex;
+			}
+
+			if (bRemoveObject)
+			{
+				UE_LOG(LogMovieScene, Display, TEXT("Removing muted object and all of its tracks: %s from: %s"), *ObjectBindings[ObjectBindingIndex].GetName(), *GetPathName());
+				ObjectBindings.RemoveAt(ObjectBindingIndex);
+			}
+			else
+			{
+				// If the object wasn't removed, look to remove tracks
+				for (int32 TrackIndex = 0; TrackIndex < ObjectBindings[ObjectBindingIndex].GetTracks().Num(); )
+				{
+					UMovieSceneTrack* Track = ObjectBindings[ObjectBindingIndex].GetTracks()[TrackIndex];
+					if (Track && Track->GetCookOptimizationFlags() == ECookOptimizationFlags::RemoveTrack)
+					{
+						Track->RemoveForCook();
+						ObjectBindings[ObjectBindingIndex].RemoveTrack(*Track);
+						UE_LOG(LogMovieScene, Display, TEXT("Removing muted track: %s from: %s"), *Track->GetDisplayName().ToString(), *GetPathName());
+						continue;
+					}
+					++TrackIndex;
+				}
+				
+				++ObjectBindingIndex;
 			}
 		}
 	}
diff --git a/Engine/Source/Runtime/MovieScene/Private/Tracks/MovieSceneSpawnTrack.cpp b/Engine/Source/Runtime/MovieScene/Private/Tracks/MovieSceneSpawnTrack.cpp
--- a/Engine/Source/Runtime/MovieScene/Private/Tracks/MovieSceneSpawnTrack.cpp
+++ b/Engine/Source/Runtime/MovieScene/Private/Tracks/MovieSceneSpawnTrack.cpp
@@ -176,6 +176,23 @@ const TArray<UMovieSceneSection*>& UMovieSceneSpawnTrack::GetAllSections() const
 	return Sections;
 }
 
+#if WITH_EDITOR
+
+ECookOptimizationFlags UMovieSceneSpawnTrack::GetCookOptimizationFlags() const
+{
+	ECookOptimizationFlags CookOptimizationFlags = Super::GetCookOptimizationFlags();
+
+	// Remove the object if the track is muted
+	if (CookOptimizationFlags == ECookOptimizationFlags::RemoveTrack)
+	{
+		return ECookOptimizationFlags::RemoveObject;
+	}
+
+	return CookOptimizationFlags;
+}
+
+#endif
+
 #if WITH_EDITORONLY_DATA
 
 FText UMovieSceneSpawnTrack::GetDisplayName() const
diff --git a/Engine/Source/Runtime/MovieScene/Public/Tracks/MovieSceneSpawnTrack.h b/Engine/Source/Runtime/MovieScene/Public/Tracks/MovieSceneSpawnTrack.h
--- a/Engine/Source/Runtime/MovieScene/Public/Tracks/MovieSceneSpawnTrack.h
+++ b/Engine/Source/Runtime/MovieScene/Public/Tracks/MovieSceneSpawnTrack.h
@@ -57,6 +57,10 @@ public:
 	//~ UObject interface
 	virtual void PostLoad() override;
 
+#if WITH_EDITOR
+	virtual ECookOptimizationFlags GetCookOptimizationFlags() const override;
+#endif
+
 #if WITH_EDITORONLY_DATA
 	virtual FText GetDisplayName() const override;
 #endif
