Sequencer: Improve how client/replicated sequences handle end of playback, and triggering events.

The first problem was with client sequences that are re-synced to the server time.

If a re-sync sends the client sequence a few frames back, it will re-evaluate these frames as it plays over them a second time. If this happens around an event, this event will fire again. So we set a timestamp to where we were at before the re-sync, which tells the event system to ignore any events before that timestamp.

The second problem was with client sequences who are *late* compared to the server.

When the server sends a Stop RPC message, if the sequence is a bit behind it would only advance to the stop time if the difference was above a threshold. But we could miss firing an event that's on the very last bit of the last frame. So now we advance to the stop time exactly.

The third problem was with client sequences who are *ahead* of the server.

When the client sequence would reach the end of playback, it would stop and reset time to the beginning all on it own. But the server would still be running the last couple frames, so it would send some network message telling the client sequence to do that. The client sequence would go back into play mode, and play all the way from the start to one of these last frames. This means it would re-evaluate the entire sequence again, including re-triggering all the events! So now, non-authoritative sequences don't stop-and-reset-time when they reach the end of playback. Instead, they stay on the last frame and wait for the server to send them the explicit "end of playback" message.

diff --git a/Engine/Source/Runtime/MovieScene/Private/MovieSceneSequencePlayer.cpp b/Engine/Source/Runtime/MovieScene/Private/MovieSceneSequencePlayer.cpp
--- a/Engine/Source/Runtime/MovieScene/Private/MovieSceneSequencePlayer.cpp
+++ b/Engine/Source/Runtime/MovieScene/Private/MovieSceneSequencePlayer.cpp
@@ -2,6 +2,7 @@
 
 #include "MovieSceneSequencePlayer.h"
 #include "MovieScene.h"
+#include "MovieSceneFwd.h"
 #include "MovieSceneTimeHelpers.h"
 #include "MovieSceneSequence.h"
 #include "MovieSceneSequenceTickManager.h"
@@ -289,6 +290,8 @@ void UMovieSceneSequencePlayer::PlayInternal()
 			UpdateMovieSceneInstance(PlayPosition.PlayTo(PlayPosition.GetCurrentPosition()), EMovieScenePlayerStatus::Playing);
 		}
 
+		RunLatentActions();
+
 		UpdateNetworkSyncProperties();
 
 		if (MovieSceneSequence)
@@ -459,6 +462,25 @@ void UMovieSceneSequencePlayer::StopInternal(FFrameTime TimeToResetTo)
 	}
 }
 
+void UMovieSceneSequencePlayer::FinishPlaybackInternal(FFrameTime TimeToFinishAt)
+{
+	if (PlaybackSettings.bPauseAtEnd)
+	{
+		Pause();
+	}
+	else
+	{
+		StopInternal(TimeToFinishAt);
+	}
+
+	TimeController->StopPlaying(GetCurrentTime());
+
+	if (OnFinished.IsBound())
+	{
+		OnFinished.Broadcast();
+	}
+}
+
 void UMovieSceneSequencePlayer::GoToEndAndStop()
 {
 	FFrameTime LastValidTime = GetLastValidTime();
@@ -978,7 +1000,7 @@ void UMovieSceneSequencePlayer::UpdateTimeCursorPosition_Internal(FFrameTime New
 			OnLooped();
 		}
 
-		// stop playback
+		// We reached the end of playback
 		else
 		{
 			// Clamp the position to the duration
@@ -987,20 +1009,21 @@ void UMovieSceneSequencePlayer::UpdateTimeCursorPosition_Internal(FFrameTime New
 			FMovieSceneEvaluationRange Range = UpdatePlayPosition(PlayPosition, NewPosition, Method);
 			UpdateMovieSceneInstance(Range, StatusOverride);
 
-			if (PlaybackSettings.bPauseAtEnd)
+			if (HasAuthority())
 			{
-				Pause();
+				FinishPlaybackInternal(NewPosition);
+
+				// Explicitly tell the clients to finish their playback. They won't have called FinishPlaybackInternal
+				// because it's in the line right above, only for sequence players with some authority
+				// (client only or server).
+				RPC_OnFinishPlaybackEvent(NewPosition);
 			}
 			else
 			{
-				StopInternal(NewPosition);
-			}
-
-			TimeController->StopPlaying(GetCurrentTime());
-
-			if (OnFinished.IsBound())
-			{
-				OnFinished.Broadcast();
+				const FString SequenceName = GetSequenceName(true);
+				UE_LOG(LogMovieScene, Verbose, 
+						TEXT("Keeping sequence %s at frame %d, subframe %f while waiting for playback finish from server."),
+						*SequenceName, NewPosition.FrameNumber.Value, NewPosition.GetSubFrame());
 			}
 		}
 
@@ -1065,14 +1088,8 @@ void UMovieSceneSequencePlayer::UpdateMovieSceneInstance(FMovieSceneEvaluationRa
 	}
 
 #if !NO_LOGGING
-	FQualifiedFrameTime CurrentTime = GetCurrentTime();
-	FString    SequenceName = MovieSceneSequence->GetName();
-
-	AActor* Actor = GetTypedOuter<AActor>();
-	if (Actor && Actor->GetWorld()->GetNetMode() == NM_Client)
-	{
-		SequenceName += FString::Printf(TEXT(" (client %d)"), GPlayInEditorID - 1);
-	}
+	const FQualifiedFrameTime CurrentTime = GetCurrentTime();
+	const FString SequenceName = GetSequenceName(true);
 	UE_LOG(LogMovieScene, VeryVerbose, TEXT("Evaluating sequence %s at frame %d, subframe %f (%f fps)."), *SequenceName, CurrentTime.Time.FrameNumber.Value, CurrentTime.Time.GetSubFrame(), CurrentTime.Rate.AsDecimal());
 #endif
 
@@ -1100,6 +1117,16 @@ void UMovieSceneSequencePlayer::UpdateMovieSceneInstance(FMovieSceneEvaluationRa
 	}
 }
 
+bool UMovieSceneSequencePlayer::IsDisablingEventTriggers(FFrameTime& DisabledUntilTime) const
+{
+	if (DisableEventTriggersUntilTime.IsSet())
+	{
+		DisabledUntilTime = DisableEventTriggersUntilTime.GetValue();
+		return true;
+	}
+	return false;
+}
+
 void UMovieSceneSequencePlayer::PreEvaluation(const FMovieSceneContext& Context)
 {
 	RunPreEvaluationCallbacks();
@@ -1138,6 +1165,28 @@ void UMovieSceneSequencePlayer::RunPostEvaluationCallbacks()
 	PostEvaluationCallbacks.Reset();
 }
 
+FString UMovieSceneSequencePlayer::GetSequenceName(bool bAddClientInfo) const
+{
+	if (Sequence)
+	{
+		FString SequenceName = Sequence->GetName();
+		if (bAddClientInfo)
+		{
+			AActor* Actor = GetTypedOuter<AActor>();
+			if (Actor && Actor->GetWorld()->GetNetMode() == NM_Client)
+			{
+				SequenceName += FString::Printf(TEXT(" (client %d)"), GPlayInEditorID - 1);
+			}
+		}
+		return SequenceName;
+	}
+	else
+	{
+		static const FName NoneName(NAME_None);
+		return NoneName.ToString();
+	}
+}
+
 void UMovieSceneSequencePlayer::SetPlaybackClient(TScriptInterface<IMovieScenePlaybackClient> InPlaybackClient)
 {
 	PlaybackClient = InPlaybackClient;
@@ -1303,15 +1352,8 @@ void UMovieSceneSequencePlayer::RPC_ExplicitServerUpdateEvent_Implementation(EUp
 	// Log the sync event if necessary
 	if (UE_LOG_ACTIVE(LogMovieScene, Verbose))
 	{
-		FFrameTime   CurrentTime     = PlayPosition.GetCurrentPosition();
-		FString      SequenceName    = RootTemplateInstance.GetSequence(MovieSceneSequenceID::Root)->GetName();
-
-		AActor* Actor = GetTypedOuter<AActor>();
-		if (Actor && Actor->GetWorld()->GetNetMode() == NM_Client)
-		{
-			SequenceName += FString::Printf(TEXT(" (client %d)"), GPlayInEditorID - 1);
-		}
-
+		const FFrameTime CurrentTime = PlayPosition.GetCurrentPosition();
+		const FString SequenceName = GetSequenceName(true);
 		UE_LOG(LogMovieScene, Verbose, TEXT("Explicit update event for sequence %s %s @ frame %d, subframe %f. Server has moved to frame %d, subframe %f with EUpdatePositionMethod::%s."),
 			*SequenceName, *UEnum::GetValueAsString(TEXT("MovieScene.EMovieScenePlayerStatus"), Status.GetValue()), CurrentTime.FrameNumber.Value, CurrentTime.GetSubFrame(),
 			NetSyncProps.LastKnownPosition.FrameNumber.Value, NetSyncProps.LastKnownPosition.GetSubFrame(), *UEnum::GetValueAsString(TEXT("MovieScene.EUpdatePositionMethod"), NetSyncProps.LastKnownStatus.GetValue()));
@@ -1343,54 +1385,56 @@ void UMovieSceneSequencePlayer::RPC_OnStopEvent_Implementation(FFrameTime Stoppe
 #if !NO_LOGGING
 	if (UE_LOG_ACTIVE(LogMovieSceneRepl, Verbose))
 	{
-		FFrameTime CurrentTime  = PlayPosition.GetCurrentPosition();
-		FString    SequenceName = RootTemplateInstance.GetSequence(MovieSceneSequenceID::Root)->GetName();
-
-		AActor* Actor = GetTypedOuter<AActor>();
-		if (Actor && Actor->GetWorld()->GetNetMode() == NM_Client)
-		{
-			SequenceName += FString::Printf(TEXT(" (client %d)"), GPlayInEditorID - 1);
-		}
-
+		const FFrameTime CurrentTime = PlayPosition.GetCurrentPosition();
+		const FString SequenceName = GetSequenceName(true);
 		UE_LOG(LogMovieSceneRepl, Verbose, TEXT("Explicit Stop() event for sequence %s %s @ frame %d, subframe %f. Server has stopped at frame %d, subframe %f."),
 			*SequenceName, *UEnum::GetValueAsString(TEXT("MovieScene.EMovieScenePlayerStatus"), Status.GetValue()), CurrentTime.FrameNumber.Value, CurrentTime.GetSubFrame(),
 			NetSyncProps.LastKnownPosition.FrameNumber.Value, NetSyncProps.LastKnownPosition.GetSubFrame());
 	}
 #endif
 
-	float PingMs = 0.f;
+	EUpdatePositionMethod UpdatePositionMethod;
+	switch (Status.GetValue())
+	{
+		case EMovieScenePlayerStatus::Playing:
+			UpdatePositionMethod = EUpdatePositionMethod::Play;
+			break;
+		case EMovieScenePlayerStatus::Scrubbing:
+			UpdatePositionMethod = EUpdatePositionMethod::Scrub;
+			break;
+		default:
+			UpdatePositionMethod = EUpdatePositionMethod::Jump;
+			break;
+	}
 
-	UWorld* PlayWorld = GetPlaybackWorld();
-	if (PlayWorld)
+	// Catch up with any loops we are missing compared to the server. This is generally just 0 or 1 loops.
+	// When it's 1 loop, it's generally because we are very close to the end, and the server somehow stopped
+	// near the beginning of the next loop, so we have just a little bit of catching up to do.
+	const int32 LoopOffset = (NetSyncProps.LastKnownNumLoops - CurrentNumLoops);
+	const FFrameTime LoopEndTime = (bReversePlayback ? StartTime : GetLastValidTime());
+	for (int32 LoopIndex = 0; LoopIndex < LoopOffset; ++LoopIndex)
 	{
-		UNetDriver* NetDriver = PlayWorld->GetNetDriver();
-		if (NetDriver && NetDriver->ServerConnection && NetDriver->ServerConnection->PlayerController && NetDriver->ServerConnection->PlayerController->PlayerState)
-		{
-			PingMs = NetDriver->ServerConnection->PlayerController->PlayerState->ExactPing * (bReversePlayback ? -1.f : 1.f);
-		}
+		SetPlaybackPosition(FMovieSceneSequencePlaybackParams(LoopEndTime, UpdatePositionMethod));
 	}
 
-	const FFrameTime PingLag = (PingMs / 1000.f) * PlayPosition.GetInputRate();
-	const FFrameTime LagThreshold = (GSequencerNetSyncThresholdMS * 0.001f) * PlayPosition.GetInputRate();
+	// Now do the last bit of catch-up for the current loop.
+	if (PlayPosition.GetCurrentPosition() < StoppedTime)
+	{
+		SetPlaybackPosition(FMovieSceneSequencePlaybackParams(StoppedTime, UpdatePositionMethod));
+	}
 
-	// When the server has stopped and a client is near the end (and is thus about to loop), we don't want to forcibly synchronize the time unless
-	// the *real* difference in time is above the threshold. We compute the real-time difference by adding SequenceDuration*LoopCountDifference to the server position:
-	const int32        LoopOffset = (NetSyncProps.LastKnownNumLoops - CurrentNumLoops) * (bReversePlayback ? -1 : 1);
-	const FFrameTime   OffsetServerTime = (NetSyncProps.LastKnownPosition + PingLag) + GetFrameDuration() * LoopOffset;
-	const FFrameTime   Difference = FMath::Abs(PlayPosition.GetCurrentPosition() - OffsetServerTime);
+	StopInternal(StoppedTime);
+}
 
-	// If the difference is large enough and the client is behind the target time to stop at, advance to the target time.
-	if (Difference > LagThreshold + PingLag)
+void UMovieSceneSequencePlayer::RPC_OnFinishPlaybackEvent_Implementation(FFrameTime StoppedTime)
+{
+	if (HasAuthority() || !Sequence)
 	{
-		const bool bBehindTime = PlayPosition.GetCurrentPosition() < StoppedTime;
-
-		if (bBehindTime)
-		{
-			SetPlaybackPosition(FMovieSceneSequencePlaybackParams(StoppedTime, (EUpdatePositionMethod)Status.GetValue()));
-		}
+		// Never run network sync operations on authoritative players or players that have not been initialized yet
+		return;
 	}
 
-	StopInternal(StoppedTime);
+	FinishPlaybackInternal(StoppedTime);
 }
 
 void UMovieSceneSequencePlayer::PostNetReceive()
@@ -1436,15 +1480,8 @@ void UMovieSceneSequencePlayer::PostNetReceive()
 #if !NO_LOGGING
 	if (UE_LOG_ACTIVE(LogMovieSceneRepl, VeryVerbose))
 	{
-		FFrameTime CurrentTime  = PlayPosition.GetCurrentPosition();
-		FString    SequenceName = RootTemplateInstance.GetSequence(MovieSceneSequenceID::Root)->GetName();
-
-		AActor* Actor = GetTypedOuter<AActor>();
-		if (Actor->GetWorld()->GetNetMode() == NM_Client)
-		{
-			SequenceName += FString::Printf(TEXT(" (client %d)"), GPlayInEditorID - 1);
-		}
-
+		const FFrameTime CurrentTime = PlayPosition.GetCurrentPosition();
+		const FString SequenceName = GetSequenceName(true);
 		UE_LOG(LogMovieSceneRepl, VeryVerbose, TEXT("Network sync for sequence %s %s @ frame %d, subframe %f. Server is %s @ frame %d, subframe %f."),
 			*SequenceName, *UEnum::GetValueAsString(TEXT("MovieScene.EMovieScenePlayerStatus"), Status.GetValue()), CurrentTime.FrameNumber.Value, CurrentTime.GetSubFrame(),
 			*UEnum::GetValueAsString(TEXT("MovieScene.EMovieScenePlayerStatus"), NetSyncProps.LastKnownStatus.GetValue()), NetSyncProps.LastKnownPosition.FrameNumber.Value, NetSyncProps.LastKnownPosition.GetSubFrame());
@@ -1519,15 +1556,8 @@ void UMovieSceneSequencePlayer::PostNetReceive()
 #if !NO_LOGGING
 					if (UE_LOG_ACTIVE(LogMovieSceneRepl, Log))
 					{
-						FFrameTime CurrentTime  = PlayPosition.GetCurrentPosition();
-						FString    SequenceName = RootTemplateInstance.GetSequence(MovieSceneSequenceID::Root)->GetName();
-
-						AActor* Actor = GetTypedOuter<AActor>();
-						if (Actor->GetWorld()->GetNetMode() == NM_Client)
-						{
-							SequenceName += FString::Printf(TEXT(" (client %d)"), GPlayInEditorID - 1);
-						}
-
+						const FFrameTime CurrentTime = PlayPosition.GetCurrentPosition();
+						const FString SequenceName = GetSequenceName(true);
 						UE_LOG(LogMovieSceneRepl, Log, TEXT("Correcting de-synced play position for sequence %s %s @ frame %d, subframe %f. Server is %s @ frame %d, subframe %f, (smoothed: frame %d, sf %f). Client ping is %.2fms."),
 							*SequenceName, *UEnum::GetValueAsString(TEXT("MovieScene.EMovieScenePlayerStatus"), Status.GetValue()), CurrentTime.FrameNumber.Value, CurrentTime.GetSubFrame(),
 							*UEnum::GetValueAsString(TEXT("MovieScene.EMovieScenePlayerStatus"), NetSyncProps.LastKnownStatus.GetValue()),
@@ -1537,7 +1567,11 @@ void UMovieSceneSequencePlayer::PostNetReceive()
 					}
 #endif
 					// We're drastically out of sync with the server so we need to forcibly set the time.
-					// Play to the time only if it is further on in the sequence (in our play direction)
+					const FFrameTime LastPosition = FFrameRate::TransformTime(
+							PlayPosition.GetCurrentPosition(), PlayPosition.GetInputRate(), PlayPosition.GetOutputRate());
+
+					// Play to the time only if it is further on in the sequence (in our play direction).
+					// Otherwise, jump backwards in time (in our play direction).
 					const bool bPlayToFrame = bReversePlayback ? SmoothedServerTime < PlayPosition.GetCurrentPosition() : SmoothedServerTime > PlayPosition.GetCurrentPosition();
 					if (bPlayToFrame)
 					{
@@ -1554,6 +1588,10 @@ void UMovieSceneSequencePlayer::PostNetReceive()
 
 					// When playing back we skip this sequence's ticked update to avoid queuing 2 updates this frame
 					bSkipNextUpdate = true;
+
+					// Also skip all events up to the last known position, otherwise if we skipped back in time we
+					// will re-trigger events again.
+					DisableEventTriggersUntilTime = LastPosition;
 				}
 			}
 			else if (Status == EMovieScenePlayerStatus::Stopped)
diff --git a/Engine/Source/Runtime/MovieScene/Public/IMovieScenePlayer.h b/Engine/Source/Runtime/MovieScene/Public/IMovieScenePlayer.h
--- a/Engine/Source/Runtime/MovieScene/Public/IMovieScenePlayer.h
+++ b/Engine/Source/Runtime/MovieScene/Public/IMovieScenePlayer.h
@@ -200,6 +200,11 @@ public:
 	 */
 	virtual TArray<UObject*> GetEventContexts() const { return TArray<UObject*>(); }
 
+	/**
+	 * Returns whether event triggers are disabled and if so, until what time.
+	 */
+	virtual bool IsDisablingEventTriggers(FFrameTime& DisabledUntilTime) const { return false; }
+
 	/**
 	 * Test whether this is a preview player or not. As such, playback range becomes insignificant for things like spawnables
 	 */
diff --git a/Engine/Source/Runtime/MovieScene/Public/MovieSceneSequencePlayer.h b/Engine/Source/Runtime/MovieScene/Public/MovieSceneSequencePlayer.h
--- a/Engine/Source/Runtime/MovieScene/Public/MovieSceneSequencePlayer.h
+++ b/Engine/Source/Runtime/MovieScene/Public/MovieSceneSequencePlayer.h
@@ -514,6 +514,14 @@ public:
 	UFUNCTION(BlueprintCallable, Category = "Sequencer|Player")
 	UMovieSceneSequence* GetSequence() const { return Sequence; }
 
+	/**
+	 * Get the name of the sequence this player is playing
+	 * @param bAddClientInfo  If true, add client index if running as a client
+	 * @return the name of the sequence, or None if no sequence is set
+	 */
+	UFUNCTION(BlueprintCallable, Category = "Sequencer|Player")
+	FString GetSequenceName(bool bAddClientInfo = false) const;
+
 	/**
 	 * Access this player's tick manager
 	 */
@@ -533,6 +541,7 @@ protected:
 
 	void PlayInternal();
 	void StopInternal(FFrameTime TimeToResetTo);
+	void FinishPlaybackInternal(FFrameTime TimeToFinishAt);
 
 	struct FMovieSceneUpdateArgs
 	{
@@ -578,6 +587,7 @@ protected:
 	virtual void UpdateCameraCut(UObject* CameraObject, const EMovieSceneCameraCutParams& CameraCutParams) override {}
 	virtual void ResolveBoundObjects(const FGuid& InBindingId, FMovieSceneSequenceID SequenceID, UMovieSceneSequence& Sequence, UObject* ResolutionContext, TArray<UObject*, TInlineAllocator<1>>& OutObjects) const override;
 	virtual IMovieScenePlaybackClient* GetPlaybackClient() override { return PlaybackClient ? &*PlaybackClient : nullptr; }
+	virtual bool IsDisablingEventTriggers(FFrameTime& DisabledUntilTime) const override;
 	virtual void PreEvaluation(const FMovieSceneContext& Context) override;
 	virtual void PostEvaluation(const FMovieSceneContext& Context) override;
 
@@ -618,6 +628,12 @@ private:
 	UFUNCTION(netmulticast, reliable)
 	void RPC_OnStopEvent(FFrameTime StoppedTime);
 
+	/**
+	 * Called on the server when playback has reached the end. Could lead to stopping or pausing.
+	 */
+	UFUNCTION(netmulticast, reliable)
+	void RPC_OnFinishPlaybackEvent(FFrameTime StoppedTime);
+
 	/**
 	 * Check whether this sequence player is an authority, as determined by its outer Actor
 	 */
@@ -687,6 +703,10 @@ protected:
 	/** Play position helper */
 	FMovieScenePlaybackPosition PlayPosition;
 
+	/** Disable event triggers until given time */
+	TOptional<FFrameTime> DisableEventTriggersUntilTime;
+
+	/** Spawn register */
 	TSharedPtr<FMovieSceneSpawnRegister> SpawnRegister;
 
 	struct FServerTimeSample
diff --git a/Engine/Source/Runtime/MovieSceneTracks/Private/Systems/MovieSceneEventSystems.cpp b/Engine/Source/Runtime/MovieSceneTracks/Private/Systems/MovieSceneEventSystems.cpp
--- a/Engine/Source/Runtime/MovieSceneTracks/Private/Systems/MovieSceneEventSystems.cpp
+++ b/Engine/Source/Runtime/MovieSceneTracks/Private/Systems/MovieSceneEventSystems.cpp
@@ -1,6 +1,7 @@
 // Copyright Epic Games, Inc. All Rights Reserved.
 
 #include "Systems/MovieSceneEventSystems.h"
+#include "Algo/RemoveIf.h"
 #include "EntitySystem/MovieSceneEntitySystemLinker.h"
 #include "EntitySystem/MovieSceneEntitySystemRunner.h"
 #include "EntitySystem/MovieSceneSpawnablesSystem.h"
@@ -72,19 +73,31 @@ void UMovieSceneEventSystem::TriggerAllEvents()
 	for (TPair<FInstanceHandle, TArray<FMovieSceneEventTriggerData>>& Pair : EventsByRoot)
 	{
 		const FSequenceInstance& RootInstance = InstanceRegistry->GetInstance(Pair.Key);
+		IMovieScenePlayer* Player = RootInstance.GetPlayer();
+
+		FFrameTime SkipUntil;
+		const bool bSkipTrigger = Player->IsDisablingEventTriggers(SkipUntil);
 
 		if (RootInstance.GetContext().GetDirection() == EPlayDirection::Forwards)
 		{
+			if (bSkipTrigger)
+			{
+				Algo::RemoveIf(Pair.Value, [SkipUntil](FMovieSceneEventTriggerData& Data) { return Data.RootTime <= SkipUntil; });
+			}
 			Algo::SortBy(Pair.Value, &FMovieSceneEventTriggerData::RootTime);
 		}
 		else
 		{
+			if (bSkipTrigger)
+			{
+				Algo::RemoveIf(Pair.Value, [SkipUntil](FMovieSceneEventTriggerData& Data) { return Data.RootTime >= SkipUntil; });
+			}
 			Algo::SortBy(Pair.Value, &FMovieSceneEventTriggerData::RootTime, TGreater<>());
 		}
 
 		FTriggerBatch& TriggerBatch = TriggerBatches.Emplace_GetRef();
 		TriggerBatch.Triggers = Pair.Value;
-		TriggerBatch.Player = RootInstance.GetPlayer();
+		TriggerBatch.Player = Player;
 	}
 
 	// We need to clean our state before actually triggering the events because one of those events could
@@ -100,6 +113,12 @@ void UMovieSceneEventSystem::TriggerAllEvents()
 
 void UMovieSceneEventSystem::TriggerEvents(TArrayView<const FMovieSceneEventTriggerData> Events, IMovieScenePlayer* Player)
 {
+#if !NO_LOGGING
+	UMovieSceneSequence* PlayerSequence = Player->GetEvaluationTemplate().GetRootSequence();
+	FString SequenceName = PlayerSequence->GetName();
+	UE_LOG(LogMovieScene, VeryVerbose, TEXT("%s: Triggering %d events"), *SequenceName, Events.Num());
+#endif
+
 	TArray<UObject*> GlobalContexts = Player->GetEventContexts();
 
 	for (const FMovieSceneEventTriggerData& Event : Events)
@@ -115,6 +134,10 @@ void UMovieSceneEventSystem::TriggerEvents(TArrayView<const FMovieSceneEventTrig
 			continue;
 		}
 
+#if !NO_LOGGING
+		UE_LOG(LogMovieScene, VeryVerbose, TEXT("%s: - Triggering event at frame %d, subframe %f"), *SequenceName, Event.RootTime.FrameNumber.Value, Event.RootTime.GetSubFrame());
+#endif
+
 #if WITH_EDITOR
 		const static FName NAME_CallInEditor(TEXT("CallInEditor"));
 
