Sequencer: Remove error from situation that is already generating a warning message.

diff --git a/Engine/Source/Runtime/MovieScene/Private/EntitySystem/MovieSceneEntitySystemRunner.cpp b/Engine/Source/Runtime/MovieScene/Private/EntitySystem/MovieSceneEntitySystemRunner.cpp
--- a/Engine/Source/Runtime/MovieScene/Private/EntitySystem/MovieSceneEntitySystemRunner.cpp
+++ b/Engine/Source/Runtime/MovieScene/Private/EntitySystem/MovieSceneEntitySystemRunner.cpp
@@ -144,7 +144,6 @@ void FMovieSceneEntitySystemRunner::Flush()
 
 	if (!Linker->StartEvaluation(*this))
 	{
-		UE_LOG(LogMovieSceneECS, Error, TEXT("An evaluation is already running on this linker, and re-entrancy isn't allow at this point"));
 		return;
 	}
 
