Fix CIS issue with non-clang compiler.

diff --git a/Engine/Source/Runtime/MovieScene/Private/MovieSceneSequencePlayer.cpp b/Engine/Source/Runtime/MovieScene/Private/MovieSceneSequencePlayer.cpp
--- a/Engine/Source/Runtime/MovieScene/Private/MovieSceneSequencePlayer.cpp
+++ b/Engine/Source/Runtime/MovieScene/Private/MovieSceneSequencePlayer.cpp
@@ -1182,8 +1182,7 @@ FString UMovieSceneSequencePlayer::GetSequenceName(bool bAddClientInfo) const
 	}
 	else
 	{
-		static const FName NoneName(NAME_None);
-		return NoneName.ToString();
+		return LexToString(NAME_None);
 	}
 }
 
