Sequencer: Enable debug visualizers for double components.

diff --git a/Engine/Source/Runtime/MovieScene/Public/EntitySystem/MovieSceneComponentDebug.h b/Engine/Source/Runtime/MovieScene/Public/EntitySystem/MovieSceneComponentDebug.h
--- a/Engine/Source/Runtime/MovieScene/Public/EntitySystem/MovieSceneComponentDebug.h
+++ b/Engine/Source/Runtime/MovieScene/Public/EntitySystem/MovieSceneComponentDebug.h
@@ -21,6 +21,7 @@ enum class EComponentDebugType
 	Uint16,
 	Int32,
 	Float,
+	Double,
 	Vector2,
 	Vector3,
 	Vector4,
@@ -46,6 +47,7 @@ template<>           struct TComponentDebugType<uint8>               { static co
 template<>           struct TComponentDebugType<uint16>              { static const EComponentDebugType Type = EComponentDebugType::Uint16;   };
 template<>           struct TComponentDebugType<int32>               { static const EComponentDebugType Type = EComponentDebugType::Int32;    };
 template<>           struct TComponentDebugType<float>               { static const EComponentDebugType Type = EComponentDebugType::Float;    };
+template<>           struct TComponentDebugType<double>              { static const EComponentDebugType Type = EComponentDebugType::Double;    };
 template<>           struct TComponentDebugType<FVector2D>           { static const EComponentDebugType Type = EComponentDebugType::Vector2;  };
 template<>           struct TComponentDebugType<FVector>             { static const EComponentDebugType Type = EComponentDebugType::Vector3;  };
 template<>           struct TComponentDebugType<FVector4>            { static const EComponentDebugType Type = EComponentDebugType::Vector4;  };
