Cherrypicked 17425747 from UE5Main to fix a sequencer issue where keep state tracks weren't guaranteed to apply state changes

diff --git a/Engine/Source/Runtime/MovieScene/Private/Compilation/MovieSceneCompiledDataManager.cpp b/Engine/Source/Runtime/MovieScene/Private/Compilation/MovieSceneCompiledDataManager.cpp
--- a/Engine/Source/Runtime/MovieScene/Private/Compilation/MovieSceneCompiledDataManager.cpp
+++ b/Engine/Source/Runtime/MovieScene/Private/Compilation/MovieSceneCompiledDataManager.cpp
@@ -32,6 +32,15 @@ FAutoConsoleVariableRef CVarMovieSceneCompilerVersion(
 );
 
 
+TAutoConsoleVariable<bool> CVarAddKeepStateDeterminismFences(
+	TEXT("Sequencer.AddKeepStateDeterminismFences"),
+	true,
+	TEXT("Whether the Sequencer compiler should auto-add determinism fences for the last frame of KeepState sections. "
+		 "This ensures that the last possible value of the section is consistently evaluated regardless of framerate, "
+		 "at the cost of an extra evaluation on frames that cross over KeepState sections' end time.\n"),
+	ECVF_Default);
+
+
 IMovieSceneModule& GetMovieSceneModule()
 {
 	static TWeakPtr<IMovieSceneModule> WeakMovieSceneModule;
@@ -1158,14 +1167,35 @@ void UMovieSceneCompiledDataManager::CompileTrack(FMovieSceneCompiledDataEntry*
 		OutCompilerData->AccumulatedMask |= EMovieSceneSequenceCompilerMask::EvaluationTemplate;
 	}
 
+	// -------------------------------------------------------------------------------------------------------------------------------------
+	// Step 2 - let the track or its sections add determinism fences
 	if (IMovieSceneDeterminismSource* DeterminismSource = Cast<IMovieSceneDeterminismSource>(Track))
 	{
 		DeterminismSource->PopulateDeterminismData(OutCompilerData->DeterminismData, TRange<FFrameNumber>::All());
 	}
 
 	const FMovieSceneTrackEvaluationField& EvaluationField = Track->GetEvaluationField();
+	const EMovieSceneCompletionMode DefaultCompletionMode = Sequence->DefaultCompletionMode;
+	const bool bAddKeepStateDeterminismFences = CVarAddKeepStateDeterminismFences.GetValueOnGameThread();
 	for (const FMovieSceneTrackEvaluationFieldEntry& Entry : EvaluationField.Entries)
 	{
+		if (bAddKeepStateDeterminismFences)
+		{
+			// If a section is KeepState, we need to make sure to evaluate it on its last frame so that the value that "sticks" is correct.
+			const TRange<FFrameNumber> SectionRange = Entry.Section->GetRange();
+			const EMovieSceneCompletionMode SectionCompletionMode = Entry.Section->GetCompletionMode();
+			if (SectionRange.HasUpperBound() &&
+					(SectionCompletionMode == EMovieSceneCompletionMode::KeepState ||
+					 (SectionCompletionMode == EMovieSceneCompletionMode::ProjectDefault && DefaultCompletionMode == EMovieSceneCompletionMode::KeepState)))
+			{
+				// We simply use the end time of the section for the fence, regardless of whether it's inclusive or exclusive.
+				// When exclusive, the ECS system will query entities just before that time, but still pass that time for
+				// evaluation purposes, so we will get the correct evaluated values.
+				const FFrameNumber FenceTime(SectionRange.GetUpperBoundValue());
+				OutCompilerData->DeterminismData.Fences.Add(FenceTime);
+			}
+		}
+
 		IMovieSceneDeterminismSource* DeterminismSource = Cast<IMovieSceneDeterminismSource>(Entry.Section);
 		if (DeterminismSource)
 		{
