Sequencer: Null section fix for determinism fences

diff --git a/Engine/Source/Runtime/MovieScene/Private/Compilation/MovieSceneCompiledDataManager.cpp b/Engine/Source/Runtime/MovieScene/Private/Compilation/MovieSceneCompiledDataManager.cpp
--- a/Engine/Source/Runtime/MovieScene/Private/Compilation/MovieSceneCompiledDataManager.cpp
+++ b/Engine/Source/Runtime/MovieScene/Private/Compilation/MovieSceneCompiledDataManager.cpp
@@ -1179,7 +1179,7 @@ void UMovieSceneCompiledDataManager::CompileTrack(FMovieSceneCompiledDataEntry*
 	const bool bAddKeepStateDeterminismFences = CVarAddKeepStateDeterminismFences.GetValueOnGameThread();
 	for (const FMovieSceneTrackEvaluationFieldEntry& Entry : EvaluationField.Entries)
 	{
-		if (bAddKeepStateDeterminismFences)
+		if (bAddKeepStateDeterminismFences && Entry.Section)
 		{
 			// If a section is KeepState, we need to make sure to evaluate it on its last frame so that the value that "sticks" is correct.
 			const TRange<FFrameNumber> SectionRange = Entry.Section->GetRange();
