Sequencer: Operate on a copy of the current instances when running RunLegacyTrackTemplates

diff --git a/Engine/Source/Runtime/MovieScene/Private/EntitySystem/MovieSceneEntitySystemRunner.cpp b/Engine/Source/Runtime/MovieScene/Private/EntitySystem/MovieSceneEntitySystemRunner.cpp
--- a/Engine/Source/Runtime/MovieScene/Private/EntitySystem/MovieSceneEntitySystemRunner.cpp
+++ b/Engine/Source/Runtime/MovieScene/Private/EntitySystem/MovieSceneEntitySystemRunner.cpp
@@ -504,12 +504,17 @@ void FMovieSceneEntitySystemRunner::GameThread_EvaluationFinalizationPhase()
 	{
 		FInstanceRegistry* InstanceRegistry = GetInstanceRegistry();
 
-		for (FInstanceHandle UpdatedInstance : CurrentInstances)
+		// Iterate on a copy of our current instances, since LegacyEvaluator->Evaluate() could change the instance handle, which would affect PostEvaluationPhase
+		TArray<FInstanceHandle> CurrentInstancesCopy(CurrentInstances);
+		for (FInstanceHandle UpdatedInstanceHandle : CurrentInstancesCopy)
 		{
-			FSequenceInstance& Instance = InstanceRegistry->MutateInstance(UpdatedInstance);
-			if (Instance.IsRootSequence())
+			if (InstanceRegistry->IsHandleValid(UpdatedInstanceHandle))
 			{
-				Instance.RunLegacyTrackTemplates();
+				FSequenceInstance& Instance = InstanceRegistry->MutateInstance(UpdatedInstanceHandle);
+				if (Instance.IsRootSequence())
+				{
+					Instance.RunLegacyTrackTemplates();
+				}
 			}
 		}
 
