Sequencer: Allow client-only sequences to finish playback on their own.

By "client-only" sequences here we refer to sequences whose actor has no network authority (client side in a networked game), but doesn't need replicated playback, so it can do whatever it wants.

diff --git a/Engine/Source/Runtime/LevelSequence/Private/LevelSequenceActor.cpp b/Engine/Source/Runtime/LevelSequence/Private/LevelSequenceActor.cpp
--- a/Engine/Source/Runtime/LevelSequence/Private/LevelSequenceActor.cpp
+++ b/Engine/Source/Runtime/LevelSequence/Private/LevelSequenceActor.cpp
@@ -116,6 +116,11 @@ TOptional<EAspectRatioAxisConstraint> ALevelSequenceActor::GetAspectRatioAxisCon
 	return AspectRatioAxisConstraint;
 }
 
+bool ALevelSequenceActor::GetIsReplicatedPlayback() const
+{
+	return bReplicatePlayback;
+}
+
 ULevelSequencePlayer* ALevelSequenceActor::GetSequencePlayer() const
 {
 	return SequencePlayer && SequencePlayer->GetSequence() ? SequencePlayer : nullptr;
diff --git a/Engine/Source/Runtime/LevelSequence/Public/LevelSequenceActor.h b/Engine/Source/Runtime/LevelSequence/Public/LevelSequenceActor.h
--- a/Engine/Source/Runtime/LevelSequence/Public/LevelSequenceActor.h
+++ b/Engine/Source/Runtime/LevelSequence/Public/LevelSequenceActor.h
@@ -262,6 +262,7 @@ protected:
 	virtual bool RetrieveBindingOverrides(const FGuid& InBindingId, FMovieSceneSequenceID InSequenceID, TArray<UObject*, TInlineAllocator<1>>& OutObjects) const override;
 	virtual UObject* GetInstanceData() const override;
 	virtual TOptional<EAspectRatioAxisConstraint> GetAspectRatioAxisConstraint() const override;
+	virtual bool GetIsReplicatedPlayback() const override;
 	//~ End IMovieScenePlaybackClient interface
 
 	//~ Begin UObject interface
diff --git a/Engine/Source/Runtime/MovieScene/Private/MovieSceneSequencePlayer.cpp b/Engine/Source/Runtime/MovieScene/Private/MovieSceneSequencePlayer.cpp
--- a/Engine/Source/Runtime/MovieScene/Private/MovieSceneSequencePlayer.cpp
+++ b/Engine/Source/Runtime/MovieScene/Private/MovieSceneSequencePlayer.cpp
@@ -1009,7 +1009,13 @@ void UMovieSceneSequencePlayer::UpdateTimeCursorPosition_Internal(FFrameTime New
 			FMovieSceneEvaluationRange Range = UpdatePlayPosition(PlayPosition, NewPosition, Method);
 			UpdateMovieSceneInstance(Range, StatusOverride);
 
-			if (HasAuthority())
+			// If we are using replicated playback and are not an authoritative sequence player, we need to wait 
+			// for the server to tell us to finish.
+			const bool bHasAuthorityToFinish = (
+					HasAuthority() || 
+					PlaybackClient == nullptr ||
+					!PlaybackClient->GetIsReplicatedPlayback());
+			if (bHasAuthorityToFinish)
 			{
 				FinishPlaybackInternal(NewPosition);
 
diff --git a/Engine/Source/Runtime/MovieScene/Public/IMovieScenePlaybackClient.h b/Engine/Source/Runtime/MovieScene/Public/IMovieScenePlaybackClient.h
--- a/Engine/Source/Runtime/MovieScene/Public/IMovieScenePlaybackClient.h
+++ b/Engine/Source/Runtime/MovieScene/Public/IMovieScenePlaybackClient.h
@@ -45,4 +45,12 @@ public:
 	{
 		return TOptional<EAspectRatioAxisConstraint>();
 	}
+
+	/*
+	 * Whether this playback client wants replicated playback.
+	 */
+	virtual bool GetIsReplicatedPlayback() const
+	{
+		return false;
+	}
 };
