Type truncation warning fixes for some MovieScene headers

diff --git a/Engine/Source/Runtime/MovieScene/Public/EntitySystem/MovieSceneEntityIDs.h b/Engine/Source/Runtime/MovieScene/Public/EntitySystem/MovieSceneEntityIDs.h
--- a/Engine/Source/Runtime/MovieScene/Public/EntitySystem/MovieSceneEntityIDs.h
+++ b/Engine/Source/Runtime/MovieScene/Public/EntitySystem/MovieSceneEntityIDs.h
@@ -181,7 +181,7 @@ struct FComponentTypeID
 	static FComponentTypeID FromBitIndex(int32 BitIndex)
 	{
 		check( (BitIndex & 0xFFFF0000) == 0 );
-		return FComponentTypeID(BitIndex);
+		return FComponentTypeID((uint16)BitIndex);
 	}
 
 	template<typename T>
diff --git a/Engine/Source/Runtime/MovieScene/Public/EntitySystem/MovieSceneEntityManager.h b/Engine/Source/Runtime/MovieScene/Public/EntitySystem/MovieSceneEntityManager.h
--- a/Engine/Source/Runtime/MovieScene/Public/EntitySystem/MovieSceneEntityManager.h
+++ b/Engine/Source/Runtime/MovieScene/Public/EntitySystem/MovieSceneEntityManager.h
@@ -884,8 +884,8 @@ private:
 		void Set(int32 InAllocationIndex, int32 InEntryIndexWithinAllocation)
 		{
 			check((InAllocationIndex & 0xFFFF0000) == 0 && (InEntryIndexWithinAllocation & 0xFFFF0000) == 0);
-			AllocationIndex = InAllocationIndex;
-			EntryIndexWithinAllocation = InEntryIndexWithinAllocation;
+			AllocationIndex = (uint16)InAllocationIndex;
+			EntryIndexWithinAllocation = (uint16)InEntryIndexWithinAllocation;
 		}
 
 		int32 GetAllocationIndex() const
diff --git a/Engine/Source/Runtime/MovieScene/Public/Evaluation/Blending/MovieSceneBlendType.h b/Engine/Source/Runtime/MovieScene/Public/Evaluation/Blending/MovieSceneBlendType.h
--- a/Engine/Source/Runtime/MovieScene/Public/Evaluation/Blending/MovieSceneBlendType.h
+++ b/Engine/Source/Runtime/MovieScene/Public/Evaluation/Blending/MovieSceneBlendType.h
@@ -170,7 +170,7 @@ private:
 
 	MOVIESCENE_API void IterateToNext();
 
-	FORCEINLINE static int8 MaxValidOffset() { return StaticEnum<EMovieSceneBlendType>()->NumEnums() - 2; }
+	FORCEINLINE static int8 MaxValidOffset() { return (int8)(StaticEnum<EMovieSceneBlendType>()->NumEnums() - 2); }
 
 	FMovieSceneBlendTypeField Field;
 	int8 Offset;
