Sequencer: Fixes to generation and traversal of determinism fences

  - Empty dissection ranges are now skipped
  - Determinism fences that exist at the exact same time as the start of the current eval time are now skipped
  - Fences are no longer added at interim segments of the sub-sequence tree, only at the true start/end points of sub sequences

diff --git a/Engine/Source/Runtime/MovieScene/Private/Compilation/MovieSceneCompiledDataManager.cpp b/Engine/Source/Runtime/MovieScene/Private/Compilation/MovieSceneCompiledDataManager.cpp
--- a/Engine/Source/Runtime/MovieScene/Private/Compilation/MovieSceneCompiledDataManager.cpp
+++ b/Engine/Source/Runtime/MovieScene/Private/Compilation/MovieSceneCompiledDataManager.cpp
@@ -828,20 +828,32 @@ void UMovieSceneCompiledDataManager::CompileSubSequences(const FMovieSceneSequen
 
 	const TMovieSceneEvaluationTree<FMovieSceneSubSequenceTreeEntry>& SubSequenceTree = Hierarchy.GetTree();
 
+	// When adding determinism fences for sub sequences, we track the iteration index for each sequence ID so that
+	// we only add a fence when the sub sequence truly ends or begins, not for every segmentation of the sub sequence tree
+	struct FSubSequenceItMetaData
+	{
+		int32 LastIterIndex = INDEX_NONE;
+		TOptional<FFrameNumber> TrailingFence;
+
+	};
+	TSortedMap<FMovieSceneSequenceID, FSubSequenceItMetaData> ItMetaData;
+
 	// Start iterating the field from the lower bound of the compile range
 	FMovieSceneEvaluationTreeRangeIterator SubSequenceIt = SubSequenceTree.IterateFromLowerBound(Params.RootClampRange.GetLowerBound());
-	for ( ; SubSequenceIt && SubSequenceIt.Range().Overlaps(Params.RootClampRange); ++SubSequenceIt)
+	for ( int32 ItIndex = 0; SubSequenceIt && SubSequenceIt.Range().Overlaps(Params.RootClampRange); ++SubSequenceIt, ++ItIndex)
 	{
 		// Iterate all sub sequences in the current range
 		for (FMovieSceneSubSequenceTreeEntry SubSequenceEntry : SubSequenceTree.GetAllData(SubSequenceIt.Node()))
 		{
-			const FMovieSceneSubSequenceData* SubData = Hierarchy.FindSubData(SubSequenceEntry.SequenceID);
+			FMovieSceneSequenceID SubSequenceID = SubSequenceEntry.SequenceID;
+
+			const FMovieSceneSubSequenceData* SubData = Hierarchy.FindSubData(SubSequenceID);
 			checkf(SubData, TEXT("Sub data could not be found for a sequence that exists in the sub sequence tree - this indicates an error while populating the sub sequence hierarchy tree."));
 
 			UMovieSceneSequence* SubSequence = SubData->GetSequence();
 			if (SubSequence)
 			{
-				FTrackGatherParameters SubSectionGatherParams = Params.CreateForSubData(*SubData, SubSequenceEntry.SequenceID);
+				FTrackGatherParameters SubSectionGatherParams = Params.CreateForSubData(*SubData, SubSequenceID);
 				SubSectionGatherParams.Flags |= SubSequenceEntry.Flags;
 				SubSectionGatherParams.SetClampRange(SubSequenceIt.Range());
 
@@ -860,14 +872,28 @@ void UMovieSceneCompiledDataManager::CompileSubSequences(const FMovieSceneSequen
 				OutCompilerData->InheritedFlags |= (CompiledDataEntries[SubDataID.Value].AccumulatedFlags & EMovieSceneSequenceFlags::InheritedFlags);
 				OutCompilerData->AccumulatedMask |= SubEntry.AccumulatedMask;
 
-				if (SubEntry.CompiledFlags.bParentSequenceRequiresLowerFence)
+				FSubSequenceItMetaData* MetaData = &ItMetaData.FindOrAdd(SubSequenceID);
+
+				const bool bWasEvaluatedLastFrame = MetaData->LastIterIndex != INDEX_NONE && MetaData->LastIterIndex == ItIndex-1;
+				if (SubEntry.CompiledFlags.bParentSequenceRequiresLowerFence && bWasEvaluatedLastFrame == false)
 				{
 					OutCompilerData->DeterminismData.Fences.Add(DiscreteInclusiveLower(SubSequenceIt.Range()));
 				}
 				if (SubEntry.CompiledFlags.bParentSequenceRequiresUpperFence)
 				{
-					OutCompilerData->DeterminismData.Fences.Add(DiscreteExclusiveUpper(SubSequenceIt.Range()));
+					MetaData->TrailingFence = DiscreteExclusiveUpper(SubSequenceIt.Range());
 				}
+
+				MetaData->LastIterIndex = ItIndex;
+			}
+		}
+
+		for (TPair<FMovieSceneSequenceID, FSubSequenceItMetaData>& Pair : ItMetaData)
+		{
+			if (Pair.Value.LastIterIndex == ItIndex-1 && Pair.Value.TrailingFence.IsSet())
+			{
+				OutCompilerData->DeterminismData.Fences.Add(Pair.Value.TrailingFence.GetValue());
+				Pair.Value.TrailingFence.Reset();
 			}
 		}
 	}
diff --git a/Engine/Source/Runtime/MovieScene/Private/EntitySystem/MovieSceneSequenceUpdaters.cpp b/Engine/Source/Runtime/MovieScene/Private/EntitySystem/MovieSceneSequenceUpdaters.cpp
--- a/Engine/Source/Runtime/MovieScene/Private/EntitySystem/MovieSceneSequenceUpdaters.cpp
+++ b/Engine/Source/Runtime/MovieScene/Private/EntitySystem/MovieSceneSequenceUpdaters.cpp
@@ -121,11 +121,14 @@ void DissectRange(TArrayView<const FFrameTime> InDissectionTimes, const TRange<F
 		FFrameTime DissectionTime = InDissectionTimes[Index];
 
 		TRange<FFrameTime> Dissection(LowerBound, TRangeBound<FFrameTime>::Exclusive(DissectionTime));
-		ensureAlwaysMsgf(Bounds.Contains(Dissection), TEXT("Dissection specified for a range outside of the current bounds"));
+		if (!Dissection.IsEmpty())
+		{
+			ensureAlwaysMsgf(Bounds.Contains(Dissection), TEXT("Dissection specified for a range outside of the current bounds"));
 
-		OutDissections.Add(Dissection);
+			OutDissections.Add(Dissection);
 
-		LowerBound = TRangeBound<FFrameTime>::Inclusive(DissectionTime);
+			LowerBound = TRangeBound<FFrameTime>::Inclusive(DissectionTime);
+		}
 	}
 
 	TRange<FFrameTime> TailRange(LowerBound, Bounds.GetUpperBound());
@@ -143,7 +146,7 @@ TArrayView<const FFrameTime> GetFencesWithinRange(TArrayView<const FFrameTime> F
 	}
 
 	// Take care to include or exclude the lower bound of the range if it's on a whole frame numbe
-	const int32 StartFence = Boundary.GetLowerBound().IsClosed() ? Algo::LowerBound(Fences, Boundary.GetLowerBoundValue()) : 0;
+	const int32 StartFence = Boundary.GetLowerBound().IsClosed() ? Algo::UpperBound(Fences, Boundary.GetLowerBoundValue()) : 0;
 	if (StartFence >= Fences.Num())
 	{
 		return TArrayView<const FFrameTime>();
