Sequencer: Use a custom interface and tick method for sequene actors.

This is because level sequences tick at a specific point in time during the game loop. Adding a Tick override makes the BP compiler enable ticking on the actor, which means that it will tick twice (once with the custom sequencer tick manager, and once with the normal actor tick groups). This will not only make the sequence go twice as fast, it will also cause other problems around starting/stopping. So here we make the sequence tick manager use a custom tick method, which frees up the standard tick to do other things.

Note that we need to introduce an interface for this custom tick method because the tick manager, which resides in the core MovieScene module, doesn't know anything about sequence actors such as LevelSequenceActor.

diff --git a/Engine/Plugins/MovieScene/TemplateSequence/Source/TemplateSequence/Private/TemplateSequenceActor.cpp b/Engine/Plugins/MovieScene/TemplateSequence/Source/TemplateSequence/Private/TemplateSequenceActor.cpp
--- a/Engine/Plugins/MovieScene/TemplateSequence/Source/TemplateSequence/Private/TemplateSequenceActor.cpp
+++ b/Engine/Plugins/MovieScene/TemplateSequence/Source/TemplateSequence/Private/TemplateSequenceActor.cpp
@@ -50,7 +50,7 @@ void ATemplateSequenceActor::PostInitializeComponents()
 
 void ATemplateSequenceActor::BeginPlay()
 {
-	UMovieSceneSequenceTickManager::Get(this)->SequenceActors.Add(this);
+	UMovieSceneSequenceTickManager::Get(this)->RegisterSequenceActor(this);
 
 	Super::BeginPlay();
 	
@@ -68,15 +68,13 @@ void ATemplateSequenceActor::EndPlay(const EEndPlayReason::Type EndPlayReason)
 		SequencePlayer->Stop();
 	}
 
-	UMovieSceneSequenceTickManager::Get(this)->SequenceActors.Remove(this);
+	UMovieSceneSequenceTickManager::Get(this)->UnregisterSequenceActor(this);
 
 	Super::EndPlay(EndPlayReason);
 }
 
-void ATemplateSequenceActor::Tick(float DeltaSeconds)
+void ATemplateSequenceActor::TickFromSequenceTickManager(float DeltaSeconds)
 {
-	Super::Tick(DeltaSeconds);
-
 	if (SequencePlayer)
 	{
 		SequencePlayer->Update(DeltaSeconds);
diff --git a/Engine/Plugins/MovieScene/TemplateSequence/Source/TemplateSequence/Public/TemplateSequenceActor.h b/Engine/Plugins/MovieScene/TemplateSequence/Source/TemplateSequence/Public/TemplateSequenceActor.h
--- a/Engine/Plugins/MovieScene/TemplateSequence/Source/TemplateSequence/Public/TemplateSequenceActor.h
+++ b/Engine/Plugins/MovieScene/TemplateSequence/Source/TemplateSequence/Public/TemplateSequenceActor.h
@@ -6,6 +6,7 @@
 #include "TemplateSequencePlayer.h"
 #include "IMovieScenePlaybackClient.h"
 #include "MovieSceneBindingOwnerInterface.h"
+#include "MovieSceneSequenceTickManager.h"
 #include "GameFramework/Actor.h"
 #include "UObject/ObjectMacros.h"
 #include "UObject/Object.h"
@@ -34,6 +35,7 @@ struct FTemplateSequenceBindingOverrideData
 UCLASS(hideCategories = (Rendering, Physics, LOD, Activation, Input))
 class TEMPLATESEQUENCE_API ATemplateSequenceActor
 	: public AActor
+	, public IMovieSceneSequenceActor
 	, public IMovieScenePlaybackClient
 {
 public:
@@ -75,13 +77,16 @@ public:
 
 protected:
 
+	//~ Begin IMovieSceneSequenceActor interface
+	virtual void TickFromSequenceTickManager(float DeltaSeconds) override;
+	//~ End IMovieSceneSequenceActor interface
+
 	//~ Begin IMovieScenePlaybackClient interface
 	virtual bool RetrieveBindingOverrides(const FGuid& InBindingId, FMovieSceneSequenceID InSequenceID, TArray<UObject*, TInlineAllocator<1>>& OutObjects) const override;
 	virtual UObject* GetInstanceData() const override;
 	//~ End IMovieScenePlaybackClient interface
 
 	//~ Begin AActor interface
-	virtual void Tick(float DeltaSeconds) override;
 	virtual void PostInitializeComponents() override;
 	virtual void BeginPlay() override;
 	virtual void EndPlay(const EEndPlayReason::Type EndPlayReason) override;
diff --git a/Engine/Source/Runtime/LevelSequence/Private/LevelSequenceActor.cpp b/Engine/Source/Runtime/LevelSequence/Private/LevelSequenceActor.cpp
--- a/Engine/Source/Runtime/LevelSequence/Private/LevelSequenceActor.cpp
+++ b/Engine/Source/Runtime/LevelSequence/Private/LevelSequenceActor.cpp
@@ -162,7 +162,7 @@ void ALevelSequenceActor::BeginPlay()
 	UMovieSceneSequenceTickManager* TickManager = SequencePlayer->GetTickManager();
 	if (ensure(TickManager))
 	{
-		TickManager->SequenceActors.Add(this);
+		TickManager->RegisterSequenceActor(this);
 	}
 
 	Super::BeginPlay();
@@ -186,17 +186,15 @@ void ALevelSequenceActor::EndPlay(const EEndPlayReason::Type EndPlayReason)
 
 		if (UMovieSceneSequenceTickManager* TickManager = SequencePlayer->GetTickManager())
 		{
-			TickManager->SequenceActors.Remove(this);
+			TickManager->UnregisterSequenceActor(this);
 		}
 	}
 
 	Super::EndPlay(EndPlayReason);
 }
 
-void ALevelSequenceActor::Tick(float DeltaSeconds)
+void ALevelSequenceActor::TickFromSequenceTickManager(float DeltaSeconds)
 {
-	Super::Tick(DeltaSeconds);
-
 	if (SequencePlayer)
 	{
 		SequencePlayer->Update(DeltaSeconds);
diff --git a/Engine/Source/Runtime/LevelSequence/Public/LevelSequenceActor.h b/Engine/Source/Runtime/LevelSequence/Public/LevelSequenceActor.h
--- a/Engine/Source/Runtime/LevelSequence/Public/LevelSequenceActor.h
+++ b/Engine/Source/Runtime/LevelSequence/Public/LevelSequenceActor.h
@@ -11,6 +11,7 @@
 #include "LevelSequencePlayer.h"
 #include "MovieSceneBindingOwnerInterface.h"
 #include "MovieSceneBindingOverrides.h"
+#include "MovieSceneSequenceTickManager.h"
 #include "LevelSequenceActor.generated.h"
 
 class ULevelSequenceBurnIn;
@@ -61,6 +62,7 @@ protected:
 UCLASS(hideCategories=(Rendering, Physics, LOD, Activation, Input))
 class LEVELSEQUENCE_API ALevelSequenceActor
 	: public AActor
+	, public IMovieSceneSequenceActor
 	, public IMovieScenePlaybackClient
 	, public IMovieSceneBindingOwnerInterface
 {
@@ -252,6 +254,10 @@ public:
 
 protected:
 
+	//~ Begin IMovieSceneSequenceActor interface
+	virtual void TickFromSequenceTickManager(float DeltaSeconds) override;
+	//~ End IMovieSceneSequenceActor interface
+
 	//~ Begin IMovieScenePlaybackClient interface
 	virtual bool RetrieveBindingOverrides(const FGuid& InBindingId, FMovieSceneSequenceID InSequenceID, TArray<UObject*, TInlineAllocator<1>>& OutObjects) const override;
 	virtual UObject* GetInstanceData() const override;
@@ -265,7 +271,6 @@ protected:
 	//~ End UObject interface
 
 	//~ Begin AActor interface
-	virtual void Tick(float DeltaSeconds) override;
 	virtual void PostInitializeComponents() override;
 	virtual void BeginPlay() override;
 	virtual void EndPlay(const EEndPlayReason::Type EndPlayReason) override;
diff --git a/Engine/Source/Runtime/MovieScene/Private/MovieSceneSequenceTickManager.cpp b/Engine/Source/Runtime/MovieScene/Private/MovieSceneSequenceTickManager.cpp
--- a/Engine/Source/Runtime/MovieScene/Private/MovieSceneSequenceTickManager.cpp
+++ b/Engine/Source/Runtime/MovieScene/Private/MovieSceneSequenceTickManager.cpp
@@ -36,6 +36,25 @@ void UMovieSceneSequenceTickManager::BeginDestroy()
 	Super::BeginDestroy();
 }
 
+void UMovieSceneSequenceTickManager::RegisterSequenceActor(AActor* InActor)
+{
+	TScriptInterface<IMovieSceneSequenceActor> SequenceActorInterface(InActor);
+	if (ensureMsgf(SequenceActorInterface, TEXT("The given actor doesn't implement the IMovieSceneSequenceActor interface!")))
+	{
+		SequenceActors.Add(FMovieSceneSequenceActorPointers{ InActor, SequenceActorInterface });
+	}
+}
+
+void UMovieSceneSequenceTickManager::UnregisterSequenceActor(AActor* InActor)
+{
+	TScriptInterface<IMovieSceneSequenceActor> SequenceActorInterface(InActor);
+	if (ensureMsgf(SequenceActorInterface, TEXT("The given actor doesn't implement the IMovieSceneSequenceActor interface!")))
+	{
+		int32 NumRemoved = SequenceActors.RemoveAll([=](FMovieSceneSequenceActorPointers& Item) { return Item.SequenceActor == InActor; });
+		ensureMsgf(NumRemoved > 0, TEXT("The given sequence actor wasn't registered"));
+	}
+}
+
 void UMovieSceneSequenceTickManager::TickSequenceActors(float DeltaSeconds)
 {
 	SCOPE_CYCLE_COUNTER(MovieSceneEval_SequenceTickManager);
@@ -56,12 +75,14 @@ void UMovieSceneSequenceTickManager::TickSequenceActors(float DeltaSeconds)
 
 	for (int32 i = SequenceActors.Num() - 1; i >= 0; --i)
 	{
-		if (AActor* SequenceActor = SequenceActors[i])
+		const FMovieSceneSequenceActorPointers& Pointers(SequenceActors[i]);
+		if (Pointers.SequenceActor)
 		{
-			if (!bIsPaused || SequenceActor->GetTickableWhenPaused())
+			if (!bIsPaused || Pointers.SequenceActor->GetTickableWhenPaused())
 			{
-				check(SequenceActor->GetWorld() == World);
-				SequenceActor->Tick(DeltaSeconds);
+				check(Pointers.SequenceActorInterface);
+				check(Pointers.SequenceActor->GetWorld() == World);
+				Pointers.SequenceActorInterface->TickFromSequenceTickManager(DeltaSeconds);
 				bHasTasks = true;
 			}
 		}
diff --git a/Engine/Source/Runtime/MovieScene/Public/MovieSceneSequenceTickManager.h b/Engine/Source/Runtime/MovieScene/Public/MovieSceneSequenceTickManager.h
--- a/Engine/Source/Runtime/MovieScene/Public/MovieSceneSequenceTickManager.h
+++ b/Engine/Source/Runtime/MovieScene/Public/MovieSceneSequenceTickManager.h
@@ -12,6 +12,9 @@ class UMovieSceneEntitySystemLinker;
 
 DECLARE_DELEGATE(FMovieSceneSequenceLatentActionDelegate);
 
+/**
+ * Utility class for running latent actions created from sequence players.
+ */
 class MOVIESCENE_API FMovieSceneLatentActionManager
 {
 public:
@@ -29,6 +32,40 @@ private:
 	bool bIsRunningLatentActions = false;
 };
 
+/**
+ * Interface for sequence actors that are to be ticked by the tick manager.
+ */
+UINTERFACE()
+class MOVIESCENE_API UMovieSceneSequenceActor
+	: public UInterface
+{
+public:
+	GENERATED_BODY()
+};
+
+class MOVIESCENE_API IMovieSceneSequenceActor
+{
+public:
+	GENERATED_BODY()
+
+	virtual void TickFromSequenceTickManager(float DeltaSeconds) = 0;
+};
+
+/**
+ * A structure for storing pointers to a sequence actor.
+ */
+USTRUCT()
+struct FMovieSceneSequenceActorPointers
+{
+	GENERATED_BODY()
+
+	UPROPERTY()
+	TObjectPtr<AActor> SequenceActor;
+
+	UPROPERTY()
+	TScriptInterface<IMovieSceneSequenceActor> SequenceActorInterface;
+};
+
 /**
  * An automatically created global object that will manage all level sequence actors' updates.
  */
@@ -43,6 +80,9 @@ public:
 	UMovieSceneEntitySystemLinker* GetLinker() { return Linker; }
 	FMovieSceneEntitySystemRunner& GetRunner() { return Runner; }
 
+	void RegisterSequenceActor(AActor* InActor);
+	void UnregisterSequenceActor(AActor* InActor);
+
 	void AddLatentAction(FMovieSceneSequenceLatentActionDelegate Delegate);
 	void RunLatentActions();
 
@@ -55,11 +95,11 @@ private:
 
 	void TickSequenceActors(float DeltaSeconds);
 
-public:
+private:
+
 	UPROPERTY(transient)
-	TArray<TObjectPtr<AActor>> SequenceActors;
+	TArray<FMovieSceneSequenceActorPointers> SequenceActors;
 
-private:
 	UPROPERTY(transient)
 	TObjectPtr<UMovieSceneEntitySystemLinker> Linker;
 
