Fix CIS issue with clang compilers

diff --git a/Engine/Source/Runtime/MovieScene/Public/EntitySystem/MovieSceneBlenderSystem.h b/Engine/Source/Runtime/MovieScene/Public/EntitySystem/MovieSceneBlenderSystem.h
--- a/Engine/Source/Runtime/MovieScene/Public/EntitySystem/MovieSceneBlenderSystem.h
+++ b/Engine/Source/Runtime/MovieScene/Public/EntitySystem/MovieSceneBlenderSystem.h
@@ -51,7 +51,7 @@ public:
 	template<typename BlenderSystemClass>
 	static FMovieSceneBlenderSystemID GetBlenderSystemID()
 	{
-		UMovieSceneBlenderSystem* DefaultObject = BlenderSystemClass::StaticClass()->GetDefaultObject<UMovieSceneBlenderSystem>();
+		UMovieSceneBlenderSystem* DefaultObject = BlenderSystemClass::StaticClass()->template GetDefaultObject<UMovieSceneBlenderSystem>();
 		check(DefaultObject);
 		return DefaultObject->SystemID;
 	}
