Sequencer: Added an enabled state for entity systems
  - Disabled systems will still be kept alive as long as they are relevant, but will not be run during their system phases

diff --git a/Engine/Source/Runtime/MovieScene/Private/EntitySystem/MovieSceneEntitySystem.cpp b/Engine/Source/Runtime/MovieScene/Private/EntitySystem/MovieSceneEntitySystem.cpp
--- a/Engine/Source/Runtime/MovieScene/Private/EntitySystem/MovieSceneEntitySystem.cpp
+++ b/Engine/Source/Runtime/MovieScene/Private/EntitySystem/MovieSceneEntitySystem.cpp
@@ -176,6 +176,7 @@ UMovieSceneEntitySystem::UMovieSceneEntitySystem(const FObjectInitializer& ObjIn
 {
 	using namespace UE::MovieScene;
 
+	bSystemIsEnabled = true;
 	SystemExclusionContext = EEntitySystemContext::None;
 
 	Phase = ESystemPhase::Evaluation;
@@ -297,6 +298,16 @@ void UMovieSceneEntitySystem::ConditionalLinkSystemImpl(UMovieSceneEntitySystemL
 	}
 }
 
+void UMovieSceneEntitySystem::Enable()
+{
+	bSystemIsEnabled = true;
+}
+
+void UMovieSceneEntitySystem::Disable()
+{
+	bSystemIsEnabled = false;
+}
+
 void UMovieSceneEntitySystem::TagGarbage()
 {
 	OnTagGarbage();
@@ -327,6 +338,11 @@ void UMovieSceneEntitySystem::FinishDestroy()
 
 void UMovieSceneEntitySystem::Run(FSystemTaskPrerequisites& InPrerequisites, FSystemSubsequentTasks& Subsequents)
 {
+	if (!bSystemIsEnabled)
+	{
+		return;
+	}
+
 #if STATS || ENABLE_STATNAMEDEVENTS
 	FScopeCycleCounter Scope(StatID);
 #endif
diff --git a/Engine/Source/Runtime/MovieScene/Public/EntitySystem/MovieSceneEntitySystem.h b/Engine/Source/Runtime/MovieScene/Public/EntitySystem/MovieSceneEntitySystem.h
--- a/Engine/Source/Runtime/MovieScene/Public/EntitySystem/MovieSceneEntitySystem.h
+++ b/Engine/Source/Runtime/MovieScene/Public/EntitySystem/MovieSceneEntitySystem.h
@@ -131,6 +131,17 @@ public:
 
 	void CleanTaggedGarbage();
 
+	/**
+	 * Enable this system if it is not already.
+	 */
+	void Enable();
+
+	/**
+	 * Disable this system if it is not already.
+	 * Disabled systems will remain in the system graph, and will stay alive as long as they are relevant, but will not be Run.
+	 */
+	void Disable();
+
 protected:
 
 	virtual bool IsReadyForFinishDestroy() override;
@@ -169,6 +180,9 @@ protected:
 
 	UE::MovieScene::EEntitySystemContext SystemExclusionContext;
 
+	/** When false, this system will not call its OnRun function, but will still be kept alive as long as IsRelevant is true */
+	bool bSystemIsEnabled;
+
 #if STATS || ENABLE_STATNAMEDEVENTS
 	TStatId StatID;
 #endif
