Sequencer: Legacy track templates are now run _before_ PostEvaluation

  - This ensures that sequencer events are triggered after track templates, that any event triggers can rely on state produced by track templates, and that they are called before the CompletionTask is signalled.

diff --git a/Engine/Source/Runtime/MovieScene/Private/EntitySystem/MovieSceneEntitySystemRunner.cpp b/Engine/Source/Runtime/MovieScene/Private/EntitySystem/MovieSceneEntitySystemRunner.cpp
--- a/Engine/Source/Runtime/MovieScene/Private/EntitySystem/MovieSceneEntitySystemRunner.cpp
+++ b/Engine/Source/Runtime/MovieScene/Private/EntitySystem/MovieSceneEntitySystemRunner.cpp
@@ -450,6 +450,8 @@ void FMovieSceneEntitySystemRunner::GameThread_EvaluationFinalizationPhase()
 	}
 	bCanQueueEventTriggers = false;
 
+	GetInstanceRegistry()->FinalizeFrame();
+
 	CurrentPhase = ESystemPhase::None;
 
 	// We are now done with the current update batch. Let's unlock the completion task to unblock
diff --git a/Engine/Source/Runtime/MovieScene/Private/EntitySystem/MovieSceneInstanceRegistry.cpp b/Engine/Source/Runtime/MovieScene/Private/EntitySystem/MovieSceneInstanceRegistry.cpp
--- a/Engine/Source/Runtime/MovieScene/Private/EntitySystem/MovieSceneInstanceRegistry.cpp
+++ b/Engine/Source/Runtime/MovieScene/Private/EntitySystem/MovieSceneInstanceRegistry.cpp
@@ -148,5 +148,16 @@ void FInstanceRegistry::CleanupLinkerEntities(const TSet<FMovieSceneEntityID>& E
 	}
 }
 
+void FInstanceRegistry::FinalizeFrame()
+{
+	for (FSequenceInstance& SequenceInstance : Instances)
+	{
+		if (SequenceInstance.IsRootSequence())
+		{
+			SequenceInstance.RunLegacyTrackTemplates();
+		}
+	}
+}
+
 } // namespace MovieScene
 } // namespace UE
diff --git a/Engine/Source/Runtime/MovieScene/Private/EntitySystem/MovieSceneSequenceInstance.cpp b/Engine/Source/Runtime/MovieScene/Private/EntitySystem/MovieSceneSequenceInstance.cpp
--- a/Engine/Source/Runtime/MovieScene/Private/EntitySystem/MovieSceneSequenceInstance.cpp
+++ b/Engine/Source/Runtime/MovieScene/Private/EntitySystem/MovieSceneSequenceInstance.cpp
@@ -264,7 +264,7 @@ void FSequenceInstance::PreEvaluation(UMovieSceneEntitySystemLinker* Linker)
 	}
 }
 
-void FSequenceInstance::PostEvaluation(UMovieSceneEntitySystemLinker* Linker)
+void FSequenceInstance::RunLegacyTrackTemplates()
 {
 	if (LegacyEvaluator)
 	{
@@ -281,7 +281,10 @@ void FSequenceInstance::PostEvaluation(UMovieSceneEntitySystemLinker* Linker)
 			}
 		}
 	}
+}
 
+void FSequenceInstance::PostEvaluation(UMovieSceneEntitySystemLinker* Linker)
+{
 	Ledger.UnlinkOneShots(Linker);
 
 	if (IsRootSequence())
diff --git a/Engine/Source/Runtime/MovieScene/Public/EntitySystem/MovieSceneSequenceInstance.h b/Engine/Source/Runtime/MovieScene/Public/EntitySystem/MovieSceneSequenceInstance.h
--- a/Engine/Source/Runtime/MovieScene/Public/EntitySystem/MovieSceneSequenceInstance.h
+++ b/Engine/Source/Runtime/MovieScene/Public/EntitySystem/MovieSceneSequenceInstance.h
@@ -88,6 +88,11 @@ public:
 	 */
 	void PostEvaluation(UMovieSceneEntitySystemLinker* Linker);
 
+	/**
+	 * Called to run legacy track templates
+	 */
+	void RunLegacyTrackTemplates();
+
 	/**
 	 * Indicate that this sequence instance should capture any and all changes of state so they can be restored later
 	 *
