Sequencer: Mutations are now allowed to be performed on _all_ entities using a default constructed filter

diff --git a/Engine/Source/Runtime/MovieScene/Private/EntitySystem/MovieSceneEntityManager.cpp b/Engine/Source/Runtime/MovieScene/Private/EntitySystem/MovieSceneEntityManager.cpp
--- a/Engine/Source/Runtime/MovieScene/Private/EntitySystem/MovieSceneEntityManager.cpp
+++ b/Engine/Source/Runtime/MovieScene/Private/EntitySystem/MovieSceneEntityManager.cpp
@@ -1606,8 +1606,6 @@ int32 FEntityManager::MutateAll(const FEntityComponentFilter& Filter, const IMov
 {
 	CheckCanChangeStructure();
 
-	check(Filter.IsValid());
-
 	TMap<int32, FComponentMask> AllocationMutations;
 
 	for (int32 AllocationIndex = 0; AllocationIndex < EntityAllocationMasks.GetMaxIndex(); ++AllocationIndex)
