Sequencer: Guard against infinite loops caused by entity mutations occuring in the finalization phase

diff --git a/Engine/Source/Runtime/MovieScene/Private/EntitySystem/MovieSceneEntitySystemRunner.cpp b/Engine/Source/Runtime/MovieScene/Private/EntitySystem/MovieSceneEntitySystemRunner.cpp
--- a/Engine/Source/Runtime/MovieScene/Private/EntitySystem/MovieSceneEntitySystemRunner.cpp
+++ b/Engine/Source/Runtime/MovieScene/Private/EntitySystem/MovieSceneEntitySystemRunner.cpp
@@ -133,12 +133,18 @@ void FMovieSceneEntitySystemRunner::Flush()
 	EntityManager.SetDispatchThread(ENamedThreads::GameThread_Local);
 	EntityManager.SetGatherThread(ENamedThreads::GameThread_Local);
 
+	// We specifically only check whether the entity manager has changed since the last instantation once
+	// to ensure that we are not vulnerable to infinite loops where components are added/removed in post-evaluation
+	bool bStructureHadChanged = Linker->EntityManager.HasStructureChangedSince(LastInstantiationVersion);
+
 	// Start flushing the update queue... keep flushing as long as we have work to do.
 	while (UpdateQueue.Num() > 0 || 
 			DissectedUpdates.Num() > 0 ||
-			Linker->EntityManager.HasStructureChangedSince(LastInstantiationVersion))
+			bStructureHadChanged)
 	{
 		DoFlushUpdateQueueOnce();
+
+		bStructureHadChanged = false;
 	}
 
 	Linker->EndEvaluation(*this);
