Sequencer: Fix CIS warning

diff --git a/Engine/Source/Runtime/MovieScene/Public/MovieSceneSequenceTickManager.h b/Engine/Source/Runtime/MovieScene/Public/MovieSceneSequenceTickManager.h
--- a/Engine/Source/Runtime/MovieScene/Public/MovieSceneSequenceTickManager.h
+++ b/Engine/Source/Runtime/MovieScene/Public/MovieSceneSequenceTickManager.h
@@ -60,7 +60,7 @@ struct FMovieSceneSequenceActorPointers
 	GENERATED_BODY()
 
 	UPROPERTY()
-	TObjectPtr<AActor> SequenceActor;
+	TObjectPtr<AActor> SequenceActor = nullptr;
 
 	UPROPERTY()
 	TScriptInterface<IMovieSceneSequenceActor> SequenceActorInterface;
