Sequencer: Move Finalization phase back onto the evaluating thread

  - In testing camera shake sequences from within sequencer, I encountered a deadlock which is being caused by task-graph re-entrancy from the camera shake player being evaluated inside the legacy track template (which I recently moved into the finalization phase).
  - To remedy this, but retain the behavior where legacy track templates are evaluated before Finalization, I have moved the Finalization phase back onto the initiating thread. This will result in no functional change in behavior because the initiating thread will have been stalled while the Finalization phase completed on a background thread anyway, so I don't forsee any added latency as a result of this change.

#tests Ran all engine tests

diff --git a/Engine/Source/Runtime/MovieScene/Private/EntitySystem/MovieSceneEntitySystemRunner.cpp b/Engine/Source/Runtime/MovieScene/Private/EntitySystem/MovieSceneEntitySystemRunner.cpp
--- a/Engine/Source/Runtime/MovieScene/Private/EntitySystem/MovieSceneEntitySystemRunner.cpp
+++ b/Engine/Source/Runtime/MovieScene/Private/EntitySystem/MovieSceneEntitySystemRunner.cpp
@@ -175,12 +175,15 @@ void FMovieSceneEntitySystemRunner::DoFlushUpdateQueueOnce()
 	{
 		GameThread_ProcessQueue();
 	}
+
 	FTaskGraphInterface::Get().WaitUntilTaskCompletes(CompletionTask->GetCompletionEvent(), ENamedThreads::GameThread_Local);
+
 	// TODO: delete?
 	CompletionTask = nullptr;
 
 	// Now run the post-evaluation logic, which contains stuff we don't want to run from inside a task
 	// graph call.
+	GameThread_EvaluationFinalizationPhase();
 	GameThread_PostEvaluationPhase();
 }
 
@@ -420,19 +423,22 @@ void FMovieSceneEntitySystemRunner::GameThread_EvaluationPhase()
 	FGraphEventArray AllTasks;
 	Linker->SystemGraph.ExecutePhase(ESystemPhase::Evaluation, Linker, AllTasks);
 
+	auto Finish = [this]
+	{
+		// We are now done with the current update batch. Let's unlock the completion task to unblock
+		// the main thread, which is waiting on it inside Flush().
+		check(this->CompletionTask != nullptr);
+		this->CompletionTask->Unlock();
+	};
+
 	if (AllTasks.Num() != 0)
 	{
 		TGraphTask<TFunctionGraphTaskImpl<void(), ESubsequentsMode::TrackSubsequents>>::CreateTask(&AllTasks, ENamedThreads::GameThread)
-		.ConstructAndDispatchWhenReady(
-			[this]
-			{
-				this->GameThread_EvaluationFinalizationPhase();
-			}
-		, TStatId(), GameThread);
+		.ConstructAndDispatchWhenReady(MoveTemp(Finish), TStatId(), GameThread);
 	}
 	else
 	{
-		this->GameThread_EvaluationFinalizationPhase();
+		Finish();
 	}
 }
 
@@ -450,20 +456,15 @@ void FMovieSceneEntitySystemRunner::GameThread_EvaluationFinalizationPhase()
 	// The events are actually executed a bit later, in GameThread_PostEvaluationPhase.
 	bCanQueueEventTriggers = true;
 	{
+		GetInstanceRegistry()->FinalizeFrame();
+
 		FGraphEventArray Tasks;
 		Linker->SystemGraph.ExecutePhase(ESystemPhase::Finalization, Linker, Tasks);
 		checkf(Tasks.Num() == 0, TEXT("Cannot dispatch new tasks during finalization"));
 	}
 	bCanQueueEventTriggers = false;
 
-	GetInstanceRegistry()->FinalizeFrame();
-
 	CurrentPhase = ESystemPhase::None;
-
-	// We are now done with the current update batch. Let's unlock the completion task to unblock
-	// the main thread, which is waiting on it inside Flush().
-	check(CompletionTask != nullptr);
-	CompletionTask->Unlock();
 }
 
 void FMovieSceneEntitySystemRunner::GameThread_PostEvaluationPhase()
