Sequencer: Added a global iteration filter

  - This allows evaluation to run on a selective set of entities

#tests Ran all engine tests

diff --git a/Engine/Source/Runtime/MovieScene/Private/EntitySystem/EntityAllocationIterator.cpp b/Engine/Source/Runtime/MovieScene/Private/EntitySystem/EntityAllocationIterator.cpp
--- a/Engine/Source/Runtime/MovieScene/Private/EntitySystem/EntityAllocationIterator.cpp
+++ b/Engine/Source/Runtime/MovieScene/Private/EntitySystem/EntityAllocationIterator.cpp
@@ -78,14 +78,27 @@ FEntityAllocationIteratorItem FEntityAllocationIterator::operator*() const
 
 int32 FEntityAllocationIterator::FindMatchingAllocationStartingAt(int32 Index) const
 {
-	while (Index < Manager->EntityAllocationMasks.GetMaxIndex())
+	const FEntityComponentFilter& GlobalIterationFilter = Manager->GetGlobalIterationFilter();
+	const bool bHasGlobalFilter = GlobalIterationFilter.IsValid();
+	for ( ; Index < Manager->EntityAllocationMasks.GetMaxIndex(); ++Index)
 	{
-		if (Manager->EntityAllocationMasks.IsAllocated(Index) && Filter->Match(Manager->EntityAllocationMasks[Index]) && Manager->EntityAllocations[Index]->Num() > 0)
+		if (!Manager->EntityAllocationMasks.IsAllocated(Index))
 		{
-			return Index;
+			continue;
+		}
+
+		if (bHasGlobalFilter == true)
+		{
+			if (!GlobalIterationFilter.Match(Manager->EntityAllocationMasks[Index]))
+			{
+				continue;
+			}
 		}
 
-		++Index;
+		if (Filter->Match(Manager->EntityAllocationMasks[Index]) && Manager->EntityAllocations[Index]->Num() > 0)
+		{
+			return Index;
+		}
 	}
 
 	return Index;
diff --git a/Engine/Source/Runtime/MovieScene/Public/EntitySystem/MovieSceneEntityManager.h b/Engine/Source/Runtime/MovieScene/Public/EntitySystem/MovieSceneEntityManager.h
--- a/Engine/Source/Runtime/MovieScene/Public/EntitySystem/MovieSceneEntityManager.h
+++ b/Engine/Source/Runtime/MovieScene/Public/EntitySystem/MovieSceneEntityManager.h
@@ -738,6 +738,26 @@ public:
 		LockdownState = ELockdownState::Unlocked;
 	}
 
+	/**
+	 * Retrieve the entity filter that should be used for any entity iteration. Can be used to constrain all iterations to specific types
+	 */
+	const FEntityComponentFilter& GetGlobalIterationFilter() const
+	{
+		return GlobalIterationFilter;
+	}
+
+	/**
+	 * Modify the entity filter that should be used for any entity iteration.
+	 * @note: Care should be take to reset this back to its original state when done,
+	 *        as leaving this populated during the instantiation phase can result in 
+	 *        undefined results.
+	 */
+	FEntityComponentFilter& ModifyGlobalIterationFilter()
+	{
+		ensureMsgf(!IsLockedDown() && IterationCount == 0, TEXT("Manipulating the global iteration filter while locked down or iterating is not recommended"));
+		return GlobalIterationFilter;
+	}
+
 	/**
 	 * Increment the current serial number for systems observing this manager. Should be called after any system is run
 	 */
@@ -923,6 +943,7 @@ private:
 	TMap<FMovieSceneEntityID, uint32> EntityGenerationMap;
 
 	FComponentMask AccumulatedMask;
+	FEntityComponentFilter GlobalIterationFilter;
 
 	FComponentRegistry* ComponentRegistry;
 
