Sequencer: Expose BlendType on FOptionalMovieSceneBlendType as BlueprintReadOnly

diff --git a/Engine/Source/Runtime/MovieScene/Public/Evaluation/Blending/MovieSceneBlendType.h b/Engine/Source/Runtime/MovieScene/Public/Evaluation/Blending/MovieSceneBlendType.h
--- a/Engine/Source/Runtime/MovieScene/Public/Evaluation/Blending/MovieSceneBlendType.h
+++ b/Engine/Source/Runtime/MovieScene/Public/Evaluation/Blending/MovieSceneBlendType.h
@@ -73,14 +73,12 @@ struct FOptionalMovieSceneBlendType
 		return A.IsValid() && A.BlendType == B;
 	}
 
-private:
-
 	/** The actual blend type */
-	UPROPERTY()
+	UPROPERTY(BlueprintReadOnly, Category="Blend Type")
 	EMovieSceneBlendType BlendType;
 
 	/** Boolean indicating whether BlendType is valid */
-	UPROPERTY()
+	UPROPERTY(BlueprintReadOnly, Category="Blend Type")
 	bool bIsValid;
 };
 
