Sequencer: Add new utility class to add and sort timestamped data from different time spaces.

diff --git a/Engine/Source/Runtime/MovieScene/Private/Evaluation/MovieSceneTimeArray.cpp b/Engine/Source/Runtime/MovieScene/Private/Evaluation/MovieSceneTimeArray.cpp
new file mode 100644
--- /dev/null
+++ b/Engine/Source/Runtime/MovieScene/Private/Evaluation/MovieSceneTimeArray.cpp
@@ -0,0 +1,4 @@
+// Copyright Epic Games, Inc. All Rights Reserved.
+
+#include "Evaluation/MovieSceneTimeArray.h"
+
diff --git a/Engine/Source/Runtime/MovieScene/Private/Tests/MovieSceneTimeArrayTests.cpp b/Engine/Source/Runtime/MovieScene/Private/Tests/MovieSceneTimeArrayTests.cpp
new file mode 100644
--- /dev/null
+++ b/Engine/Source/Runtime/MovieScene/Private/Tests/MovieSceneTimeArrayTests.cpp
@@ -0,0 +1,70 @@
+// Copyright Epic Games, Inc. All Rights Reserved.
+
+#include "CoreMinimal.h"
+#include "Evaluation/MovieSceneTimeArray.h"
+#include "Misc/AutomationTest.h"
+
+#define LOCTEXT_NAMESPACE "MovieSceneTimeArrayTests"
+
+template<typename DataType>
+bool TestArrayEntries(FAutomationTestBase& Test, const TMovieSceneTimeArray<DataType>& TimeArray, const TArray<int32>& FrameNumbers, const TArray<DataType>& Data)
+{
+	if (!Test.TestEqual("Unexpected test arguments", Data.Num(), FrameNumbers.Num()))
+	{
+		return false;
+	}
+
+	TArrayView<const TMovieSceneTimeArrayEntry<DataType>> Entries = TimeArray.GetEntries();
+	if (!Test.TestEqual("Unexpected number of entries", Entries.Num(), FrameNumbers.Num()))
+	{
+		return false;
+	}
+
+	for (int32 Index = 0, Num = Entries.Num(); Index < Num; ++Index)
+	{
+		const TMovieSceneTimeArrayEntry<DataType>& Entry(Entries[Index]);
+		if (!Test.TestEqual(FString::Format(TEXT("Invalid time for entry {0}"), { Index }), Entry.RootTime, FFrameTime(FrameNumbers[Index])))
+		{
+			return false;
+		}
+		if (!Test.TestEqual(FString::Format(TEXT("Invalid data for entry {0}"), { Index }), Entry.Datum, Data[Index]))
+		{
+			return false;
+		}
+	}
+
+	return true;
+}
+
+#define UTEST_TIME_ARRAY(TimeArray, FrameNumbers, Data)\
+	if (!TestArrayEntries(*this, TimeArray, FrameNumbers, Data)) return false;
+
+IMPLEMENT_SIMPLE_AUTOMATION_TEST(FMovieSceneSimpleTimeArrayTests,
+		"System.Engine.Sequencer.Core.SimpleTimeArrays",
+		EAutomationTestFlags::EditorContext | EAutomationTestFlags::EngineFilter)
+bool FMovieSceneSimpleTimeArrayTests::RunTest(const FString& Parameters)
+{
+	{
+		TMovieSceneTimeArray<int32> Array;
+		Array.Add(0, 0);
+		Array.Add(10, 2);
+		Array.Add(5, 1);
+		UTEST_TIME_ARRAY(Array, TArray<int32>({ 0, 5, 10 }), TArray<int32>({ 0, 1, 2 }));
+	}
+
+	{
+		TMovieSceneTimeArray<int32> Array;
+		Array.Add(0, 0);
+		Array.PushTransform(-200);
+		Array.Add(0, 1);
+		Array.Add(100, 2);
+		Array.PopTransform();
+		UTEST_TIME_ARRAY(Array, TArray<int32>({ 0, 200, 300 }), TArray<int32>({ 0, 1, 2 }));
+	}
+
+	return true;
+}
+
+#undef UTEST_TIME_ARRAY
+
+#undef LOCTEXT_NAMESPACE
diff --git a/Engine/Source/Runtime/MovieScene/Public/Evaluation/MovieSceneTimeArray.h b/Engine/Source/Runtime/MovieScene/Public/Evaluation/MovieSceneTimeArray.h
new file mode 100644
--- /dev/null
+++ b/Engine/Source/Runtime/MovieScene/Public/Evaluation/MovieSceneTimeArray.h
@@ -0,0 +1,210 @@
+// Copyright Epic Games, Inc. All Rights Reserved.
+
+#pragma once
+
+#include "CoreMinimal.h"
+#include "Misc/FrameNumber.h"
+#include "MovieSceneSequenceTransform.h"
+#include "MovieSceneTimeTransform.h"
+#include "MovieSceneTimeWarping.h"
+
+/** Entry struct for the TMovieSceneTimeArray class */
+template<typename DataType>
+struct TMovieSceneTimeArrayEntry
+{
+	FFrameTime RootTime;
+	DataType Datum;
+};
+
+/**
+ * A utility class that lets you store a collection of timestamped data originating from various time bases.
+ *
+ * All of the data is stored in "root" time space. That is: as you add timestamped data, these timestamps are
+ * converted back to "root times" using the inverse of the current time transform. Pushing and popping time
+ * transforms, and incrementing loop counts, makes it possible to change what's considered "local time", 
+ * which affects this inverse transformation.
+ */
+template<typename DataType>
+struct TMovieSceneTimeArray
+{
+	/** Starts a scope where local times are offset and scaled as specified */
+	void PushTransform(FFrameNumber FrameOffset, float TimeScale = 1.f)
+	{
+		PushTransform(FMovieSceneTimeTransform(FrameOffset, TimeScale));
+	}
+
+	/** Starts a scope where local times are transformed as specified */
+	void PushTransform(const FMovieSceneTimeTransform& InTransform)
+	{
+		FTransformStep Step;
+		Step.LinearTransform = InTransform;
+		TransformStack.Push(Step);
+
+		UpdateCachedInverseTransform();
+	}
+
+	/** Starts a scope where local times are looped and offset as specified */
+	void PushTransform(FFrameNumber FirstLoopStartFrameOffset, FFrameNumber LoopStart, FFrameNumber LoopEnd, float TimeScale = 1.f)
+	{
+		PushTransform(
+				FMovieSceneTimeTransform(FirstLoopStartFrameOffset, TimeScale),
+				FMovieSceneTimeWarping(LoopStart, LoopEnd));
+	}
+
+	/** Starts a scope where local times are transformed as specified */
+	void PushTransform(const FMovieSceneTimeWarping& InWarping)
+	{
+		PushTransform(FMovieSceneTimeTransform(), InWarping);
+	}
+
+	/** Starts a scope where local times are transformed as specified */
+	void PushTransform(const FMovieSceneTimeTransform& InTransform, const FMovieSceneTimeWarping& InWarping)
+	{
+		FTransformStep Step;
+		Step.LinearTransform = InTransform;
+		Step.Warping = InWarping;
+		Step.WarpCounter = InWarping.IsValid() ? 0 : FMovieSceneTimeWarping::InvalidWarpCount;
+		TransformStack.Push(Step);
+
+		UpdateCachedInverseTransform();
+	}
+
+	/** Incremenets the warp counter for the local time space, i.e. all local times will be considered as being in a further loop */
+	void IncrementWarpCounter()
+	{
+		if (ensure(TransformStack.Num() > 0))
+		{
+			FTransformStep& LastStep = TransformStack.Last();
+			if (ensure(LastStep.Warping.IsValid() && LastStep.WarpCounter != FMovieSceneTimeWarping::InvalidWarpCount))
+			{
+				++LastStep.WarpCounter;
+			}
+		}
+
+		UpdateCachedInverseTransform();
+	}
+
+	/** Removes the last pushed local time transform */
+	void PopTransform()
+	{
+		TransformStack.Pop();
+
+		UpdateCachedInverseTransform();
+	}
+
+	/** Adds a new piece of data, timestamped to the given local time */
+	void Add(FFrameNumber LocalTime, const DataType& Datum)
+	{
+		const FFrameTime RootTime = LocalTime * CachedInverseTransform;
+
+		// Insert the new entry by keeping the array sorted by time.
+		int32 NextEntry = Entries.IndexOfByPredicate([RootTime](const FEntry& Entry) { return RootTime < Entry.RootTime; });
+		if (NextEntry != INDEX_NONE)
+		{
+			Entries.Insert(FEntry{ RootTime, Datum }, NextEntry);
+		}
+		else
+		{
+			Entries.Add(FEntry{ RootTime, Datum });
+		}
+	}
+
+	/** Clears all entries */
+	void Clear()
+	{
+		Entries.Clear();
+	}
+
+	/** Gets the current list of entries in the array */
+	TArrayView<const TMovieSceneTimeArrayEntry<DataType>> GetEntries() const
+	{
+		return TArrayView<const FEntry>(Entries);
+	}
+
+private:
+
+	void UpdateCachedInverseTransform()
+	{
+		// Build the total transform from root to local, and keep track of the loop indices.
+		FMovieSceneWarpCounter LoopCounter;
+		FMovieSceneSequenceTransform RootToLocalTransform;
+
+		for (const FTransformStep& Step : TransformStack)
+		{
+			FMovieSceneSequenceTransform StepTransform;
+			StepTransform.LinearTransform = Step.LinearTransform;
+			if (Step.Warping.IsValid())
+			{
+				StepTransform.NestedTransforms.Add(FMovieSceneNestedSequenceTransform(Step.Warping));
+			}
+			RootToLocalTransform = StepTransform * RootToLocalTransform;
+
+			ensure(Step.WarpCounter == FMovieSceneTimeWarping::InvalidWarpCount || Step.Warping.IsValid());
+			LoopCounter.AddWarpingLevel(Step.WarpCounter);
+		}
+
+		CachedInverseTransform = RootToLocalTransform.InverseFromWarp(LoopCounter);
+	}
+
+private:
+
+	using FEntry = TMovieSceneTimeArrayEntry<DataType>;
+
+	TArray<FEntry> Entries;
+
+	struct FTransformStep
+	{
+		FMovieSceneTimeTransform LinearTransform;
+		FMovieSceneTimeWarping Warping;
+		uint32 WarpCounter = FMovieSceneTimeWarping::InvalidWarpCount;
+	};
+	TArray<FTransformStep> TransformStack;
+
+	FMovieSceneTimeTransform CachedInverseTransform;
+};
+
+template<typename DataType>
+struct TMovieSceneTimeArrayTransformScope
+{
+	TMovieSceneTimeArrayTransformScope(TMovieSceneTimeArray<DataType>& InArray, FFrameNumber Offset, float TimeScale = 1.f)
+		: Array(&InArray)
+	{
+		Array->PushTransform(Offset, TimeScale);
+	}
+
+	TMovieSceneTimeArrayTransformScope(TMovieSceneTimeArray<DataType>& InArray, FMovieSceneTimeTransform&& InTransform)
+		: Array(&InArray)
+	{
+		Array->PushTransform(InTransform);
+	}
+
+	TMovieSceneTimeArrayTransformScope(TMovieSceneTimeArray<DataType>& InArray, FFrameNumber FirstLoopStartFrameOffset, FFrameNumber LoopStart, FFrameNumber LoopEnd, float TimeScale = 1.f)
+		: Array(&InArray)
+	{
+		Array->PushTransform(FirstLoopStartFrameOffset, LoopStart, LoopEnd, TimeScale);
+	}
+
+	TMovieSceneTimeArrayTransformScope(TMovieSceneTimeArray<DataType>& InArray, FMovieSceneTimeTransform&& InTransform, FMovieSceneTimeWarping&& InWarping)
+		: Array(&InArray)
+	{
+		Array->PushTransform(InTransform, InWarping);
+	}
+
+	TMovieSceneTimeArrayTransformScope(TMovieSceneTimeArrayTransformScope&& Other)
+		: Array(Other.Array)
+	{
+		Other.Array = nullptr;
+	}
+
+	~TMovieSceneTimeArrayTransformScope()
+	{
+		if (Array)
+		{
+			Array->PopTransform();
+		}
+	}
+
+private:
+	TMovieSceneTimeArray<DataType>* Array = nullptr;
+};
+
