Sequencer: Add some new utility methods for time transforms, and fix bug in inverse computation.

Computing the inverse of a time transform wasn't taking into account the linear part of a looping transform.

diff --git a/Engine/Source/Runtime/MovieScene/Private/Tests/MovieSceneTransformTests.cpp b/Engine/Source/Runtime/MovieScene/Private/Tests/MovieSceneTransformTests.cpp
--- a/Engine/Source/Runtime/MovieScene/Private/Tests/MovieSceneTransformTests.cpp
+++ b/Engine/Source/Runtime/MovieScene/Private/Tests/MovieSceneTransformTests.cpp
@@ -416,6 +416,20 @@ bool FMovieSceneSubSectionCoreWarpingAndScalingTransformsTest::RunTest(const FSt
 		TestEqual("Inverse time 18", FFrameNumber(4) * Inv, FFrameTime(32));
 	}
 
+	{
+		// Sub-sequence at 3, playing at x2, with start offset 20, but all contained inside a higher offset of 100.
+		FMovieSceneSequenceTransform Transform;
+		Transform.LinearTransform.Offset = FFrameTime(-100);
+		Transform.NestedTransforms.Emplace(FMovieSceneTimeTransform(-6 + 20, 2.f), FMovieSceneTimeWarping(20, 50));
+		TestEqual("Transform time 19", FFrameNumber(113) * Transform, FFrameTime(40));
+		TestEqual("Transform time 20", FFrameNumber(121) * Transform, FFrameTime(26));
+
+		FMovieSceneTimeTransform Inv = Transform.InverseFromWarp(TArray<uint32> { 0 });
+		TestEqual("Inverse time 21", FFrameNumber(40) * Inv, FFrameTime(113));
+		Inv = Transform.InverseFromWarp(TArray<uint32> { 1 });
+		TestEqual("Inverse time 22", FFrameNumber(40) * Inv, FFrameTime(128));
+	}
+
 	return true;
 }
 
diff --git a/Engine/Source/Runtime/MovieScene/Public/Evaluation/MovieSceneSequenceTransform.h b/Engine/Source/Runtime/MovieScene/Public/Evaluation/MovieSceneSequenceTransform.h
--- a/Engine/Source/Runtime/MovieScene/Public/Evaluation/MovieSceneSequenceTransform.h
+++ b/Engine/Source/Runtime/MovieScene/Public/Evaluation/MovieSceneSequenceTransform.h
@@ -105,6 +105,10 @@ struct FMovieSceneNestedSequenceTransform
 		: LinearTransform(InLinearTransform)
 	{}
 
+	FMovieSceneNestedSequenceTransform(FMovieSceneTimeWarping InWarping)
+		: Warping(InWarping)
+	{}
+
 	FMovieSceneNestedSequenceTransform(FMovieSceneTimeTransform InLinearTransform, FMovieSceneTimeWarping InWarping)
 		: LinearTransform(InLinearTransform)
 		, Warping(InWarping)
@@ -244,6 +248,17 @@ struct FMovieSceneSequenceTransform
 		return TimeScale;
 	}
 
+	/**
+	 * Transforms the given time, returning the transformed time.
+	 */
+	FFrameTime TransformTime(FFrameTime InTime) const
+	{
+		FFrameTime OutTime;
+		FMovieSceneWarpCounter WarpCounter;
+		TransformTime(InTime, OutTime, WarpCounter);
+		return OutTime;
+	}
+
 	/**
 	 * Transforms the given time, returning the transformed time along with a series of loop indices that
 	 * indicate which loops each nested level is in.
@@ -418,9 +433,6 @@ struct FMovieSceneSequenceTransform
 			const FMovieSceneNestedSequenceTransform& LastNesting = NestedTransforms.Last();
 			FMovieSceneTimeTransform Result = LastNesting.InverseFromWarp(WarpCounts.Last());
 
-			// If there's a tail linear transform, the WarpCounts array has one less item than the
-			// NestedTransforms array, and we need to iterate with different indices that are offset by 1.
-			// Otherwise, we can iterate with the same index.
 			size_t Index = NestedTransformsSize - 1;
 			while (Index > 0)
 			{
@@ -430,6 +442,10 @@ struct FMovieSceneSequenceTransform
 				const FMovieSceneNestedSequenceTransform& CurNesting = NestedTransforms[Index];
 				Result = CurNesting.InverseFromWarp(CurWarpCount) * Result;
 			}
+
+			// Add the inverse of the main linear transform.
+			Result = LinearTransform.Inverse() * Result;
+
 			return Result;
 		}
 		else
