Sequencer: Add more information to sub-sequence data.

The new data lets us track the play range of the parent section.

diff --git a/Engine/Source/Runtime/MovieScene/Private/Evaluation/MovieSceneSequenceHierarchy.cpp b/Engine/Source/Runtime/MovieScene/Private/Evaluation/MovieSceneSequenceHierarchy.cpp
--- a/Engine/Source/Runtime/MovieScene/Private/Evaluation/MovieSceneSequenceHierarchy.cpp
+++ b/Engine/Source/Runtime/MovieScene/Private/Evaluation/MovieSceneSequenceHierarchy.cpp
@@ -15,6 +15,11 @@ FMovieSceneSubSequenceData::FMovieSceneSubSequenceData()
 FMovieSceneSubSequenceData::FMovieSceneSubSequenceData(const UMovieSceneSubSection& InSubSection)
 	: Sequence(InSubSection.GetSequence())
 	, DeterministicSequenceID(InSubSection.GetSequenceID())
+	, ParentPlayRange(InSubSection.GetTrueRange())
+	, ParentStartFrameOffset(InSubSection.Parameters.StartFrameOffset)
+	, ParentEndFrameOffset(InSubSection.Parameters.EndFrameOffset)
+	, ParentFirstLoopStartFrameOffset(InSubSection.Parameters.FirstLoopStartFrameOffset)
+	, bCanLoop(InSubSection.Parameters.bCanLoop)
 	, HierarchicalBias(InSubSection.Parameters.HierarchicalBias)
 	, bHasHierarchicalEasing(false)
 #if WITH_EDITORONLY_DATA
@@ -94,6 +99,17 @@ bool FMovieSceneSubSequenceData::IsDirty(const UMovieSceneSubSection& InSubSecti
 	return InSubSection.GetSignature() != SubSectionSignature || InSubSection.OuterToInnerTransform() != OuterToInnerTransform;
 }
 
+FMovieSceneSectionParameters FMovieSceneSubSequenceData::ToSubSectionParameters() const
+{
+	FMovieSceneSectionParameters Parameters;
+	Parameters.StartFrameOffset = ParentStartFrameOffset;
+	Parameters.bCanLoop = bCanLoop;
+	Parameters.EndFrameOffset = ParentEndFrameOffset;
+	Parameters.FirstLoopStartFrameOffset = ParentFirstLoopStartFrameOffset;
+	Parameters.TimeScale = OuterToInnerTransform.GetTimeScale();
+	return Parameters;
+}
+
 void FMovieSceneSequenceHierarchy::Add(const FMovieSceneSubSequenceData& Data, FMovieSceneSequenceIDRef ThisSequenceID, FMovieSceneSequenceIDRef ParentID)
 {
 	check(ParentID != MovieSceneSequenceID::Invalid);
diff --git a/Engine/Source/Runtime/MovieScene/Public/Evaluation/MovieSceneSequenceHierarchy.h b/Engine/Source/Runtime/MovieScene/Public/Evaluation/MovieSceneSequenceHierarchy.h
--- a/Engine/Source/Runtime/MovieScene/Public/Evaluation/MovieSceneSequenceHierarchy.h
+++ b/Engine/Source/Runtime/MovieScene/Public/Evaluation/MovieSceneSequenceHierarchy.h
@@ -58,10 +58,19 @@ struct FMovieSceneSubSequenceData
 	*/
 	MOVIESCENE_API FGuid GetSubSectionSignature() const { return SubSectionSignature; }
 
+	/**
+	 * Re-creates a sub-section parameter struct.
+	 */
+	MOVIESCENE_API FMovieSceneSectionParameters ToSubSectionParameters() const;
+
 	/** The sequence that the sub section references */
 	UPROPERTY(meta=(AllowedClasses="MovieSceneSequence"))
 	FSoftObjectPath Sequence;
 
+	/** The transform from this sub sequence's parent to its own play space. */
+	UPROPERTY()
+	FMovieSceneSequenceTransform OuterToInnerTransform;
+
 	/** Transform that transforms a given time from the sequences outer space, to its authored space. */
 	UPROPERTY()
 	FMovieSceneSequenceTransform RootToSequenceTransform;
@@ -73,6 +82,26 @@ struct FMovieSceneSubSequenceData
 	/** This sequence's deterministic sequence ID. Used in editor to reduce the risk of collisions on recompilation. */ 
 	UPROPERTY()
 	FMovieSceneSequenceID DeterministicSequenceID;
+
+	/** The play range of the parent section */
+	UPROPERTY()
+	FMovieSceneFrameRange ParentPlayRange;
+
+	/** The start frame offset of the parent section */
+	UPROPERTY()
+	FFrameNumber ParentStartFrameOffset;
+
+	/** The end frame offset of the parent section */
+	UPROPERTY()
+	FFrameNumber ParentEndFrameOffset;
+
+	/** The offset for the first loop of the sub-sequence */
+	UPROPERTY()
+	FFrameNumber ParentFirstLoopStartFrameOffset;
+
+	/** Whether this sub-sequence can loop */
+	UPROPERTY()
+	bool bCanLoop = false;
 	
 	/** This sub sequence's playback range according to its parent sub section. Clamped recursively during template generation */
 	UPROPERTY()
@@ -127,10 +156,6 @@ private:
 	/** The sub section's signature at the time this structure was populated. */
 	UPROPERTY()
 	FGuid SubSectionSignature;
-
-	/** The transform from this sub sequence's parent to its own play space. */
-	UPROPERTY()
-	FMovieSceneSequenceTransform OuterToInnerTransform;
 };
 
 /**
