Sequencer: Speculative fix for GetKeyTimes out of bounds

diff --git a/Engine/Source/Runtime/MovieScene/Private/Channels/MovieSceneChannelData.cpp b/Engine/Source/Runtime/MovieScene/Private/Channels/MovieSceneChannelData.cpp
--- a/Engine/Source/Runtime/MovieScene/Private/Channels/MovieSceneChannelData.cpp
+++ b/Engine/Source/Runtime/MovieScene/Private/Channels/MovieSceneChannelData.cpp
@@ -218,7 +218,7 @@ void FMovieSceneChannelData::GetKeyTimes(TArrayView<const FKeyHandle> InHandles,
 	for (int32 Index = 0; Index < InHandles.Num(); ++Index)
 	{
 		const int32 KeyIndex = GetIndex(InHandles[Index]);
-		if (KeyIndex != INDEX_NONE)
+		if (Times->IsValidIndex(KeyIndex))
 		{
 			OutKeyTimes[Index] = (*Times)[KeyIndex];
 		}
