Sequencer: Remove output log for possessed object class now that it is an editor only property

diff --git a/Engine/Source/Runtime/MovieScene/Private/Evaluation/MovieSceneEvaluationState.cpp b/Engine/Source/Runtime/MovieScene/Private/Evaluation/MovieSceneEvaluationState.cpp
--- a/Engine/Source/Runtime/MovieScene/Private/Evaluation/MovieSceneEvaluationState.cpp
+++ b/Engine/Source/Runtime/MovieScene/Private/Evaluation/MovieSceneEvaluationState.cpp
@@ -351,13 +351,6 @@ void FMovieSceneObjectCache::UpdateBindings(const FGuid& InGuid, IMovieScenePlay
 					for (UObject* Object : FoundObjects)
 					{
 						Bindings->Objects.Add(Object);
-
-#if WITH_EDITORONLY_DATA
-						if (!Object->GetClass()->IsChildOf(Possessable->GetPossessedObjectClass()))
-						{
-							UE_LOG(LogMovieScene, Error, TEXT("Mismatch in %s with possessed object class for: %s, Expected: %s, but was: %s"), *GetNameSafe(Sequence), *Possessable->GetName(), *GetNameSafe(Possessable->GetPossessedObjectClass()), *GetNameSafe(Object->GetClass()));
-						}
-#endif
 					}
 				}
 			}
@@ -369,13 +362,6 @@ void FMovieSceneObjectCache::UpdateBindings(const FGuid& InGuid, IMovieScenePlay
 				for (UObject* Object : FoundObjects)
 				{
 					Bindings->Objects.Add(Object);
-
-#if WITH_EDITORONLY_DATA
-					if (!Object->GetClass()->IsChildOf(Possessable->GetPossessedObjectClass()))
-					{
-						UE_LOG(LogMovieScene, Error, TEXT("Mismatch in %s with possessed object class for: %s, Expected: %s, but was: %s"), *GetNameSafe(Sequence), *Possessable->GetName(), *GetNameSafe(Possessable->GetPossessedObjectClass()), *GetNameSafe(Object->GetClass()));
-					}
-#endif
 				}
 			}
 		}
