Sequencer: More fixes to looping sub-sequences.

- Gather params need the entire warp counter, not just the index of the latest level, because sub-sequences are not necessarily compiled in top-down order.
- Transforming time from the sub-sequence to the root should use the sub-data, not the parent sub-sequence params.

diff --git a/Engine/Source/Runtime/MovieScene/Private/Compilation/MovieSceneCompiledDataManager.cpp b/Engine/Source/Runtime/MovieScene/Private/Compilation/MovieSceneCompiledDataManager.cpp
--- a/Engine/Source/Runtime/MovieScene/Private/Compilation/MovieSceneCompiledDataManager.cpp
+++ b/Engine/Source/Runtime/MovieScene/Private/Compilation/MovieSceneCompiledDataManager.cpp
@@ -144,10 +144,10 @@ struct FGatherParameters
 
 	FGatherParameters CreateForSubData(const FMovieSceneSubSequenceData& SubData, FMovieSceneSequenceID InSubSequenceID) const
 	{
-		return CreateForSubData(SubData, InSubSequenceID, FMovieSceneTimeWarping::InvalidWarpCount);
+		return CreateForSubData(SubData, InSubSequenceID, FMovieSceneWarpCounter());
 	}
 
-	FGatherParameters CreateForSubData(const FMovieSceneSubSequenceData& SubData, FMovieSceneSequenceID InSubSequenceID, uint32 WarpIndex) const
+	FGatherParameters CreateForSubData(const FMovieSceneSubSequenceData& SubData, FMovieSceneSequenceID InSubSequenceID, FMovieSceneWarpCounter WarpCounter) const
 	{
 		FGatherParameters SubParams = *this;
 
@@ -155,21 +155,9 @@ struct FGatherParameters
 		SubParams.HierarchicalBias          = SubData.HierarchicalBias;
 		SubParams.bHasHierarchicalEasing    = SubData.bHasHierarchicalEasing;
 		SubParams.SequenceID                = InSubSequenceID;
-		SubParams.LocalClampRange           = SubData.RootToSequenceTransform.TransformRangeUnwarped(SubParams.RootClampRange);
+		SubParams.RootToSequenceWarpCounter = WarpCounter;
 
-		if (WarpIndex != FMovieSceneTimeWarping::InvalidWarpCount)
-		{
-		SubParams.RootToSequenceWarpCounter.AddWarpingLevel(WarpIndex);
-		}
-		else
-		{
-			const int32 NumWarpCounts = SubParams.RootToSequenceWarpCounter.NumWarpCounts();
-			const int32 NumNestedTransforms = SubParams.RootToSequenceTransform.NestedTransforms.Num();
-			if (NumNestedTransforms > NumWarpCounts)
-			{
-				SubParams.RootToSequenceWarpCounter.AddNonWarpingLevel();
-			}
-		}
+		SubParams.LocalClampRange			= SubData.RootToSequenceTransform.TransformRangeUnwarped(SubParams.RootClampRange);
 
 		return SubParams;
 	}
@@ -219,10 +207,10 @@ struct FTrackGatherParameters : FGatherParameters
 		: TemplateGenerator(InCompiledDataManager)
 	{}
 
-	FTrackGatherParameters CreateForSubData(const FMovieSceneSubSequenceData& SubData, FMovieSceneSequenceID InSubSequenceID, uint32 WarpIndex) const
+	FTrackGatherParameters CreateForSubData(const FMovieSceneSubSequenceData& SubData, FMovieSceneSequenceID InSubSequenceID, FMovieSceneWarpCounter WarpCounter) const
 	{
 		FTrackGatherParameters SubParams = *this;
-		static_cast<FGatherParameters&>(SubParams) = FGatherParameters::CreateForSubData(SubData, InSubSequenceID, WarpIndex);
+		static_cast<FGatherParameters&>(SubParams) = FGatherParameters::CreateForSubData(SubData, InSubSequenceID, WarpCounter);
 
 		return SubParams;
 	}
@@ -911,8 +899,7 @@ void UMovieSceneCompiledDataManager::CompileSubSequences(const FMovieSceneSequen
 			UMovieSceneSequence* SubSequence = SubData->GetSequence();
 			if (SubSequence)
 			{
-				const uint32 LastSubSequenceEntryWarpCount = SubSequenceEntry.RootToSequenceWarpCounter.LastWarpCount();
-				FTrackGatherParameters SubSectionGatherParams = Params.CreateForSubData(*SubData, SubSequenceID, LastSubSequenceEntryWarpCount);
+				FTrackGatherParameters SubSectionGatherParams = Params.CreateForSubData(*SubData, SubSequenceID, SubSequenceEntry.RootToSequenceWarpCounter);
 				SubSectionGatherParams.Flags |= SubSequenceEntry.Flags;
 				SubSectionGatherParams.SetClampRange(SubSequenceIt.Range());
 
@@ -1572,10 +1559,11 @@ void UMovieSceneCompiledDataManager::PopulateSubSequenceTree(UMovieSceneSubTrack
 
 		if (!SubSectionParams.bCanLoop)
 		{
-			FGatherParameters SubParams = Params.CreateForSubData(*SubData, SubSequenceID);
+			FGatherParameters SubParams = Params.CreateForSubData(*SubData, SubSequenceID, Params.RootToSequenceWarpCounter);
 			SubParams.SetClampRange(EffectiveRange);
 			SubParams.Flags |= Entry.Flags;
 			SubParams.NetworkMask = NewMask;
+			SubParams.RootToSequenceWarpCounter.AddNonWarpingLevel();
 
 			// The section isn't looping, so we can just add it to the tree.
 			InOutHierarchy->AddRange(EffectiveRange, SubSequenceID, SubEntryFlags, SubParams.RootToSequenceWarpCounter);
@@ -1592,7 +1580,7 @@ void UMovieSceneCompiledDataManager::PopulateSubSequenceTree(UMovieSceneSubTrack
 			// The section is looping so we need to add its contents to the tree as many times as it has loops.
 			const FMovieSceneSequenceTransform SequenceToRootTransform = Params.RootToSequenceTransform.InverseFromWarp(Params.RootToSequenceWarpCounter);
 
-			const float RootToSubSequenceTimeScale = Params.RootToSequenceTransform.GetTimeScale();
+			const float RootToSubSequenceTimeScale = SubData->RootToSequenceTransform.GetTimeScale();
 			const float SubSequenceToRootTimeScale = (RootToSubSequenceTimeScale != 0.f) ? 1.0f / RootToSubSequenceTimeScale : 1.f;
 
 			UMovieSceneSequence* SubSequence = SubData->GetSequence();
@@ -1618,10 +1606,11 @@ void UMovieSceneCompiledDataManager::PopulateSubSequenceTree(UMovieSceneSubTrack
 				{
 					if (CurRootRange.Overlaps(Params.RootClampRange))
 					{
-						FGatherParameters CurLoopParams = Params.CreateForSubData(*SubData, SubSequenceID, LoopCount);
+						FGatherParameters CurLoopParams = Params.CreateForSubData(*SubData, SubSequenceID, Params.RootToSequenceWarpCounter);
 						CurLoopParams.SetClampRange(EffectiveRange);
 						CurLoopParams.Flags |= Entry.Flags;
 						CurLoopParams.NetworkMask = NewMask;
+						CurLoopParams.RootToSequenceWarpCounter.AddWarpingLevel(LoopCount);
 
 						// Add the section to the tree for the current loop.
 						const TRange<FFrameNumber> ClampedCurRootRange = TRange<FFrameNumber>::Intersection(CurRootRange, Params.RootClampRange);
