Sequencer: Add support for renaming track rows.
Enable renaming for 3d Transform track

diff --git a/Engine/Source/Runtime/MovieScene/Public/MovieSceneTrack.h b/Engine/Source/Runtime/MovieScene/Public/MovieSceneTrack.h
--- a/Engine/Source/Runtime/MovieScene/Public/MovieSceneTrack.h
+++ b/Engine/Source/Runtime/MovieScene/Public/MovieSceneTrack.h
@@ -430,7 +430,7 @@ public:
 	 *
 	 * @return Display name text.
 	 */
-	virtual FText GetTrackRowDisplayName(int32 RowIndex) const PURE_VIRTUAL(UMovieSceneTrack::GetTrackRowDisplayName, return FText::FromString(TEXT("Unnamed Track")););
+	virtual FText GetTrackRowDisplayName(int32 RowIndex) const { return FText::FromString(TEXT("Unnamed Track")); }
 
 	/**
 	 * Get the track's display tooltip text to be shown on the track's name.
