Sequencer: Don't flag objects that are going away immediately as requiring binding

When PIE stops, if the Sequencer editor is open, it will update which world it is hooked to. It will change its mind twice: first when the client world is destroyed, by hooking itself to the server world. And then when the server world is destroyed, by hooking itself to the editor world.

So for just a fraction of a frame, we are getting objects that have the NeedsLink tag, but then are immediately flagged with NeedsUnlink because the server world is going away. When the GC kicks in, we incorrectly flag those with the tag HasUnresolvedBindings, which the bound object instantiators pick up, thinking those need to be resolved. And because that's the "clean up" update before we tear everything down, we end up with freshly bound objects that didn't get cleaned up.

This fix simply doesn't tag objects with HasUnresolvedBindings if these objects are meant to be unlinked.

diff --git a/Engine/Source/Runtime/MovieScene/Private/EntitySystem/MovieSceneEntitySystemLinker.cpp b/Engine/Source/Runtime/MovieScene/Private/EntitySystem/MovieSceneEntitySystemLinker.cpp
--- a/Engine/Source/Runtime/MovieScene/Private/EntitySystem/MovieSceneEntitySystemLinker.cpp
+++ b/Engine/Source/Runtime/MovieScene/Private/EntitySystem/MovieSceneEntitySystemLinker.cpp
@@ -277,6 +277,7 @@ void UMovieSceneEntitySystemLinker::CleanGarbage()
 	FEntityTaskBuilder()
 	.Read(BuiltInComponents->BoundObject)
 	.Read(BuiltInComponents->ParentEntity)
+	.FilterNone({ BuiltInComponents->Tags.NeedsUnlink, BuiltInComponents->Tags.Ignored, BuiltInComponents->Tags.Finished })
 	.Iterate_PerEntity(&EntityManager, [&UnresolvedEntities](UObject* Object, FMovieSceneEntityID ParentEntityID)
 	{
 		if (!Object)
