Sequencer: Time controller to play every frame. Defaults to 1 second, but can be set with the cvar Sequencer.SecondsPerFrame.

This is useful in gameplay especially for debugging per frame evaluations

diff --git a/Engine/Source/Editor/EditorStyle/Private/SlateEditorStyle.cpp b/Engine/Source/Editor/EditorStyle/Private/SlateEditorStyle.cpp
--- a/Engine/Source/Editor/EditorStyle/Private/SlateEditorStyle.cpp
+++ b/Engine/Source/Editor/EditorStyle/Private/SlateEditorStyle.cpp
@@ -2577,6 +2577,7 @@ void FSlateEditorStyle::FStyle::SetupSequencerStyles()
 		Set("Sequencer.ClockSource.Audio", new IMAGE_BRUSH("Sequencer/Main_Icons/Icon_ClockSource_Audio_16x", Icon14x14));
 		Set("Sequencer.ClockSource.RelativeTimecode", new IMAGE_BRUSH("Sequencer/Main_Icons/Icon_ClockSource_RelativeTimecode_16x", Icon14x14));
 		Set("Sequencer.ClockSource.Timecode", new IMAGE_BRUSH("Sequencer/Main_Icons/Icon_ClockSource_Timecode_16x", Icon14x14));
+		Set("Sequencer.ClockSource.PlayEveryFrame", new IMAGE_BRUSH("Sequencer/Main_Icons/Icon_ClockSource_PlayEveryFrame_16x", Icon14x14));
 		Set("Sequencer.ClockSource.Custom", new IMAGE_BRUSH("Sequencer/Main_Icons/Icon_ClockSource_Custom_16x", Icon14x14));
 
 
diff --git a/Engine/Source/Editor/EditorStyle/Private/StarshipStyle.cpp b/Engine/Source/Editor/EditorStyle/Private/StarshipStyle.cpp
--- a/Engine/Source/Editor/EditorStyle/Private/StarshipStyle.cpp
+++ b/Engine/Source/Editor/EditorStyle/Private/StarshipStyle.cpp
@@ -2152,6 +2152,7 @@ void FStarshipEditorStyle::FStyle::SetupSequencerStyles()
 		Set("Sequencer.ClockSource.Audio", new IMAGE_BRUSH("Sequencer/Main_Icons/Icon_ClockSource_Audio_16x", Icon14x14));
 		Set("Sequencer.ClockSource.RelativeTimecode", new IMAGE_BRUSH("Sequencer/Main_Icons/Icon_ClockSource_RelativeTimecode_16x", Icon14x14));
 		Set("Sequencer.ClockSource.Timecode", new IMAGE_BRUSH("Sequencer/Main_Icons/Icon_ClockSource_Timecode_16x", Icon14x14));
+		Set("Sequencer.ClockSource.PlayEveryFrame", new IMAGE_BRUSH("Sequencer/Main_Icons/Icon_ClockSource_PlayEveryFrame_16x", Icon14x14));
 		Set("Sequencer.ClockSource.Custom", new IMAGE_BRUSH("Sequencer/Main_Icons/Icon_ClockSource_Custom_16x", Icon14x14));
 
 		Set("Sequencer.BreadcrumbText", FTextBlockStyle(NormalText).SetFont(FStyleFonts::Get().NormalBold));
diff --git a/Engine/Source/Editor/Sequencer/Private/SSequencerPlayRateCombo.cpp b/Engine/Source/Editor/Sequencer/Private/SSequencerPlayRateCombo.cpp
--- a/Engine/Source/Editor/Sequencer/Private/SSequencerPlayRateCombo.cpp
+++ b/Engine/Source/Editor/Sequencer/Private/SSequencerPlayRateCombo.cpp
@@ -568,6 +568,8 @@ const FSlateBrush* SSequencerPlayRateCombo::GetClockSourceImage() const
 					return FEditorStyle::GetBrush("Sequencer.ClockSource.RelativeTimecode");
 				case EUpdateClockSource::Timecode:
 					return FEditorStyle::GetBrush("Sequencer.ClockSource.Timecode");
+				case EUpdateClockSource::PlayEveryFrame:
+					return FEditorStyle::GetBrush("Sequencer.ClockSource.PlayEveryFrame");
 				case EUpdateClockSource::Custom:
 					return FEditorStyle::GetBrush("Sequencer.ClockSource.Custom");
 				default:
diff --git a/Engine/Source/Editor/Sequencer/Private/Sequencer.cpp b/Engine/Source/Editor/Sequencer/Private/Sequencer.cpp
--- a/Engine/Source/Editor/Sequencer/Private/Sequencer.cpp
+++ b/Engine/Source/Editor/Sequencer/Private/Sequencer.cpp
@@ -758,7 +758,7 @@ void FSequencer::Tick(float InDeltaTime)
 
 	if (PlaybackState == EMovieScenePlayerStatus::Playing)
 	{
-		FFrameTime NewGlobalTime = TimeController->RequestCurrentTime(GlobalTime, PlaybackSpeed * Dilation);
+		FFrameTime NewGlobalTime = TimeController->RequestCurrentTime(GlobalTime, PlaybackSpeed * Dilation, GetFocusedDisplayRate());
 
 		// Put the time into local space
  		SetLocalTimeLooped(NewGlobalTime * RootToLocalTransform);
@@ -13724,6 +13724,7 @@ void FSequencer::ResetTimeController()
 	case EUpdateClockSource::Platform: TimeController = MakeShared<FMovieSceneTimeController_PlatformClock>();      break;
 	case EUpdateClockSource::RelativeTimecode: TimeController = MakeShared<FMovieSceneTimeController_RelativeTimecodeClock>();      break;
 	case EUpdateClockSource::Timecode: TimeController = MakeShared<FMovieSceneTimeController_TimecodeClock>();      break;
+	case EUpdateClockSource::PlayEveryFrame: TimeController = MakeShared<FMovieSceneTimeController_PlayEveryFrame>();      break;
 	case EUpdateClockSource::Custom:   TimeController = MovieScene->MakeCustomTimeController(GetPlaybackContext()); break;
 	default:                           TimeController = MakeShared<FMovieSceneTimeController_Tick>();               break;
 	}
diff --git a/Engine/Source/Runtime/MovieScene/Private/MovieSceneSequencePlayer.cpp b/Engine/Source/Runtime/MovieScene/Private/MovieSceneSequencePlayer.cpp
--- a/Engine/Source/Runtime/MovieScene/Private/MovieSceneSequencePlayer.cpp
+++ b/Engine/Source/Runtime/MovieScene/Private/MovieSceneSequencePlayer.cpp
@@ -662,6 +662,11 @@ FFrameTime UMovieSceneSequencePlayer::GetLastValidTime() const
 	}
 }
 
+FFrameRate UMovieSceneSequencePlayer::GetDisplayRate() const
+{
+	return Sequence && Sequence->GetMovieScene() ? Sequence->GetMovieScene()->GetDisplayRate() : FFrameRate();
+}
+
 bool UMovieSceneSequencePlayer::ShouldStopOrLoop(FFrameTime NewPosition) const
 {
 	bool bShouldStopOrLoop = false;
@@ -817,6 +822,7 @@ void UMovieSceneSequencePlayer::Initialize(UMovieSceneSequence* InSequence, cons
 		case EUpdateClockSource::Platform: TimeController = MakeShared<FMovieSceneTimeController_PlatformClock>(); break;
 		case EUpdateClockSource::RelativeTimecode: TimeController = MakeShared<FMovieSceneTimeController_RelativeTimecodeClock>(); break;
 		case EUpdateClockSource::Timecode: TimeController = MakeShared<FMovieSceneTimeController_TimecodeClock>(); break;
+		case EUpdateClockSource::PlayEveryFrame: TimeController = MakeShared<FMovieSceneTimeController_PlayEveryFrame>(); break;
 		default:                           TimeController = MakeShared<FMovieSceneTimeController_Tick>();          break;
 		}
 
@@ -867,7 +873,7 @@ void UMovieSceneSequencePlayer::Update(const float DeltaSeconds)
 		{
 			check(!bIsEvaluating);
 
-			FFrameTime NewTime = TimeController->RequestCurrentTime(GetCurrentTime(), PlayRate);
+			FFrameTime NewTime = TimeController->RequestCurrentTime(GetCurrentTime(), PlayRate, GetDisplayRate());
 			UpdateTimeCursorPosition(NewTime, EUpdatePositionMethod::Play);
 		}
 
diff --git a/Engine/Source/Runtime/MovieScene/Private/MovieSceneTimeController.cpp b/Engine/Source/Runtime/MovieScene/Private/MovieSceneTimeController.cpp
--- a/Engine/Source/Runtime/MovieScene/Private/MovieSceneTimeController.cpp
+++ b/Engine/Source/Runtime/MovieScene/Private/MovieSceneTimeController.cpp
@@ -13,6 +13,12 @@ static TAutoConsoleVariable<int32> CVarSequencerRelativeTimecodeSmoothing(
 	TEXT("If nonzero, accumulate with platform time since when the timecodes were equal."),
 	ECVF_Default);
 
+static TAutoConsoleVariable<float> CVarSequencerSecondsPerFrame(
+	TEXT("Sequencer.SecondsPerFrame"),
+	1.f,
+	TEXT("Seconds per frame to wait when in play every frame mode."),
+	ECVF_Default);
+
 void FMovieSceneTimeController::Tick(float DeltaSeconds, float InPlayRate)
 {
 	OnTick(DeltaSeconds, InPlayRate);
@@ -56,8 +62,9 @@ void FMovieSceneTimeController::StopPlaying(const FQualifiedFrameTime& InStopTim
 	PlaybackStartTime.Reset();
 }
 
-FFrameTime FMovieSceneTimeController::RequestCurrentTime(const FQualifiedFrameTime& InCurrentTime, float InPlayRate)
+FFrameTime FMovieSceneTimeController::RequestCurrentTime(const FQualifiedFrameTime& InCurrentTime, float InPlayRate, FFrameRate InDisplayRate)
 {
+	DisplayRate = InDisplayRate;
 	return OnRequestCurrentTime(InCurrentTime, InPlayRate);
 }
 
@@ -204,4 +211,24 @@ FFrameTime FMovieSceneTimeController_Tick::OnRequestCurrentTime(const FQualified
 		FFrameTime StartTime = StartTimeIfPlaying->ConvertTo(InCurrentTime.Rate);
 		return StartTime + CurrentOffsetSeconds * InCurrentTime.Rate;
 	}
+}
+
+void FMovieSceneTimeController_PlayEveryFrame::OnStartPlaying(const FQualifiedFrameTime& InStartTime)
+{
+	PreviousPlatformTime = FPlatformTime::Seconds();
+	CurrentTime = InStartTime.Time;
+}
+
+FFrameTime FMovieSceneTimeController_PlayEveryFrame::OnRequestCurrentTime(const FQualifiedFrameTime& InCurrentTime, float InPlayRate)
+{
+	double TimeDiff = FPlatformTime::Seconds() - PreviousPlatformTime;
+	if (TimeDiff > CVarSequencerSecondsPerFrame->GetFloat())
+	{
+		// Update current time by moving onto the next frame
+		FFrameRate CurrentDisplayRate = GetDisplayRate();
+		FFrameTime OneDisplayFrame = CurrentDisplayRate.AsFrameTime(CurrentDisplayRate.AsInterval());
+		CurrentTime = CurrentTime + ConvertFrameTime(OneDisplayFrame, CurrentDisplayRate, InCurrentTime.Rate);
+		PreviousPlatformTime = FPlatformTime::Seconds();
+	}
+	return CurrentTime;
 }
\ No newline at end of file
diff --git a/Engine/Source/Runtime/MovieScene/Public/MovieSceneFwd.h b/Engine/Source/Runtime/MovieScene/Public/MovieSceneFwd.h
--- a/Engine/Source/Runtime/MovieScene/Public/MovieSceneFwd.h
+++ b/Engine/Source/Runtime/MovieScene/Public/MovieSceneFwd.h
@@ -55,6 +55,9 @@ enum class EUpdateClockSource : uint8
 	/** Use current timecode provider for timing. Does not honor world or actor pause state. */
 	Timecode,
 
+	/** Debugging Tool: Hold on each whole frame for a Sequencer.SecondsPerFrame many wall-clock seconds before advancing to the next one. Does not honor world or actor pause state or time dilation and audio will be out of sync. */
+	PlayEveryFrame,
+
 	/** Custom clock source created and defined externally. */
 	Custom,
 };
diff --git a/Engine/Source/Runtime/MovieScene/Public/MovieSceneSequencePlayer.h b/Engine/Source/Runtime/MovieScene/Public/MovieSceneSequencePlayer.h
--- a/Engine/Source/Runtime/MovieScene/Public/MovieSceneSequencePlayer.h
+++ b/Engine/Source/Runtime/MovieScene/Public/MovieSceneSequencePlayer.h
@@ -568,6 +568,8 @@ protected:
 
 	FFrameTime GetLastValidTime() const;
 
+	FFrameRate GetDisplayRate() const;
+
 	bool NeedsQueueLatentAction() const;
 	void QueueLatentAction(FMovieSceneSequenceLatentActionDelegate Delegate);
 	void RunLatentActions();
diff --git a/Engine/Source/Runtime/MovieScene/Public/MovieSceneTimeController.h b/Engine/Source/Runtime/MovieScene/Public/MovieSceneTimeController.h
--- a/Engine/Source/Runtime/MovieScene/Public/MovieSceneTimeController.h
+++ b/Engine/Source/Runtime/MovieScene/Public/MovieSceneTimeController.h
@@ -41,7 +41,7 @@ public:
 	 * @param InCurrentTime    The current time of the sequence
 	 * @param InPlayRate       The current play rate of the sequence, multiplied by any world actor settings global dilation
 	 */
-	FFrameTime RequestCurrentTime(const FQualifiedFrameTime& InCurrentTime, float InPlayRate);
+	FFrameTime RequestCurrentTime(const FQualifiedFrameTime& InCurrentTime, float InPlayRate, FFrameRate InDisplayRate);
 
 	/**
 	 * Called when the status of the owning IMovieScenePlayer has changed
@@ -67,9 +67,16 @@ protected:
 		return PlaybackStartTime;
 	}
 
+	FFrameRate GetDisplayRate() const
+	{
+		return DisplayRate;
+	}
+
 private:
 
 	TOptional<FQualifiedFrameTime> PlaybackStartTime;
+
+	FFrameRate DisplayRate;
 };
 
 /**
@@ -162,4 +169,24 @@ private:
 	double CurrentOffsetSeconds;
 };
 
+/**
+ * A timing manager that plays every display frame for a certain number of seconds
+ */
+struct MOVIESCENE_API FMovieSceneTimeController_PlayEveryFrame : FMovieSceneTimeController
+{
+	FMovieSceneTimeController_PlayEveryFrame()
+		: PreviousPlatformTime(0.f)
+	{}
+
+protected:
+
+	virtual void OnStartPlaying(const FQualifiedFrameTime& InStartTime) override;
+	virtual FFrameTime OnRequestCurrentTime(const FQualifiedFrameTime& InCurrentTime, float InPlayRate) override;
+
+private:
+	double PreviousPlatformTime;
+	FFrameTime CurrentTime;
+};
+
+
 
