Sequencer: Change log for mismatched possessed object class to Error

diff --git a/Engine/Source/Runtime/MovieScene/Private/Evaluation/MovieSceneEvaluationState.cpp b/Engine/Source/Runtime/MovieScene/Private/Evaluation/MovieSceneEvaluationState.cpp
--- a/Engine/Source/Runtime/MovieScene/Private/Evaluation/MovieSceneEvaluationState.cpp
+++ b/Engine/Source/Runtime/MovieScene/Private/Evaluation/MovieSceneEvaluationState.cpp
@@ -354,7 +354,7 @@ void FMovieSceneObjectCache::UpdateBindings(const FGuid& InGuid, IMovieScenePlay
 
 						if (!Object->GetClass()->IsChildOf(Possessable->GetPossessedObjectClass()))
 						{
-							UE_LOG(LogMovieScene, Warning, TEXT("Mismatch in %s with posssessed object class for: %s, Expected: %s, but was: %s"), *GetNameSafe(Sequence), *Possessable->GetName(), *GetNameSafe(Possessable->GetPossessedObjectClass()), *GetNameSafe(Object->GetClass()));
+							UE_LOG(LogMovieScene, Error, TEXT("Mismatch in %s with posssessed object class for: %s, Expected: %s, but was: %s"), *GetNameSafe(Sequence), *Possessable->GetName(), *GetNameSafe(Possessable->GetPossessedObjectClass()), *GetNameSafe(Object->GetClass()));
 						}
 					}
 				}
@@ -370,7 +370,7 @@ void FMovieSceneObjectCache::UpdateBindings(const FGuid& InGuid, IMovieScenePlay
 
 					if (!Object->GetClass()->IsChildOf(Possessable->GetPossessedObjectClass()))
 					{
-						UE_LOG(LogMovieScene, Warning, TEXT("Mismatch in %s with posssessed object class for: %s, Expected: %s, but was: %s"), *GetNameSafe(Sequence), *Possessable->GetName(), *GetNameSafe(Possessable->GetPossessedObjectClass()), *GetNameSafe(Object->GetClass()));
+						UE_LOG(LogMovieScene, Error, TEXT("Mismatch in %s with posssessed object class for: %s, Expected: %s, but was: %s"), *GetNameSafe(Sequence), *Possessable->GetName(), *GetNameSafe(Possessable->GetPossessedObjectClass()), *GetNameSafe(Object->GetClass()));
 					}
 				}
 			}
