MoviePipeline: When iterating through the sequence hierarchy to build the tree, skip sections that were generated only because they were pre/post roll for an actual shot.

diff --git a/Engine/Source/Runtime/MovieScene/Private/MovieSceneSequenceVisitor.cpp b/Engine/Source/Runtime/MovieScene/Private/MovieSceneSequenceVisitor.cpp
--- a/Engine/Source/Runtime/MovieScene/Private/MovieSceneSequenceVisitor.cpp
+++ b/Engine/Source/Runtime/MovieScene/Private/MovieSceneSequenceVisitor.cpp
@@ -34,6 +34,11 @@ void VisitSubTrackImpl(const FSequenceVisitParams& InParams, UMovieSceneSubTrack
 
 	for (const FMovieSceneTrackEvaluationFieldEntry& Entry : SubTrack->GetEvaluationField().Entries)
 	{
+		if (EnumHasAnyFlags(Entry.Flags, ESectionEvaluationFlags::PreRoll) || EnumHasAnyFlags(Entry.Flags, ESectionEvaluationFlags::PostRoll))
+		{
+			continue;
+		}
+
 		UMovieSceneSubSection* SubSection  = Cast<UMovieSceneSubSection>(Entry.Section);
 	
 		if (SubSection && SubTrack->IsRowEvalDisabled(SubSection->GetRowIndex()) && !InParams.bVisitDisabledSubSequences)
