Sequencer: make sure FActuatorTokenStackPtr inline size is big enough to avoid temporary memory allocations.

diff --git a/Engine/Source/Runtime/MovieScene/Public/Evaluation/Blending/MovieSceneAccumulatedBlendState.h b/Engine/Source/Runtime/MovieScene/Public/Evaluation/Blending/MovieSceneAccumulatedBlendState.h
--- a/Engine/Source/Runtime/MovieScene/Public/Evaluation/Blending/MovieSceneAccumulatedBlendState.h
+++ b/Engine/Source/Runtime/MovieScene/Public/Evaluation/Blending/MovieSceneAccumulatedBlendState.h
@@ -13,7 +13,7 @@
 /**
  * Container struct that houses all tokens of a single data type that are to be applied using a specific actuator
  */
-typedef TInlineValue<IBlendableTokenStack, 32> FActuatorTokenStackPtr;
+typedef TInlineValue<IBlendableTokenStack, 64> FActuatorTokenStackPtr;
 
 struct FMovieSceneBlendingKey
 {
