Sequencer: Remove spawnables and possessables from arrays. Remove BP events from event keys when muting an event track.

diff --git a/Engine/Source/Runtime/MovieScene/Public/MovieScene.h b/Engine/Source/Runtime/MovieScene/Public/MovieScene.h
--- a/Engine/Source/Runtime/MovieScene/Public/MovieScene.h
+++ b/Engine/Source/Runtime/MovieScene/Public/MovieScene.h
@@ -1154,6 +1154,7 @@ protected:
 private:
 
 #if WITH_EDITOR
+	void OptimizeForCook();
 	void RemoveNullTracks();
 #endif
 
