Sequencer: Changed PossessedObjectClass to editor only and changed it from a TObjectPtr<UClass> to a TSoftClassPtr<UObject>. Note that it's not a TSoftClassPtr<UClass> because Get() will test IsChildOf(UClass)

diff --git a/Engine/Source/Runtime/MovieScene/Private/Evaluation/MovieSceneEvaluationState.cpp b/Engine/Source/Runtime/MovieScene/Private/Evaluation/MovieSceneEvaluationState.cpp
--- a/Engine/Source/Runtime/MovieScene/Private/Evaluation/MovieSceneEvaluationState.cpp
+++ b/Engine/Source/Runtime/MovieScene/Private/Evaluation/MovieSceneEvaluationState.cpp
@@ -352,10 +352,12 @@ void FMovieSceneObjectCache::UpdateBindings(const FGuid& InGuid, IMovieScenePlay
 					{
 						Bindings->Objects.Add(Object);
 
+#if WITH_EDITORONLY_DATA
 						if (!Object->GetClass()->IsChildOf(Possessable->GetPossessedObjectClass()))
 						{
 							UE_LOG(LogMovieScene, Error, TEXT("Mismatch in %s with posssessed object class for: %s, Expected: %s, but was: %s"), *GetNameSafe(Sequence), *Possessable->GetName(), *GetNameSafe(Possessable->GetPossessedObjectClass()), *GetNameSafe(Object->GetClass()));
 						}
+#endif
 					}
 				}
 			}
@@ -368,10 +370,12 @@ void FMovieSceneObjectCache::UpdateBindings(const FGuid& InGuid, IMovieScenePlay
 				{
 					Bindings->Objects.Add(Object);
 
+#if WITH_EDITORONLY_DATA
 					if (!Object->GetClass()->IsChildOf(Possessable->GetPossessedObjectClass()))
 					{
 						UE_LOG(LogMovieScene, Error, TEXT("Mismatch in %s with posssessed object class for: %s, Expected: %s, but was: %s"), *GetNameSafe(Sequence), *Possessable->GetName(), *GetNameSafe(Possessable->GetPossessedObjectClass()), *GetNameSafe(Object->GetClass()));
 					}
+#endif
 				}
 			}
 		}
