Sequencer: Move instantiation serial number from the ECS runner into the linker.

This prevents extra updates from occuring when there are multiple runners, because after one updates another one might incorrectly think instantiation changes have occured and need more updates.

Generally speaking, any "true" state (i.e. state that carries outside of a single update) needs to be avoided in the ECS runner.

diff --git a/Engine/Source/Runtime/MovieScene/Private/EntitySystem/MovieSceneEntitySystemLinker.cpp b/Engine/Source/Runtime/MovieScene/Private/EntitySystem/MovieSceneEntitySystemLinker.cpp
--- a/Engine/Source/Runtime/MovieScene/Private/EntitySystem/MovieSceneEntitySystemLinker.cpp
+++ b/Engine/Source/Runtime/MovieScene/Private/EntitySystem/MovieSceneEntitySystemLinker.cpp
@@ -47,6 +47,7 @@ UMovieSceneEntitySystemLinker::UMovieSceneEntitySystemLinker(const FObjectInitia
 	using namespace UE::MovieScene;
 
 	LastSystemLinkVersion = 0;
+	LastInstantiationVersion = 0;
 	AutoLinkMode = EAutoLinkRelevantSystems::Enabled;
 	SystemContext = EEntitySystemContext::Runtime;
 
@@ -202,6 +203,11 @@ void UMovieSceneEntitySystemLinker::TagInvalidBoundObjects()
 	}
 }
 
+bool UMovieSceneEntitySystemLinker::HasStructureChangedSinceLastRun() const
+{
+	return EntityManager.HasStructureChangedSince(LastInstantiationVersion);
+}
+
 bool UMovieSceneEntitySystemLinker::StartEvaluation(FMovieSceneEntitySystemRunner& InRunner)
 {
 	if (ActiveRunners.Num() == 0 || ActiveRunners.Last().bIsReentrancyAllowed)
@@ -225,6 +231,13 @@ FMovieSceneEntitySystemRunner* UMovieSceneEntitySystemLinker::GetActiveRunner()
 	return nullptr;
 }
 
+void UMovieSceneEntitySystemLinker::PostInstantation(FMovieSceneEntitySystemRunner& InRunner)
+{
+	LastInstantiationVersion = EntityManager.GetSystemSerial();
+
+	GetInstanceRegistry()->PostInstantation();
+}
+
 void UMovieSceneEntitySystemLinker::EndEvaluation(FMovieSceneEntitySystemRunner& InRunner)
 {
 	if (ensureMsgf((ActiveRunners.Num() > 0 && ActiveRunners.Last().Runner == &InRunner),
@@ -265,6 +278,10 @@ void UMovieSceneEntitySystemLinker::CleanGarbage()
 		return;
 	}
 
+	// Clear the instantiation serial to indicate that we probably need to re-run the instantiation systems
+	// the next time a runner gets flushed
+	LastInstantiationVersion = 0;
+
 	// Allow any other system to tag garbage
 	Events.CleanTaggedGarbage.Broadcast(this);
 
diff --git a/Engine/Source/Runtime/MovieScene/Private/EntitySystem/MovieSceneEntitySystemRunner.cpp b/Engine/Source/Runtime/MovieScene/Private/EntitySystem/MovieSceneEntitySystemRunner.cpp
--- a/Engine/Source/Runtime/MovieScene/Private/EntitySystem/MovieSceneEntitySystemRunner.cpp
+++ b/Engine/Source/Runtime/MovieScene/Private/EntitySystem/MovieSceneEntitySystemRunner.cpp
@@ -49,8 +49,6 @@ void FMovieSceneEntitySystemRunner::AttachToLinker(UMovieSceneEntitySystemLinker
 	}
 
 	WeakLinker = InLinker;
-	LastInstantiationVersion = 0;
-	InLinker->Events.CleanTaggedGarbage.AddRaw(this, &FMovieSceneEntitySystemRunner::OnLinkerGarbageCleaned);
 	InLinker->Events.AbandonLinker.AddRaw(this, &FMovieSceneEntitySystemRunner::OnLinkerAbandon);
 }
 
@@ -103,9 +101,9 @@ bool FMovieSceneEntitySystemRunner::HasQueuedUpdates() const
 	{
 		return true;
 	}
-	if (UMovieSceneEntitySystemLinker* Linker = GetLinker())
+	if (const UMovieSceneEntitySystemLinker* Linker = GetLinker())
 	{
-		return Linker->EntityManager.HasStructureChangedSince(LastInstantiationVersion);
+		return Linker->HasStructureChangedSinceLastRun();
 	}
 	return false;
 }
@@ -166,7 +164,7 @@ void FMovieSceneEntitySystemRunner::Flush()
 
 	// We specifically only check whether the entity manager has changed since the last instantation once
 	// to ensure that we are not vulnerable to infinite loops where components are added/removed in post-evaluation
-	bool bStructureHadChanged = Linker->EntityManager.HasStructureChangedSince(LastInstantiationVersion);
+	bool bStructureHadChanged = Linker->HasStructureChangedSinceLastRun();
 
 	// Start flushing the update queue... keep flushing as long as we have work to do.
 	while (UpdateQueue.Num() > 0 ||
@@ -338,7 +336,7 @@ void FMovieSceneEntitySystemRunner::GameThread_SpawnPhase()
 		DissectedUpdates.RemoveAt(0, Index);
 	}
 
-	const bool bInstantiationDirty = Linker->EntityManager.HasStructureChangedSince(LastInstantiationVersion) || InstanceRegistry->HasInvalidatedBindings();
+	const bool bInstantiationDirty = Linker->HasStructureChangedSinceLastRun() || InstanceRegistry->HasInvalidatedBindings();
 
 	FGraphEventArray AllTasks;
 
@@ -440,8 +438,7 @@ void FMovieSceneEntitySystemRunner::GameThread_PostInstantiation()
 		UMovieSceneEntitySystemLinker* Linker = GetLinker();
 		check(Linker);
 
-		LastInstantiationVersion = Linker->EntityManager.GetSystemSerial();
-		Linker->GetInstanceRegistry()->PostInstantation();
+		Linker->PostInstantation(*this);
 
 		FEntityManager& EntityManager = Linker->EntityManager;
 		FBuiltInComponentTypes* BuiltInComponentTypes = FBuiltInComponentTypes::Get();
@@ -613,7 +610,7 @@ void FMovieSceneEntitySystemRunner::FinishInstance(FInstanceHandle InInstanceHan
 	}
 
 	InstanceRegistry->MutateInstance(InInstanceHandle).Finish(Linker);
-	if (Linker->EntityManager.HasStructureChangedSince(LastInstantiationVersion))
+	if (Linker->HasStructureChangedSinceLastRun())
 	{
 		MarkForUpdate(InInstanceHandle);
 		Flush();
@@ -629,17 +626,10 @@ void FMovieSceneEntitySystemRunner::MarkForUpdate(FInstanceHandle InInstanceHand
 	CurrentInstances.AddUnique(InInstanceHandle);
 }
 
-void FMovieSceneEntitySystemRunner::OnLinkerGarbageCleaned(UMovieSceneEntitySystemLinker* InLinker)
-{
-	check(WeakLinker.Get() == InLinker);
-	LastInstantiationVersion = 0;
-}
-
 void FMovieSceneEntitySystemRunner::OnLinkerAbandon(UMovieSceneEntitySystemLinker* InLinker)
 {
 	if (ensure(InLinker))
 	{
-		InLinker->Events.CleanTaggedGarbage.RemoveAll(this);
 		InLinker->Events.AbandonLinker.RemoveAll(this);
 	}
 	WeakLinker.Reset();
diff --git a/Engine/Source/Runtime/MovieScene/Public/EntitySystem/MovieSceneEntitySystemLinker.h b/Engine/Source/Runtime/MovieScene/Public/EntitySystem/MovieSceneEntitySystemLinker.h
--- a/Engine/Source/Runtime/MovieScene/Public/EntitySystem/MovieSceneEntitySystemLinker.h
+++ b/Engine/Source/Runtime/MovieScene/Public/EntitySystem/MovieSceneEntitySystemLinker.h
@@ -234,11 +234,14 @@ public:
 	void LinkRelevantSystems();
 	void AutoLinkRelevantSystems();
 
+	bool HasStructureChangedSinceLastRun() const;
+
 	void InvalidateObjectBinding(const FGuid& ObjectBindingID, FInstanceHandle InstanceHandle);
 	void CleanupInvalidBoundObjects();
 
 	bool StartEvaluation(FMovieSceneEntitySystemRunner& InRunner);
 	FMovieSceneEntitySystemRunner* GetActiveRunner() const;
+	void PostInstantation(FMovieSceneEntitySystemRunner& InRunner);
 	void EndEvaluation(FMovieSceneEntitySystemRunner& InRunner);
 
 private:
@@ -287,6 +290,7 @@ public:
 private:
 
 	uint64 LastSystemLinkVersion;
+	uint64 LastInstantiationVersion;
 
 	TWeakPtr<bool> GlobalStateCaptureToken;
 
diff --git a/Engine/Source/Runtime/MovieScene/Public/EntitySystem/MovieSceneEntitySystemRunner.h b/Engine/Source/Runtime/MovieScene/Public/EntitySystem/MovieSceneEntitySystemRunner.h
--- a/Engine/Source/Runtime/MovieScene/Public/EntitySystem/MovieSceneEntitySystemRunner.h
+++ b/Engine/Source/Runtime/MovieScene/Public/EntitySystem/MovieSceneEntitySystemRunner.h
@@ -80,7 +80,6 @@ public:
 
 private:
 
-	void OnLinkerGarbageCleaned(UMovieSceneEntitySystemLinker* Linker);
 	void OnLinkerAbandon(UMovieSceneEntitySystemLinker* Linker);
 
 private:
@@ -124,8 +123,6 @@ private:
 	/** When an update is running, the list of sub-contexts for the requested update */
 	TArray<FDissectedUpdate> DissectedUpdates;
 
-	uint64 LastInstantiationVersion = 0;
-
 	TGraphTask<FNullGraphTask>* CompletionTask;
 
 	bool bCanQueueEventTriggers;
