Sequencer: When running evaluations with no particular sequence to update, update all of them.

This fixes a bug where an evaluation caused by an outdated instantiation serial number in the linker does not correctly run the legacy evaluation templates.

diff --git a/Engine/Source/Runtime/MovieScene/Private/EntitySystem/MovieSceneEntitySystemRunner.cpp b/Engine/Source/Runtime/MovieScene/Private/EntitySystem/MovieSceneEntitySystemRunner.cpp
--- a/Engine/Source/Runtime/MovieScene/Private/EntitySystem/MovieSceneEntitySystemRunner.cpp
+++ b/Engine/Source/Runtime/MovieScene/Private/EntitySystem/MovieSceneEntitySystemRunner.cpp
@@ -293,6 +293,18 @@ void FMovieSceneEntitySystemRunner::GameThread_ProcessQueue()
 		}
 	}
 
+	// If we have no instances marked for update, we are running an evaluation probably because some
+	// structural changes have occured in the entity manager (out of date instantiation serial number
+	// in the linker). So we mark everything for update, so that PreEvaluation/PostEvaluation callbacks
+	// and legacy templates are correctly executed.
+	if (CurrentInstances.Num() == 0)
+	{
+		for (const FSequenceInstance& Instance : InstanceRegistry->GetSparseInstances())
+		{
+			MarkForUpdate(Instance.GetInstanceHandle());
+		}
+	}
+
 	// Let sequence instances do any pre-evaluation work.
 	for (FInstanceHandle UpdatedInstanceHandle : CurrentInstances)
 	{
diff --git a/Engine/Source/Runtime/MovieScene/Public/EntitySystem/MovieSceneSequenceInstance.h b/Engine/Source/Runtime/MovieScene/Public/EntitySystem/MovieSceneSequenceInstance.h
--- a/Engine/Source/Runtime/MovieScene/Public/EntitySystem/MovieSceneSequenceInstance.h
+++ b/Engine/Source/Runtime/MovieScene/Public/EntitySystem/MovieSceneSequenceInstance.h
@@ -134,6 +134,14 @@ public:
 		return Context;
 	}
 
+	/**
+	 * Gets the handle to this instance
+	 */
+	FInstanceHandle GetInstanceHandle() const
+	{
+		return InstanceHandle;
+	}
+
 	/**
 	 * Get a handle to the root instance for this sub sequence instance
 	 *
