Sequencer: FindTrack should return derived classes

diff --git a/Engine/Source/Runtime/MovieScene/Private/MovieScene.cpp b/Engine/Source/Runtime/MovieScene/Private/MovieScene.cpp
--- a/Engine/Source/Runtime/MovieScene/Private/MovieScene.cpp
+++ b/Engine/Source/Runtime/MovieScene/Private/MovieScene.cpp
@@ -877,7 +877,7 @@ UMovieSceneTrack* UMovieScene::FindTrack(TSubclassOf<UMovieSceneTrack> TrackClas
 
 		for (const auto& Track : Binding.GetTracks())
 		{
-			if (TrackClass.GetDefaultObject() == nullptr ||  Track->GetClass() == TrackClass)
+			if (TrackClass.GetDefaultObject() == nullptr || Track->GetClass()->IsChildOf(TrackClass))
 			{
 				if (TrackName == NAME_None || Track->GetTrackName() == TrackName)
 				{
@@ -903,7 +903,7 @@ TArray<UMovieSceneTrack*> UMovieScene::FindTracks(TSubclassOf<UMovieSceneTrack>
 
 		for (const auto& Track : Binding.GetTracks())
 		{
-			if (TrackClass.GetDefaultObject() == nullptr || Track->GetClass() == TrackClass)
+			if (TrackClass.GetDefaultObject() == nullptr || Track->GetClass()->IsChildOf(TrackClass))
 			{
 				if (TrackName == NAME_None || Track->GetTrackName() == TrackName)
 				{
@@ -1000,7 +1000,7 @@ UMovieSceneTrack* UMovieScene::FindMasterTrack( TSubclassOf<UMovieSceneTrack> Tr
 
 	for (const auto Track : MasterTracks)
 	{
-		if( Track->GetClass() == TrackClass )
+		if (Track->GetClass()->IsChildOf(TrackClass))
 		{
 			FoundTrack = Track;
 			break;
