Sequencer: Fix networked play smoothing for reverse playback.

diff --git a/Engine/Source/Runtime/MovieScene/Private/MovieSceneSequencePlayer.cpp b/Engine/Source/Runtime/MovieScene/Private/MovieSceneSequencePlayer.cpp
--- a/Engine/Source/Runtime/MovieScene/Private/MovieSceneSequencePlayer.cpp
+++ b/Engine/Source/Runtime/MovieScene/Private/MovieSceneSequencePlayer.cpp
@@ -1256,6 +1256,7 @@ FFrameTime UMovieSceneSequencePlayer::UpdateServerTimeSamples()
 	// We need to reproject the samples to the current wall-clock time, based on when they were taken
 	const double CurrentWallClock = FPlatformTime::Seconds();
 	const double Lifetime         = CurrentWallClock - float(GSequencerMaxSmoothedNetSyncSampleAge) / 1000.f;
+	const float PlaybackMultiplier = bReversePlayback ? -1.f : 1.f;
 
 	// Cull any old samples that were taken more than GSequencerMaxSmoothedNetSyncSampleAge ms ago by
 	// Finding the index of the first sample younger than this time
@@ -1298,7 +1299,7 @@ FFrameTime UMovieSceneSequencePlayer::UpdateServerTimeSamples()
 	double MeanTime = 0;
 	for (const FServerTimeSample& Sample : ServerTimeSamples)
 	{
-		const double ThisSample = Sample.ServerTime + (CurrentWallClock - Sample.ReceievedTime);
+		const double ThisSample = Sample.ServerTime + (CurrentWallClock - Sample.ReceievedTime) * PlaybackMultiplier;
 		MeanTime += ThisSample;
 	}
 	MeanTime = MeanTime / ServerTimeSamples.Num();
@@ -1307,7 +1308,7 @@ FFrameTime UMovieSceneSequencePlayer::UpdateServerTimeSamples()
 	double StandardDeviation = 0;
 	for (const FServerTimeSample& Sample : ServerTimeSamples)
 	{
-		const double ThisSample = Sample.ServerTime + (CurrentWallClock - Sample.ReceievedTime);
+		const double ThisSample = Sample.ServerTime + (CurrentWallClock - Sample.ReceievedTime) * PlaybackMultiplier;
 		StandardDeviation += FMath::Square(ThisSample - MeanTime);
 	}
 	StandardDeviation = StandardDeviation / ServerTimeSamples.Num();
@@ -1325,7 +1326,7 @@ FFrameTime UMovieSceneSequencePlayer::UpdateServerTimeSamples()
 		// Discard anything outside the standard deviation in the hopes that future samples will converge
 		for (int32 SampleIndex = ServerTimeSamples.Num()-1; SampleIndex >= 0; --SampleIndex)
 		{
-			const double ThisSample = ServerTimeSamples[SampleIndex].ServerTime + (CurrentWallClock - ServerTimeSamples[SampleIndex].ReceievedTime);
+			const double ThisSample = ServerTimeSamples[SampleIndex].ServerTime + (CurrentWallClock - ServerTimeSamples[SampleIndex].ReceievedTime) * PlaybackMultiplier;
 			if (FMath::Abs(ThisSample - MeanTime) > StandardDeviation)
 			{
 				ServerTimeSamples.RemoveAt(SampleIndex, 1, false);
