Sequencer: Check for null sequence

diff --git a/Engine/Source/Runtime/MovieScene/Private/Compilation/MovieSceneCompiledDataManager.cpp b/Engine/Source/Runtime/MovieScene/Private/Compilation/MovieSceneCompiledDataManager.cpp
--- a/Engine/Source/Runtime/MovieScene/Private/Compilation/MovieSceneCompiledDataManager.cpp
+++ b/Engine/Source/Runtime/MovieScene/Private/Compilation/MovieSceneCompiledDataManager.cpp
@@ -612,6 +612,11 @@ void UMovieSceneCompiledDataManager::DestroyTemplate(FMovieSceneCompiledDataID D
 
 bool UMovieSceneCompiledDataManager::IsDirty(const FMovieSceneCompiledDataEntry& Entry) const
 {
+	if (!Entry.GetSequence())
+	{
+		return false;
+	}
+
 	if (Entry.CompiledSignature != Entry.GetSequence()->GetSignature())
 	{
 		return true;
