MovieScene: deprecate UMovieSceneSequenceExtensions::GetTimecodeSource() in favor of GetEarliestTimecodeSource() on UMovieScene and UMovieSceneSequence

UMovieScene is a container that may represent multiple pieces of media as
sections, each of which may have their own distinct timecode source. As a result,
the timecode source associated with the UMovieScene itself doesn't correspond to
any particular piece of media, but rather was just intended as the timecode at which
the movie scene was recorded.

This is the first in a series of changes that aim to remove the TimecodeSource property
of UMovieScene to try to disambiguate it from the similarly named property on
UMovieSceneSections that *does* correspond to the section's media. A subsequent change
will remove the property from UMovieScene, possibly replacing it with one that more
explicitly represents the time at which the movie scene was recorded.

UMovieSceneSequenceExtensions::GetTimecodeSource() previously would return the
TimecodeSource of its UMovieScene. With that property going away, the more interesting
query is likely to look through the media contained by the movie scene and return the
earliest timecode source of any of the scene's sections. A GetEarliestTimecodeSource()
function was added to UMovieScene that does this and the UMovieSceneSequenceExtensions
function was deprecated. A GetEarliestTimecodeSource() function was added to
UMovieSceneSequence as well that simply calls through to its movie scene.

These changes also ensure that the Timecode property of FMovieSceneTimecodeSource
is accessible from Python.

diff --git a/Engine/Plugins/MovieScene/SequencerScripting/Source/SequencerScripting/Private/ExtensionLibraries/MovieSceneSequenceExtensions.cpp b/Engine/Plugins/MovieScene/SequencerScripting/Source/SequencerScripting/Private/ExtensionLibraries/MovieSceneSequenceExtensions.cpp
--- a/Engine/Plugins/MovieScene/SequencerScripting/Source/SequencerScripting/Private/ExtensionLibraries/MovieSceneSequenceExtensions.cpp
+++ b/Engine/Plugins/MovieScene/SequencerScripting/Source/SequencerScripting/Private/ExtensionLibraries/MovieSceneSequenceExtensions.cpp
@@ -442,17 +442,12 @@ EUpdateClockSource UMovieSceneSequenceExtensions::GetClockSource(UMovieSceneSequ
 FTimecode UMovieSceneSequenceExtensions::GetTimecodeSource(UMovieSceneSequence* Sequence)
 {
 	UMovieScene* MovieScene = GetMovieScene(Sequence);
-	if (MovieScene)
-	{
-#if WITH_EDITORONLY_DATA
-		return MovieScene->TimecodeSource.Timecode;
-#endif
-		return FTimecode();
-	}
-	else
+	if (!MovieScene)
 	{
 		return FTimecode();
 	}
+
+	return MovieScene->GetEarliestTimecodeSource().Timecode;
 }
 
 FSequencerBindingProxy UMovieSceneSequenceExtensions::FindBindingByName(UMovieSceneSequence* Sequence, FString Name)
diff --git a/Engine/Plugins/MovieScene/SequencerScripting/Source/SequencerScripting/Public/ExtensionLibraries/MovieSceneSequenceExtensions.h b/Engine/Plugins/MovieScene/SequencerScripting/Source/SequencerScripting/Public/ExtensionLibraries/MovieSceneSequenceExtensions.h
--- a/Engine/Plugins/MovieScene/SequencerScripting/Source/SequencerScripting/Public/ExtensionLibraries/MovieSceneSequenceExtensions.h
+++ b/Engine/Plugins/MovieScene/SequencerScripting/Source/SequencerScripting/Public/ExtensionLibraries/MovieSceneSequenceExtensions.h
@@ -341,7 +341,8 @@ public:
 	 * @param Sequence        The sequence within which to get the timecode source
 	 * @return Timecode source of this sequence
 	 */
-	UFUNCTION(BlueprintCallable, Category = "Sequencer|Sequence", meta = (ScriptMethod))
+	UE_DEPRECATED(5.0, "GetTimecodeSource() is no longer supported for movie scene sequences. Please use GetEarliestTimecodeSource() instead.")
+	UFUNCTION(BlueprintCallable, Category = "Sequencer|Sequence", meta = (ScriptMethod, DeprecatedFunction, DeprecationMessage = "GetTimecodeSource() is no longer supported for movie scene sequences. Please use GetEarliestTimecodeSource() instead."))
 	static FTimecode GetTimecodeSource(UMovieSceneSequence* Sequence);
 
 	/**
diff --git a/Engine/Source/Runtime/MovieScene/Private/MovieScene.cpp b/Engine/Source/Runtime/MovieScene/Private/MovieScene.cpp
--- a/Engine/Source/Runtime/MovieScene/Private/MovieScene.cpp
+++ b/Engine/Source/Runtime/MovieScene/Private/MovieScene.cpp
@@ -4,6 +4,7 @@
 
 #include "MovieSceneTrack.h"
 #include "MovieSceneFolder.h"
+#include "MovieSceneSection.h"
 #include "MovieSceneSequence.h"
 #include "Evaluation/MovieSceneEvaluationCustomVersion.h"
 #include "Compilation/MovieSceneSegmentCompiler.h"
@@ -11,6 +12,7 @@
 #include "Evaluation/IMovieSceneCustomClockSource.h"
 #include "CommonFrameRates.h"
 #include "EntitySystem/IMovieSceneEntityProvider.h"
+#include "Misc/FrameRate.h"
 #include "UObject/ObjectSaveContext.h"
 #include "UObject/UObjectHash.h"
 
@@ -1408,6 +1410,49 @@ TSharedPtr<FMovieSceneTimeController> UMovieScene::MakeCustomTimeController(UObj
 	return MakeShared<FMovieSceneTimeController_Custom>(CustomClockSourcePath, PlaybackContext);
 }
 
+FMovieSceneTimecodeSource UMovieScene::GetEarliestTimecodeSource() const
+{
+	FMovieSceneTimecodeSource EarliestTimecodeSource;
+
+#if WITH_EDITORONLY_DATA
+
+	const TArray<UMovieSceneSection*> MovieSceneSections = GetAllSections();
+
+	// Find the first non-default timecode source.
+	const FMovieSceneTimecodeSource DefaultTimecodeSource;
+	int32 Index = 0;
+	while (Index < MovieSceneSections.Num() && EarliestTimecodeSource == DefaultTimecodeSource)
+	{
+		if (MovieSceneSections[Index])
+		{
+			EarliestTimecodeSource = MovieSceneSections[Index]->TimecodeSource;
+		}
+
+		++Index;
+	}
+
+	// Continue searching through the sections where we left off looking for any earlier timecodes.
+	// Any subsequently found default timecode source could be considered earlier.
+	for (; Index < MovieSceneSections.Num(); ++Index)
+	{
+		if (!MovieSceneSections[Index])
+		{
+			continue;
+		}
+
+		const FMovieSceneTimecodeSource SectionTimecodeSource = MovieSceneSections[Index]->TimecodeSource;
+		const FFrameRate ComparisonFrameRate;
+		if (SectionTimecodeSource.Timecode.ToFrameNumber(ComparisonFrameRate) < EarliestTimecodeSource.Timecode.ToFrameNumber(ComparisonFrameRate))
+		{
+			EarliestTimecodeSource = SectionTimecodeSource;
+		}
+	}
+
+#endif
+
+	return EarliestTimecodeSource;
+}
+
 void UMovieScene::SetMarkedFrame(int32 InMarkIndex, FFrameNumber InFrameNumber)
 {
 	if (InMarkIndex < 0 && InMarkIndex > MarkedFrames.Num()-1)
diff --git a/Engine/Source/Runtime/MovieScene/Private/MovieSceneSequence.cpp b/Engine/Source/Runtime/MovieScene/Private/MovieSceneSequence.cpp
--- a/Engine/Source/Runtime/MovieScene/Private/MovieSceneSequence.cpp
+++ b/Engine/Source/Runtime/MovieScene/Private/MovieSceneSequence.cpp
@@ -341,3 +341,14 @@ const TArray<FMovieSceneObjectBindingID>& UMovieSceneSequence::FindBindingsByTag
 	static TArray<FMovieSceneObjectBindingID> EmptyBindings;
 	return EmptyBindings;
 }
+
+FMovieSceneTimecodeSource UMovieSceneSequence::GetEarliestTimecodeSource() const
+{
+	const UMovieScene* MovieScene = GetMovieScene();
+	if (!MovieScene)
+	{
+		return FMovieSceneTimecodeSource();
+	}
+
+	return MovieScene->GetEarliestTimecodeSource();
+}
diff --git a/Engine/Source/Runtime/MovieScene/Public/MovieScene.h b/Engine/Source/Runtime/MovieScene/Public/MovieScene.h
--- a/Engine/Source/Runtime/MovieScene/Public/MovieScene.h
+++ b/Engine/Source/Runtime/MovieScene/Public/MovieScene.h
@@ -43,7 +43,7 @@ struct FMovieSceneExpansionState
 	bool bExpanded;
 };
 
-USTRUCT()
+USTRUCT(BlueprintType)
 struct FMovieSceneTimecodeSource
 {
 	GENERATED_BODY()
@@ -68,7 +68,7 @@ struct FMovieSceneTimecodeSource
 public:
 
 	/** The global timecode at which this target is based (ie. the timecode at the beginning of the movie scene section when it was recorded) */
-	UPROPERTY(EditAnywhere, Category="Timecode")
+	UPROPERTY(EditAnywhere, BlueprintReadOnly, Category="Timecode")
 	FTimecode Timecode;
 };
 
@@ -854,6 +854,11 @@ public:
 		CustomClockSourcePath = InNewClockSource;
 	}
 
+	/**
+	 * Get the earliest timecode source out of all of the movie scene sections contained within this movie scene.
+	 */
+	FMovieSceneTimecodeSource GetEarliestTimecodeSource() const;
+
 	/*
 	* Replace an existing binding with another 
 	*/
diff --git a/Engine/Source/Runtime/MovieScene/Public/MovieSceneSequence.h b/Engine/Source/Runtime/MovieScene/Public/MovieSceneSequence.h
--- a/Engine/Source/Runtime/MovieScene/Public/MovieSceneSequence.h
+++ b/Engine/Source/Runtime/MovieScene/Public/MovieSceneSequence.h
@@ -17,6 +17,7 @@ class UMovieSceneEntitySystemLinker;
 
 struct FMovieScenePossessable;
 struct FMovieSceneObjectCache;
+struct FMovieSceneTimecodeSource;
 
 enum class ETrackSupport
 {
@@ -216,6 +217,12 @@ public:
 	UFUNCTION(BlueprintCallable, Category = "Sequencer|Sequence")
 	MOVIESCENE_API const TArray<FMovieSceneObjectBindingID>& FindBindingsByTag(FName InBindingName) const;
 
+	/**
+	 * Get the earliest timecode source out of all of the movie scene sections contained within this sequence's movie scene.
+	 */
+	UFUNCTION(BlueprintCallable, Category = "Sequencer|Sequence")
+	MOVIESCENE_API FMovieSceneTimecodeSource GetEarliestTimecodeSource() const;
+
 public:
 
 	MOVIESCENE_API virtual void PostLoad() override;
