Sequencer: Check for any sub tracks not just master sub tracks when looking for circular dependencies. This fixes a crash when trying to add the same sequence to a template sequence track.

diff --git a/Engine/Source/Runtime/MovieScene/Private/Sections/MovieSceneSubSection.cpp b/Engine/Source/Runtime/MovieScene/Private/Sections/MovieSceneSubSection.cpp
--- a/Engine/Source/Runtime/MovieScene/Private/Sections/MovieSceneSubSection.cpp
+++ b/Engine/Source/Runtime/MovieScene/Private/Sections/MovieSceneSubSection.cpp
@@ -221,20 +221,38 @@ void UMovieSceneSubSection::PostEditChangeProperty(FPropertyChangedEvent& Proper
 		{
 			if (UMovieSceneSequence* CurrentSequence = TrackOuter->GetTypedOuter<UMovieSceneSequence>())
 			{
+				TArray<UMovieSceneSubTrack*> SubTracks;
+
 				for (UMovieSceneTrack* MasterTrack : SubSequenceMovieScene->GetMasterTracks())
 				{
 					if (UMovieSceneSubTrack* SubTrack = Cast<UMovieSceneSubTrack>(MasterTrack))
 					{
-						if ( SubTrack->ContainsSequence(*CurrentSequence, true))
-						{
-							UE_LOG(LogMovieScene, Error, TEXT("Invalid level sequence %s. It is already contained by: %s."), *SubSequence->GetDisplayName().ToString(), *CurrentSequence->GetDisplayName().ToString());
+						SubTracks.Add(SubTrack);
+					}
+				}
 
-							// Restore to the previous sub sequence because there was a circular dependency
-							SubSequence = PreviousSubSequence;
-							break;
+				for (const FMovieSceneBinding& Binding : SubSequenceMovieScene->GetBindings())
+				{
+					for (UMovieSceneTrack* Track : SubSequenceMovieScene->FindTracks(UMovieSceneSubTrack::StaticClass(), Binding.GetObjectGuid()))
+					{
+						if (UMovieSceneSubTrack* SubTrack = Cast<UMovieSceneSubTrack>(Track))
+						{
+							SubTracks.Add(SubTrack);
 						}
 					}
 				}
+
+				for (UMovieSceneSubTrack* SubTrack : SubTracks)
+				{
+					if ( SubTrack->ContainsSequence(*CurrentSequence, true))
+					{
+						UE_LOG(LogMovieScene, Error, TEXT("Invalid level sequence %s. It is already contained by: %s."), *SubSequence->GetDisplayName().ToString(), *CurrentSequence->GetDisplayName().ToString());
+
+						// Restore to the previous sub sequence because there was a circular dependency
+						SubSequence = PreviousSubSequence;
+						break;
+					}
+				}
 			}
 		}
 
