a bunch of cosmetic changes

diff --git a/Engine/Source/Runtime/CoreUObject/Public/UObject/FastReferenceCollector.h b/Engine/Source/Runtime/CoreUObject/Public/UObject/FastReferenceCollector.h
--- a/Engine/Source/Runtime/CoreUObject/Public/UObject/FastReferenceCollector.h
+++ b/Engine/Source/Runtime/CoreUObject/Public/UObject/FastReferenceCollector.h
@@ -496,7 +496,7 @@ private:
 		{
 			return ESubsequentsMode::TrackSubsequents;
 		}
-		void DoTask(ENamedThreads::Type CurrentThread, FGraphEventRef& MyCompletionGraphEvent)
+		void DoTask(ENamedThreads::Type CurrentThread, const FGraphEventRef& MyCompletionGraphEvent)
 		{
 			Owner->ProcessObjectArray(*ArrayStruct, MyCompletionGraphEvent);
 		}
diff --git a/Engine/Source/Runtime/Engine/Classes/Engine/EngineBaseTypes.h b/Engine/Source/Runtime/Engine/Classes/Engine/EngineBaseTypes.h
--- a/Engine/Source/Runtime/Engine/Classes/Engine/EngineBaseTypes.h
+++ b/Engine/Source/Runtime/Engine/Classes/Engine/EngineBaseTypes.h
@@ -266,7 +266,7 @@ private:
 		int32 TickQueuedGFrameCounter;
 
 		/** Pointer to the task, only used during setup. This is often stale. **/
-		void* TaskPointer;
+		FBaseGraphTask* TaskPointer;
 
 		/** The next function in the cooling down list for ticks with an interval*/
 		FTickFunction* Next;
diff --git a/Engine/Source/Runtime/MovieScene/Private/EntitySystem/MovieSceneEntitySystemRunner.cpp b/Engine/Source/Runtime/MovieScene/Private/EntitySystem/MovieSceneEntitySystemRunner.cpp
--- a/Engine/Source/Runtime/MovieScene/Private/EntitySystem/MovieSceneEntitySystemRunner.cpp
+++ b/Engine/Source/Runtime/MovieScene/Private/EntitySystem/MovieSceneEntitySystemRunner.cpp
@@ -208,7 +208,6 @@ void FMovieSceneEntitySystemRunner::DoFlushUpdateQueueOnce()
 
 	FTaskGraphInterface::Get().WaitUntilTaskCompletes(CompletionTask->GetCompletionEvent(), ENamedThreads::GameThread_Local);
 
-	// TODO: delete?
 	CompletionTask = nullptr;
 
 	// Now run the post-evaluation logic, which contains stuff we don't want to run from inside a task
