Sequencer: Add the ability to bind a possessable to a spawnable object

FMovieScenePossessable now has an FMovieSceneObjectBindingID which can be set to bind to a spawnable object. Calling sites should invoke BindSpawanbleObject which will establish the FMovieSceneObjectBindingID which is resolved in FMovieSceneObjectCache::UpdateBindings

diff --git a/Engine/Source/Editor/Sequencer/Private/Sequencer.cpp b/Engine/Source/Editor/Sequencer/Private/Sequencer.cpp
--- a/Engine/Source/Editor/Sequencer/Private/Sequencer.cpp
+++ b/Engine/Source/Editor/Sequencer/Private/Sequencer.cpp
@@ -1138,8 +1138,11 @@ FGuid FSequencer::CreateBinding(UObject& InObject, const FString& InName)
 		}
 	}
 
-	OwnerSequence->BindPossessableObject(PossessableGuid, InObject, BindingContext);
-	
+	if (!OwnerMovieScene->FindPossessable(PossessableGuid)->BindSpawnableObject(GetFocusedTemplateID(), &InObject, this))
+	{
+		OwnerSequence->BindPossessableObject(PossessableGuid, InObject, BindingContext);
+	}
+		
 	// Broadcast if a parent actor was added as a result of adding this object
 	if (ParentActorAdded && ParentGuid.IsValid())
 	{
@@ -7537,7 +7540,10 @@ void FSequencer::AddActorsToBinding(FGuid InObjectBinding, const TArray<AActor*>
 				}
 
 				ActorToAdd->Modify();
-				OwnerSequence->BindPossessableObject(InObjectBinding, *ActorToAdd, GetPlaybackContext());
+				if (!OwnerMovieScene->FindPossessable(InObjectBinding)->BindSpawnableObject(GetFocusedTemplateID(), ActorToAdd, this))
+				{
+					OwnerSequence->BindPossessableObject(InObjectBinding, *ActorToAdd, GetPlaybackContext());
+				}
 				++NumObjectsAdded;
 			}
 			else
@@ -10114,7 +10120,11 @@ TArray<FGuid> FSequencer::ExpandMultiplePossessableBindings(FGuid PossessableGui
 				NewPossessable->SetParent(ParentGuid);
 			}
 
-			Sequence->BindPossessableObject(NewPossessableGuid, *FoundObject, BindingContext);
+			if (!NewPossessable->BindSpawnableObject(GetFocusedTemplateID(), FoundObject, this))
+			{
+				Sequence->BindPossessableObject(NewPossessableGuid, *FoundObject, BindingContext);
+			}
+
 			NewPossessableGuids.Add(NewPossessableGuid);
 
 			// Create copies of the tracks
diff --git a/Engine/Source/Editor/Sequencer/Private/SequencerUtilities.cpp b/Engine/Source/Editor/Sequencer/Private/SequencerUtilities.cpp
--- a/Engine/Source/Editor/Sequencer/Private/SequencerUtilities.cpp
+++ b/Engine/Source/Editor/Sequencer/Private/SequencerUtilities.cpp
@@ -388,7 +388,10 @@ FGuid FSequencerUtilities::DoAssignActor(ISequencer * InSequencerPtr, AActor* Ac
 		// Add this object
 		NewPossessableActor = FMovieScenePossessable(NewActorLabel, Actor->GetClass());
 		NewGuid = NewPossessableActor.GetGuid();
-		OwnerSequence->BindPossessableObject(NewPossessableActor.GetGuid(), *Actor, InSequencerPtr->GetPlaybackContext());
+		if (!NewPossessableActor.BindSpawnableObject(InSequencerPtr->GetFocusedTemplateID(), Actor, InSequencerPtr))
+		{
+			OwnerSequence->BindPossessableObject(NewPossessableActor.GetGuid(), *Actor, InSequencerPtr->GetPlaybackContext());
+		}
 
 		// Defer replacing this object until the components have been updated
 	}
diff --git a/Engine/Source/Runtime/MovieScene/Private/Evaluation/MovieSceneEvaluationState.cpp b/Engine/Source/Runtime/MovieScene/Private/Evaluation/MovieSceneEvaluationState.cpp
--- a/Engine/Source/Runtime/MovieScene/Private/Evaluation/MovieSceneEvaluationState.cpp
+++ b/Engine/Source/Runtime/MovieScene/Private/Evaluation/MovieSceneEvaluationState.cpp
@@ -346,7 +346,22 @@ void FMovieSceneObjectCache::UpdateBindings(const FGuid& InGuid, IMovieScenePlay
 					}
 
 					TArray<UObject*, TInlineAllocator<1>> FoundObjects;
-					Player.ResolveBoundObjects(InGuid, SequenceID, *Sequence, ResolutionContext, FoundObjects);
+					
+					if (Possessable->GetSpawnableObjectBindingID().IsValid())
+					{
+						for (TWeakObjectPtr<> BoundObject : Possessable->GetSpawnableObjectBindingID().ResolveBoundObjects(SequenceID, Player))
+						{
+							if (BoundObject.IsValid())
+							{
+								FoundObjects.Add(BoundObject.Get());
+							}
+						}
+					}
+					else
+					{
+						Player.ResolveBoundObjects(InGuid, SequenceID, *Sequence, ResolutionContext, FoundObjects);
+					}
+					
 					Bindings = BoundObjects.Find(InGuid);
 					for (UObject* Object : FoundObjects)
 					{
@@ -357,7 +372,22 @@ void FMovieSceneObjectCache::UpdateBindings(const FGuid& InGuid, IMovieScenePlay
 			else
 			{
 				TArray<UObject*, TInlineAllocator<1>> FoundObjects;
-				Player.ResolveBoundObjects(InGuid, SequenceID, *Sequence, ResolutionContext, FoundObjects);
+
+				if (Possessable->GetSpawnableObjectBindingID().IsValid())
+				{
+					for (TWeakObjectPtr<> BoundObject : Possessable->GetSpawnableObjectBindingID().ResolveBoundObjects(SequenceID, Player))
+					{
+						if (BoundObject.IsValid())
+						{
+							FoundObjects.Add(BoundObject.Get());
+						}
+					}				
+				}
+				else
+				{
+					Player.ResolveBoundObjects(InGuid, SequenceID, *Sequence, ResolutionContext, FoundObjects);
+				}
+				
 				Bindings = BoundObjects.Find(InGuid);
 				for (UObject* Object : FoundObjects)
 				{
diff --git a/Engine/Source/Runtime/MovieScene/Private/MovieScenePossessable.cpp b/Engine/Source/Runtime/MovieScene/Private/MovieScenePossessable.cpp
new file mode 100644
--- /dev/null
+++ b/Engine/Source/Runtime/MovieScene/Private/MovieScenePossessable.cpp
@@ -0,0 +1,21 @@
+// Copyright Epic Games, Inc. All Rights Reserved.
+
+#include "MovieScenePossessable.h"
+#include "MovieSceneSpawnableAnnotation.h"
+#include "MovieSceneSequence.h"
+#include "MovieSceneSequenceID.h"
+#include "MovieScene.h"
+#include "IMovieScenePlayer.h"
+
+bool FMovieScenePossessable::BindSpawnableObject(FMovieSceneSequenceID SequenceID, UObject* Object, IMovieScenePlayer* Player)
+{
+	TOptional<FMovieSceneSpawnableAnnotation> Spawnable = FMovieSceneSpawnableAnnotation::Find(Object);
+	if (Spawnable.IsSet())
+	{
+		// Check whether the spawnable is underneath the current sequence, if so, we can remap it to a local sequence ID
+		SetSpawnableObjectBindingID(UE::MovieScene::FRelativeObjectBindingID(SequenceID, Spawnable->SequenceID, Spawnable->ObjectBindingID, *Player));
+		return true;
+	}
+
+	return false;
+}
diff --git a/Engine/Source/Runtime/MovieScene/Public/MovieScenePossessable.h b/Engine/Source/Runtime/MovieScene/Public/MovieScenePossessable.h
--- a/Engine/Source/Runtime/MovieScene/Public/MovieScenePossessable.h
+++ b/Engine/Source/Runtime/MovieScene/Public/MovieScenePossessable.h
@@ -5,6 +5,9 @@
 #include "CoreMinimal.h"
 #include "UObject/ObjectMacros.h"
 #include "Misc/Guid.h"
+#include "IMovieScenePlayer.h"
+#include "MovieSceneObjectBindingID.h"
+#include "MovieSceneSequenceID.h"
 #include "MovieScenePossessable.generated.h"
 
 /**
@@ -129,6 +132,21 @@ public:
 	UPROPERTY(EditAnywhere, AdvancedDisplay, Category=Actor)
 	TArray<FName> Tags;
 
+	/* Get the optional binding id for binding to a spawnable */
+	const FMovieSceneObjectBindingID& GetSpawnableObjectBindingID() const 
+	{
+		return SpawnableObjectBindingID;
+	}
+
+	/* Set the optional binding id for binding to a spawnable */
+	void SetSpawnableObjectBindingID(const FMovieSceneObjectBindingID& InSpawnableObjectBindingID)
+	{
+		SpawnableObjectBindingID = InSpawnableObjectBindingID;
+	}
+
+	/* Bind the potential spawnable object to this possessable by setting the ObjectBindingID */
+	MOVIESCENE_API bool BindSpawnableObject(FMovieSceneSequenceID SequenceID, UObject* Object, IMovieScenePlayer* Player);
+
 private:
 
 	/** Unique identifier of the possessable object. */
@@ -153,4 +171,8 @@ private:
 	/** GUID relating to this possessable's parent, if applicable. */
 	UPROPERTY()
 	FGuid ParentGuid;
+
+	/** Optional object binding ID if this possessable possesses a spawnable */
+	UPROPERTY()
+	FMovieSceneObjectBindingID SpawnableObjectBindingID;
 };
