MovieScene: migrate FMovieSceneTimecodeSource from MovieScene.h to MovieSceneSection.h

Now that the TimecodeSource property exists only on UMovieSceneSections rather
than on UMovieScenes, this moves the definition of the struct to the header file for
UMovieSceneSection.

diff --git a/Engine/Source/Editor/SequenceRecorderSections/Public/MovieScene3DAttachSectionRecorder.h b/Engine/Source/Editor/SequenceRecorderSections/Public/MovieScene3DAttachSectionRecorder.h
--- a/Engine/Source/Editor/SequenceRecorderSections/Public/MovieScene3DAttachSectionRecorder.h
+++ b/Engine/Source/Editor/SequenceRecorderSections/Public/MovieScene3DAttachSectionRecorder.h
@@ -6,6 +6,7 @@
 #include "Misc/Guid.h"
 #include "GameFramework/Actor.h"
 #include "MovieScene.h"
+#include "MovieSceneSection.h"
 #include "IMovieSceneSectionRecorder.h"
 #include "IMovieSceneSectionRecorderFactory.h"
 
diff --git a/Engine/Source/Runtime/MovieScene/Public/MovieScene.h b/Engine/Source/Runtime/MovieScene/Public/MovieScene.h
--- a/Engine/Source/Runtime/MovieScene/Public/MovieScene.h
+++ b/Engine/Source/Runtime/MovieScene/Public/MovieScene.h
@@ -5,7 +5,6 @@
 #include "CoreMinimal.h"
 #include "UObject/ObjectMacros.h"
 #include "Misc/Guid.h"
-#include "Misc/Timecode.h"
 #include "Templates/SubclassOf.h"
 #include "Templates/Casts.h"
 #include "MovieSceneFwd.h"
@@ -25,6 +24,7 @@ class UMovieSceneFolder;
 class UMovieSceneSection;
 class UMovieSceneTrack;
 struct FMovieSceneChannelMetaData;
+struct FMovieSceneTimecodeSource;
 
 //delegates for use when some data in the MovieScene changes, WIP right now, hopefully will replace delegates on ISequencer
 //and be used for moving towards a true MVC system
@@ -43,35 +43,6 @@ struct FMovieSceneExpansionState
 	bool bExpanded;
 };
 
-USTRUCT(BlueprintType)
-struct FMovieSceneTimecodeSource
-{
-	GENERATED_BODY()
-
-	FMovieSceneTimecodeSource(FTimecode InTimecode)
-		: Timecode(InTimecode)
-	{}
-
-	FMovieSceneTimecodeSource()
-		: Timecode(FTimecode())
-	{}
-
-	FORCEINLINE bool operator==(const FMovieSceneTimecodeSource& Other) const
-	{
-		return Timecode == Other.Timecode;
-	}
-	FORCEINLINE bool operator!=(const FMovieSceneTimecodeSource& Other) const
-	{
-		return Timecode != Other.Timecode;
-	}
-
-public:
-
-	/** The global timecode at which this target is based (ie. the timecode at the beginning of the movie scene section when it was recorded) */
-	UPROPERTY(EditAnywhere, BlueprintReadOnly, Category="Timecode")
-	FTimecode Timecode;
-};
-
 USTRUCT(BlueprintType)
 struct FMovieSceneMarkedFrame
 {
diff --git a/Engine/Source/Runtime/MovieScene/Public/MovieSceneSection.h b/Engine/Source/Runtime/MovieScene/Public/MovieSceneSection.h
--- a/Engine/Source/Runtime/MovieScene/Public/MovieSceneSection.h
+++ b/Engine/Source/Runtime/MovieScene/Public/MovieSceneSection.h
@@ -5,6 +5,8 @@
 #include "CoreMinimal.h"
 #include "Containers/ContainersFwd.h"
 #include "Misc/FrameTime.h"
+#include "Misc/QualifiedFrameTime.h"
+#include "Misc/Timecode.h"
 #include "UObject/ObjectMacros.h"
 #include "MovieSceneFwd.h"
 #include "KeyParams.h"
@@ -15,7 +17,6 @@
 #include "Generators/MovieSceneEasingFunction.h"
 #include "Evaluation/MovieSceneSequenceHierarchy.h"
 #include "MovieSceneFrameMigration.h"
-#include "Misc/QualifiedFrameTime.h"
 #include "Evaluation/MovieSceneEvaluationCustomVersion.h"
 #include "EntitySystem/MovieSceneEntityBuilder.h"
 #include "MovieSceneSection.generated.h"
@@ -150,6 +151,35 @@ public:
 #endif
 };
 
+USTRUCT(BlueprintType)
+struct FMovieSceneTimecodeSource
+{
+	GENERATED_BODY()
+
+	FMovieSceneTimecodeSource(FTimecode InTimecode)
+		: Timecode(InTimecode)
+	{}
+
+	FMovieSceneTimecodeSource()
+		: Timecode(FTimecode())
+	{}
+
+	FORCEINLINE bool operator==(const FMovieSceneTimecodeSource& Other) const
+	{
+		return Timecode == Other.Timecode;
+	}
+	FORCEINLINE bool operator!=(const FMovieSceneTimecodeSource& Other) const
+	{
+		return Timecode != Other.Timecode;
+	}
+
+public:
+
+	/** The global timecode at which this target is based (ie. the timecode at the beginning of the movie scene section when it was recorded) */
+	UPROPERTY(EditAnywhere, BlueprintReadOnly, Category="Timecode")
+	FTimecode Timecode;
+};
+
 /**
  * Base class for movie scene sections
  */
