Sequencer: Pre-Animated state changes to support track instance inputs

  - Track Instance Inputs can now be used as capture sources to capture pre-animated state
  - PreAnimatedState extension now simply exists as a member of UMovieSceneEntitySystem. Previously the existence of the extension was used to determine whether any pre-animated state logic was required, but in practice this created more problems than it solved and led to some convoluted lifetime management code.
  - Added a templated group manager type from which the object group manager can inherit. This simplifies the introduction of new group managers
  - Removed the requirement that all pre-animated state must exist as part of a valid group. This forced state from 'master' tracks to be arbitrarily grouped together which put an unnecessary burden on the storage classes.

diff --git a/Engine/Plugins/Media/MediaCompositing/Source/MediaCompositing/Private/MovieScene/MovieSceneMediaPlayerPropertyTemplate.cpp b/Engine/Plugins/Media/MediaCompositing/Source/MediaCompositing/Private/MovieScene/MovieSceneMediaPlayerPropertyTemplate.cpp
--- a/Engine/Plugins/Media/MediaCompositing/Source/MediaCompositing/Private/MovieScene/MovieSceneMediaPlayerPropertyTemplate.cpp
+++ b/Engine/Plugins/Media/MediaCompositing/Source/MediaCompositing/Private/MovieScene/MovieSceneMediaPlayerPropertyTemplate.cpp
@@ -8,6 +8,7 @@
 
 #include "MovieSceneMediaPlayerPropertySection.h"
 #include "MovieSceneMediaPlayerPropertyTrack.h"
+#include "Evaluation/PreAnimatedState/MovieScenePreAnimatedCaptureSource.h"
 
 
 struct FMediaPlayerPreAnimatedState : IMovieScenePreAnimatedToken
diff --git a/Engine/Plugins/VirtualProduction/LiveLinkCamera/Source/LiveLinkCameraRecording/Private/MovieSceneLiveLinkCameraControllerSection.cpp b/Engine/Plugins/VirtualProduction/LiveLinkCamera/Source/LiveLinkCameraRecording/Private/MovieSceneLiveLinkCameraControllerSection.cpp
--- a/Engine/Plugins/VirtualProduction/LiveLinkCamera/Source/LiveLinkCameraRecording/Private/MovieSceneLiveLinkCameraControllerSection.cpp
+++ b/Engine/Plugins/VirtualProduction/LiveLinkCamera/Source/LiveLinkCameraRecording/Private/MovieSceneLiveLinkCameraControllerSection.cpp
@@ -9,6 +9,7 @@
 #include "LiveLinkControllerBase.h"
 #include "MovieSceneExecutionToken.h"
 #include "Roles/LiveLinkCameraRole.h"
+#include "Evaluation/PreAnimatedState/MovieScenePreAnimatedCaptureSource.h"
 
 struct FPreAnimatedLensFileTokenProducer : IMovieScenePreAnimatedTokenProducer
 {
diff --git a/Engine/Source/Runtime/MovieScene/Private/EntitySystem/MovieSceneEntitySystemLinker.cpp b/Engine/Source/Runtime/MovieScene/Private/EntitySystem/MovieSceneEntitySystemLinker.cpp
--- a/Engine/Source/Runtime/MovieScene/Private/EntitySystem/MovieSceneEntitySystemLinker.cpp
+++ b/Engine/Source/Runtime/MovieScene/Private/EntitySystem/MovieSceneEntitySystemLinker.cpp
@@ -53,6 +53,8 @@ UMovieSceneEntitySystemLinker::UMovieSceneEntitySystemLinker(const FObjectInitia
 
 	if (!HasAnyFlags(RF_ClassDefaultObject))
 	{
+		PreAnimatedState.Initialize(this);
+
 		FCoreUObjectDelegates::GetPostGarbageCollect().AddUObject(this, &UMovieSceneEntitySystemLinker::HandlePostGarbageCollection);
 
 		EntityManager.SetDebugName(GetName() + TEXT("[Entity Manager]"));
diff --git a/Engine/Source/Runtime/MovieScene/Private/EntitySystem/MovieSceneEvaluationHookSystem.cpp b/Engine/Source/Runtime/MovieScene/Private/EntitySystem/MovieSceneEvaluationHookSystem.cpp
--- a/Engine/Source/Runtime/MovieScene/Private/EntitySystem/MovieSceneEvaluationHookSystem.cpp
+++ b/Engine/Source/Runtime/MovieScene/Private/EntitySystem/MovieSceneEvaluationHookSystem.cpp
@@ -5,6 +5,7 @@
 #include "EntitySystem/MovieSceneEntitySystemRunner.h"
 #include "EntitySystem/MovieSceneSpawnablesSystem.h"
 #include "EntitySystem/MovieSceneEntitySystemTask.h"
+#include "Evaluation/PreAnimatedState/MovieScenePreAnimatedCaptureSource.h"
 
 #include "Evaluation/MovieSceneEvaluationTemplateInstance.h"
 #include "IMovieScenePlayer.h"
diff --git a/Engine/Source/Runtime/MovieScene/Private/EntitySystem/MovieScenePreAnimatedStateSystem.cpp b/Engine/Source/Runtime/MovieScene/Private/EntitySystem/MovieScenePreAnimatedStateSystem.cpp
--- a/Engine/Source/Runtime/MovieScene/Private/EntitySystem/MovieScenePreAnimatedStateSystem.cpp
+++ b/Engine/Source/Runtime/MovieScene/Private/EntitySystem/MovieScenePreAnimatedStateSystem.cpp
@@ -12,59 +12,6 @@
 
 
 
-namespace UE
-{
-namespace MovieScene
-{
-
-TSharedPtr<FPreAnimatedStateExtension> FPreAnimatedStateExtensionReference::Get() const
-{
-	return WeakPreAnimatedStateExtension.Pin();
-}
-
-TSharedPtr<FPreAnimatedStateExtension> FPreAnimatedStateExtensionReference::Update(UMovieSceneEntitySystemLinker* Linker)
-{
-	FPreAnimatedStateExtension* ExistingExtension = Linker->FindExtension<FPreAnimatedStateExtension>();
-
-	const bool bHasRestoreStateEntities = Linker->EntityManager.ContainsComponent(FBuiltInComponentTypes::Get()->Tags.RestoreState);
-	const bool bIsCapturingGlobalState  = ExistingExtension && ExistingExtension->IsCapturingGlobalState();
-
-	if (bIsCapturingGlobalState && WeakPreAnimatedStateExtension.Pin() == nullptr)
-	{
-		WeakPreAnimatedStateExtension = ExistingExtension->AsShared();
-	}
-
-	if (bHasRestoreStateEntities && !PreAnimatedStateExtensionRef)
-	{
-		if (ExistingExtension)
-		{
-			PreAnimatedStateExtensionRef = ExistingExtension->AsShared();
-			WeakPreAnimatedStateExtension = PreAnimatedStateExtensionRef;
-		}
-		else
-		{
-			// FPreAnimatedStateExtension automatically adds itself to the linker
-			PreAnimatedStateExtensionRef = MakeShared<FPreAnimatedStateExtension>(Linker);
-			WeakPreAnimatedStateExtension = PreAnimatedStateExtensionRef;
-		}
-	}
-	else if (!bHasRestoreStateEntities && PreAnimatedStateExtensionRef)
-	{
-		PreAnimatedStateExtensionRef = nullptr;
-		if (!bIsCapturingGlobalState)
-		{
-			WeakPreAnimatedStateExtension = nullptr;
-		}
-	}
-
-	return WeakPreAnimatedStateExtension.Pin();
-}
-
-
-} // namespace MovieScene
-} // namespace UE
-
-
 UMovieSceneCachePreAnimatedStateSystem::UMovieSceneCachePreAnimatedStateSystem(const FObjectInitializer& ObjInit)
 	: Super(ObjInit)
 {
@@ -79,34 +26,26 @@ UMovieSceneCachePreAnimatedStateSystem::UMovieSceneCachePreAnimatedStateSystem(c
 
 bool UMovieSceneCachePreAnimatedStateSystem::IsRelevantImpl(UMovieSceneEntitySystemLinker* InLinker) const
 {
+	using namespace UE::MovieScene;
+
 	// This function can be called on the CDO and instances, so care is taken to do the right thing
-	if (HasAnyFlags(RF_ClassDefaultObject))
-	{
-		return UE::MovieScene::FPreAnimatedStateExtensionReference(InLinker).Get().IsValid();
-	}
-	return PreAnimatedStateRef.Get().IsValid();
+	const bool bHasRestoreStateEntities = InLinker->EntityManager.ContainsComponent(FBuiltInComponentTypes::Get()->Tags.RestoreState);
+	const bool bIsCapturingGlobalState  = InLinker->PreAnimatedState.IsCapturingGlobalState();
+
+	return bHasRestoreStateEntities || bIsCapturingGlobalState;
 }
 
 void UMovieSceneCachePreAnimatedStateSystem::OnLink()
 {
-	using namespace UE::MovieScene;
-
-	PreAnimatedStateRef.Update(Linker);
 }
 
 void UMovieSceneCachePreAnimatedStateSystem::OnUnlink()
 {
-	PreAnimatedStateRef = UE::MovieScene::FPreAnimatedStateExtensionReference();
 }
 
 void UMovieSceneCachePreAnimatedStateSystem::OnRun(FSystemTaskPrerequisites& InPrerequisites, FSystemSubsequentTasks& Subsequents)
 {
 	using namespace UE::MovieScene;
-	TSharedPtr<UE::MovieScene::FPreAnimatedStateExtension> Extension = PreAnimatedStateRef.Update(Linker);
-	if (!Extension)
-	{
-		return;
-	}
 
 	TArray<IMovieScenePreAnimatedStateSystemInterface*, TInlineAllocator<16>> Interfaces;
 	auto ForEachSystem = [&Interfaces](UMovieSceneEntitySystem* InSystem)
@@ -120,7 +59,7 @@ void UMovieSceneCachePreAnimatedStateSystem::OnRun(FSystemTaskPrerequisites& InP
 	Linker->SystemGraph.IteratePhase(ESystemPhase::Instantiation, ForEachSystem);
 	Linker->SystemGraph.IteratePhase(ESystemPhase::Evaluation, ForEachSystem);
 
-	IMovieScenePreAnimatedStateSystemInterface::FPreAnimationParameters Params{ &InPrerequisites, &Subsequents, Extension.Get() };
+	IMovieScenePreAnimatedStateSystemInterface::FPreAnimationParameters Params{ &InPrerequisites, &Subsequents, &Linker->PreAnimatedState };
 	for (IMovieScenePreAnimatedStateSystemInterface* Interface : Interfaces)
 	{
 		Interface->SavePreAnimatedState(Params);
@@ -149,16 +88,6 @@ bool UMovieSceneRestorePreAnimatedStateSystem::IsRelevantImpl(UMovieSceneEntityS
 void UMovieSceneRestorePreAnimatedStateSystem::OnLink()
 {
 	using namespace UE::MovieScene;
-	FPreAnimatedStateExtension* ExistingExtension = Linker->FindExtension<FPreAnimatedStateExtension>();
-	if (ExistingExtension)
-	{
-		PreAnimatedStateRef = ExistingExtension->AsShared();
-	}
-	else
-	{
-		// FPreAnimatedStateExtension automatically adds itself to the linker
-		PreAnimatedStateRef = MakeShared<FPreAnimatedStateExtension>(Linker);
-	}
 
 	UMovieSceneCachePreAnimatedStateSystem* CacheSystem = Linker->LinkSystem<UMovieSceneCachePreAnimatedStateSystem>();
 	Linker->SystemGraph.AddReference(this, CacheSystem);
@@ -166,7 +95,6 @@ void UMovieSceneRestorePreAnimatedStateSystem::OnLink()
 
 void UMovieSceneRestorePreAnimatedStateSystem::OnUnlink()
 {
-	PreAnimatedStateRef = nullptr;
 }
 
 void UMovieSceneRestorePreAnimatedStateSystem::OnRun(FSystemTaskPrerequisites& InPrerequisites, FSystemSubsequentTasks& Subsequents)
@@ -187,7 +115,7 @@ void UMovieSceneRestorePreAnimatedStateSystem::OnRun(FSystemTaskPrerequisites& I
 	Linker->SystemGraph.IteratePhase(ESystemPhase::Instantiation, ForEachSystem);
 	Linker->SystemGraph.IteratePhase(ESystemPhase::Evaluation, ForEachSystem);
 
-	IMovieScenePreAnimatedStateSystemInterface::FPreAnimationParameters Params{ &InPrerequisites, &Subsequents, PreAnimatedStateRef.Get() };
+	IMovieScenePreAnimatedStateSystemInterface::FPreAnimationParameters Params{ &InPrerequisites, &Subsequents, &Linker->PreAnimatedState };
 
 	// Iterate backwards restoring stale state
 	for (int32 Index = Interfaces.Num()-1; Index >= 0; --Index)
@@ -195,7 +123,7 @@ void UMovieSceneRestorePreAnimatedStateSystem::OnRun(FSystemTaskPrerequisites& I
 		Interfaces[Index]->RestorePreAnimatedState(Params);
 	}
 
-	FPreAnimatedEntityCaptureSource* EntityMetaData = PreAnimatedStateRef->GetEntityMetaData();
+	FPreAnimatedEntityCaptureSource* EntityMetaData = Linker->PreAnimatedState.GetEntityMetaData();
 	if (EntityMetaData)
 	{
 		auto CleanupExpiredObjects = [EntityMetaData](FMovieSceneEntityID EntityID)
diff --git a/Engine/Source/Runtime/MovieScene/Private/EntitySystem/MovieSceneSequenceInstance.cpp b/Engine/Source/Runtime/MovieScene/Private/EntitySystem/MovieSceneSequenceInstance.cpp
--- a/Engine/Source/Runtime/MovieScene/Private/EntitySystem/MovieSceneSequenceInstance.cpp
+++ b/Engine/Source/Runtime/MovieScene/Private/EntitySystem/MovieSceneSequenceInstance.cpp
@@ -158,32 +158,6 @@ void FSequenceInstance::DissectContext(UMovieSceneEntitySystemLinker* Linker, co
 	SequenceUpdater->DissectContext(Linker, Player, InContext, OutDissections);
 }
 
-void FSequenceInstance::EnableGlobalPreAnimatedStateCapture(UMovieSceneEntitySystemLinker* Linker)
-{
-	if (ensure(Linker) && !GlobalPreAnimatedState)
-	{
-		FPreAnimatedStateExtension* Existing = Linker->FindExtension<FPreAnimatedStateExtension>();
-		if (Existing)
-		{
-			GlobalPreAnimatedState = Existing->AsShared();
-		}
-		else
-		{
-			// FPreAnimatedStateExtension automatically adds itself to the linker
-			GlobalPreAnimatedState = MakeShared<FPreAnimatedStateExtension>(Linker);
-		}
-
-		++GlobalPreAnimatedState->NumRequestsForGlobalState;
-
-		GetPlayer()->PreAnimatedState.OnEnableGlobalCapture(GlobalPreAnimatedState);
-	}
-}
-
-bool FSequenceInstance::IsCapturingGlobalPreAnimatedState() const
-{
-	return GlobalPreAnimatedState != nullptr;
-}
-
 void FSequenceInstance::Start(UMovieSceneEntitySystemLinker* Linker, const FMovieSceneContext& InContext)
 {
 	check(SequenceID == MovieSceneSequenceID::Root);
@@ -250,9 +224,9 @@ void FSequenceInstance::Finish(UMovieSceneEntitySystemLinker* Linker)
 		SpawnRegister.ForgetExternallyOwnedSpawnedObjects(Player->State, *Player);
 		SpawnRegister.CleanUp(*Player);
 
-		if (GlobalPreAnimatedState)
+		if (Player->PreAnimatedState.IsCapturingGlobalPreAnimatedState())
 		{
-			GlobalPreAnimatedState->RestoreGlobalState(FRestoreStateParams{ Linker, RootInstanceHandle });
+			Linker->PreAnimatedState.RestoreGlobalState(FRestoreStateParams{ Linker, RootInstanceHandle });
 		}
 	}
 }
@@ -316,14 +290,6 @@ void FSequenceInstance::DestroyImmediately(UMovieSceneEntitySystemLinker* Linker
 	{
 		SequenceUpdater->Destroy(Linker);
 	}
-
-	if (GlobalPreAnimatedState)
-	{
-		--GlobalPreAnimatedState->NumRequestsForGlobalState;
-		GlobalPreAnimatedState = nullptr;
-
-		GetPlayer()->PreAnimatedState.OnDisableGlobalCapture();
-	}
 }
 
 void FSequenceInstance::OverrideRootSequence(UMovieSceneEntitySystemLinker* Linker, FMovieSceneSequenceID NewRootSequenceID)
diff --git a/Engine/Source/Runtime/MovieScene/Private/EntitySystem/TrackInstance/MovieSceneTrackInstance.cpp b/Engine/Source/Runtime/MovieScene/Private/EntitySystem/TrackInstance/MovieSceneTrackInstance.cpp
--- a/Engine/Source/Runtime/MovieScene/Private/EntitySystem/TrackInstance/MovieSceneTrackInstance.cpp
+++ b/Engine/Source/Runtime/MovieScene/Private/EntitySystem/TrackInstance/MovieSceneTrackInstance.cpp
@@ -2,6 +2,8 @@
 
 #include "EntitySystem/TrackInstance/MovieSceneTrackInstance.h"
 #include "EntitySystem/BuiltInComponentTypes.h"
+#include "Evaluation/PreAnimatedState/MovieScenePreAnimatedCaptureSource.h"
+#include "Evaluation/PreAnimatedState/MovieScenePreAnimatedCaptureSources.h"
 #include "Algo/Sort.h"
 
 void UMovieSceneTrackInstance::Initialize(UObject* InAnimatedObject, UMovieSceneEntitySystemLinker* InLinker)
@@ -11,7 +13,7 @@ void UMovieSceneTrackInstance::Initialize(UObject* InAnimatedObject, UMovieScene
 	AnimatedObject = InAnimatedObject;
 	bIsMasterTrackInstance = (InAnimatedObject == nullptr);
 
-	Linker = InLinker;
+	PrivateLinker = InLinker;
 
 	OnInitialize();
 }
@@ -23,22 +25,46 @@ void UMovieSceneTrackInstance::Animate()
 
 void UMovieSceneTrackInstance::Destroy()
 {
+	using namespace UE::MovieScene;
+
 	if (bIsMasterTrackInstance || !UE::MovieScene::FBuiltInComponentTypes::IsBoundObjectGarbage(AnimatedObject))
 	{
 		OnDestroyed();
 	}
+
+	FPreAnimatedTrackInstanceInputCaptureSources* InputMetaData = PrivateLinker->PreAnimatedState.GetTrackInstanceInputMetaData();
+	if (InputMetaData)
+	{
+		for (const FMovieSceneTrackInstanceInput& Input : Inputs)
+		{
+			InputMetaData->StopTrackingCaptureSource(Input);
+		}
+	}
+
+	FPreAnimatedTrackInstanceCaptureSources* TrackInstanceMetaData = PrivateLinker->PreAnimatedState.GetTrackInstanceMetaData();
+	if (TrackInstanceMetaData)
+	{
+		TrackInstanceMetaData->StopTrackingCaptureSource(this);
+	}
 }
 
 void UMovieSceneTrackInstance::UpdateInputs(TArray<FMovieSceneTrackInstanceInput>&& InNewInputs)
 {
+	using namespace UE::MovieScene;
+
 	Algo::Sort(InNewInputs);
 
 	// Fast path if they are the same
 	if (Inputs == InNewInputs)
 	{
+		// We still call OnBegin/EndUpdateInputs since some of them must have been invalidated for us to get this far
+		OnBeginUpdateInputs();
+		OnEndUpdateInputs();
 		return;
 	}
 
+	FPreAnimatedTrackInstanceInputCaptureSources* InputMetaData = PrivateLinker->PreAnimatedState.GetTrackInstanceInputMetaData();
+
 	// We know they are different in some way - now 
 	OnBeginUpdateInputs();
 
@@ -66,11 +92,17 @@ void UMovieSceneTrackInstance::UpdateInputs(TArray<FMovieSceneTrackInstanceInput
 			{
 				// Out with the old
 				OnInputRemoved(Inputs[OldIndex]);
+				if (InputMetaData)
+				{
+					InputMetaData->StopTrackingCaptureSource(Inputs[OldIndex]);
+				}
 				++OldIndex;
 			}
 			else
 			{
 				// and in with the new
+				FScopedPreAnimatedCaptureSource CaptureSource(PrivateLinker, InNewInputs[NewIndex]);
+
 				OnInputAdded(InNewInputs[NewIndex]);
 				++NewIndex;
 			}
@@ -79,6 +111,10 @@ void UMovieSceneTrackInstance::UpdateInputs(TArray<FMovieSceneTrackInstanceInput
 		{
 			// Out with the old
 			OnInputRemoved(Inputs[OldIndex]);
+			if (InputMetaData)
+			{
+				InputMetaData->StopTrackingCaptureSource(Inputs[OldIndex]);
+			}
 			++OldIndex;
 		}
 		else if (ensure(NewIndex < NewNum))
diff --git a/Engine/Source/Runtime/MovieScene/Private/EntitySystem/TrackInstance/MovieSceneTrackInstanceSystem.cpp b/Engine/Source/Runtime/MovieScene/Private/EntitySystem/TrackInstance/MovieSceneTrackInstanceSystem.cpp
--- a/Engine/Source/Runtime/MovieScene/Private/EntitySystem/TrackInstance/MovieSceneTrackInstanceSystem.cpp
+++ b/Engine/Source/Runtime/MovieScene/Private/EntitySystem/TrackInstance/MovieSceneTrackInstanceSystem.cpp
@@ -11,8 +11,6 @@
 #include "EntitySystem/MovieSceneInstanceRegistry.h"
 #include "EntitySystem/MovieSceneEntitySystemLinker.h"
 #include "EntitySystem/MovieSceneEntityFactoryTemplates.h"
-#include "Evaluation/PreAnimatedState/MovieScenePreAnimatedStateExtension.h"
-#include "Evaluation/PreAnimatedState/MovieScenePreAnimatedCaptureSources.h"
 #include "MovieSceneSection.h"
 
 DECLARE_CYCLE_STAT(TEXT("Generic Track Instances"), MovieSceneEval_GenericTrackInstances, STATGROUP_MovieSceneECS);
@@ -198,9 +196,6 @@ void UMovieSceneTrackInstanceInstantiator::OnRun(FSystemTaskPrerequisites& InPre
 		return;
 	}
 
-	FPreAnimatedStateExtension*              PreAnimatedStateExtension = Linker->FindExtension<FPreAnimatedStateExtension>();
-	FPreAnimatedTrackInstanceCaptureSources* TrackInstanceMetaData     = PreAnimatedStateExtension ? PreAnimatedStateExtension->GetTrackInstanceMetaData() : nullptr;
-
 	// Gather all the inputs for any invalidated output indices
 	TSortedMap<int32, TArray<FMovieSceneTrackInstanceInput> > NewInputs;
 	{
@@ -236,11 +231,6 @@ void UMovieSceneTrackInstanceInstantiator::OnRun(FSystemTaskPrerequisites& InPre
 		FMovieSceneTrackInstanceEntry& Entry = TrackInstances[DestroyIndex];
 		Entry.TrackInstance->Destroy();
 
-		if (TrackInstanceMetaData)
-		{
-			TrackInstanceMetaData->StopTrackingCaptureSource(Entry.TrackInstance);
-		}
-
 		// Remove the entry from our LUTs
 		BoundObjectToInstances.Remove(Entry.BoundObject, DestroyIndex);
 		TrackInstances.RemoveAt(DestroyIndex);
diff --git a/Engine/Source/Runtime/MovieScene/Private/Evaluation/MovieSceneEvaluationTemplateInstance.cpp b/Engine/Source/Runtime/MovieScene/Private/Evaluation/MovieSceneEvaluationTemplateInstance.cpp
--- a/Engine/Source/Runtime/MovieScene/Private/Evaluation/MovieSceneEvaluationTemplateInstance.cpp
+++ b/Engine/Source/Runtime/MovieScene/Private/Evaluation/MovieSceneEvaluationTemplateInstance.cpp
@@ -160,9 +160,12 @@ void FMovieSceneRootEvaluationTemplateInstance::Finish(IMovieScenePlayer& Player
 
 void FMovieSceneRootEvaluationTemplateInstance::EnableGlobalPreAnimatedStateCapture()
 {
+	using namespace UE::MovieScene;
+
 	if (ensure(EntitySystemLinker))
 	{
-		EntitySystemLinker->GetInstanceRegistry()->MutateInstance(RootInstanceHandle).EnableGlobalPreAnimatedStateCapture(EntitySystemLinker);
+		const FSequenceInstance& Instance = EntitySystemLinker->GetInstanceRegistry()->GetInstance(RootInstanceHandle);
+		Instance.GetPlayer()->PreAnimatedState.EnableGlobalPreAnimatedStateCapture();
 	}
 }
 
@@ -344,9 +347,7 @@ void FMovieSceneRootEvaluationTemplateInstance::PlaybackContextChanged(IMovieSce
 {
 	using namespace UE::MovieScene;
 
-	const bool bGlobalCapture = EntitySystemLinker
-		&& RootInstanceHandle.IsValid()
-		&& EntitySystemLinker->GetInstanceRegistry()->GetInstance(RootInstanceHandle).IsCapturingGlobalPreAnimatedState();
+	const bool bGlobalCapture = Player.PreAnimatedState.IsCapturingGlobalPreAnimatedState();
 
 	if (EntitySystemLinker && IsValidChecked(EntitySystemLinker) && !EntitySystemLinker->IsUnreachable() && !EntitySystemLinker->HasAnyFlags(RF_BeginDestroyed))
 	{
@@ -373,7 +374,7 @@ void FMovieSceneRootEvaluationTemplateInstance::PlaybackContextChanged(IMovieSce
 	Player.PreAnimatedState.Initialize(EntitySystemLinker, RootInstanceHandle);
 	if (bGlobalCapture)
 	{
-		EnableGlobalPreAnimatedStateCapture();
+		Player.PreAnimatedState.EnableGlobalPreAnimatedStateCapture();
 	}
 }
 
diff --git a/Engine/Source/Runtime/MovieScene/Private/Evaluation/MovieScenePreAnimatedState.cpp b/Engine/Source/Runtime/MovieScene/Private/Evaluation/MovieScenePreAnimatedState.cpp
--- a/Engine/Source/Runtime/MovieScene/Private/Evaluation/MovieScenePreAnimatedState.cpp
+++ b/Engine/Source/Runtime/MovieScene/Private/Evaluation/MovieScenePreAnimatedState.cpp
@@ -32,45 +32,32 @@ IMovieScenePlayer* FRestoreStateParams::GetTerminalPlayer() const
 } // namespace UE
 
 
-FScopedPreAnimatedCaptureSource::FScopedPreAnimatedCaptureSource(FMovieScenePreAnimatedState* InPreAnimatedState, const FMovieSceneEvaluationKey& InEvalKey, bool bInWantsRestoreState)
-	: Variant(TInPlaceType<FMovieSceneEvaluationKey>(), InEvalKey)
-	, PreAnimatedState(InPreAnimatedState)
-	, PrevCaptureSource(PreAnimatedState->CaptureSource)
-	, bWantsRestoreState(bInWantsRestoreState)
-{
-	InPreAnimatedState->CaptureSource = this;
-}
-FScopedPreAnimatedCaptureSource::FScopedPreAnimatedCaptureSource(FMovieScenePreAnimatedState* InPreAnimatedState, const UObject* InEvalHook, FMovieSceneSequenceID InSequenceID, bool bInWantsRestoreState)
-	: Variant(TInPlaceType<FEvalHookType>(), FEvalHookType{InEvalHook, InSequenceID} )
-	, PreAnimatedState(InPreAnimatedState)
-	, PrevCaptureSource(PreAnimatedState->CaptureSource)
-	, bWantsRestoreState(bInWantsRestoreState)
-{
-	InPreAnimatedState->CaptureSource = this;
-}
-FScopedPreAnimatedCaptureSource::FScopedPreAnimatedCaptureSource(FMovieScenePreAnimatedState* InPreAnimatedState, UMovieSceneTrackInstance* InTrackInstance, bool bInWantsRestoreState)
-	: Variant(TInPlaceType<UMovieSceneTrackInstance*>(), InTrackInstance)
-	, PreAnimatedState(InPreAnimatedState)
-	, PrevCaptureSource(PreAnimatedState->CaptureSource)
-	, bWantsRestoreState(bInWantsRestoreState)
-{
-	InPreAnimatedState->CaptureSource = this;
-}
-FScopedPreAnimatedCaptureSource::~FScopedPreAnimatedCaptureSource()
-{
-	PreAnimatedState->CaptureSource = PrevCaptureSource;
-}
 
 FMovieScenePreAnimatedState::~FMovieScenePreAnimatedState()
 {
+	// Ensure that the global state capture request is removed
+	UMovieSceneEntitySystemLinker* CurrentLinker = WeakLinker.Get();
+	if (CurrentLinker && bCapturingGlobalPreAnimatedState)
+	{
+		checkf(CurrentLinker->PreAnimatedState.NumRequestsForGlobalState > 0, TEXT("Increment/Decrement mismatch on FPreAnimatedState::NumRequestsForGlobalState"));
+		--CurrentLinker->PreAnimatedState.NumRequestsForGlobalState;
+	}
 }
 
 void FMovieScenePreAnimatedState::Initialize(UMovieSceneEntitySystemLinker* Linker, UE::MovieScene::FInstanceHandle InInstanceHandle)
 {
-	using namespace UE::MovieScene;
+	// If we're re-using a pre-animated state class and it was previously
+	// capturing global state make sure to decrement that request
+	UMovieSceneEntitySystemLinker* PreviousLinker = WeakLinker.Get();
+	if (PreviousLinker && bCapturingGlobalPreAnimatedState)
+	{
+		checkf(PreviousLinker->PreAnimatedState.NumRequestsForGlobalState > 0, TEXT("Increment/Decrement mismatch on FPreAnimatedState::NumRequestsForGlobalState"));
+		--PreviousLinker->PreAnimatedState.NumRequestsForGlobalState;
+
+	}
+
+	bCapturingGlobalPreAnimatedState = false;
 
-	WeakExtension = nullptr;
-	EntityExtensionRef = nullptr;
 	WeakObjectStorage = nullptr;
 	WeakMasterStorage = nullptr;
 	TemplateMetaData = nullptr;
@@ -80,102 +67,29 @@ void FMovieScenePreAnimatedState::Initialize(UMovieSceneEntitySystemLinker* Link
 	InstanceHandle = InInstanceHandle;
 }
 
-void FMovieScenePreAnimatedState::OnEnableGlobalCapture(TSharedPtr<UE::MovieScene::FPreAnimatedStateExtension> InExtension)
+UMovieSceneEntitySystemLinker* FMovieScenePreAnimatedState::GetLinker() const
 {
-	InitializeStorage(InExtension);
+	return WeakLinker.Get();
 }
 
-void FMovieScenePreAnimatedState::OnDisableGlobalCapture()
+bool FMovieScenePreAnimatedState::IsCapturingGlobalPreAnimatedState() const
 {
-	if (!EntityExtensionRef)
-	{
-		WeakObjectStorage = nullptr;
-		WeakMasterStorage = nullptr;
-
-		WeakExtension = nullptr;
-	}
+	return bCapturingGlobalPreAnimatedState;
 }
 
-void FMovieScenePreAnimatedState::ConditionalInitializeEntityStorage(bool bOverrideWantsRestoreState)
+void FMovieScenePreAnimatedState::EnableGlobalPreAnimatedStateCapture()
 {
-	using namespace UE::MovieScene;
-
-	if (bOverrideWantsRestoreState && EntityExtensionRef == nullptr)
-	{
-		if (UMovieSceneEntitySystemLinker* Linker = WeakLinker.Get())
-		{
-			FPreAnimatedStateExtension* Extension = Linker->FindExtension<FPreAnimatedStateExtension>();
-			if (Extension)
-			{
-				EntityExtensionRef = Extension->AsShared();
-			}
-			else
-			{
-				EntityExtensionRef = MakeShared<FPreAnimatedStateExtension>(Linker);
-			}
-		}
-
-		if (EntityExtensionRef)
-		{
-			InitializeStorage(EntityExtensionRef);
-		}
-	}
-}
-
-void FMovieScenePreAnimatedState::InitializeStorage(TSharedPtr<UE::MovieScene::FPreAnimatedStateExtension> Extension)
-{
-	using namespace UE::MovieScene;
-
-	WeakExtension = Extension;
-	WeakObjectStorage = Extension->GetOrCreateStorage<FAnimTypePreAnimatedStateObjectStorage>();
-	WeakMasterStorage = Extension->GetOrCreateStorage<FAnimTypePreAnimatedStateMasterStorage>();
-}
-
-void FMovieScenePreAnimatedState::AddSourceMetaData(const UE::MovieScene::FPreAnimatedStateEntry& Entry)
-{
-	using namespace UE::MovieScene;
-
-	TSharedPtr<FPreAnimatedStateExtension> Extension = WeakExtension.Pin();
-	if (!Extension)
-	{
-		return;
-	}
-
-	if (!CaptureSource)
+	if (bCapturingGlobalPreAnimatedState)
 	{
-		Extension->EnsureMetaData(Entry);
 		return;
 	}
 
-	FPreAnimatedStateMetaData MetaData;
-	MetaData.Entry = Entry;
-	MetaData.RootInstanceHandle = InstanceHandle;
-	MetaData.bWantsRestoreState = CaptureSource->bWantsRestoreState;
+	bCapturingGlobalPreAnimatedState = true;
 
-	if (FMovieSceneEvaluationKey* EvalKey = CaptureSource->Variant.TryGet<FMovieSceneEvaluationKey>())
-	{
-		// Make the association to this track template key
-		if (!TemplateMetaData)
-		{
-			TemplateMetaData = MakeShared<FPreAnimatedTemplateCaptureSources>(Extension.Get());
-			Extension->AddWeakCaptureSource(TemplateMetaData);
-		}
-		TemplateMetaData->BeginTrackingCaptureSource(*EvalKey, MetaData);
-	}
-	else if (FScopedPreAnimatedCaptureSource::FEvalHookType* EvalHook = CaptureSource->Variant.TryGet<FScopedPreAnimatedCaptureSource::FEvalHookType>())
-	{
-		if (!EvaluationHookMetaData)
-		{
-			EvaluationHookMetaData = MakeShared<FPreAnimatedEvaluationHookCaptureSources>(Extension.Get());
-			Extension->AddWeakCaptureSource(EvaluationHookMetaData);
-		}
-		EvaluationHookMetaData->BeginTrackingCaptureSource(EvalHook->EvalHook, EvalHook->SequenceID, MetaData);
-	}
-	else if (UMovieSceneTrackInstance* const * TrackInstance = CaptureSource->Variant.TryGet<UMovieSceneTrackInstance*>())
+	UMovieSceneEntitySystemLinker* Linker = WeakLinker.Get();
+	if (ensure(Linker))
 	{
-		// Track instance meta-data is shared between all players
-		FPreAnimatedTrackInstanceCaptureSources* TrackInstanceMetaData = Extension->GetOrCreateTrackInstanceMetaData();
-		TrackInstanceMetaData->BeginTrackingCaptureSource(*TrackInstance, MetaData);
+		++Linker->PreAnimatedState.NumRequestsForGlobalState;
 	}
 }
 
@@ -183,46 +97,15 @@ void FMovieScenePreAnimatedState::SavePreAnimatedState(UObject& InObject, FMovie
 {
 	using namespace UE::MovieScene;
 
-	TSharedPtr<FAnimTypePreAnimatedStateObjectStorage> ObjectStorage = WeakObjectStorage.Pin();
-
-	const bool bWantsRestoreState = CaptureSource && CaptureSource->bWantsRestoreState;
-	if (!bWantsRestoreState && !ObjectStorage)
+	UMovieSceneEntitySystemLinker* Linker = WeakLinker.Get();
+	if (!Linker)
 	{
 		return;
 	}
 
-	ConditionalInitializeEntityStorage(bWantsRestoreState);
-	if (!ObjectStorage)
-	{
-		// Re-resolve the ptr as it may have changed inside ConditionalInitializeEntityStorage
-		ObjectStorage = WeakObjectStorage.Pin();
-	}
-
-	if (ObjectStorage)
+	if (bCapturingGlobalPreAnimatedState || Linker->PreAnimatedState.HasActiveCaptureSource())
 	{
-		FPreAnimatedStateEntry   Entry        = ObjectStorage->MakeEntry(&InObject, InTokenType);
-		FPreAnimatedStorageIndex StorageIndex = Entry.ValueHandle.StorageIndex;
-
-		AddSourceMetaData(Entry);
-
-		EPreAnimatedStorageRequirement Requirement = bWantsRestoreState
-			? EPreAnimatedStorageRequirement::Transient
-			: EPreAnimatedStorageRequirement::Persistent;
-
-		if (!ObjectStorage->IsStorageRequirementSatisfied(StorageIndex, Requirement))
-		{
-			IMovieScenePreAnimatedTokenPtr Token = Producer.CacheExistingState(InObject);
-			if (Token.IsValid())
-			{
-				const bool bHasEverAninmated = ObjectStorage->HasEverAnimated(StorageIndex);
-				if (!bHasEverAninmated)
-				{
-					Producer.InitializeObjectForAnimation(InObject);
-				}
-
-				ObjectStorage->AssignPreAnimatedValue(StorageIndex, Requirement, MoveTemp(Token));
-			}
-		}
+		Linker->PreAnimatedState.SavePreAnimatedStateDirectly(InObject, InTokenType, Producer);
 	}
 }
 
@@ -230,46 +113,15 @@ void FMovieScenePreAnimatedState::SavePreAnimatedState(FMovieSceneAnimTypeID InT
 {
 	using namespace UE::MovieScene;
 
-	TSharedPtr<FAnimTypePreAnimatedStateMasterStorage> MasterStorage = WeakMasterStorage.Pin();
-
-	const bool bWantsRestoreState = CaptureSource && CaptureSource->bWantsRestoreState;
-	if (!bWantsRestoreState && !MasterStorage)
+	UMovieSceneEntitySystemLinker* Linker = WeakLinker.Get();
+	if (!Linker)
 	{
 		return;
 	}
 
-	ConditionalInitializeEntityStorage(bWantsRestoreState);
-	if (!MasterStorage)
-	{
-		// Re-resolve the ptr as it may have changed inside ConditionalInitializeEntityStorage
-		MasterStorage = WeakMasterStorage.Pin();
-	}
-
-	if (MasterStorage)
+	if (bCapturingGlobalPreAnimatedState || Linker->PreAnimatedState.HasActiveCaptureSource())
 	{
-		FPreAnimatedStateEntry   Entry        = MasterStorage->MakeEntry(InTokenType);
-		FPreAnimatedStorageIndex StorageIndex = Entry.ValueHandle.StorageIndex;
-
-		AddSourceMetaData(Entry);
-
-		EPreAnimatedStorageRequirement Requirement = bWantsRestoreState
-			? EPreAnimatedStorageRequirement::Transient
-			: EPreAnimatedStorageRequirement::Persistent;
-
-		if (!MasterStorage->IsStorageRequirementSatisfied(StorageIndex, Requirement))
-		{
-			IMovieScenePreAnimatedGlobalTokenPtr Token = Producer.CacheExistingState();
-			if (Token.IsValid())
-			{
-				const bool bHasEverAninmated = MasterStorage->HasEverAnimated(StorageIndex);
-				if (!bHasEverAninmated)
-				{
-					Producer.InitializeForAnimation();
-				}
-
-				MasterStorage->AssignPreAnimatedValue(StorageIndex, Requirement, MoveTemp(Token));
-			}
-		}
+		Linker->PreAnimatedState.SavePreAnimatedStateDirectly(InTokenType, Producer);
 	}
 }
 
@@ -280,11 +132,7 @@ void FMovieScenePreAnimatedState::RestorePreAnimatedState()
 	UMovieSceneEntitySystemLinker* Linker = WeakLinker.Get();
 	if (Linker)
 	{
-		TSharedPtr<FPreAnimatedStateExtension> Extension = WeakExtension.Pin();
-		if (Extension)
-		{
-			Extension->RestoreGlobalState(FRestoreStateParams{Linker, InstanceHandle});
-		}
+		Linker->PreAnimatedState.RestoreGlobalState(FRestoreStateParams{Linker, InstanceHandle});
 	}
 }
 
@@ -308,40 +156,38 @@ void FMovieScenePreAnimatedState::RestorePreAnimatedState(UObject& Object)
 {
 	using namespace UE::MovieScene;
 
-	UMovieSceneEntitySystemLinker*         Linker    = WeakLinker.Get();
-	TSharedPtr<FPreAnimatedStateExtension> Extension = WeakExtension.Pin();
-	if (!Linker || !Extension)
+	UMovieSceneEntitySystemLinker* Linker = WeakLinker.Get();
+	if (!Linker)
 	{
 		return;
 	}
 
-	TSharedPtr<FPreAnimatedObjectGroupManager> ObjectGroupManager = Extension->FindGroupManager<FPreAnimatedObjectGroupManager>();
+	TSharedPtr<FPreAnimatedObjectGroupManager> ObjectGroupManager = Linker->PreAnimatedState.FindGroupManager<FPreAnimatedObjectGroupManager>();
 	if (!ObjectGroupManager)
 	{
 		return;
 	}
 
-	FPreAnimatedStorageGroupHandle Group = ObjectGroupManager->FindGroupForObject(&Object);
+	FPreAnimatedStorageGroupHandle Group = ObjectGroupManager->FindGroupForKey(&Object);
 	if (!Group)
 	{
 		return;
 	}
 
-	Extension->RestoreStateForGroup(Group, FRestoreStateParams{Linker, InstanceHandle});
+	Linker->PreAnimatedState.RestoreStateForGroup(Group, FRestoreStateParams{Linker, InstanceHandle});
 }
 
 void FMovieScenePreAnimatedState::RestorePreAnimatedState(UClass* GeneratedClass)
 {
 	using namespace UE::MovieScene;
 
-	UMovieSceneEntitySystemLinker*         Linker    = WeakLinker.Get();
-	TSharedPtr<FPreAnimatedStateExtension> Extension = WeakExtension.Pin();
-	if (!Linker || !Extension)
+	UMovieSceneEntitySystemLinker* Linker = WeakLinker.Get();
+	if (!Linker)
 	{
 		return;
 	}
 
-	TSharedPtr<FPreAnimatedObjectGroupManager> ObjectGroupManager = Extension->FindGroupManager<FPreAnimatedObjectGroupManager>();
+	TSharedPtr<FPreAnimatedObjectGroupManager> ObjectGroupManager = Linker->PreAnimatedState.FindGroupManager<FPreAnimatedObjectGroupManager>();
 	if (ObjectGroupManager)
 	{
 		TArray<FPreAnimatedStorageGroupHandle> Handles;
@@ -350,7 +196,7 @@ void FMovieScenePreAnimatedState::RestorePreAnimatedState(UClass* GeneratedClass
 		FRestoreStateParams Params{Linker, InstanceHandle};
 		for (FPreAnimatedStorageGroupHandle GroupHandle : Handles)
 		{
-			Extension->RestoreStateForGroup(GroupHandle, Params);
+			Linker->PreAnimatedState.RestoreStateForGroup(GroupHandle, Params);
 		}
 	}
 }
@@ -388,10 +234,10 @@ void FMovieScenePreAnimatedState::DiscardEntityTokens()
 {
 	using namespace UE::MovieScene;
 
-	TSharedPtr<FPreAnimatedStateExtension> Extension = WeakExtension.Pin();
-	if (Extension)
+	UMovieSceneEntitySystemLinker* Linker = WeakLinker.Get();
+	if (Linker)
 	{
-		Extension->DiscardTransientState();
+		Linker->PreAnimatedState.DiscardTransientState();
 	}
 }
 
@@ -399,40 +245,38 @@ void FMovieScenePreAnimatedState::DiscardAndRemoveEntityTokensForObject(UObject&
 {
 	using namespace UE::MovieScene;
 
-	UMovieSceneEntitySystemLinker*         Linker    = WeakLinker.Get();
-	TSharedPtr<FPreAnimatedStateExtension> Extension = WeakExtension.Pin();
-
-	if (!Linker || !Extension)
+	UMovieSceneEntitySystemLinker* Linker = WeakLinker.Get();
+	if (!Linker)
 	{
 		return;
 	}
 
-	TSharedPtr<FPreAnimatedObjectGroupManager> ObjectGroupManager = Extension->FindGroupManager<FPreAnimatedObjectGroupManager>();
+	TSharedPtr<FPreAnimatedObjectGroupManager> ObjectGroupManager = Linker->PreAnimatedState.FindGroupManager<FPreAnimatedObjectGroupManager>();
 	if (!ObjectGroupManager)
 	{
 		return;
 	}
 
-	FPreAnimatedStorageGroupHandle Group = ObjectGroupManager->FindGroupForObject(&Object);
+	FPreAnimatedStorageGroupHandle Group = ObjectGroupManager->FindGroupForKey(&Object);
 	if (!Group)
 	{
 		return;
 	}
 
-	Extension->DiscardStateForGroup(Group);
+	Linker->PreAnimatedState.DiscardStateForGroup(Group);
 }
 
 void FMovieScenePreAnimatedState::OnObjectsReplaced(const TMap<UObject*, UObject*>& ReplacementMap)
 {
 	using namespace UE::MovieScene;
 
-	TSharedPtr<FPreAnimatedStateExtension> Extension = WeakExtension.Pin();
-	if (!Extension)
+	UMovieSceneEntitySystemLinker* Linker = WeakLinker.Get();
+	if (!Linker)
 	{
 		return;
 	}
 
-	TSharedPtr<FPreAnimatedObjectGroupManager> ObjectGroupManager = Extension->FindGroupManager<FPreAnimatedObjectGroupManager>();
+	TSharedPtr<FPreAnimatedObjectGroupManager> ObjectGroupManager = Linker->PreAnimatedState.FindGroupManager<FPreAnimatedObjectGroupManager>();
 	if (ObjectGroupManager)
 	{
 		ObjectGroupManager->OnObjectsReplaced(ReplacementMap);
@@ -443,6 +287,6 @@ bool FMovieScenePreAnimatedState::ContainsAnyStateForSequence() const
 {
 	using namespace UE::MovieScene;
 
-	TSharedPtr<FPreAnimatedStateExtension> Extension = WeakExtension.Pin();
-	return Extension && InstanceHandle.IsValid() && Extension->ContainsAnyStateForInstanceHandle(InstanceHandle);
+	UMovieSceneEntitySystemLinker* Linker = WeakLinker.Get();
+	return Linker && InstanceHandle.IsValid() && Linker->PreAnimatedState.ContainsAnyStateForInstanceHandle(InstanceHandle);
 }
\ No newline at end of file
diff --git a/Engine/Source/Runtime/MovieScene/Private/Evaluation/PreAnimatedState/MovieScenePreAnimatedCaptureSource.cpp b/Engine/Source/Runtime/MovieScene/Private/Evaluation/PreAnimatedState/MovieScenePreAnimatedCaptureSource.cpp
new file mode 100644
--- /dev/null
+++ b/Engine/Source/Runtime/MovieScene/Private/Evaluation/PreAnimatedState/MovieScenePreAnimatedCaptureSource.cpp
@@ -0,0 +1,114 @@
+// Copyright Epic Games, Inc. All Rights Reserved.
+
+#include "Evaluation/PreAnimatedState/MovieScenePreAnimatedCaptureSource.h"
+#include "Evaluation/PreAnimatedState/MovieScenePreAnimatedCaptureSources.h"
+#include "Evaluation/MovieScenePreAnimatedState.h"
+
+FScopedPreAnimatedCaptureSource::FScopedPreAnimatedCaptureSource(FMovieScenePreAnimatedState* InPreAnimatedState, const FMovieSceneEvaluationKey& InEvalKey, bool bInWantsRestoreState)
+	: Variant(TInPlaceType<FMovieSceneEvaluationKey>(), InEvalKey)
+	, OptionalSequencePreAnimatedState(InPreAnimatedState)
+	, bWantsRestoreState(bInWantsRestoreState)
+{
+	FScopedPreAnimatedCaptureSource*& CaptureSourcePtr = GetCaptureSourcePtr();
+	PrevCaptureSource = CaptureSourcePtr;
+	CaptureSourcePtr = this;
+}
+FScopedPreAnimatedCaptureSource::FScopedPreAnimatedCaptureSource(FMovieScenePreAnimatedState* InPreAnimatedState, const UObject* InEvalHook, FMovieSceneSequenceID InSequenceID, bool bInWantsRestoreState)
+	: Variant(TInPlaceType<FEvalHookType>(), FEvalHookType{InEvalHook, InSequenceID} )
+	, OptionalSequencePreAnimatedState(InPreAnimatedState)
+	, bWantsRestoreState(bInWantsRestoreState)
+{
+	FScopedPreAnimatedCaptureSource*& CaptureSourcePtr = GetCaptureSourcePtr();
+	PrevCaptureSource = CaptureSourcePtr;
+	CaptureSourcePtr = this;
+}
+FScopedPreAnimatedCaptureSource::FScopedPreAnimatedCaptureSource(UMovieSceneEntitySystemLinker* InLinker, UMovieSceneTrackInstance* InTrackInstance, bool bInWantsRestoreState)
+	: Variant(TInPlaceType<UMovieSceneTrackInstance*>(), InTrackInstance)
+	, OptionalSequencePreAnimatedState(nullptr)
+	, bWantsRestoreState(bInWantsRestoreState)
+{
+	FScopedPreAnimatedCaptureSource*& CaptureSourcePtr = GetCaptureSourcePtr();
+	PrevCaptureSource = CaptureSourcePtr;
+	CaptureSourcePtr = this;
+}
+FScopedPreAnimatedCaptureSource::FScopedPreAnimatedCaptureSource(UMovieSceneEntitySystemLinker* InLinker, const FMovieSceneTrackInstanceInput& TrackInstanceInput)
+	: Variant(TInPlaceType<FMovieSceneTrackInstanceInput>(), TrackInstanceInput)
+	, OptionalSequencePreAnimatedState(nullptr)
+{
+	EMovieSceneCompletionMode CompletionMode = TrackInstanceInput.Section->GetCompletionMode();
+	if (CompletionMode == EMovieSceneCompletionMode::ProjectDefault)
+	{
+		CompletionMode = TrackInstanceInput.Section->GetTypedOuter<UMovieSceneSequence>()->DefaultCompletionMode;
+	}
+	bWantsRestoreState = (CompletionMode == EMovieSceneCompletionMode::RestoreState);
+
+	FScopedPreAnimatedCaptureSource*& CaptureSourcePtr = GetCaptureSourcePtr();
+	PrevCaptureSource = CaptureSourcePtr;
+	CaptureSourcePtr = this;
+}
+
+FScopedPreAnimatedCaptureSource::~FScopedPreAnimatedCaptureSource()
+{
+	GetCaptureSourcePtr() = PrevCaptureSource;
+}
+
+FScopedPreAnimatedCaptureSource*& FScopedPreAnimatedCaptureSource::GetCaptureSourcePtr()
+{
+	// Implemented as a static thread-local for now since there are some tests that run without a linker,
+	// so we can't put this on UMovieSceneEntitySystemLinker::PreAnimatedState where it should probably belong
+	static thread_local FScopedPreAnimatedCaptureSource* GCaptureSource = nullptr;
+	return GCaptureSource;
+}
+
+UE::MovieScene::FInstanceHandle FScopedPreAnimatedCaptureSource::GetRootInstanceHandle(UMovieSceneEntitySystemLinker* Linker) const
+{
+	if (OptionalSequencePreAnimatedState)
+	{
+		return OptionalSequencePreAnimatedState->InstanceHandle;
+	}
+
+	if (const FMovieSceneTrackInstanceInput* TrackInstanceInput = Variant.TryGet<FMovieSceneTrackInstanceInput>())
+	{
+		return Linker->GetInstanceRegistry()->GetInstance(TrackInstanceInput->InstanceHandle).GetRootInstanceHandle();
+	}
+	return UE::MovieScene::FInstanceHandle();
+}
+
+void FScopedPreAnimatedCaptureSource::BeginTracking(const UE::MovieScene::FPreAnimatedStateMetaData& MetaData, UMovieSceneEntitySystemLinker* Linker)
+{
+	using namespace UE::MovieScene;
+
+	if (FMovieSceneEvaluationKey* EvalKey = Variant.TryGet<FMovieSceneEvaluationKey>())
+	{
+		check(OptionalSequencePreAnimatedState);
+		// Make the association to this track template key
+		if (!OptionalSequencePreAnimatedState->TemplateMetaData)
+		{
+			OptionalSequencePreAnimatedState->TemplateMetaData = MakeShared<FPreAnimatedTemplateCaptureSources>(&Linker->PreAnimatedState);
+			Linker->PreAnimatedState.AddWeakCaptureSource(OptionalSequencePreAnimatedState->TemplateMetaData);
+		}
+		OptionalSequencePreAnimatedState->TemplateMetaData->BeginTrackingCaptureSource(*EvalKey, MetaData);
+	}
+	else if (FScopedPreAnimatedCaptureSource::FEvalHookType* EvalHook = Variant.TryGet<FScopedPreAnimatedCaptureSource::FEvalHookType>())
+	{
+		check(OptionalSequencePreAnimatedState);
+		if (!OptionalSequencePreAnimatedState->EvaluationHookMetaData)
+		{
+			OptionalSequencePreAnimatedState->EvaluationHookMetaData = MakeShared<FPreAnimatedEvaluationHookCaptureSources>(&Linker->PreAnimatedState);
+			Linker->PreAnimatedState.AddWeakCaptureSource(OptionalSequencePreAnimatedState->EvaluationHookMetaData);
+		}
+		OptionalSequencePreAnimatedState->EvaluationHookMetaData->BeginTrackingCaptureSource(EvalHook->EvalHook, EvalHook->SequenceID, MetaData);
+	}
+	else if (UMovieSceneTrackInstance* const * TrackInstance = Variant.TryGet<UMovieSceneTrackInstance*>())
+	{
+		// Track instance meta-data is shared between all players
+		FPreAnimatedTrackInstanceCaptureSources* TrackInstanceMetaData = Linker->PreAnimatedState.GetOrCreateTrackInstanceMetaData();
+		TrackInstanceMetaData->BeginTrackingCaptureSource(*TrackInstance, MetaData);
+	}
+	else if (const FMovieSceneTrackInstanceInput* TrackInstanceInput = Variant.TryGet<FMovieSceneTrackInstanceInput>())
+	{
+		// Track instance meta-data is shared between all players
+		FPreAnimatedTrackInstanceInputCaptureSources* TrackInstanceInputMetaData = Linker->PreAnimatedState.GetOrCreateTrackInstanceInputMetaData();
+		TrackInstanceInputMetaData->BeginTrackingCaptureSource(*TrackInstanceInput, MetaData);
+	}
+}
diff --git a/Engine/Source/Runtime/MovieScene/Private/Evaluation/PreAnimatedState/MovieScenePreAnimatedCaptureSources.cpp b/Engine/Source/Runtime/MovieScene/Private/Evaluation/PreAnimatedState/MovieScenePreAnimatedCaptureSources.cpp
--- a/Engine/Source/Runtime/MovieScene/Private/Evaluation/PreAnimatedState/MovieScenePreAnimatedCaptureSources.cpp
+++ b/Engine/Source/Runtime/MovieScene/Private/Evaluation/PreAnimatedState/MovieScenePreAnimatedCaptureSources.cpp
@@ -50,7 +50,9 @@ void FPreAnimatedTrackInstanceCaptureSources::StopTrackingCaptureSource(UMovieSc
 	TPreAnimatedCaptureSources<FObjectKey>::StopTrackingCaptureSource(Key);
 }
 
-
+FPreAnimatedTrackInstanceInputCaptureSources::FPreAnimatedTrackInstanceInputCaptureSources(FPreAnimatedStateExtension* InOwner)
+	: TPreAnimatedCaptureSources<FMovieSceneTrackInstanceInput>(InOwner)
+{}
 
 FPreAnimatedTemplateCaptureSources::FPreAnimatedTemplateCaptureSources(FPreAnimatedStateExtension* InOwner)
 	: TPreAnimatedCaptureSources<FMovieSceneEvaluationKey>(InOwner)
diff --git a/Engine/Source/Runtime/MovieScene/Private/Evaluation/PreAnimatedState/MovieScenePreAnimatedMasterTokenStorage.cpp b/Engine/Source/Runtime/MovieScene/Private/Evaluation/PreAnimatedState/MovieScenePreAnimatedMasterTokenStorage.cpp
--- a/Engine/Source/Runtime/MovieScene/Private/Evaluation/PreAnimatedState/MovieScenePreAnimatedMasterTokenStorage.cpp
+++ b/Engine/Source/Runtime/MovieScene/Private/Evaluation/PreAnimatedState/MovieScenePreAnimatedMasterTokenStorage.cpp
@@ -15,25 +15,11 @@ void FAnimTypePreAnimatedStateMasterStorage::Initialize(FPreAnimatedStorageID In
 	TPreAnimatedStateStorage<FPreAnimatedMasterTokenTraits>::Initialize(InStorageID, InParentExtension);
 }
 
-void FAnimTypePreAnimatedStateMasterStorage::InitializeGroupManager(FPreAnimatedStateExtension* Extension)
-{}
-
-void FAnimTypePreAnimatedStateMasterStorage::OnGroupDestroyed(FPreAnimatedStorageGroupHandle InGroup)
-{
-	check(InGroup == GroupHandle)
-	GroupHandle = FPreAnimatedStorageGroupHandle();
-}
-
 FPreAnimatedStateEntry FAnimTypePreAnimatedStateMasterStorage::MakeEntry(FMovieSceneAnimTypeID AnimTypeID)
 {
-	if (!GroupHandle)
-	{
-		GroupHandle = ParentExtension->AllocateGroup(SharedThis(this));
-	}
-
 	// Find the storage index for the specific anim-type and object we're animating
 	FPreAnimatedStorageIndex StorageIndex = GetOrCreateStorageIndex(AnimTypeID);
-	return FPreAnimatedStateEntry{ GroupHandle, FPreAnimatedStateCachedValueHandle{ StorageID, StorageIndex } };
+	return FPreAnimatedStateEntry{ FPreAnimatedStorageGroupHandle(), FPreAnimatedStateCachedValueHandle{ StorageID, StorageIndex } };
 }
 
 
diff --git a/Engine/Source/Runtime/MovieScene/Private/Evaluation/PreAnimatedState/MovieScenePreAnimatedObjectGroupManager.cpp b/Engine/Source/Runtime/MovieScene/Private/Evaluation/PreAnimatedState/MovieScenePreAnimatedObjectGroupManager.cpp
--- a/Engine/Source/Runtime/MovieScene/Private/Evaluation/PreAnimatedState/MovieScenePreAnimatedObjectGroupManager.cpp
+++ b/Engine/Source/Runtime/MovieScene/Private/Evaluation/PreAnimatedState/MovieScenePreAnimatedObjectGroupManager.cpp
@@ -10,45 +10,13 @@ namespace MovieScene
 
 TAutoRegisterPreAnimatedStorageID<FPreAnimatedObjectGroupManager> FPreAnimatedObjectGroupManager::GroupManagerID;
 
-void FPreAnimatedObjectGroupManager::InitializeGroupManager(FPreAnimatedStateExtension* InExtension)
-{
-	Extension = InExtension;
-}
-
-void FPreAnimatedObjectGroupManager::OnGroupDestroyed(FPreAnimatedStorageGroupHandle Group)
-{
-	FObjectKey Temp = StorageGroupsToObject.FindChecked(Group);
-
-	StorageGroupsByObject.Remove(Temp);
-	StorageGroupsToObject.Remove(Group);
-}
-
-FPreAnimatedStorageGroupHandle FPreAnimatedObjectGroupManager::FindGroupForObject(const FObjectKey& Object) const
-{
-	return StorageGroupsByObject.FindRef(Object);
-}
-
-FPreAnimatedStorageGroupHandle FPreAnimatedObjectGroupManager::MakeGroupForObject(const FObjectKey& Object)
-{
-	FPreAnimatedStorageGroupHandle GroupHandle = StorageGroupsByObject.FindRef(Object);
-	if (GroupHandle)
-	{
-		return GroupHandle;
-	}
-
-	GroupHandle = Extension->AllocateGroup(AsShared());
-	StorageGroupsByObject.Add(Object, GroupHandle);
-	StorageGroupsToObject.Add(GroupHandle, Object);
-	return GroupHandle;
-}
-
 void FPreAnimatedObjectGroupManager::OnObjectsReplaced(const TMap<UObject*, UObject*>& ReplacementMap)
 {
-	TMap<FObjectKey, FPreAnimatedStorageGroupHandle> OldStorageGroupsByObject = MoveTemp(StorageGroupsByObject);
-	StorageGroupsByObject.Reset();
-	StorageGroupsByObject.Reserve(OldStorageGroupsByObject.Num());
+	TMap<FObjectKey, FPreAnimatedStorageGroupHandle> OldStorageGroupsByKey = MoveTemp(StorageGroupsByKey);
+	StorageGroupsByKey.Reset();
+	StorageGroupsByKey.Reserve(OldStorageGroupsByKey.Num());
 
-	for (auto It = OldStorageGroupsByObject.CreateIterator(); It; ++It)
+	for (auto It = OldStorageGroupsByKey.CreateIterator(); It; ++It)
 	{
 		FPreAnimatedStorageGroupHandle GroupHandle = It.Value();
 
@@ -57,22 +25,22 @@ void FPreAnimatedObjectGroupManager::OnObjectsReplaced(const TMap<UObject*, UObj
 		{
 			FObjectKey NewKey(ReplacedObject);
 
-			StorageGroupsByObject.Add(NewKey, GroupHandle);
+			StorageGroupsByKey.Add(NewKey, GroupHandle);
 			// This will overwrite the existing entry
-			StorageGroupsToObject.Add(GroupHandle, NewKey);
+			StorageGroupsToKey.Add(GroupHandle, NewKey);
 
 			Extension->ReplaceObjectForGroup(GroupHandle, It.Key(), NewKey);
 		}
 		else
 		{
-			StorageGroupsByObject.Add(It.Key(), It.Value());
+			StorageGroupsByKey.Add(It.Key(), It.Value());
 		}
 	}
 }
 
 void FPreAnimatedObjectGroupManager::GetGroupsByClass(UClass* GeneratedClass, TArray<FPreAnimatedStorageGroupHandle>& OutGroupHandles)
 {
-	for (auto It = StorageGroupsByObject.CreateConstIterator(); It; ++It)
+	for (auto It = StorageGroupsByKey.CreateConstIterator(); It; ++It)
 	{
 		UObject* Object = It.Key().ResolveObjectPtrEvenIfPendingKill();
 		if (Object && Object->IsA(GeneratedClass))
diff --git a/Engine/Source/Runtime/MovieScene/Private/Evaluation/PreAnimatedState/MovieScenePreAnimatedObjectTokenStorage.cpp b/Engine/Source/Runtime/MovieScene/Private/Evaluation/PreAnimatedState/MovieScenePreAnimatedObjectTokenStorage.cpp
--- a/Engine/Source/Runtime/MovieScene/Private/Evaluation/PreAnimatedState/MovieScenePreAnimatedObjectTokenStorage.cpp
+++ b/Engine/Source/Runtime/MovieScene/Private/Evaluation/PreAnimatedState/MovieScenePreAnimatedObjectTokenStorage.cpp
@@ -31,7 +31,7 @@ FPreAnimatedStateEntry FAnimTypePreAnimatedStateObjectStorage::MakeEntry(UObject
 	FPreAnimatedObjectTokenTraits::FAnimatedKey Key{ Object, AnimTypeID };
 
 	// Begin by finding or creating a pre-animated state group for this bound object
-	FPreAnimatedStorageGroupHandle Group = ObjectGroupManager->MakeGroupForObject(Object);
+	FPreAnimatedStorageGroupHandle Group = ObjectGroupManager->MakeGroupForKey(Object);
 
 	// Find the storage index for the specific anim-type and object we're animating
 	FPreAnimatedStorageIndex StorageIndex = GetOrCreateStorageIndex(Key);
diff --git a/Engine/Source/Runtime/MovieScene/Private/Evaluation/PreAnimatedState/MovieScenePreAnimatedStateExtension.cpp b/Engine/Source/Runtime/MovieScene/Private/Evaluation/PreAnimatedState/MovieScenePreAnimatedStateExtension.cpp
--- a/Engine/Source/Runtime/MovieScene/Private/Evaluation/PreAnimatedState/MovieScenePreAnimatedStateExtension.cpp
+++ b/Engine/Source/Runtime/MovieScene/Private/Evaluation/PreAnimatedState/MovieScenePreAnimatedStateExtension.cpp
@@ -16,19 +16,17 @@ namespace UE
 namespace MovieScene
 {
 
-TEntitySystemLinkerExtensionID<FPreAnimatedStateExtension> FPreAnimatedStateExtension::GetExtensionID()
-{
-	static TEntitySystemLinkerExtensionID<FPreAnimatedStateExtension> ID = UMovieSceneEntitySystemLinker::RegisterExtension<FPreAnimatedStateExtension>();
-	return ID;
-}
-
-
-FPreAnimatedStateExtension::FPreAnimatedStateExtension(UMovieSceneEntitySystemLinker* InLinker)
-	: TSharedEntitySystemLinkerExtension<FPreAnimatedStateExtension>(InLinker)
-	, NumRequestsForGlobalState(0)
+FPreAnimatedStateExtension::FPreAnimatedStateExtension()
+	: NumRequestsForGlobalState(0)
+	, Linker(nullptr)
 	, bEntriesInvalidated(false)
 {}
 
+void FPreAnimatedStateExtension::Initialize(UMovieSceneEntitySystemLinker* InLinker)
+{
+	Linker = InLinker;
+}
+
 FPreAnimatedStateExtension::~FPreAnimatedStateExtension()
 {}
 
@@ -73,10 +71,57 @@ void FPreAnimatedStateExtension::ReplaceObjectForGroup(FPreAnimatedStorageGroupH
 	}
 }
 
+FPreAnimatedStateExtension::FAggregatePreAnimatedStateMetaData* FPreAnimatedStateExtension::FindMetaData(const FPreAnimatedStateEntry& Entry)
+{
+	if (Entry.GroupHandle.IsValid())
+	{
+		FPreAnimatedGroupMetaData&          Group     = GroupMetaData[Entry.GroupHandle.Value];
+		FAggregatePreAnimatedStateMetaData* Aggregate = Algo::FindBy(Group.AggregateMetaData, Entry.ValueHandle, &FAggregatePreAnimatedStateMetaData::ValueHandle);
+
+		return Aggregate;
+	}
+	else
+	{
+		return UngroupedMetaData.Find(Entry.ValueHandle);
+	}
+}
+
+const FPreAnimatedStateExtension::FAggregatePreAnimatedStateMetaData* FPreAnimatedStateExtension::FindMetaData(const FPreAnimatedStateEntry& Entry) const
+{
+	if (Entry.GroupHandle.IsValid())
+	{
+		const FPreAnimatedGroupMetaData&          Group     = GroupMetaData[Entry.GroupHandle.Value];
+		const FAggregatePreAnimatedStateMetaData* Aggregate = Algo::FindBy(Group.AggregateMetaData, Entry.ValueHandle, &FAggregatePreAnimatedStateMetaData::ValueHandle);
+
+		return Aggregate;
+	}
+	else
+	{
+		return UngroupedMetaData.Find(Entry.ValueHandle);
+	}
+}
+
+FPreAnimatedStateExtension::FAggregatePreAnimatedStateMetaData* FPreAnimatedStateExtension::GetOrAddMetaDataInternal(const FPreAnimatedStateEntry& Entry)
+{
+	FAggregatePreAnimatedStateMetaData* Aggregate = FindMetaData(Entry);
+	if (!Aggregate)
+	{
+		if (Entry.GroupHandle.IsValid())
+		{
+			Aggregate = &GroupMetaData[Entry.GroupHandle.Value].AggregateMetaData.Emplace_GetRef(Entry.ValueHandle);
+		}
+		else
+		{
+			Aggregate = &UngroupedMetaData.Add(Entry.ValueHandle, FAggregatePreAnimatedStateMetaData{Entry.ValueHandle});
+		}
+	}
+
+	return Aggregate;
+}
+
 EPreAnimatedStorageRequirement FPreAnimatedStateExtension::GetStorageRequirement(const FPreAnimatedStateEntry& Entry) const
 {
-	const FPreAnimatedGroupMetaData&          Group     = GroupMetaData[Entry.GroupHandle.Value];
-	const FAggregatePreAnimatedStateMetaData* Aggregate = Algo::FindBy(Group.AggregateMetaData, Entry.ValueHandle, &FAggregatePreAnimatedStateMetaData::ValueHandle);
+	const FAggregatePreAnimatedStateMetaData* Aggregate = FindMetaData(Entry);
 
 	if (ensure(Aggregate))
 	{
@@ -91,22 +136,17 @@ EPreAnimatedStorageRequirement FPreAnimatedStateExtension::GetStorageRequirement
 
 void FPreAnimatedStateExtension::EnsureMetaData(const FPreAnimatedStateEntry& Entry)
 {
-	FPreAnimatedGroupMetaData&          Group     = GroupMetaData[Entry.GroupHandle.Value];
-	FAggregatePreAnimatedStateMetaData* Aggregate = Algo::FindBy(Group.AggregateMetaData, Entry.ValueHandle, &FAggregatePreAnimatedStateMetaData::ValueHandle);
-	if (!Aggregate)
-	{
-		Aggregate = &Group.AggregateMetaData.Emplace_GetRef(Entry.ValueHandle);
-	}
+	GetOrAddMetaDataInternal(Entry);
+}
+
+bool FPreAnimatedStateExtension::MetaDataExists(const FPreAnimatedStateEntry& Entry) const
+{
+	return FindMetaData(Entry) != nullptr;
 }
 
 void FPreAnimatedStateExtension::AddMetaData(const FPreAnimatedStateMetaData& MetaData)
 {
-	FPreAnimatedGroupMetaData&          Group     = GroupMetaData[MetaData.Entry.GroupHandle.Value];
-	FAggregatePreAnimatedStateMetaData* Aggregate = Algo::FindBy(Group.AggregateMetaData, MetaData.Entry.ValueHandle, &FAggregatePreAnimatedStateMetaData::ValueHandle);
-	if (!Aggregate)
-	{
-		Aggregate = &Group.AggregateMetaData.Emplace_GetRef(MetaData.Entry.ValueHandle);
-	}
+	FAggregatePreAnimatedStateMetaData* Aggregate = GetOrAddMetaDataInternal(MetaData.Entry);
 
 	++Aggregate->NumContributors;
 	if (MetaData.bWantsRestoreState)
@@ -118,20 +158,17 @@ void FPreAnimatedStateExtension::AddMetaData(const FPreAnimatedStateMetaData& Me
 
 void FPreAnimatedStateExtension::RemoveMetaData(const FPreAnimatedStateMetaData& MetaData)
 {
-	FPreAnimatedGroupMetaData& Group          = GroupMetaData[MetaData.Entry.GroupHandle.Value];
-	const int32                AggregateIndex = Algo::IndexOfBy(Group.AggregateMetaData, MetaData.Entry.ValueHandle, &FAggregatePreAnimatedStateMetaData::ValueHandle);
+	FAggregatePreAnimatedStateMetaData* Aggregate = FindMetaData(MetaData.Entry);
 
-	if (!ensure(AggregateIndex != INDEX_NONE))
+	if (!Aggregate)
 	{
 		return;
 	}
 
-	FAggregatePreAnimatedStateMetaData& Aggregate = Group.AggregateMetaData[AggregateIndex];
-
-	const int32 TotalNum = --Aggregate.NumContributors;
+	const int32 TotalNum = --Aggregate->NumContributors;
 	if (MetaData.bWantsRestoreState)
 	{
-		if (--Aggregate.NumRestoreContributors == 0)
+		if (--Aggregate->NumRestoreContributors == 0)
 		{
 			EPreAnimatedStorageRequirement NewRequirement = TotalNum != 0
 				? EPreAnimatedStorageRequirement::Persistent
@@ -139,20 +176,29 @@ void FPreAnimatedStateExtension::RemoveMetaData(const FPreAnimatedStateMetaData&
 
 			TSharedPtr<IPreAnimatedStorage> Storage = GetStorageChecked(MetaData.Entry.ValueHandle.TypeID);
 
-			FRestoreStateParams Params = {WeakLinker.Get(), MetaData.RootInstanceHandle};
+			FRestoreStateParams Params = {Linker, MetaData.RootInstanceHandle};
 			NewRequirement = Storage->RestorePreAnimatedStateStorage(MetaData.Entry.ValueHandle.StorageIndex, EPreAnimatedStorageRequirement::Transient, NewRequirement, Params);
 
 			if (NewRequirement == EPreAnimatedStorageRequirement::None)
 			{
-				if (Group.AggregateMetaData.Num() == 1)
+				if (MetaData.Entry.GroupHandle.IsValid())
 				{
-					// If the group is going to be empty - just remove it all
-					FreeGroupInternal(MetaData.Entry.GroupHandle.Value);
+					FPreAnimatedGroupMetaData& Group = GroupMetaData[MetaData.Entry.GroupHandle.Value];
+					if (Group.AggregateMetaData.Num() == 1)
+					{
+						// If the group is going to be empty - just remove it all
+						FreeGroupInternal(MetaData.Entry.GroupHandle.Value);
+					}
+					else
+					{
+						const int32 AggregateIndex = Aggregate - Group.AggregateMetaData.GetData();
+						// Otherwise remove just this aggregate
+						Group.AggregateMetaData.RemoveAt(AggregateIndex, 1, false);
+					}
 				}
 				else
 				{
-					// Otherwise remove just this aggregate
-					Group.AggregateMetaData.RemoveAt(AggregateIndex, 1, false);
+					UngroupedMetaData.Remove(MetaData.Entry.ValueHandle);
 				}
 
 				return;
@@ -162,16 +208,14 @@ void FPreAnimatedStateExtension::RemoveMetaData(const FPreAnimatedStateMetaData&
 
 	if (TotalNum == 0)
 	{
-		Aggregate.bWantedRestore = false;
-		Aggregate.TerminalInstanceHandle = MetaData.RootInstanceHandle;
+		Aggregate->bWantedRestore = false;
+		Aggregate->TerminalInstanceHandle = MetaData.RootInstanceHandle;
 	}
 }
 
 void FPreAnimatedStateExtension::UpdateMetaData(const FPreAnimatedStateMetaData& MetaData)
 {
-	FPreAnimatedGroupMetaData&          Group     = GroupMetaData[MetaData.Entry.GroupHandle.Value];
-	FAggregatePreAnimatedStateMetaData* Aggregate = Algo::FindBy(Group.AggregateMetaData, MetaData.Entry.ValueHandle, &FAggregatePreAnimatedStateMetaData::ValueHandle);
-
+	FAggregatePreAnimatedStateMetaData* Aggregate = FindMetaData(MetaData.Entry);
 	if (ensure(Aggregate))
 	{
 		if (MetaData.bWantsRestoreState)
@@ -214,6 +258,20 @@ FPreAnimatedTrackInstanceCaptureSources* FPreAnimatedStateExtension::GetOrCreate
 	return TrackInstanceCaptureSource.Get();
 }
 
+FPreAnimatedTrackInstanceInputCaptureSources* FPreAnimatedStateExtension::GetTrackInstanceInputMetaData() const
+{
+	return TrackInstanceInputCaptureSource.Get();
+}
+
+FPreAnimatedTrackInstanceInputCaptureSources* FPreAnimatedStateExtension::GetOrCreateTrackInstanceInputMetaData()
+{
+	if (!TrackInstanceInputCaptureSource)
+	{
+		TrackInstanceInputCaptureSource = MakeUnique<FPreAnimatedTrackInstanceInputCaptureSources>(this);
+	}
+	return TrackInstanceInputCaptureSource.Get();
+}
+
 void FPreAnimatedStateExtension::AddWeakCaptureSource(TWeakPtr<IPreAnimatedCaptureSource> InWeakCaptureSource)
 {
 	WeakExternalCaptureSources.Add(InWeakCaptureSource);
@@ -264,8 +322,7 @@ void FPreAnimatedStateExtension::RestoreGlobalState(const FRestoreStateParams& P
 	// Remove all contributions
 	for (const FPreAnimatedStateMetaData& MetaData : ExpiredMetaData)
 	{
-		FPreAnimatedGroupMetaData&          Group     = GroupMetaData[MetaData.Entry.GroupHandle.Value];
-		FAggregatePreAnimatedStateMetaData* Aggregate = Algo::FindBy(Group.AggregateMetaData, MetaData.Entry.ValueHandle, &FAggregatePreAnimatedStateMetaData::ValueHandle);
+		FAggregatePreAnimatedStateMetaData* Aggregate = FindMetaData(MetaData.Entry);
 		if (ensure(Aggregate))
 		{
 			const int32 TotalNum = --Aggregate->NumContributors;
@@ -312,6 +369,18 @@ void FPreAnimatedStateExtension::RestoreGlobalState(const FRestoreStateParams& P
 		}
 	}
 
+	for (auto UngroupedIt = UngroupedMetaData.CreateIterator(); UngroupedIt; ++UngroupedIt)
+	{
+		FAggregatePreAnimatedStateMetaData& Aggregate = UngroupedIt.Value();
+		if (Aggregate.NumContributors == 0 && (!Aggregate.TerminalInstanceHandle.IsValid() || Aggregate.TerminalInstanceHandle == Params.TerminalInstanceHandle))
+		{
+			TSharedPtr<IPreAnimatedStorage> Storage = GetStorageChecked(Aggregate.ValueHandle.TypeID);
+			Storage->RestorePreAnimatedStateStorage(Aggregate.ValueHandle.StorageIndex, EPreAnimatedStorageRequirement::Persistent, EPreAnimatedStorageRequirement::None, Params);
+
+			UngroupedIt.RemoveCurrent();
+		}
+	}
+
 	GroupMetaData.Shrink();
 
 	bEntriesInvalidated = true;
@@ -343,7 +412,15 @@ void FPreAnimatedStateExtension::DiscardTransientState()
 			Storage->DiscardPreAnimatedStateStorage(Aggregate.ValueHandle.StorageIndex, EPreAnimatedStorageRequirement::Transient);
 		}
 	}
+	for (TPair<FPreAnimatedStateCachedValueHandle, FAggregatePreAnimatedStateMetaData> Pair : UngroupedMetaData)
+	{
+		FAggregatePreAnimatedStateMetaData& Aggregate = Pair.Value;
+
+		Aggregate = FAggregatePreAnimatedStateMetaData(Pair.Key);
 
+		TSharedPtr<IPreAnimatedStorage> Storage = GetStorageChecked(Pair.Key.TypeID);
+		Storage->DiscardPreAnimatedStateStorage(Pair.Key.StorageIndex, EPreAnimatedStorageRequirement::Transient);
+	}
 	bEntriesInvalidated = true;
 }
 
@@ -401,5 +478,124 @@ bool FPreAnimatedStateExtension::ContainsAnyStateForInstanceHandle(FInstanceHand
 	return false;
 }
 
+bool FPreAnimatedStateExtension::HasActiveCaptureSource() const
+{
+	FScopedPreAnimatedCaptureSource* CaptureSource = FScopedPreAnimatedCaptureSource::GetCaptureSourcePtr();
+	return (CaptureSource && CaptureSource->bWantsRestoreState);
+}
+
+bool FPreAnimatedStateExtension::ShouldCaptureAnyState() const
+{
+	FScopedPreAnimatedCaptureSource* CaptureSource = FScopedPreAnimatedCaptureSource::GetCaptureSourcePtr();
+	return (CaptureSource && CaptureSource->bWantsRestoreState) || IsCapturingGlobalState();
+}
+
+void FPreAnimatedStateExtension::AddSourceMetaData(const UE::MovieScene::FPreAnimatedStateEntry& Entry)
+{
+	using namespace UE::MovieScene;
+
+	FScopedPreAnimatedCaptureSource* CaptureSource = FScopedPreAnimatedCaptureSource::GetCaptureSourcePtr();
+	if (!CaptureSource)
+	{
+		EnsureMetaData(Entry);
+	}
+	else
+	{
+		FPreAnimatedStateMetaData MetaData;
+		MetaData.Entry = Entry;
+		MetaData.RootInstanceHandle = CaptureSource->GetRootInstanceHandle(Linker);
+		MetaData.bWantsRestoreState = CaptureSource->bWantsRestoreState;
+
+		CaptureSource->BeginTracking(MetaData, Linker);
+	}
+}
+
+void FPreAnimatedStateExtension::SavePreAnimatedState(UObject& InObject, FMovieSceneAnimTypeID InTokenType, const IMovieScenePreAnimatedTokenProducer& Producer)
+{
+	if (HasActiveCaptureSource() || IsCapturingGlobalState())
+	{
+		SavePreAnimatedStateDirectly(InObject, InTokenType, Producer);
+	}
+}
+
+void FPreAnimatedStateExtension::SavePreAnimatedStateDirectly(UObject& InObject, FMovieSceneAnimTypeID InTokenType, const IMovieScenePreAnimatedTokenProducer& Producer)
+{
+	TSharedPtr<FAnimTypePreAnimatedStateObjectStorage> ObjectStorage = WeakGenericObjectStorage.Pin();
+	if (!ObjectStorage)
+	{
+		ObjectStorage = GetOrCreateStorage<FAnimTypePreAnimatedStateObjectStorage>();
+		WeakGenericObjectStorage = ObjectStorage;
+	}
+
+	FPreAnimatedStateEntry   Entry        = ObjectStorage->MakeEntry(&InObject, InTokenType);
+	FPreAnimatedStorageIndex StorageIndex = Entry.ValueHandle.StorageIndex;
+
+	AddSourceMetaData(Entry);
+
+	EPreAnimatedStorageRequirement Requirement = HasActiveCaptureSource()
+		? EPreAnimatedStorageRequirement::Transient
+		: EPreAnimatedStorageRequirement::Persistent;
+
+	if (!ObjectStorage->IsStorageRequirementSatisfied(StorageIndex, Requirement))
+	{
+		IMovieScenePreAnimatedTokenPtr Token = Producer.CacheExistingState(InObject);
+		if (Token.IsValid())
+		{
+			const bool bHasEverAninmated = ObjectStorage->HasEverAnimated(StorageIndex);
+			if (!bHasEverAninmated)
+			{
+				Producer.InitializeObjectForAnimation(InObject);
+			}
+
+			ObjectStorage->AssignPreAnimatedValue(StorageIndex, Requirement, MoveTemp(Token));
+		}
+	}
+}
+
+void FPreAnimatedStateExtension::SavePreAnimatedState(FMovieSceneAnimTypeID InTokenType, const IMovieScenePreAnimatedGlobalTokenProducer& Producer)
+{
+	if (HasActiveCaptureSource() || IsCapturingGlobalState())
+	{
+		SavePreAnimatedStateDirectly(InTokenType, Producer);
+	}
+}
+
+void FPreAnimatedStateExtension::SavePreAnimatedStateDirectly(FMovieSceneAnimTypeID InTokenType, const IMovieScenePreAnimatedGlobalTokenProducer& Producer)
+{
+	using namespace UE::MovieScene;
+
+	TSharedPtr<FAnimTypePreAnimatedStateMasterStorage> MasterStorage = WeakGenericMasterStorage.Pin();
+	if (!MasterStorage)
+	{
+		MasterStorage = GetOrCreateStorage<FAnimTypePreAnimatedStateMasterStorage>();
+		WeakGenericMasterStorage = MasterStorage;
+	}
+
+	FPreAnimatedStateEntry   Entry        = MasterStorage->MakeEntry(InTokenType);
+	FPreAnimatedStorageIndex StorageIndex = Entry.ValueHandle.StorageIndex;
+
+	AddSourceMetaData(Entry);
+
+	EPreAnimatedStorageRequirement Requirement = HasActiveCaptureSource()
+		? EPreAnimatedStorageRequirement::Transient
+		: EPreAnimatedStorageRequirement::Persistent;
+
+	if (!MasterStorage->IsStorageRequirementSatisfied(StorageIndex, Requirement))
+	{
+		IMovieScenePreAnimatedGlobalTokenPtr Token = Producer.CacheExistingState();
+		if (Token.IsValid())
+		{
+			const bool bHasEverAninmated = MasterStorage->HasEverAnimated(StorageIndex);
+			if (!bHasEverAninmated)
+			{
+				Producer.InitializeForAnimation();
+			}
+
+			MasterStorage->AssignPreAnimatedValue(StorageIndex, Requirement, MoveTemp(Token));
+		}
+	}
+}
+
+
 } // namespace MovieScene
 } // namespace UE
diff --git a/Engine/Source/Runtime/MovieScene/Private/Tests/MovieScenePreAnimatedStateTests.cpp b/Engine/Source/Runtime/MovieScene/Private/Tests/MovieScenePreAnimatedStateTests.cpp
--- a/Engine/Source/Runtime/MovieScene/Private/Tests/MovieScenePreAnimatedStateTests.cpp
+++ b/Engine/Source/Runtime/MovieScene/Private/Tests/MovieScenePreAnimatedStateTests.cpp
@@ -149,12 +149,9 @@ bool FMovieScenePreAnimatedStateGlobalTest::RunTest(const FString& Parameters)
 
 	ResetValues();
 
-	TSharedPtr<FPreAnimatedStateExtension> GlobalPreAnimatedState = MakeShared<FPreAnimatedStateExtension>(Linker);
-	++GlobalPreAnimatedState->NumRequestsForGlobalState;
-
 	FMovieScenePreAnimatedState State;
 	State.Initialize(Linker, FInstanceHandle());
-	State.OnEnableGlobalCapture(GlobalPreAnimatedState);
+	State.EnableGlobalPreAnimatedStateCapture();
 
 	FPreAnimatedTokenProducer Producer(&TestValue1);
 
@@ -186,12 +183,9 @@ bool FMovieScenePreAnimatedStateEntityTest::RunTest(const FString& Parameters)
 
 	UMovieSceneEntitySystemLinker* Linker = GetTestLinker();
 
-	TSharedPtr<FPreAnimatedStateExtension> GlobalPreAnimatedState = MakeShared<FPreAnimatedStateExtension>(Linker);
-	++GlobalPreAnimatedState->NumRequestsForGlobalState;
-
 	FMovieScenePreAnimatedState State;
 	State.Initialize(Linker, FInstanceHandle());
-	State.OnEnableGlobalCapture(GlobalPreAnimatedState);
+	State.EnableGlobalPreAnimatedStateCapture();
 
 	FPreAnimatedTokenProducer Producer(&TestValue1);
 
@@ -226,12 +220,9 @@ bool FMovieScenePreAnimatedStateOverlappingEntitiesTest::RunTest(const FString&
 
 	UMovieSceneEntitySystemLinker* Linker = GetTestLinker();
 
-	TSharedPtr<FPreAnimatedStateExtension> GlobalPreAnimatedState = MakeShared<FPreAnimatedStateExtension>(Linker);
-	++GlobalPreAnimatedState->NumRequestsForGlobalState;
-
 	FMovieScenePreAnimatedState State;
 	State.Initialize(Linker, FInstanceHandle());
-	State.OnEnableGlobalCapture(GlobalPreAnimatedState);
+	State.EnableGlobalPreAnimatedStateCapture();
 
 	FPreAnimatedTokenProducer Producer(&TestValue1);
 
@@ -301,12 +292,9 @@ bool FMovieScenePreAnimatedStateKeepThenRestoreEntityTest::RunTest(const FString
 
 	UMovieSceneEntitySystemLinker* Linker = GetTestLinker();
 
-	TSharedPtr<FPreAnimatedStateExtension> GlobalPreAnimatedState = MakeShared<FPreAnimatedStateExtension>(Linker);
-	++GlobalPreAnimatedState->NumRequestsForGlobalState;
-
 	FMovieScenePreAnimatedState State;
 	State.Initialize(Linker, FInstanceHandle());
-	State.OnEnableGlobalCapture(GlobalPreAnimatedState);
+	State.EnableGlobalPreAnimatedStateCapture();
 
 	FPreAnimatedTokenProducer Producer(&TestValue1);
 
@@ -528,12 +516,9 @@ bool FMovieScenePreAnimatedStatePerformanceTest::RunTest(const FString& Paramete
 
 	UMovieSceneEntitySystemLinker* Linker = GetTestLinker();
 
-	TSharedPtr<FPreAnimatedStateExtension> GlobalPreAnimatedState = MakeShared<FPreAnimatedStateExtension>(Linker);
-	++GlobalPreAnimatedState->NumRequestsForGlobalState;
-
 	FMovieScenePreAnimatedState State;
 	State.Initialize(Linker, FInstanceHandle());
-	State.OnEnableGlobalCapture(GlobalPreAnimatedState);
+	State.EnableGlobalPreAnimatedStateCapture();
 
 	FPreAnimatedTokenProducer Producer(&TestValue1);
 
diff --git a/Engine/Source/Runtime/MovieScene/Public/EntitySystem/MovieSceneEntitySystemLinker.h b/Engine/Source/Runtime/MovieScene/Public/EntitySystem/MovieSceneEntitySystemLinker.h
--- a/Engine/Source/Runtime/MovieScene/Public/EntitySystem/MovieSceneEntitySystemLinker.h
+++ b/Engine/Source/Runtime/MovieScene/Public/EntitySystem/MovieSceneEntitySystemLinker.h
@@ -14,6 +14,7 @@
 #include "EntitySystem/MovieSceneEntitySystemGraphs.h"
 #include "EntitySystem/MovieSceneSequenceInstance.h"
 #include "EntitySystem/MovieSceneEntitySystemLinkerExtension.h"
+#include "Evaluation/PreAnimatedState/MovieScenePreAnimatedStateExtension.h"
 
 #include "MovieSceneEntitySystemLinker.generated.h"
 
@@ -64,6 +65,8 @@ public:
 
 	GENERATED_BODY()
 
+	UE::MovieScene::FPreAnimatedStateExtension PreAnimatedState;
+
 	UMovieSceneEntitySystemLinker(const FObjectInitializer& ObjInit);
 
 	static FComponentRegistry* GetComponents();
diff --git a/Engine/Source/Runtime/MovieScene/Public/EntitySystem/MovieSceneInstanceRegistry.h b/Engine/Source/Runtime/MovieScene/Public/EntitySystem/MovieSceneInstanceRegistry.h
--- a/Engine/Source/Runtime/MovieScene/Public/EntitySystem/MovieSceneInstanceRegistry.h
+++ b/Engine/Source/Runtime/MovieScene/Public/EntitySystem/MovieSceneInstanceRegistry.h
@@ -105,7 +105,6 @@ struct FInstanceRegistry
 		return InvalidatedObjectBindings.Num() != 0;
 	}
 
-	bool RemoveInvalidHandles();
 	void PostInstantation();
 
 	void TagGarbage();
diff --git a/Engine/Source/Runtime/MovieScene/Public/EntitySystem/MovieScenePreAnimatedStateSystem.h b/Engine/Source/Runtime/MovieScene/Public/EntitySystem/MovieScenePreAnimatedStateSystem.h
--- a/Engine/Source/Runtime/MovieScene/Public/EntitySystem/MovieScenePreAnimatedStateSystem.h
+++ b/Engine/Source/Runtime/MovieScene/Public/EntitySystem/MovieScenePreAnimatedStateSystem.h
@@ -7,33 +7,6 @@
 
 #include "MovieScenePreAnimatedStateSystem.generated.h"
 
-namespace UE
-{
-namespace MovieScene
-{
-
-/** Structure that manages the lifetime of the pre-animated state extension while entities exist with the restore state tag */
-struct FPreAnimatedStateExtensionReference
-{
-	FPreAnimatedStateExtensionReference() = default;
-	FPreAnimatedStateExtensionReference(UMovieSceneEntitySystemLinker* Linker)
-	{
-		Update(Linker);
-	}
-
-	TSharedPtr<FPreAnimatedStateExtension> Get() const;
-
-	TSharedPtr<FPreAnimatedStateExtension> Update(UMovieSceneEntitySystemLinker* Linker);
-
-private:
-	/** Weak ref to the extension - this is always used for access, and will remain valid as long as there are any global state captures, or RestoreState entities */
-	TWeakPtr<FPreAnimatedStateExtension>   WeakPreAnimatedStateExtension;
-	/** Strong ref to the extension that keeps the extension alive if there are RestoreState entities in the entity manager */
-	TSharedPtr<FPreAnimatedStateExtension> PreAnimatedStateExtensionRef;
-};
-
-} // namespace MovieScene
-} // namespace UE
 
 UINTERFACE()
 class UMovieScenePreAnimatedStateSystemInterface : public UInterface
@@ -97,8 +70,6 @@ private:
 	virtual void OnLink() override;
 	virtual void OnUnlink() override;
 	virtual void OnRun(FSystemTaskPrerequisites& InPrerequisites, FSystemSubsequentTasks& Subsequents) override final;
-
-	UE::MovieScene::FPreAnimatedStateExtensionReference PreAnimatedStateRef;
 };
 
 
diff --git a/Engine/Source/Runtime/MovieScene/Public/EntitySystem/MovieSceneSequenceInstance.h b/Engine/Source/Runtime/MovieScene/Public/EntitySystem/MovieSceneSequenceInstance.h
--- a/Engine/Source/Runtime/MovieScene/Public/EntitySystem/MovieSceneSequenceInstance.h
+++ b/Engine/Source/Runtime/MovieScene/Public/EntitySystem/MovieSceneSequenceInstance.h
@@ -93,18 +93,6 @@ public:
 	 */
 	void RunLegacyTrackTemplates();
 
-	/**
-	 * Indicate that this sequence instance should capture any and all changes of state so they can be restored later
-	 *
-	 * @param Linker     The linker that owns this sequence instance
-	 */
-	void EnableGlobalPreAnimatedStateCapture(UMovieSceneEntitySystemLinker* Linker);
-
-	/**
-	 * Check whether this sequence instance is capturing any and all changes of state so they can be restored later
-	 */
-	bool IsCapturingGlobalPreAnimatedState() const;
-
 public:
 
 	/**
@@ -287,8 +275,6 @@ private:
 	TUniquePtr<ISequenceUpdater> SequenceUpdater;
 	/** For top-level sequences only - a utility class that is used to ensure that volatile sequences are up to date. Only valid in editor, or for sequences that have the volatile flag. */
 	TUniquePtr<FCompiledDataVolatilityManager> VolatilityManager;
-	/** For top-level sequences only - a marker that keeps this instance's linker saving global preanimated state. */
-	TSharedPtr<FPreAnimatedStateExtension> GlobalPreAnimatedState;
 
 
 	/** Delegate Binding for when an object binding is invalidated in this instance . */
diff --git a/Engine/Source/Runtime/MovieScene/Public/EntitySystem/TrackInstance/MovieSceneTrackInstance.h b/Engine/Source/Runtime/MovieScene/Public/EntitySystem/TrackInstance/MovieSceneTrackInstance.h
--- a/Engine/Source/Runtime/MovieScene/Public/EntitySystem/TrackInstance/MovieSceneTrackInstance.h
+++ b/Engine/Source/Runtime/MovieScene/Public/EntitySystem/TrackInstance/MovieSceneTrackInstance.h
@@ -10,6 +10,18 @@
 class UMovieSceneSection;
 class UMovieSceneEntitySystemLinker;
 
+struct FMovieSceneTrackInstanceComponent;
+
+namespace UE
+{
+namespace MovieScene
+{
+
+struct FPreAnimatedTrackInstanceInputCaptureSources;
+
+}
+}
+
 USTRUCT()
 struct FMovieSceneTrackInstanceInput
 {
@@ -117,7 +129,7 @@ public:
 	 */
 	UMovieSceneEntitySystemLinker* GetLinker() const
 	{
-		return Linker;
+		return PrivateLinker;
 	}
 
 
@@ -163,7 +175,7 @@ private:
 	bool bIsMasterTrackInstance;
 
 	UPROPERTY()
-	TObjectPtr<UMovieSceneEntitySystemLinker> Linker;
+	TObjectPtr<UMovieSceneEntitySystemLinker> PrivateLinker;
 
 	UPROPERTY()
 	TArray<FMovieSceneTrackInstanceInput> Inputs;
diff --git a/Engine/Source/Runtime/MovieScene/Public/Evaluation/Blending/BlendableTokenStack.h b/Engine/Source/Runtime/MovieScene/Public/Evaluation/Blending/BlendableTokenStack.h
--- a/Engine/Source/Runtime/MovieScene/Public/Evaluation/Blending/BlendableTokenStack.h
+++ b/Engine/Source/Runtime/MovieScene/Public/Evaluation/Blending/BlendableTokenStack.h
@@ -10,6 +10,7 @@
 #include "Evaluation/MovieSceneEvaluationOperand.h"
 #include "Evaluation/Blending/BlendableToken.h"
 #include "IMovieScenePlayer.h"
+#include "Evaluation/PreAnimatedState/MovieScenePreAnimatedCaptureSource.h"
 
 struct FMovieSceneInterrogationData;
 struct FMovieSceneContext;
diff --git a/Engine/Source/Runtime/MovieScene/Public/Evaluation/MovieScenePreAnimatedState.h b/Engine/Source/Runtime/MovieScene/Public/Evaluation/MovieScenePreAnimatedState.h
--- a/Engine/Source/Runtime/MovieScene/Public/Evaluation/MovieScenePreAnimatedState.h
+++ b/Engine/Source/Runtime/MovieScene/Public/Evaluation/MovieScenePreAnimatedState.h
@@ -37,62 +37,6 @@ struct FPreAnimatedEvaluationHookCaptureSources;
 }
 }
 
-/** Enumeration that defines at what level to capture animating state for */
-enum class ECapturePreAnimatedState : uint8
-{
-	/** Don't capture anything */
-	None,
-	/** Capture anything that animates, but only store such state globally across the sequence */
-	Global,
-	/** Capture anything for the current entity (track or section), such that it will be restored when that entity stops evaluating */
-	Entity,
-};
-
-
-/**
- * Scoped structure that can be used to wrap a call to SavePreAnimatedState to associate the capture with a specific capture source.
- */
-struct FScopedPreAnimatedCaptureSource
-{
-	/**
-	 * Construct this capture source from a template (FMovieSceneEvalTemplate) evaluation key, and whether this should restore state when the template is finished
-	 */
-	MOVIESCENE_API explicit FScopedPreAnimatedCaptureSource(FMovieScenePreAnimatedState* InPreAnimatedState, const FMovieSceneEvaluationKey& InEvalKey, bool bInWantsRestoreState);
-
-	/**
-	 * Construct this capture source from an evaluation hook (UMovieSceneEvaluationHookSection), its SequenceID, and whether this should restore state when the template is finished
-	 */
-	MOVIESCENE_API explicit FScopedPreAnimatedCaptureSource(FMovieScenePreAnimatedState* InPreAnimatedState, const UObject* InEvalHook, FMovieSceneSequenceID SequenceID, bool bInWantsRestoreState);
-
-	/**
-	 * Construct this capture source from a track instance (UMovieSceneTrackInstance) and whether this should restore state when the template is finished
-	 */
-	MOVIESCENE_API explicit FScopedPreAnimatedCaptureSource(FMovieScenePreAnimatedState* InPreAnimatedState, UMovieSceneTrackInstance* InTrackInstance, bool bInWantsRestoreState);
-
-	FScopedPreAnimatedCaptureSource(const FScopedPreAnimatedCaptureSource&) = delete;
-	void operator=(const FScopedPreAnimatedCaptureSource&) = delete;
-
-	FScopedPreAnimatedCaptureSource(FScopedPreAnimatedCaptureSource&&) = delete;
-	void operator=(FScopedPreAnimatedCaptureSource&&) = delete;
-
-	MOVIESCENE_API ~FScopedPreAnimatedCaptureSource();
-
-private:
-	friend class FMovieScenePreAnimatedState;
-	struct FEvalHookType
-	{
-		const UObject* EvalHook;
-		FMovieSceneSequenceID SequenceID;
-	};
-	using CaptureSourceType = TVariant<FMovieSceneEvaluationKey, FEvalHookType, UMovieSceneTrackInstance*>;
-
-	CaptureSourceType Variant;
-	FMovieScenePreAnimatedState* PreAnimatedState;
-	FScopedPreAnimatedCaptureSource* PrevCaptureSource;
-	bool bWantsRestoreState;
-};
-
-
 /**
  * Class that caches pre-animated state for objects that were manipulated by sequencer
  */
@@ -110,14 +54,19 @@ public:
 	MOVIESCENE_API void Initialize(UMovieSceneEntitySystemLinker* Linker, UE::MovieScene::FInstanceHandle InstanceHandle);
 
 	/**
-	 * Called when global capture has been enabled for this player
+	 * Check whether this sequence instance is capturing any and all changes of state so they can be restored later
+	 */
+	MOVIESCENE_API bool IsCapturingGlobalPreAnimatedState() const;
+
+	/**
+	 * Check whether this sequence instance is capturing any and all changes of state so they can be restored later
 	 */
-	void OnEnableGlobalCapture(TSharedPtr<UE::MovieScene::FPreAnimatedStateExtension> InExtension);
+	MOVIESCENE_API void EnableGlobalPreAnimatedStateCapture();
 
 	/**
-	 * Called when global capture has been disabled for this player
+	 * Retrieve the linker this container is bound to
 	 */
-	void OnDisableGlobalCapture();
+	MOVIESCENE_API UMovieSceneEntitySystemLinker* GetLinker() const;
 
 public:
 
@@ -184,20 +133,11 @@ private:
 
 private:
 
-	friend FScopedPreAnimatedCaptureSource;
+	friend struct FScopedPreAnimatedCaptureSource;
 
 	/** Weak pointer to the linker that we're associated with */
 	TWeakObjectPtr<UMovieSceneEntitySystemLinker> WeakLinker;
 
-	/**
-	 * Weak pointer to a pre-animated state extension in the linker.
-	 * This is kept alive either by FSequenceInstance::GlobalPreAnimatedState if global state capture is active,
-	 * or by EntityExtensionRef if there are entries that have captured state that need restoring when done
-	 */
-	TWeakPtr<UE::MovieScene::FPreAnimatedStateExtension> WeakExtension;
-	/** Strong pointer to the extension that keeps it alive while state needs restoring on completion */
-	TSharedPtr<UE::MovieScene::FPreAnimatedStateExtension> EntityExtensionRef;
-
 	/** Pointers to the storage for state bound to objects, organized by FMovieSceneAnimTypeID */
 	TWeakPtr<UE::MovieScene::FAnimTypePreAnimatedStateObjectStorage> WeakObjectStorage;
 	/** Pointers to the storage for state created from master tracks, or otherwise not bound to objects */
@@ -211,6 +151,5 @@ private:
 	/** The instance handle for the root sequence instance */
 	UE::MovieScene::FInstanceHandle InstanceHandle;
 
-	/** Current capture source, if any */
-	FScopedPreAnimatedCaptureSource* CaptureSource = nullptr;
+	bool bCapturingGlobalPreAnimatedState;
 };
diff --git a/Engine/Source/Runtime/MovieScene/Public/Evaluation/PreAnimatedState/MovieScenePreAnimatedCaptureSource.h b/Engine/Source/Runtime/MovieScene/Public/Evaluation/PreAnimatedState/MovieScenePreAnimatedCaptureSource.h
new file mode 100644
--- /dev/null
+++ b/Engine/Source/Runtime/MovieScene/Public/Evaluation/PreAnimatedState/MovieScenePreAnimatedCaptureSource.h
@@ -0,0 +1,88 @@
+// Copyright Epic Games, Inc. All Rights Reserved.
+
+#pragma once
+
+#include "CoreTypes.h"
+#include "Misc/TVariant.h"
+#include "MovieSceneSequenceID.h"
+#include "Evaluation/MovieSceneEvaluationKey.h"
+#include "Evaluation/PreAnimatedState/MovieScenePreAnimatedStateTypes.h"
+#include "EntitySystem/MovieSceneSequenceInstanceHandle.h"
+#include "EntitySystem/TrackInstance/MovieSceneTrackInstance.h"
+
+struct FMovieSceneSequenceID;
+struct FMovieSceneEvaluationKey;
+
+class UObject;
+class UMovieSceneTrackInstance;
+class FMovieScenePreAnimatedState;
+
+namespace UE
+{
+namespace MovieScene
+{
+
+struct FPreAnimatedStateMetaData;
+
+}
+}
+
+
+/**
+ * Scoped structure that can be used to wrap a call to SavePreAnimatedState to associate the capture with a specific capture source.
+ */
+struct FScopedPreAnimatedCaptureSource
+{
+	/**
+	 * Construct this capture source from a template (FMovieSceneEvalTemplate) evaluation key, and whether this should restore state when the template is finished
+	 */
+	MOVIESCENE_API explicit FScopedPreAnimatedCaptureSource(FMovieScenePreAnimatedState* InPreAnimatedState, const FMovieSceneEvaluationKey& InEvalKey, bool bInWantsRestoreState);
+
+	/**
+	 * Construct this capture source from an evaluation hook (UMovieSceneEvaluationHookSection), its SequenceID, and whether this should restore state when the template is finished
+	 */
+	MOVIESCENE_API explicit FScopedPreAnimatedCaptureSource(FMovieScenePreAnimatedState* InPreAnimatedState, const UObject* InEvalHook, FMovieSceneSequenceID SequenceID, bool bInWantsRestoreState);
+
+	/**
+	 * Construct this capture source from a track instance (UMovieSceneTrackInstance) and whether this should restore state when the template is finished
+	 */
+	MOVIESCENE_API explicit FScopedPreAnimatedCaptureSource(UMovieSceneEntitySystemLinker* InLinker, UMovieSceneTrackInstance* InTrackInstance, bool bInWantsRestoreState);
+
+	/**
+	 * Construct this capture source from a track instance input (UMovieSceneTrackInstance + FMovieSceneTrackInstanceInput) and whether this should restore state when the template is finished
+	 */
+	MOVIESCENE_API explicit FScopedPreAnimatedCaptureSource(UMovieSceneEntitySystemLinker* InLinker, const FMovieSceneTrackInstanceInput& TrackInstanceInput);
+
+	FScopedPreAnimatedCaptureSource(const FScopedPreAnimatedCaptureSource&) = delete;
+	void operator=(const FScopedPreAnimatedCaptureSource&) = delete;
+
+	FScopedPreAnimatedCaptureSource(FScopedPreAnimatedCaptureSource&&) = delete;
+	void operator=(FScopedPreAnimatedCaptureSource&&) = delete;
+
+	MOVIESCENE_API ~FScopedPreAnimatedCaptureSource();
+
+	bool WantsRestoreState() const
+	{
+		return bWantsRestoreState;
+	}
+
+private:
+
+	static FScopedPreAnimatedCaptureSource*& GetCaptureSourcePtr();
+
+	void BeginTracking(const UE::MovieScene::FPreAnimatedStateMetaData& MetaData, UMovieSceneEntitySystemLinker* Linker);
+	UE::MovieScene::FInstanceHandle GetRootInstanceHandle(UMovieSceneEntitySystemLinker* Linker) const;
+
+	friend UE::MovieScene::FPreAnimatedStateExtension;
+	struct FEvalHookType
+	{
+		const UObject* EvalHook;
+		FMovieSceneSequenceID SequenceID;
+	};
+	using CaptureSourceType = TVariant<FMovieSceneEvaluationKey, FEvalHookType, UMovieSceneTrackInstance*, FMovieSceneTrackInstanceInput>;
+
+	CaptureSourceType Variant;
+	FMovieScenePreAnimatedState* OptionalSequencePreAnimatedState;
+	FScopedPreAnimatedCaptureSource* PrevCaptureSource;
+	bool bWantsRestoreState;
+};
\ No newline at end of file
diff --git a/Engine/Source/Runtime/MovieScene/Public/Evaluation/PreAnimatedState/MovieScenePreAnimatedCaptureSources.h b/Engine/Source/Runtime/MovieScene/Public/Evaluation/PreAnimatedState/MovieScenePreAnimatedCaptureSources.h
--- a/Engine/Source/Runtime/MovieScene/Public/Evaluation/PreAnimatedState/MovieScenePreAnimatedCaptureSources.h
+++ b/Engine/Source/Runtime/MovieScene/Public/Evaluation/PreAnimatedState/MovieScenePreAnimatedCaptureSources.h
@@ -103,7 +103,7 @@ struct FPreAnimatedEvaluationHookCaptureSources : TPreAnimatedCaptureSources<FPr
 
 
 /**
- * Structure responsible for tracking contributions to pre-eanimated state entries that originate from track templates (ie, from an IMovieSceneExecutionToken::Execute)
+ * Structure responsible for tracking contributions to pre-eanimated state entries that originate from specific track template inputs
  */
 struct FPreAnimatedTrackInstanceCaptureSources : TPreAnimatedCaptureSources<FObjectKey>
 {
@@ -122,5 +122,15 @@ struct FPreAnimatedTrackInstanceCaptureSources : TPreAnimatedCaptureSources<FObj
 };
 
 
+/**
+ * Structure responsible for tracking contributions to pre-eanimated state entries that originate from track templates (ie, from an IMovieSceneExecutionToken::Execute)
+ */
+struct FPreAnimatedTrackInstanceInputCaptureSources : TPreAnimatedCaptureSources<FMovieSceneTrackInstanceInput>
+{
+	MOVIESCENE_API FPreAnimatedTrackInstanceInputCaptureSources(FPreAnimatedStateExtension* InOwner);
+};
+
+
+
 } // namespace MovieScene
 } // namespace UE
diff --git a/Engine/Source/Runtime/MovieScene/Public/Evaluation/PreAnimatedState/MovieScenePreAnimatedMasterTokenStorage.h b/Engine/Source/Runtime/MovieScene/Public/Evaluation/PreAnimatedState/MovieScenePreAnimatedMasterTokenStorage.h
--- a/Engine/Source/Runtime/MovieScene/Public/Evaluation/PreAnimatedState/MovieScenePreAnimatedMasterTokenStorage.h
+++ b/Engine/Source/Runtime/MovieScene/Public/Evaluation/PreAnimatedState/MovieScenePreAnimatedMasterTokenStorage.h
@@ -29,7 +29,7 @@ struct FPreAnimatedMasterTokenTraits
 	}
 };
 
-struct MOVIESCENE_API FAnimTypePreAnimatedStateMasterStorage : TPreAnimatedStateStorage<FPreAnimatedMasterTokenTraits>, IPreAnimatedStateGroupManager
+struct MOVIESCENE_API FAnimTypePreAnimatedStateMasterStorage : TPreAnimatedStateStorage<FPreAnimatedMasterTokenTraits>
 {
 	static TAutoRegisterPreAnimatedStorageID<FAnimTypePreAnimatedStateMasterStorage> StorageID;
 
@@ -39,15 +39,6 @@ public:
 
 	FPreAnimatedStorageID GetStorageType() const override { return StorageID; }
 	void Initialize(FPreAnimatedStorageID InStorageID, FPreAnimatedStateExtension* ParentExtension) override;
-
-	void InitializeGroupManager(FPreAnimatedStateExtension* Extension) override;
-	void OnGroupDestroyed(FPreAnimatedStorageGroupHandle Group) override;
-
-private:
-
-
-	TMap<FMovieSceneAnimTypeID, FPreAnimatedStorageGroupHandle> GroupsByAnimTypeID;
-	FPreAnimatedStorageGroupHandle GroupHandle;
 };
 
 
diff --git a/Engine/Source/Runtime/MovieScene/Public/Evaluation/PreAnimatedState/MovieScenePreAnimatedObjectGroupManager.h b/Engine/Source/Runtime/MovieScene/Public/Evaluation/PreAnimatedState/MovieScenePreAnimatedObjectGroupManager.h
--- a/Engine/Source/Runtime/MovieScene/Public/Evaluation/PreAnimatedState/MovieScenePreAnimatedObjectGroupManager.h
+++ b/Engine/Source/Runtime/MovieScene/Public/Evaluation/PreAnimatedState/MovieScenePreAnimatedObjectGroupManager.h
@@ -13,27 +13,13 @@ namespace UE
 namespace MovieScene
 {
 
-struct MOVIESCENE_API FPreAnimatedObjectGroupManager : IPreAnimatedStateGroupManager, TSharedFromThis<FPreAnimatedObjectGroupManager>
+struct MOVIESCENE_API FPreAnimatedObjectGroupManager : TPreAnimatedStateGroupManager<FObjectKey>, TSharedFromThis<FPreAnimatedObjectGroupManager>
 {
 	static TAutoRegisterPreAnimatedStorageID<FPreAnimatedObjectGroupManager> GroupManagerID;
 
-	void InitializeGroupManager(FPreAnimatedStateExtension* Extension) override;
-	void OnGroupDestroyed(FPreAnimatedStorageGroupHandle Group) override;
-
-	FPreAnimatedStorageGroupHandle FindGroupForObject(const FObjectKey& Object) const;
-
-	FPreAnimatedStorageGroupHandle MakeGroupForObject(const FObjectKey& Object);
-
 	void OnObjectsReplaced(const TMap<UObject*, UObject*>& ReplacementMap);
 
 	void GetGroupsByClass(UClass* GeneratedClass, TArray<FPreAnimatedStorageGroupHandle>& OutGroupHandles);
-
-private:
-
-	TMap<FObjectKey, FPreAnimatedStorageGroupHandle> StorageGroupsByObject;
-	TMap<FPreAnimatedStorageGroupHandle, FObjectKey> StorageGroupsToObject;
-
-	FPreAnimatedStateExtension* Extension;
 };
 
 
diff --git a/Engine/Source/Runtime/MovieScene/Public/Evaluation/PreAnimatedState/MovieScenePreAnimatedObjectStorage.h b/Engine/Source/Runtime/MovieScene/Public/Evaluation/PreAnimatedState/MovieScenePreAnimatedObjectStorage.h
--- a/Engine/Source/Runtime/MovieScene/Public/Evaluation/PreAnimatedState/MovieScenePreAnimatedObjectStorage.h
+++ b/Engine/Source/Runtime/MovieScene/Public/Evaluation/PreAnimatedState/MovieScenePreAnimatedObjectStorage.h
@@ -60,17 +60,21 @@ public:
 
 	void BeginTrackingEntities(const FPreAnimatedTrackerParams& Params, TRead<FMovieSceneEntityID> EntityIDs, TRead<FInstanceHandle> InstanceHandles, TRead<UObject*> BoundObjects) override
 	{
-		FPreAnimatedEntityCaptureSource* EntityMetaData = this->ParentExtension->GetOrCreateEntityMetaData();
-
 		const int32 Num = Params.Num;
 		const bool  bWantsRestore = Params.bWantsRestoreState;
 
+		if (!this->ParentExtension->IsCapturingGlobalState() && !bWantsRestore)
+		{
+			return;
+		}
+
+		FPreAnimatedEntityCaptureSource* EntityMetaData = this->ParentExtension->GetOrCreateEntityMetaData();
 		for (int32 Index = 0; Index < Num; ++Index)
 		{
 			UObject* BoundObject = BoundObjects[Index];
 			FObjectKey Key{ BoundObject };
 
-			FPreAnimatedStorageGroupHandle GroupHandle  = this->ObjectGroupManager->MakeGroupForObject(Key);
+			FPreAnimatedStorageGroupHandle GroupHandle  = this->ObjectGroupManager->MakeGroupForKey(Key);
 			FPreAnimatedStorageIndex       StorageIndex = this->GetOrCreateStorageIndex(Key);
 
 			FPreAnimatedStateEntry Entry{ GroupHandle, FPreAnimatedStateCachedValueHandle{ this->StorageID, StorageIndex } };
@@ -91,10 +95,18 @@ public:
 		FObjectKey Key{ BoundObject };
 
 		FPreAnimatedStorageIndex       StorageIndex = this->GetOrCreateStorageIndex(Key);
-		FPreAnimatedStorageGroupHandle GroupHandle  = this->ObjectGroupManager->MakeGroupForObject(Key);
+		FPreAnimatedStorageGroupHandle GroupHandle  = this->ObjectGroupManager->MakeGroupForKey(Key);
 		FPreAnimatedStateEntry         Entry        = FPreAnimatedStateEntry{ GroupHandle, FPreAnimatedStateCachedValueHandle{ this->StorageID, StorageIndex } };
 
-		this->ParentExtension->EnsureMetaData(Entry);
+		if (this->ParentExtension->IsCapturingGlobalState())
+		{
+			this->ParentExtension->EnsureMetaData(Entry);
+		}
+		else if (!this->ParentExtension->MetaDataExists(Entry))
+		{
+			return;
+		}
+
 		CachePreAnimatedValue(Params, Entry, BoundObject);
 	}
 
diff --git a/Engine/Source/Runtime/MovieScene/Public/Evaluation/PreAnimatedState/MovieScenePreAnimatedPropertyStorage.h b/Engine/Source/Runtime/MovieScene/Public/Evaluation/PreAnimatedState/MovieScenePreAnimatedPropertyStorage.h
--- a/Engine/Source/Runtime/MovieScene/Public/Evaluation/PreAnimatedState/MovieScenePreAnimatedPropertyStorage.h
+++ b/Engine/Source/Runtime/MovieScene/Public/Evaluation/PreAnimatedState/MovieScenePreAnimatedPropertyStorage.h
@@ -85,11 +85,15 @@ struct TPreAnimatedPropertyStorageImpl<PropertyTraits, TPropertyMetaData<MetaDat
 
 	void BeginTrackingEntities(const FPreAnimatedTrackerParams& Params, TRead<FMovieSceneEntityID> EntityIDs, TRead<FInstanceHandle> InstanceHandles, TRead<UObject*> BoundObjects, TRead<FMovieScenePropertyBinding> PropertyBindings) override
 	{
-		FPreAnimatedEntityCaptureSource* EntityMetaData = ParentExtension->GetOrCreateEntityMetaData();
-
 		const int32 Num = Params.Num;
 		const bool  bWantsRestore = Params.bWantsRestoreState;
 
+		if (!this->ParentExtension->IsCapturingGlobalState() && !bWantsRestore)
+		{
+			return;
+		}
+
+		FPreAnimatedEntityCaptureSource* EntityMetaData = ParentExtension->GetOrCreateEntityMetaData();
 		for (int32 Index = 0; Index < Num; ++Index)
 		{
 			UObject* BoundObject  = BoundObjects[Index];
@@ -97,7 +101,7 @@ struct TPreAnimatedPropertyStorageImpl<PropertyTraits, TPropertyMetaData<MetaDat
 
 			FAnimatedPropertyKey Key{ BoundObject, PropertyPath };
 
-			FPreAnimatedStorageGroupHandle GroupHandle  = this->ObjectGroupManager->MakeGroupForObject(BoundObject);
+			FPreAnimatedStorageGroupHandle GroupHandle  = this->ObjectGroupManager->MakeGroupForKey(BoundObject);
 			FPreAnimatedStorageIndex       StorageIndex = this->Storage.GetOrCreateStorageIndex(Key);
 
 			FPreAnimatedStateEntry Entry{ GroupHandle, FPreAnimatedStateCachedValueHandle{ StorageID, StorageIndex } };
@@ -125,12 +129,19 @@ struct TPreAnimatedPropertyStorageImpl<PropertyTraits, TPropertyMetaData<MetaDat
 
 			FAnimatedPropertyKey Key{ BoundObject, PropertyPath };
 
-			FPreAnimatedStorageGroupHandle GroupHandle  = this->ObjectGroupManager->MakeGroupForObject(Key.BoundObject);
+			FPreAnimatedStorageGroupHandle GroupHandle  = this->ObjectGroupManager->MakeGroupForKey(Key.BoundObject);
 			FPreAnimatedStorageIndex       StorageIndex = this->Storage.GetOrCreateStorageIndex(Key);
 
 			FPreAnimatedStateEntry Entry{ GroupHandle, FPreAnimatedStateCachedValueHandle{ StorageID, StorageIndex } };
 
-			this->ParentExtension->EnsureMetaData(Entry);
+			if (this->ParentExtension->IsCapturingGlobalState())
+			{
+				this->ParentExtension->EnsureMetaData(Entry);
+			}
+			else if (!this->ParentExtension->MetaDataExists(Entry))
+			{
+				continue;
+			}
 
 			EPreAnimatedStorageRequirement StorageRequirement = this->ParentExtension->GetStorageRequirement(Entry);
 			if (!Storage.IsStorageRequirementSatisfied(StorageIndex, StorageRequirement))
diff --git a/Engine/Source/Runtime/MovieScene/Public/Evaluation/PreAnimatedState/MovieScenePreAnimatedStateExtension.h b/Engine/Source/Runtime/MovieScene/Public/Evaluation/PreAnimatedState/MovieScenePreAnimatedStateExtension.h
--- a/Engine/Source/Runtime/MovieScene/Public/Evaluation/PreAnimatedState/MovieScenePreAnimatedStateExtension.h
+++ b/Engine/Source/Runtime/MovieScene/Public/Evaluation/PreAnimatedState/MovieScenePreAnimatedStateExtension.h
@@ -5,12 +5,17 @@
 #include "CoreTypes.h"
 #include "Containers/Map.h"
 #include "Containers/Array.h"
-#include "EntitySystem/MovieSceneEntitySystemLinkerSharedExtension.h"
 #include "EntitySystem/MovieScenePropertySystemTypes.h"
+#include "Evaluation/MovieSceneAnimTypeID.h"
 #include "Evaluation/PreAnimatedState/MovieScenePreAnimatedStorageID.h"
 #include "Evaluation/PreAnimatedState/IMovieScenePreAnimatedStorage.h"
 #include "Evaluation/PreAnimatedState/MovieScenePreAnimatedStateTypes.h"
 
+struct FScopedPreAnimatedCaptureSource;
+
+struct IMovieScenePreAnimatedTokenProducer;
+struct IMovieScenePreAnimatedGlobalTokenProducer;
+
 namespace UE
 {
 namespace MovieScene
@@ -21,8 +26,10 @@ struct IPreAnimatedStorage;
 struct IPreAnimatedCaptureSource;
 struct FPreAnimatedStateExtension;
 struct FPreAnimatedEntityCaptureSource;
+struct FAnimTypePreAnimatedStateObjectStorage;
+struct FAnimTypePreAnimatedStateMasterStorage;
 struct FPreAnimatedTrackInstanceCaptureSources;
-
+struct FPreAnimatedTrackInstanceInputCaptureSources;
 
 /**
  * Interface required for any logical grouping of pre-animated state
@@ -34,6 +41,48 @@ struct IPreAnimatedStateGroupManager
 	virtual void OnGroupDestroyed(FPreAnimatedStorageGroupHandle Group) = 0;
 };
 
+template<typename KeyType>
+struct TPreAnimatedStateGroupManager : IPreAnimatedStateGroupManager, TSharedFromThis<TPreAnimatedStateGroupManager<KeyType>>
+{
+	void InitializeGroupManager(FPreAnimatedStateExtension* InExtension) override
+	{
+		Extension = InExtension;
+	}
+
+	void OnGroupDestroyed(FPreAnimatedStorageGroupHandle Group) override
+	{
+		KeyType Temp = StorageGroupsToKey.FindChecked(Group);
+
+		StorageGroupsByKey.Remove(Temp);
+		StorageGroupsToKey.Remove(Group);
+	}
+
+	FPreAnimatedStorageGroupHandle FindGroupForKey(const KeyType& InKey) const
+	{
+		return StorageGroupsByKey.FindRef(InKey);
+	}
+
+	FPreAnimatedStorageGroupHandle MakeGroupForKey(const KeyType& InKey)
+	{
+		FPreAnimatedStorageGroupHandle GroupHandle = StorageGroupsByKey.FindRef(InKey);
+		if (GroupHandle)
+		{
+			return GroupHandle;
+		}
+
+		GroupHandle = Extension->AllocateGroup(this->AsShared());
+		StorageGroupsByKey.Add(InKey, GroupHandle);
+		StorageGroupsToKey.Add(GroupHandle, InKey);
+		return GroupHandle;
+	}
+
+protected:
+
+	TMap<KeyType, FPreAnimatedStorageGroupHandle> StorageGroupsByKey;
+	TMap<FPreAnimatedStorageGroupHandle, KeyType> StorageGroupsToKey;
+
+	FPreAnimatedStateExtension* Extension;
+};
 
 
 /**
@@ -46,23 +95,16 @@ struct IPreAnimatedStateGroupManager
  *     - Pre-animated state is grouped into buckets in order to preserve significant ordering constraints (for instance, each object must restore state in the reverse order it was cached)
  *     - Groups are managed by registered IPreAnimatedStateGroupManager instances. The main implementation is FPreAnimatedObjectGroupManager, which maintains an object to group mapping
  */
-struct MOVIESCENE_API FPreAnimatedStateExtension : TSharedEntitySystemLinkerExtension<FPreAnimatedStateExtension>
+struct MOVIESCENE_API FPreAnimatedStateExtension
 {
 public:
 
-	FPreAnimatedStateExtension(UMovieSceneEntitySystemLinker* InLinker);
-
-	FPreAnimatedStateExtension(const FPreAnimatedStateExtension&) = delete;
-	FPreAnimatedStateExtension& operator=(const FPreAnimatedStateExtension&) = delete;
-
+	FPreAnimatedStateExtension();
 	~FPreAnimatedStateExtension();
 
 public:
 
-	/**
-	 * Get this extension's application-wide ID for identification within UMovieSceneEntitySystemLinker
-	 */
-	static TEntitySystemLinkerExtensionID<FPreAnimatedStateExtension> GetExtensionID();
+	void Initialize(UMovieSceneEntitySystemLinker* InLinker);
 
 	/**
 	 * Called from TAutoRegisterPreAnimatedStorageID to register a new application-wide storage type
@@ -249,6 +291,14 @@ public:
 	 */
 	bool ContainsAnyStateForInstanceHandle(FInstanceHandle RootInstanceHandle) const;
 
+
+	// Use FScopedPreAnimatedCaptureSource to capture from a specific source rather than globally
+	void SavePreAnimatedState(FMovieSceneAnimTypeID InTokenType, const IMovieScenePreAnimatedGlobalTokenProducer& Producer);
+	void SavePreAnimatedState(UObject& InObject, FMovieSceneAnimTypeID InTokenType, const IMovieScenePreAnimatedTokenProducer& Producer);
+
+	void SavePreAnimatedStateDirectly(FMovieSceneAnimTypeID InTokenType, const IMovieScenePreAnimatedGlobalTokenProducer& Producer);
+	void SavePreAnimatedStateDirectly(UObject& InObject, FMovieSceneAnimTypeID InTokenType, const IMovieScenePreAnimatedTokenProducer& Producer);
+
 public:
 
 	FPreAnimatedEntityCaptureSource* GetEntityMetaData() const;
@@ -257,10 +307,17 @@ public:
 	FPreAnimatedTrackInstanceCaptureSources* GetTrackInstanceMetaData() const;
 	FPreAnimatedTrackInstanceCaptureSources* GetOrCreateTrackInstanceMetaData();
 
+	FPreAnimatedTrackInstanceInputCaptureSources* GetTrackInstanceInputMetaData() const;
+	FPreAnimatedTrackInstanceInputCaptureSources* GetOrCreateTrackInstanceInputMetaData();
+
+	bool HasActiveCaptureSource() const;
+
 	void AddWeakCaptureSource(TWeakPtr<IPreAnimatedCaptureSource> InWeakMetaData);
 	void RemoveWeakCaptureSource(TWeakPtr<IPreAnimatedCaptureSource> InWeakMetaData);
 
 	void EnsureMetaData(const FPreAnimatedStateEntry& Entry);
+	void AddSourceMetaData(const FPreAnimatedStateEntry& Entry);
+	bool MetaDataExists(const FPreAnimatedStateEntry& Entry) const;
 
 	void AddMetaData(const FPreAnimatedStateMetaData& MetaData);
 	void RemoveMetaData(const FPreAnimatedStateMetaData& MetaData);
@@ -271,6 +328,7 @@ public:
 private:
 
 	void FreeGroupInternal(FPreAnimatedStorageGroupHandle Handle);
+	bool ShouldCaptureAnyState() const;
 
 public:
 
@@ -285,11 +343,21 @@ private:
 	/** Meta-data pertaining to pre-animated state originating from track instances */
 	TUniquePtr<FPreAnimatedTrackInstanceCaptureSources> TrackInstanceCaptureSource;
 
+	/** Meta-data pertaining to pre-animated state originating from track instances from a specific input */
+	TUniquePtr<FPreAnimatedTrackInstanceInputCaptureSources> TrackInstanceInputCaptureSource;
+
 	/** Weakly held meta data provided by FMovieScenePreAnimatedState for various other origins */
 	TArray<TWeakPtr<IPreAnimatedCaptureSource>> WeakExternalCaptureSources;
 
+	/** Pointers to the storage for state bound to objects, organized by FMovieSceneAnimTypeID */
+	TWeakPtr<FAnimTypePreAnimatedStateObjectStorage> WeakGenericObjectStorage;
+	/** Pointers to the storage for state created from master tracks, or otherwise not bound to objects */
+	TWeakPtr<FAnimTypePreAnimatedStateMasterStorage> WeakGenericMasterStorage;
+
 private:
 
+	friend struct FScopedPreAnimatedCaptureSource;
+
 	struct FAggregatePreAnimatedStateMetaData
 	{
 		FAggregatePreAnimatedStateMetaData(FPreAnimatedStateCachedValueHandle InValueHandle)
@@ -305,6 +373,11 @@ private:
 		bool bEligibleForGlobalRestore = false;
 	};
 
+	FAggregatePreAnimatedStateMetaData* FindMetaData(const FPreAnimatedStateEntry& Entry);
+	const FAggregatePreAnimatedStateMetaData* FindMetaData(const FPreAnimatedStateEntry& Entry) const;
+
+	FAggregatePreAnimatedStateMetaData* GetOrAddMetaDataInternal(const FPreAnimatedStateEntry& Entry);
+
 	struct FPreAnimatedGroupMetaData
 	{
 		TSharedPtr<IPreAnimatedStateGroupManager> GroupManagerPtr;
@@ -313,9 +386,14 @@ private:
 
 	TSparseArray<FPreAnimatedGroupMetaData> GroupMetaData;
 
+	TMap<FPreAnimatedStateCachedValueHandle, FAggregatePreAnimatedStateMetaData> UngroupedMetaData;
+
 	TSortedMap<FPreAnimatedStorageID, TSharedPtr<IPreAnimatedStorage>> StorageImplementations;
 	TSortedMap<FPreAnimatedStorageID, TWeakPtr<IPreAnimatedStateGroupManager>> GroupManagers;
 
+	/** Linker is always valid because this extension is owned by the Linker */
+	UMovieSceneEntitySystemLinker* Linker;
+
 	bool bEntriesInvalidated;
 
 private:
diff --git a/Engine/Source/Runtime/MovieScene/Public/Evaluation/PreAnimatedState/MovieScenePreAnimatedStateStorage.h b/Engine/Source/Runtime/MovieScene/Public/Evaluation/PreAnimatedState/MovieScenePreAnimatedStateStorage.h
--- a/Engine/Source/Runtime/MovieScene/Public/Evaluation/PreAnimatedState/MovieScenePreAnimatedStateStorage.h
+++ b/Engine/Source/Runtime/MovieScene/Public/Evaluation/PreAnimatedState/MovieScenePreAnimatedStateStorage.h
@@ -55,6 +55,27 @@ public:
 		return StorageID;
 	}
 
+	template<typename OnCacheValue /* StorageType(const KeyType&) */>
+	void SavePreAnimatedState(FPreAnimatedStorageGroupHandle GroupHandle, const KeyType& InKey, OnCacheValue&& CacheCallback)
+	{
+		// Find the storage index for the specific key we're animating
+		FPreAnimatedStorageIndex StorageIndex = GetOrCreateStorageIndex(InKey);
+
+		FPreAnimatedStateEntry Entry{ GroupHandle, FPreAnimatedStateCachedValueHandle{ StorageID, StorageIndex } };
+
+		// Begin tracking the entry from its capture source
+		ParentExtension->AddSourceMetaData(Entry);
+
+		EPreAnimatedStorageRequirement Requirement = ParentExtension->HasActiveCaptureSource()
+			? EPreAnimatedStorageRequirement::Transient
+			: EPreAnimatedStorageRequirement::Persistent;
+
+		if (!IsStorageRequirementSatisfied(Entry.ValueHandle.StorageIndex, Requirement))
+		{
+			AssignPreAnimatedValue(StorageIndex, Requirement, CacheCallback(InKey));
+		}
+	}
+
 	EPreAnimatedStorageRequirement RestorePreAnimatedStateStorage(FPreAnimatedStorageIndex StorageIndex, EPreAnimatedStorageRequirement SourceRequirement, EPreAnimatedStorageRequirement TargetRequirement, const FRestoreStateParams& Params) override
 	{
 		if (RestoreMask)
diff --git a/Engine/Source/Runtime/MovieSceneTracks/Private/PreAnimatedState/MovieScenePreAnimatedComponentTransformStorage.cpp b/Engine/Source/Runtime/MovieSceneTracks/Private/PreAnimatedState/MovieScenePreAnimatedComponentTransformStorage.cpp
--- a/Engine/Source/Runtime/MovieSceneTracks/Private/PreAnimatedState/MovieScenePreAnimatedComponentTransformStorage.cpp
+++ b/Engine/Source/Runtime/MovieSceneTracks/Private/PreAnimatedState/MovieScenePreAnimatedComponentTransformStorage.cpp
@@ -84,7 +84,7 @@ void FPreAnimatedComponentTransformStorage::CachePreAnimatedTransforms(const FCa
 
 		FAnimatedPropertyKey Key{ BoundObject, PropertyBinding.PropertyPath };
 
-		FPreAnimatedStorageGroupHandle GroupHandle  = ObjectGroupManager->MakeGroupForObject(BoundObject);
+		FPreAnimatedStorageGroupHandle GroupHandle  = ObjectGroupManager->MakeGroupForKey(BoundObject);
 		FPreAnimatedStorageIndex       StorageIndex = Storage.GetOrCreateStorageIndex(Key);
 
 		FPreAnimatedStateEntry Entry{ GroupHandle, FPreAnimatedStateCachedValueHandle{ StorageID, StorageIndex } };
diff --git a/Engine/Source/Runtime/MovieSceneTracks/Private/Systems/MovieSceneComponentMobilitySystem.cpp b/Engine/Source/Runtime/MovieSceneTracks/Private/Systems/MovieSceneComponentMobilitySystem.cpp
--- a/Engine/Source/Runtime/MovieSceneTracks/Private/Systems/MovieSceneComponentMobilitySystem.cpp
+++ b/Engine/Source/Runtime/MovieSceneTracks/Private/Systems/MovieSceneComponentMobilitySystem.cpp
@@ -83,7 +83,7 @@ static void GetFlattenedHierarchy(USceneComponent* SceneComponent, TArray<UScene
 
 FPreAnimatedStateEntry FPreAnimatedComponentMobilityStorage::MakeEntry(USceneComponent* InSceneComponent)
 {
-	FPreAnimatedStorageGroupHandle GroupHandle  = ObjectGroupManager->MakeGroupForObject(InSceneComponent);
+	FPreAnimatedStorageGroupHandle GroupHandle  = ObjectGroupManager->MakeGroupForKey(InSceneComponent);
 	FPreAnimatedStorageIndex       StorageIndex = GetOrCreateStorageIndex(InSceneComponent);
 
 	return FPreAnimatedStateEntry{ GroupHandle, FPreAnimatedStateCachedValueHandle{ StorageID, StorageIndex } };
diff --git a/Engine/Source/Runtime/MovieSceneTracks/Private/Systems/MovieSceneDataLayerSystem.cpp b/Engine/Source/Runtime/MovieSceneTracks/Private/Systems/MovieSceneDataLayerSystem.cpp
--- a/Engine/Source/Runtime/MovieSceneTracks/Private/Systems/MovieSceneDataLayerSystem.cpp
+++ b/Engine/Source/Runtime/MovieSceneTracks/Private/Systems/MovieSceneDataLayerSystem.cpp
@@ -621,19 +621,18 @@ void UMovieSceneDataLayerSystem::BeginTrackingEntities()
 
 	UWorld*                     World              = GetWorld();
 	UDataLayerSubsystem*        DataLayerSubsystem = World ? World->GetSubsystem<UDataLayerSubsystem>() : nullptr;
-	FPreAnimatedStateExtension* PreAnimatedState   = Linker->FindExtension<FPreAnimatedStateExtension>();
 
-	if (!DataLayerSubsystem || !PreAnimatedState)
+	if (!DataLayerSubsystem)
 	{
 		return;
 	}
 
 	FBuiltInComponentTypes*          BuiltInComponents = FBuiltInComponentTypes::Get();
 	FMovieSceneTracksComponentTypes* TracksComponents  = FMovieSceneTracksComponentTypes::Get();
-	FPreAnimatedEntityCaptureSource* EntityMetaData    = PreAnimatedState->GetOrCreateEntityMetaData();
+	FPreAnimatedEntityCaptureSource* EntityMetaData    = Linker->PreAnimatedState.GetOrCreateEntityMetaData();
 
 	// Cache the preanimated storage
-	TSharedPtr<FPreAnimatedDataLayerStorage> PreAnimatedStorage = PreAnimatedState->GetOrCreateStorage<FPreAnimatedDataLayerStorage>();
+	TSharedPtr<FPreAnimatedDataLayerStorage> PreAnimatedStorage = Linker->PreAnimatedState.GetOrCreateStorage<FPreAnimatedDataLayerStorage>();
 	WeakPreAnimatedStorage = PreAnimatedStorage;
 
 	// ---------------------------------------------------------------------------------
@@ -667,12 +666,19 @@ void UMovieSceneDataLayerSystem::BeginTrackingEntities()
 		}
 	};
 
+	FComponentMask Filter{ BuiltInComponents->Tags.NeedsLink };
+	if (!Linker->PreAnimatedState.IsCapturingGlobalState())
+	{
+		// If we're not capturing global state, only visit entities with the RestoreState tag
+		Filter.Set(BuiltInComponents->Tags.RestoreState);
+	}
+
 	// Iterate any data layer components that need link
 	FEntityTaskBuilder()
 	.ReadEntityIDs()
 	.Read(BuiltInComponents->RootInstanceHandle)
 	.Read(TracksComponents->DataLayer)
-	.FilterAll({ BuiltInComponents->Tags.NeedsLink })
+	.FilterAll(Filter)
 	.Iterate_PerAllocation(&Linker->EntityManager, GatherDataLayers);
 }
 
