Non Unity Fixes - Movie Scene

diff --git a/Engine/Source/Runtime/MovieScene/Private/EntitySystem/TrackInstance/MovieSceneTrackInstance.cpp b/Engine/Source/Runtime/MovieScene/Private/EntitySystem/TrackInstance/MovieSceneTrackInstance.cpp
--- a/Engine/Source/Runtime/MovieScene/Private/EntitySystem/TrackInstance/MovieSceneTrackInstance.cpp
+++ b/Engine/Source/Runtime/MovieScene/Private/EntitySystem/TrackInstance/MovieSceneTrackInstance.cpp
@@ -2,6 +2,7 @@
 
 #include "EntitySystem/TrackInstance/MovieSceneTrackInstance.h"
 #include "EntitySystem/BuiltInComponentTypes.h"
+#include "EntitySystem/MovieSceneEntitySystemLinker.h"
 #include "Evaluation/PreAnimatedState/MovieScenePreAnimatedCaptureSource.h"
 #include "Evaluation/PreAnimatedState/MovieScenePreAnimatedCaptureSources.h"
 #include "Algo/Sort.h"
diff --git a/Engine/Source/Runtime/MovieScene/Private/Evaluation/PreAnimatedState/MovieScenePreAnimatedStateExtension.cpp b/Engine/Source/Runtime/MovieScene/Private/Evaluation/PreAnimatedState/MovieScenePreAnimatedStateExtension.cpp
--- a/Engine/Source/Runtime/MovieScene/Private/Evaluation/PreAnimatedState/MovieScenePreAnimatedStateExtension.cpp
+++ b/Engine/Source/Runtime/MovieScene/Private/Evaluation/PreAnimatedState/MovieScenePreAnimatedStateExtension.cpp
@@ -4,6 +4,9 @@
 #include "Evaluation/PreAnimatedState/MovieScenePreAnimatedEntityCaptureSource.h"
 #include "Evaluation/PreAnimatedState/MovieScenePreAnimatedCaptureSources.h"
 #include "Evaluation/PreAnimatedState/MovieSceneRestoreStateParams.h"
+#include "Evaluation/PreAnimatedState/MovieScenePreAnimatedCaptureSource.h"
+#include "Evaluation/PreAnimatedState/MovieScenePreAnimatedObjectTokenStorage.h"
+#include "Evaluation/PreAnimatedState/MovieScenePreAnimatedMasterTokenStorage.h"
 #include "EntitySystem/MovieSceneEntitySystemTask.h"
 #include "EntitySystem/MovieSceneComponentPtr.h"
 #include "EntitySystem/BuiltInComponentTypes.h"
diff --git a/Engine/Source/Runtime/MovieScene/Public/Evaluation/PreAnimatedState/IMovieScenePreAnimatedStorage.h b/Engine/Source/Runtime/MovieScene/Public/Evaluation/PreAnimatedState/IMovieScenePreAnimatedStorage.h
--- a/Engine/Source/Runtime/MovieScene/Public/Evaluation/PreAnimatedState/IMovieScenePreAnimatedStorage.h
+++ b/Engine/Source/Runtime/MovieScene/Public/Evaluation/PreAnimatedState/IMovieScenePreAnimatedStorage.h
@@ -13,6 +13,7 @@ struct FObjectKey;
 struct FMovieSceneEvaluationKey;
 
 class FTrackInstancePropertyBindings;
+class UMovieSceneEntitySystemLinker;
 
 namespace UE
 {
diff --git a/Engine/Source/Runtime/MovieScene/Public/Evaluation/PreAnimatedState/MovieScenePreAnimatedCaptureSource.h b/Engine/Source/Runtime/MovieScene/Public/Evaluation/PreAnimatedState/MovieScenePreAnimatedCaptureSource.h
--- a/Engine/Source/Runtime/MovieScene/Public/Evaluation/PreAnimatedState/MovieScenePreAnimatedCaptureSource.h
+++ b/Engine/Source/Runtime/MovieScene/Public/Evaluation/PreAnimatedState/MovieScenePreAnimatedCaptureSource.h
@@ -15,6 +15,7 @@ struct FMovieSceneEvaluationKey;
 
 class UObject;
 class UMovieSceneTrackInstance;
+class UMovieSceneEntitySystemLinker;
 class FMovieScenePreAnimatedState;
 
 namespace UE
diff --git a/Engine/Source/Runtime/MovieScene/Public/Evaluation/PreAnimatedState/MovieScenePreAnimatedCaptureSources.h b/Engine/Source/Runtime/MovieScene/Public/Evaluation/PreAnimatedState/MovieScenePreAnimatedCaptureSources.h
--- a/Engine/Source/Runtime/MovieScene/Public/Evaluation/PreAnimatedState/MovieScenePreAnimatedCaptureSources.h
+++ b/Engine/Source/Runtime/MovieScene/Public/Evaluation/PreAnimatedState/MovieScenePreAnimatedCaptureSources.h
@@ -12,6 +12,7 @@
 
 #include "Evaluation/PreAnimatedState/MovieScenePreAnimatedStateTypes.h"
 #include "Evaluation/PreAnimatedState/IMovieScenePreAnimatedCaptureSource.h"
+#include "EntitySystem/TrackInstance/MovieSceneTrackInstance.h"
 
 class UMovieSceneTrackInstance;
 
diff --git a/Engine/Source/Runtime/MovieScene/Public/Evaluation/PreAnimatedState/MovieScenePreAnimatedStateExtension.h b/Engine/Source/Runtime/MovieScene/Public/Evaluation/PreAnimatedState/MovieScenePreAnimatedStateExtension.h
--- a/Engine/Source/Runtime/MovieScene/Public/Evaluation/PreAnimatedState/MovieScenePreAnimatedStateExtension.h
+++ b/Engine/Source/Runtime/MovieScene/Public/Evaluation/PreAnimatedState/MovieScenePreAnimatedStateExtension.h
@@ -4,6 +4,7 @@
 
 #include "CoreTypes.h"
 #include "Containers/Map.h"
+#include "Containers/SortedMap.h"
 #include "Containers/Array.h"
 #include "EntitySystem/MovieScenePropertySystemTypes.h"
 #include "Evaluation/MovieSceneAnimTypeID.h"
@@ -15,6 +16,7 @@ struct FScopedPreAnimatedCaptureSource;
 
 struct IMovieScenePreAnimatedTokenProducer;
 struct IMovieScenePreAnimatedGlobalTokenProducer;
+class UMovieSceneEntitySystemLinker;
 
 namespace UE
 {
