Fix typo in comment

diff --git a/Engine/Source/Runtime/MovieScene/Public/IMovieScenePlayer.h b/Engine/Source/Runtime/MovieScene/Public/IMovieScenePlayer.h
--- a/Engine/Source/Runtime/MovieScene/Public/IMovieScenePlayer.h
+++ b/Engine/Source/Runtime/MovieScene/Public/IMovieScenePlayer.h
@@ -171,7 +171,7 @@ public:
 	virtual FMovieSceneSpawnRegister& GetSpawnRegister() { return NullRegister; }
 
 	/*
-	 * Called wehn an object is spawned by sequencer
+	 * Called when an object is spawned by sequencer
 	 * 
 	 */
 	virtual void OnObjectSpawned(UObject* InObject, const FMovieSceneEvaluationOperand& Operand) {}
