Sequencer: Pre-Animated state changes to support track instance inputs [Fix case]

  - Track Instance Inputs can now be used as capture sources to capture pre-animated state
  - PreAnimatedState extension now simply exists as a member of UMovieSceneEntitySystem. Previously the existence of the extension was used to determine whether any pre-animated state logic was required, but in practice this created more problems than it solved and led to some convoluted lifetime management code.
  - Added a templated group manager type from which the object group manager can inherit. This simplifies the introduction of new group managers
  - Removed the requirement that all pre-animated state must exist as part of a valid group. This forced state from 'master' tracks to be arbitrarily grouped together which put an unnecessary burden on the storage classes.

diff --git a/Engine/Source/Runtime/MovieScene/Private/Evaluation/PreAnimatedState/MovieScenePreAnimatedCaptureSource.cpp b/Engine/Source/Runtime/MovieScene/Private/Evaluation/PreAnimatedState/MovieScenePreAnimatedCaptureSource.cpp
new file mode 100644
--- /dev/null
+++ b/Engine/Source/Runtime/MovieScene/Private/Evaluation/PreAnimatedState/MovieScenePreAnimatedCaptureSource.cpp
@@ -0,0 +1,114 @@
+// Copyright Epic Games, Inc. All Rights Reserved.
+
+#include "Evaluation/PreAnimatedState/MovieScenePreAnimatedCaptureSource.h"
+#include "Evaluation/PreAnimatedState/MovieScenePreAnimatedCaptureSources.h"
+#include "Evaluation/MovieScenePreAnimatedState.h"
+
+FScopedPreAnimatedCaptureSource::FScopedPreAnimatedCaptureSource(FMovieScenePreAnimatedState* InPreAnimatedState, const FMovieSceneEvaluationKey& InEvalKey, bool bInWantsRestoreState)
+	: Variant(TInPlaceType<FMovieSceneEvaluationKey>(), InEvalKey)
+	, OptionalSequencePreAnimatedState(InPreAnimatedState)
+	, bWantsRestoreState(bInWantsRestoreState)
+{
+	FScopedPreAnimatedCaptureSource*& CaptureSourcePtr = GetCaptureSourcePtr();
+	PrevCaptureSource = CaptureSourcePtr;
+	CaptureSourcePtr = this;
+}
+FScopedPreAnimatedCaptureSource::FScopedPreAnimatedCaptureSource(FMovieScenePreAnimatedState* InPreAnimatedState, const UObject* InEvalHook, FMovieSceneSequenceID InSequenceID, bool bInWantsRestoreState)
+	: Variant(TInPlaceType<FEvalHookType>(), FEvalHookType{InEvalHook, InSequenceID} )
+	, OptionalSequencePreAnimatedState(InPreAnimatedState)
+	, bWantsRestoreState(bInWantsRestoreState)
+{
+	FScopedPreAnimatedCaptureSource*& CaptureSourcePtr = GetCaptureSourcePtr();
+	PrevCaptureSource = CaptureSourcePtr;
+	CaptureSourcePtr = this;
+}
+FScopedPreAnimatedCaptureSource::FScopedPreAnimatedCaptureSource(UMovieSceneEntitySystemLinker* InLinker, UMovieSceneTrackInstance* InTrackInstance, bool bInWantsRestoreState)
+	: Variant(TInPlaceType<UMovieSceneTrackInstance*>(), InTrackInstance)
+	, OptionalSequencePreAnimatedState(nullptr)
+	, bWantsRestoreState(bInWantsRestoreState)
+{
+	FScopedPreAnimatedCaptureSource*& CaptureSourcePtr = GetCaptureSourcePtr();
+	PrevCaptureSource = CaptureSourcePtr;
+	CaptureSourcePtr = this;
+}
+FScopedPreAnimatedCaptureSource::FScopedPreAnimatedCaptureSource(UMovieSceneEntitySystemLinker* InLinker, const FMovieSceneTrackInstanceInput& TrackInstanceInput)
+	: Variant(TInPlaceType<FMovieSceneTrackInstanceInput>(), TrackInstanceInput)
+	, OptionalSequencePreAnimatedState(nullptr)
+{
+	EMovieSceneCompletionMode CompletionMode = TrackInstanceInput.Section->GetCompletionMode();
+	if (CompletionMode == EMovieSceneCompletionMode::ProjectDefault)
+	{
+		CompletionMode = TrackInstanceInput.Section->GetTypedOuter<UMovieSceneSequence>()->DefaultCompletionMode;
+	}
+	bWantsRestoreState = (CompletionMode == EMovieSceneCompletionMode::RestoreState);
+
+	FScopedPreAnimatedCaptureSource*& CaptureSourcePtr = GetCaptureSourcePtr();
+	PrevCaptureSource = CaptureSourcePtr;
+	CaptureSourcePtr = this;
+}
+
+FScopedPreAnimatedCaptureSource::~FScopedPreAnimatedCaptureSource()
+{
+	GetCaptureSourcePtr() = PrevCaptureSource;
+}
+
+FScopedPreAnimatedCaptureSource*& FScopedPreAnimatedCaptureSource::GetCaptureSourcePtr()
+{
+	// Implemented as a static thread-local for now since there are some tests that run without a linker,
+	// so we can't put this on UMovieSceneEntitySystemLinker::PreAnimatedState where it should probably belong
+	static thread_local FScopedPreAnimatedCaptureSource* GCaptureSource = nullptr;
+	return GCaptureSource;
+}
+
+UE::MovieScene::FInstanceHandle FScopedPreAnimatedCaptureSource::GetRootInstanceHandle(UMovieSceneEntitySystemLinker* Linker) const
+{
+	if (OptionalSequencePreAnimatedState)
+	{
+		return OptionalSequencePreAnimatedState->InstanceHandle;
+	}
+
+	if (const FMovieSceneTrackInstanceInput* TrackInstanceInput = Variant.TryGet<FMovieSceneTrackInstanceInput>())
+	{
+		return Linker->GetInstanceRegistry()->GetInstance(TrackInstanceInput->InstanceHandle).GetRootInstanceHandle();
+	}
+	return UE::MovieScene::FInstanceHandle();
+}
+
+void FScopedPreAnimatedCaptureSource::BeginTracking(const UE::MovieScene::FPreAnimatedStateMetaData& MetaData, UMovieSceneEntitySystemLinker* Linker)
+{
+	using namespace UE::MovieScene;
+
+	if (FMovieSceneEvaluationKey* EvalKey = Variant.TryGet<FMovieSceneEvaluationKey>())
+	{
+		check(OptionalSequencePreAnimatedState);
+		// Make the association to this track template key
+		if (!OptionalSequencePreAnimatedState->TemplateMetaData)
+		{
+			OptionalSequencePreAnimatedState->TemplateMetaData = MakeShared<FPreAnimatedTemplateCaptureSources>(&Linker->PreAnimatedState);
+			Linker->PreAnimatedState.AddWeakCaptureSource(OptionalSequencePreAnimatedState->TemplateMetaData);
+		}
+		OptionalSequencePreAnimatedState->TemplateMetaData->BeginTrackingCaptureSource(*EvalKey, MetaData);
+	}
+	else if (FScopedPreAnimatedCaptureSource::FEvalHookType* EvalHook = Variant.TryGet<FScopedPreAnimatedCaptureSource::FEvalHookType>())
+	{
+		check(OptionalSequencePreAnimatedState);
+		if (!OptionalSequencePreAnimatedState->EvaluationHookMetaData)
+		{
+			OptionalSequencePreAnimatedState->EvaluationHookMetaData = MakeShared<FPreAnimatedEvaluationHookCaptureSources>(&Linker->PreAnimatedState);
+			Linker->PreAnimatedState.AddWeakCaptureSource(OptionalSequencePreAnimatedState->EvaluationHookMetaData);
+		}
+		OptionalSequencePreAnimatedState->EvaluationHookMetaData->BeginTrackingCaptureSource(EvalHook->EvalHook, EvalHook->SequenceID, MetaData);
+	}
+	else if (UMovieSceneTrackInstance* const * TrackInstance = Variant.TryGet<UMovieSceneTrackInstance*>())
+	{
+		// Track instance meta-data is shared between all players
+		FPreAnimatedTrackInstanceCaptureSources* TrackInstanceMetaData = Linker->PreAnimatedState.GetOrCreateTrackInstanceMetaData();
+		TrackInstanceMetaData->BeginTrackingCaptureSource(*TrackInstance, MetaData);
+	}
+	else if (const FMovieSceneTrackInstanceInput* TrackInstanceInput = Variant.TryGet<FMovieSceneTrackInstanceInput>())
+	{
+		// Track instance meta-data is shared between all players
+		FPreAnimatedTrackInstanceInputCaptureSources* TrackInstanceInputMetaData = Linker->PreAnimatedState.GetOrCreateTrackInstanceInputMetaData();
+		TrackInstanceInputMetaData->BeginTrackingCaptureSource(*TrackInstanceInput, MetaData);
+	}
+}
diff --git a/Engine/Source/Runtime/MovieScene/Public/Evaluation/PreAnimatedState/MovieScenePreAnimatedCaptureSource.h b/Engine/Source/Runtime/MovieScene/Public/Evaluation/PreAnimatedState/MovieScenePreAnimatedCaptureSource.h
new file mode 100644
--- /dev/null
+++ b/Engine/Source/Runtime/MovieScene/Public/Evaluation/PreAnimatedState/MovieScenePreAnimatedCaptureSource.h
@@ -0,0 +1,89 @@
+// Copyright Epic Games, Inc. All Rights Reserved.
+
+#pragma once
+
+#include "CoreTypes.h"
+#include "Misc/TVariant.h"
+#include "MovieSceneSequenceID.h"
+#include "Evaluation/MovieSceneEvaluationKey.h"
+#include "Evaluation/PreAnimatedState/MovieScenePreAnimatedStateTypes.h"
+#include "EntitySystem/MovieSceneSequenceInstanceHandle.h"
+#include "EntitySystem/TrackInstance/MovieSceneTrackInstance.h"
+
+struct FMovieSceneSequenceID;
+struct FMovieSceneEvaluationKey;
+
+class UObject;
+class UMovieSceneTrackInstance;
+class UMovieSceneEntitySystemLinker;
+class FMovieScenePreAnimatedState;
+
+namespace UE
+{
+namespace MovieScene
+{
+
+struct FPreAnimatedStateMetaData;
+
+}
+}
+
+
+/**
+ * Scoped structure that can be used to wrap a call to SavePreAnimatedState to associate the capture with a specific capture source.
+ */
+struct FScopedPreAnimatedCaptureSource
+{
+	/**
+	 * Construct this capture source from a template (FMovieSceneEvalTemplate) evaluation key, and whether this should restore state when the template is finished
+	 */
+	MOVIESCENE_API explicit FScopedPreAnimatedCaptureSource(FMovieScenePreAnimatedState* InPreAnimatedState, const FMovieSceneEvaluationKey& InEvalKey, bool bInWantsRestoreState);
+
+	/**
+	 * Construct this capture source from an evaluation hook (UMovieSceneEvaluationHookSection), its SequenceID, and whether this should restore state when the template is finished
+	 */
+	MOVIESCENE_API explicit FScopedPreAnimatedCaptureSource(FMovieScenePreAnimatedState* InPreAnimatedState, const UObject* InEvalHook, FMovieSceneSequenceID SequenceID, bool bInWantsRestoreState);
+
+	/**
+	 * Construct this capture source from a track instance (UMovieSceneTrackInstance) and whether this should restore state when the template is finished
+	 */
+	MOVIESCENE_API explicit FScopedPreAnimatedCaptureSource(UMovieSceneEntitySystemLinker* InLinker, UMovieSceneTrackInstance* InTrackInstance, bool bInWantsRestoreState);
+
+	/**
+	 * Construct this capture source from a track instance input (UMovieSceneTrackInstance + FMovieSceneTrackInstanceInput) and whether this should restore state when the template is finished
+	 */
+	MOVIESCENE_API explicit FScopedPreAnimatedCaptureSource(UMovieSceneEntitySystemLinker* InLinker, const FMovieSceneTrackInstanceInput& TrackInstanceInput);
+
+	FScopedPreAnimatedCaptureSource(const FScopedPreAnimatedCaptureSource&) = delete;
+	void operator=(const FScopedPreAnimatedCaptureSource&) = delete;
+
+	FScopedPreAnimatedCaptureSource(FScopedPreAnimatedCaptureSource&&) = delete;
+	void operator=(FScopedPreAnimatedCaptureSource&&) = delete;
+
+	MOVIESCENE_API ~FScopedPreAnimatedCaptureSource();
+
+	bool WantsRestoreState() const
+	{
+		return bWantsRestoreState;
+	}
+
+private:
+
+	static FScopedPreAnimatedCaptureSource*& GetCaptureSourcePtr();
+
+	void BeginTracking(const UE::MovieScene::FPreAnimatedStateMetaData& MetaData, UMovieSceneEntitySystemLinker* Linker);
+	UE::MovieScene::FInstanceHandle GetRootInstanceHandle(UMovieSceneEntitySystemLinker* Linker) const;
+
+	friend UE::MovieScene::FPreAnimatedStateExtension;
+	struct FEvalHookType
+	{
+		const UObject* EvalHook;
+		FMovieSceneSequenceID SequenceID;
+	};
+	using CaptureSourceType = TVariant<FMovieSceneEvaluationKey, FEvalHookType, UMovieSceneTrackInstance*, FMovieSceneTrackInstanceInput>;
+
+	CaptureSourceType Variant;
+	FMovieScenePreAnimatedState* OptionalSequencePreAnimatedState;
+	FScopedPreAnimatedCaptureSource* PrevCaptureSource;
+	bool bWantsRestoreState;
+};
\ No newline at end of file
