Sequencer: Fix non-unity

diff --git a/Engine/Source/Runtime/MovieScene/Private/Evaluation/PreAnimatedState/MovieScenePreAnimatedCaptureSource.cpp b/Engine/Source/Runtime/MovieScene/Private/Evaluation/PreAnimatedState/MovieScenePreAnimatedCaptureSource.cpp
--- a/Engine/Source/Runtime/MovieScene/Private/Evaluation/PreAnimatedState/MovieScenePreAnimatedCaptureSource.cpp
+++ b/Engine/Source/Runtime/MovieScene/Private/Evaluation/PreAnimatedState/MovieScenePreAnimatedCaptureSource.cpp
@@ -1,10 +1,11 @@
 // Copyright Epic Games, Inc. All Rights Reserved.
 
 #include "Evaluation/PreAnimatedState/MovieScenePreAnimatedCaptureSource.h"
-#include "Evaluation/PreAnimatedState/MovieScenePreAnimatedCaptureSources.h"
+#include "Evaluation/PreAnimatedState/MovieScenePreAnimatedCaptureSources.inl"
 #include "Evaluation/MovieScenePreAnimatedState.h"
 #include "EntitySystem/MovieSceneEntitySystemLinker.h"
 #include "MovieSceneSection.h"
+#include "MovieSceneSequence.h"
 
 FScopedPreAnimatedCaptureSource::FScopedPreAnimatedCaptureSource(FMovieScenePreAnimatedState* InPreAnimatedState, const FMovieSceneEvaluationKey& InEvalKey, bool bInWantsRestoreState)
 	: Variant(TInPlaceType<FMovieSceneEvaluationKey>(), InEvalKey)
diff --git a/Engine/Source/Runtime/MovieScene/Private/Evaluation/PreAnimatedState/MovieScenePreAnimatedCaptureSources.inl b/Engine/Source/Runtime/MovieScene/Private/Evaluation/PreAnimatedState/MovieScenePreAnimatedCaptureSources.inl
--- a/Engine/Source/Runtime/MovieScene/Private/Evaluation/PreAnimatedState/MovieScenePreAnimatedCaptureSources.inl
+++ b/Engine/Source/Runtime/MovieScene/Private/Evaluation/PreAnimatedState/MovieScenePreAnimatedCaptureSources.inl
@@ -1,5 +1,7 @@
 // Copyright Epic Games, Inc. All Rights Reserved.
 
+#pragma once
+
 #include "Evaluation/PreAnimatedState/MovieScenePreAnimatedCaptureSources.h"
 #include "Evaluation/PreAnimatedState/MovieScenePreAnimatedStateExtension.h"
 #include "Evaluation/PreAnimatedState/MovieSceneRestoreStateParams.h"
