Sequencer: Relinked track instance inputs are now have OnInputRemoved/OnInputAdded called on them

  - When track instance input is re-linked (ie, something modified a section causing it to be invalidated in the track instance) we track whether the input previously existed in the inputs, and remove + re-add it if so. This ensures that OnInputAdded is always called when a section is modified. If OnInputAdded is not called then logic that operates on the section's data (which probably changed during the Modify that invalidated the input) will get missed.

In the case of the CVar track, all its pre-animated state capture happens in OnInputAdded - when CVars were added or removed from an existing section, pre-animated state would never be captured for the new data.

diff --git a/Engine/Source/Runtime/MovieScene/Private/EntitySystem/TrackInstance/MovieSceneTrackInstance.cpp b/Engine/Source/Runtime/MovieScene/Private/EntitySystem/TrackInstance/MovieSceneTrackInstance.cpp
--- a/Engine/Source/Runtime/MovieScene/Private/EntitySystem/TrackInstance/MovieSceneTrackInstance.cpp
+++ b/Engine/Source/Runtime/MovieScene/Private/EntitySystem/TrackInstance/MovieSceneTrackInstance.cpp
@@ -55,7 +55,8 @@ void UMovieSceneTrackInstance::UpdateInputs(TArray<FMovieSceneTrackInstanceInput
 
 	Algo::Sort(InNewInputs);
 
-	// Fast path if they are the same
+	// Fast path if they are the same - this includes checking whether bInputHasBeenProcessed is the
+	// same on all pre/post inputs, which technically should never happen for us to have gotten this far
 	if (Inputs == InNewInputs)
 	{
 		// We still call OnBegin/EndUpdateInputs since some of them must have been invalidated for us to get this far
@@ -77,8 +78,22 @@ void UMovieSceneTrackInstance::UpdateInputs(TArray<FMovieSceneTrackInstanceInput
 
 	for ( ; OldIndex < OldNum || NewIndex < NewNum; )
 	{
-		while (OldIndex < OldNum && NewIndex < NewNum && Inputs[OldIndex] == InNewInputs[NewIndex])
+		while (OldIndex < OldNum && NewIndex < NewNum && Inputs[OldIndex].IsSameInput(InNewInputs[NewIndex]))
 		{
+			if (!InNewInputs[NewIndex].bInputHasBeenProcessed)
+			{
+				// This input is the same as one already existing in the track instance
+				// But it is being reimported - remove it and add it again
+				OnInputRemoved(Inputs[OldIndex]);
+				if (InputMetaData)
+				{
+					InputMetaData->StopTrackingCaptureSource(Inputs[OldIndex]);
+				}
+
+				InNewInputs[NewIndex].bInputHasBeenProcessed = true;
+				OnInputAdded(InNewInputs[NewIndex]);
+			}
+
 			++OldIndex;
 			++NewIndex;
 		}
@@ -104,6 +119,7 @@ void UMovieSceneTrackInstance::UpdateInputs(TArray<FMovieSceneTrackInstanceInput
 				// and in with the new
 				FScopedPreAnimatedCaptureSource CaptureSource(PrivateLinker, InNewInputs[NewIndex]);
 
+				InNewInputs[NewIndex].bInputHasBeenProcessed = true;
 				OnInputAdded(InNewInputs[NewIndex]);
 				++NewIndex;
 			}
@@ -121,6 +137,7 @@ void UMovieSceneTrackInstance::UpdateInputs(TArray<FMovieSceneTrackInstanceInput
 		else if (ensure(NewIndex < NewNum))
 		{
 			// and in with the new
+			InNewInputs[NewIndex].bInputHasBeenProcessed = true;
 			OnInputAdded(InNewInputs[NewIndex]);
 			++NewIndex;
 		}
diff --git a/Engine/Source/Runtime/MovieScene/Private/EntitySystem/TrackInstance/MovieSceneTrackInstanceSystem.cpp b/Engine/Source/Runtime/MovieScene/Private/EntitySystem/TrackInstance/MovieSceneTrackInstanceSystem.cpp
--- a/Engine/Source/Runtime/MovieScene/Private/EntitySystem/TrackInstance/MovieSceneTrackInstanceSystem.cpp
+++ b/Engine/Source/Runtime/MovieScene/Private/EntitySystem/TrackInstance/MovieSceneTrackInstanceSystem.cpp
@@ -199,14 +199,30 @@ void UMovieSceneTrackInstanceInstantiator::OnRun(FSystemTaskPrerequisites& InPre
 	// Gather all the inputs for any invalidated output indices
 	TSortedMap<int32, TArray<FMovieSceneTrackInstanceInput> > NewInputs;
 	{
-		auto ReLinkInputs = [this, &NewInputs](FInstanceHandle SourceInstance, FTrackInstanceInputComponent InputComponent)
+		auto ReLinkInputs = [this, &NewInputs, BuiltInComponents](FEntityAllocationIteratorItem Item, const FInstanceHandle* SourceInstances, const FTrackInstanceInputComponent* InputComponents)
 		{
-			if (this->InvalidatedOutputs.IsValidIndex(InputComponent.OutputIndex) && this->InvalidatedOutputs[InputComponent.OutputIndex] == true)
+			const int32 Num = Item.GetAllocation()->Num();
+
+			// If the input does not have the NeedsLink tag, it has already been processed, so doesn't need removing and re-adding
+			const bool bInputHasBeenProcessed = !Item.GetAllocationType().Contains(BuiltInComponents->Tags.NeedsLink);
+			
+			for (int32 Index = 0; Index < Num; ++Index)
 			{
-				NewInputs.FindOrAdd(InputComponent.OutputIndex).Add(FMovieSceneTrackInstanceInput{ InputComponent.Section, SourceInstance });
+				const int32 OutputIndex = InputComponents[Index].OutputIndex;
+
+				FMovieSceneTrackInstanceInput NewInput{ InputComponents[Index].Section, SourceInstances[Index], bInputHasBeenProcessed };
+				if (this->InvalidatedOutputs.IsValidIndex(OutputIndex) && this->InvalidatedOutputs[OutputIndex] == true)
+				{
+					NewInputs.FindOrAdd(OutputIndex).Add(NewInput);
+				}
 			}
+			
 		};
-		FEntityTaskBuilder().Read(BuiltInComponents->InstanceHandle).Read(BuiltInComponents->TrackInstanceInput).FilterNone({ BuiltInComponents->Tags.NeedsUnlink }).Iterate_PerEntity(&Linker->EntityManager, ReLinkInputs);
+		FEntityTaskBuilder()
+		.Read(BuiltInComponents->InstanceHandle)
+		.Read(BuiltInComponents->TrackInstanceInput)
+		.FilterNone({ BuiltInComponents->Tags.NeedsUnlink })
+		.Iterate_PerAllocation(&Linker->EntityManager, ReLinkInputs);
 	}
 
 	// Update the inputs for each of the invalidated indices
diff --git a/Engine/Source/Runtime/MovieScene/Public/EntitySystem/TrackInstance/MovieSceneTrackInstance.h b/Engine/Source/Runtime/MovieScene/Public/EntitySystem/TrackInstance/MovieSceneTrackInstance.h
--- a/Engine/Source/Runtime/MovieScene/Public/EntitySystem/TrackInstance/MovieSceneTrackInstance.h
+++ b/Engine/Source/Runtime/MovieScene/Public/EntitySystem/TrackInstance/MovieSceneTrackInstance.h
@@ -34,13 +34,23 @@ struct FMovieSceneTrackInstanceInput
 	/** The instance that this input relates to */
 	UE::MovieScene::FInstanceHandle InstanceHandle;
 
+	/** Whether the input has been processed by being added to a track instance */
+	bool bInputHasBeenProcessed = false;
+
+	/** Check whether this input matches another, ignoring the bInputHasBeenProcessed flag */
+	bool IsSameInput(const FMovieSceneTrackInstanceInput& OtherInput) const
+	{
+		return Section == OtherInput.Section && InstanceHandle == OtherInput.InstanceHandle;
+	}
+
 	friend bool operator==(const FMovieSceneTrackInstanceInput& A, const FMovieSceneTrackInstanceInput& B)
 	{
-		return A.Section == B.Section && A.InstanceHandle == B.InstanceHandle;
+		return A.Section == B.Section && A.InstanceHandle == B.InstanceHandle && A.bInputHasBeenProcessed == B.bInputHasBeenProcessed;
 	}
 
 	friend bool operator<(const FMovieSceneTrackInstanceInput& A, const FMovieSceneTrackInstanceInput& B)
 	{
+		// bInputHasBeenProcessed is not considered for sorting
 		if (A.Section == B.Section)
 		{
 			return A.InstanceHandle < B.InstanceHandle;
@@ -60,7 +70,7 @@ struct FMovieSceneTrackInstanceInput
 	/** Serializer that ensures both the section and instance handle are copied over when an input is copied for a reinstanced animator */
 	bool Serialize(FArchive& Ar)
 	{
-		Ar << (UObject*&)Section << InstanceHandle.InstanceID << InstanceHandle.InstanceSerial;
+		Ar << (UObject*&)Section << InstanceHandle.InstanceID << InstanceHandle.InstanceSerial << bInputHasBeenProcessed;
 		return true;
 	}
 };
