Sequencer: Remove erroneous TSharedFromThis<> inheritance

With the introduction of TPreAnimatedStateGroupManager, FPreAnimatedObjectGroupManager ended up having multiple instances of TSharedFromThis implemented with different types and members resulting in one being used in some compilation contexts, and the other being used elsewhere meaning one wouldn't get initialized correctly.

diff --git a/Engine/Source/Runtime/MovieScene/Public/Evaluation/PreAnimatedState/MovieScenePreAnimatedObjectGroupManager.h b/Engine/Source/Runtime/MovieScene/Public/Evaluation/PreAnimatedState/MovieScenePreAnimatedObjectGroupManager.h
--- a/Engine/Source/Runtime/MovieScene/Public/Evaluation/PreAnimatedState/MovieScenePreAnimatedObjectGroupManager.h
+++ b/Engine/Source/Runtime/MovieScene/Public/Evaluation/PreAnimatedState/MovieScenePreAnimatedObjectGroupManager.h
@@ -13,7 +13,7 @@ namespace UE
 namespace MovieScene
 {
 
-struct MOVIESCENE_API FPreAnimatedObjectGroupManager : TPreAnimatedStateGroupManager<FObjectKey>, TSharedFromThis<FPreAnimatedObjectGroupManager>
+struct MOVIESCENE_API FPreAnimatedObjectGroupManager : TPreAnimatedStateGroupManager<FObjectKey>
 {
 	static TAutoRegisterPreAnimatedStorageID<FPreAnimatedObjectGroupManager> GroupManagerID;
 
