Sequencer: Fix FPreAnimatedTrackInstanceCaptureSources include

diff --git a/Engine/Source/Runtime/MovieScene/Private/EntitySystem/TrackInstance/MovieSceneTrackInstance.cpp b/Engine/Source/Runtime/MovieScene/Private/EntitySystem/TrackInstance/MovieSceneTrackInstance.cpp
--- a/Engine/Source/Runtime/MovieScene/Private/EntitySystem/TrackInstance/MovieSceneTrackInstance.cpp
+++ b/Engine/Source/Runtime/MovieScene/Private/EntitySystem/TrackInstance/MovieSceneTrackInstance.cpp
@@ -4,7 +4,7 @@
 #include "EntitySystem/BuiltInComponentTypes.h"
 #include "EntitySystem/MovieSceneEntitySystemLinker.h"
 #include "Evaluation/PreAnimatedState/MovieScenePreAnimatedCaptureSource.h"
-#include "Evaluation/PreAnimatedState/MovieScenePreAnimatedCaptureSources.h"
+#include "Evaluation/PreAnimatedState/MovieScenePreAnimatedCaptureSources.inl"
 #include "Algo/Sort.h"
 
 void UMovieSceneTrackInstance::Initialize(UObject* InAnimatedObject, UMovieSceneEntitySystemLinker* InLinker)
