Sequencer: Fixed bound object ptrs not being available to child entity instantiators

diff --git a/Engine/Source/Runtime/MovieScene/Private/EntitySystem/BuiltInComponentTypes.cpp b/Engine/Source/Runtime/MovieScene/Private/EntitySystem/BuiltInComponentTypes.cpp
--- a/Engine/Source/Runtime/MovieScene/Private/EntitySystem/BuiltInComponentTypes.cpp
+++ b/Engine/Source/Runtime/MovieScene/Private/EntitySystem/BuiltInComponentTypes.cpp
@@ -192,29 +192,9 @@ FBuiltInComponentTypes::FBuiltInComponentTypes()
 	ComponentRegistry->Factories.DuplicateChildComponent(PropertyBinding);
 	ComponentRegistry->Factories.DuplicateChildComponent(HierarchicalBias);
 
-	// Children always need a Parent
+	// Children always need a Parent - these are initialized by the tasks that create them
 	{
-		struct FParentEntityInitializer : TChildEntityInitializer<FMovieSceneEntityID, FMovieSceneEntityID>
-		{
-			explicit FParentEntityInitializer(TComponentTypeID<FMovieSceneEntityID> ParentEntity)
-				: TChildEntityInitializer(TComponentTypeID<FMovieSceneEntityID>(), ParentEntity)
-			{}
-
-			virtual void Run(const FEntityRange& ChildRange, const FEntityAllocation* ParentAllocation, TArrayView<const int32> ParentAllocationOffsets) override
-			{
-				TArrayView<const FMovieSceneEntityID> ParentIDs       = ParentAllocation->GetEntityIDs();
-				TComponentWriter<FMovieSceneEntityID> ChildComponents = GetChildComponents(ChildRange.Allocation);
-
-				for (int32 Index = 0; Index < ChildRange.Num; ++Index)
-				{
-					const int32 ParentIndex = ParentAllocationOffsets[Index];
-					const int32 ChildIndex  = ChildRange.ComponentStartOffset + Index;
-
-					ChildComponents[ChildIndex] = ParentIDs[ParentIndex];
-				}
-			}
-		};
-		ComponentRegistry->Factories.DefineChildComponent(FParentEntityInitializer(ParentEntity));
+		ComponentRegistry->Factories.DefineChildComponent(FComponentTypeID::Invalid(), ParentEntity);
 	}
 	
 	// Bool channel relationships
diff --git a/Engine/Source/Runtime/MovieScene/Private/EntitySystem/MovieSceneEntityFactory.cpp b/Engine/Source/Runtime/MovieScene/Private/EntitySystem/MovieSceneEntityFactory.cpp
--- a/Engine/Source/Runtime/MovieScene/Private/EntitySystem/MovieSceneEntityFactory.cpp
+++ b/Engine/Source/Runtime/MovieScene/Private/EntitySystem/MovieSceneEntityFactory.cpp
@@ -49,6 +49,7 @@ void FChildEntityFactory::Apply(UMovieSceneEntitySystemLinker* Linker, FEntityAl
 
 	int32 CurrentParentOffset = 0;
 	const FEntityAllocation* ParentAllocation = ParentAllocationProxy.GetAllocation();
+	FEntityAllocationWriteContext WriteContext = FEntityAllocationWriteContext::NewAllocation();
 
 	// We attempt to allocate all the linker entities contiguously in memory for efficient initialization,
 	// but we may reach capacity constraints within allocations so we may have to run the factories more than once
@@ -62,11 +63,24 @@ void FChildEntityFactory::Apply(UMovieSceneEntitySystemLinker* Linker, FEntityAl
 
 		CurrentEntityOffsets = MakeArrayView(ParentEntityOffsets.GetData() + CurrentParentOffset, NumAdded);
 
-		Linker->EntityManager.InitializeChildAllocation(ParentType, DerivedEntityType, ParentAllocation, CurrentEntityOffsets, ChildRange);
+		if (TOptionalComponentWriter<FMovieSceneEntityID> ParentEntityIDs =
+			ChildRange.Allocation->TryWriteComponents(FBuiltInComponentTypes::Get()->ParentEntity, WriteContext))
+		{
+			TArrayView<const FMovieSceneEntityID> ParentIDs = ParentAllocation->GetEntityIDs();
+			for (int32 Index = 0; Index < ChildRange.Num; ++Index)
+			{
+				const int32 ParentIndex = CurrentEntityOffsets[Index];
+				const int32 ChildIndex  = ChildRange.ComponentStartOffset + Index;
 
-		// Important: This must go after Linker->EntityManager.InitializeChildAllocation so that we know that parent entity IDs are initialized correctly
+				ParentEntityIDs[ChildIndex] = ParentIDs[ParentIndex];
+			}
+		}
+
+		// Initialize the bound objects before we call child initializers
 		InitializeAllocation(Linker, ParentType, DerivedEntityType, ParentAllocation, CurrentEntityOffsets, ChildRange);
 
+		Linker->EntityManager.InitializeChildAllocation(ParentType, DerivedEntityType, ParentAllocation, CurrentEntityOffsets, ChildRange);
+
 		CurrentParentOffset += NumAdded;
 	}
 
