Added a GetDescendantSubSections function to MovieSceneCommonHelpers

diff --git a/Engine/Source/Runtime/MovieScene/Private/MovieSceneCommonHelpers.cpp b/Engine/Source/Runtime/MovieScene/Private/MovieSceneCommonHelpers.cpp
--- a/Engine/Source/Runtime/MovieScene/Private/MovieSceneCommonHelpers.cpp
+++ b/Engine/Source/Runtime/MovieScene/Private/MovieSceneCommonHelpers.cpp
@@ -348,6 +348,26 @@ void MovieSceneHelpers::GetDescendantMovieScenes(UMovieSceneSequence* InSequence
 	}
 }
 
+void MovieSceneHelpers::GetDescendantSubSections(const UMovieScene* InMovieScene, TArray<UMovieSceneSubSection*>& InSubSections)
+{
+	if (!IsValid(InMovieScene))
+	{
+		return;
+	}
+
+	for (UMovieSceneSection* Section : InMovieScene->GetAllSections())
+	{
+		if (UMovieSceneSubSection* SubSection = Cast<UMovieSceneSubSection>(Section))
+		{
+			InSubSections.Add(SubSection);
+			
+			if (const UMovieSceneSequence* SubSequence = SubSection->GetSequence())
+			{
+				GetDescendantSubSections(SubSequence->GetMovieScene(), InSubSections);
+			}
+		}
+	}
+}
 
 USceneComponent* MovieSceneHelpers::SceneComponentFromRuntimeObject(UObject* Object)
 {
diff --git a/Engine/Source/Runtime/MovieScene/Public/MovieSceneCommonHelpers.h b/Engine/Source/Runtime/MovieScene/Public/MovieSceneCommonHelpers.h
--- a/Engine/Source/Runtime/MovieScene/Public/MovieSceneCommonHelpers.h
+++ b/Engine/Source/Runtime/MovieScene/Public/MovieSceneCommonHelpers.h
@@ -14,6 +14,7 @@ class AActor;
 class UCameraComponent;
 class UMovieScene;
 class UMovieSceneSection;
+class UMovieSceneSubSection;
 class UMovieSceneSequence;
 class USceneComponent;
 class USoundBase;
@@ -73,6 +74,11 @@ public:
 	 */
 	static void GetDescendantMovieScenes(UMovieSceneSequence* InSequence, TArray<UMovieScene*> & InMovieScenes);
 
+	/*
+	 * Gather up descendant movie scene sub-sections from the incoming movie scene
+	 */
+	static void GetDescendantSubSections(const UMovieScene* InMovieScene, TArray<UMovieSceneSubSection*>& InSubSections);
+	
 	/**
 	 * Get the scene component from the runtime object
 	 *
