Fix spamming of Triggers by adding new Track type

diff --git a/Engine/Plugins/Runtime/Metasound/Source/MetasoundEngine/Private/MetasoundSource.cpp b/Engine/Plugins/Runtime/Metasound/Source/MetasoundEngine/Private/MetasoundSource.cpp
--- a/Engine/Plugins/Runtime/Metasound/Source/MetasoundEngine/Private/MetasoundSource.cpp
+++ b/Engine/Plugins/Runtime/Metasound/Source/MetasoundEngine/Private/MetasoundSource.cpp
@@ -265,8 +265,30 @@ void UMetaSoundSource::SetRegistryAssetClassInfo(const Metasound::Frontend::FNod
 
 void UMetaSoundSource::InitParameters(TArray<FAudioParameter>& InParametersToInit, FName InFeatureName)
 {
-	auto Sanitize = [](FAudioParameter& Parameter)
+	TMap<FName, const FMetasoundFrontendClassInput*> TransmittableInputMap;
+	const TArray<const FMetasoundFrontendClassInput*> TransmittableInputs = GetTransmittableClassInputs();
+	Algo::Transform(TransmittableInputs, TransmittableInputMap, [](const FMetasoundFrontendClassInput* Input)
 	{
+		return TPair<FName, const FMetasoundFrontendClassInput*>(Input->Name, Input);
+	});
+	
+	// Removes values that are not explicitly defined by the ParamType and returns
+	// whether or not the parameter is a valid input and should be included.
+	auto Sanitize = [&TransmittableInputMap](FAudioParameter& Parameter) -> bool
+	{
+		if (const FMetasoundFrontendClassInput* Input = TransmittableInputMap.FindRef(Parameter.ParamName))
+		{
+			// Only remove if parameter typename is specified and does not match
+			if (Parameter.TypeName != FName() && Parameter.TypeName != Input->TypeName)
+			{
+				return false;
+			}
+		}
+		else
+		{
+			return false;
+		}
+
 		switch (Parameter.ParamType)
 		{
 			case EAudioParameterType::Boolean:
@@ -339,6 +361,8 @@ void UMetaSoundSource::InitParameters(TArray<FAudioParameter>& InParametersToIni
 			default:
 			break;
 		}
+
+		return true;
 	};
 
 	auto ConstructProxies = [this, FeatureName = InFeatureName](FAudioParameter& OutParamToInit)
@@ -392,17 +416,24 @@ void UMetaSoundSource::InitParameters(TArray<FAudioParameter>& InParametersToIni
 		return TPair<FName, FName>(Input.Name, Input.TypeName);
 	});
 
-	for (FAudioParameter& Parameter : InParametersToInit)
+	for (int32 i = InParametersToInit.Num() - 1; i >= 0; --i)
 	{
-		Sanitize(Parameter);
-
-		if (IsParameterValid(Parameter, InputNameTypeMap))
+		FAudioParameter& Parameter = InParametersToInit[i];
+		if (Sanitize(Parameter))
 		{
-			ConstructProxies(Parameter);
+			if (IsParameterValid(Parameter, InputNameTypeMap))
+			{
+				ConstructProxies(Parameter);
+			}
+			else
+			{
+				UE_LOG(LogMetaSound, Error, TEXT("Failed to set invalid parameter '%s': Either does not exist or is unsupported type"), *Parameter.ParamName.ToString());
+			}
 		}
 		else
 		{
-			UE_LOG(LogMetaSound, Error, TEXT("Failed to set invalid parameter '%s': Either does not exist or is unsupported type"), *Parameter.ParamName.ToString());
+			constexpr bool bAllowShrinking = false;
+			InParametersToInit.RemoveAtSwap(i, 1, bAllowShrinking);
 		}
 	}
 }
@@ -492,6 +523,7 @@ bool UMetaSoundSource::GetAllDefaultParameters(TArray<FAudioParameter>& OutParam
 
 		FAudioParameter Params;
 		Params.ParamName = Input->Name;
+		Params.TypeName = Input->TypeName;
 
 		switch (Input->DefaultLiteral.GetType())
 		{
diff --git a/Engine/Source/Editor/MovieSceneTools/Private/MovieSceneToolsModule.cpp b/Engine/Source/Editor/MovieSceneTools/Private/MovieSceneToolsModule.cpp
--- a/Engine/Source/Editor/MovieSceneTools/Private/MovieSceneToolsModule.cpp
+++ b/Engine/Source/Editor/MovieSceneTools/Private/MovieSceneToolsModule.cpp
@@ -73,6 +73,7 @@
 #include "Channels/MovieSceneCameraShakeSourceTriggerChannel.h"
 #include "Channels/EventChannelCurveModel.h"
 #include "Channels/SCurveEditorEventChannelView.h"
+#include "Channels/MovieSceneAudioTriggerChannel.h"
 #include "Sections/MovieSceneEventSection.h"
 
 #include "MovieSceneEventUtils.h"
@@ -165,6 +166,8 @@ void FMovieSceneToolsModule::StartupModule()
 		SequencerModule.RegisterChannelInterface<FMovieSceneEventSectionData>();
 		SequencerModule.RegisterChannelInterface<FMovieSceneObjectPathChannel>();
 
+		SequencerModule.RegisterChannelInterface<FMovieSceneAudioTriggerChannel>();
+
 		SequencerModule.RegisterChannelInterface<FMovieSceneEventChannel>();
 
 		SequencerModule.RegisterChannelInterface<FMovieSceneCameraShakeSourceTriggerChannel>();
diff --git a/Engine/Source/Runtime/AudioExtensions/Public/AudioParameter.h b/Engine/Source/Runtime/AudioExtensions/Public/AudioParameter.h
--- a/Engine/Source/Runtime/AudioExtensions/Public/AudioParameter.h
+++ b/Engine/Source/Runtime/AudioExtensions/Public/AudioParameter.h
@@ -227,7 +227,8 @@ struct AUDIOEXTENSIONS_API FAudioParameter
 		, ArrayObjectParam(MoveTemp(InParameter.ArrayObjectParam))
 		, ArrayStringParam(MoveTemp(InParameter.ArrayStringParam))
 		, ParamType(MoveTemp(InParameter.ParamType))
-		, ObjectProxies(MoveTemp(InParameter.ObjectProxies))
+		, TypeName(MoveTemp(InParameter.TypeName))
+		, ObjectProxies(MoveTemp(InParameter.ObjectProxies))		
 	{
 	}
 
@@ -244,6 +245,7 @@ struct AUDIOEXTENSIONS_API FAudioParameter
 		, ArrayObjectParam(InParameter.ArrayObjectParam)
 		, ArrayStringParam(InParameter.ArrayStringParam)
 		, ParamType(InParameter.ParamType)
+		, TypeName(InParameter.TypeName)
 	{
 		ObjectProxies.Reset();
 		for (const Audio::IProxyDataPtr& DataPtr : InParameter.ObjectProxies)
@@ -266,6 +268,7 @@ struct AUDIOEXTENSIONS_API FAudioParameter
 		ArrayObjectParam = InParameter.ArrayObjectParam;
 		ArrayStringParam = InParameter.ArrayStringParam;
 		ParamType = InParameter.ParamType;
+		TypeName = InParameter.TypeName;
 
 		ObjectProxies.Reset();
 		for (const Audio::IProxyDataPtr& DataPtr : InParameter.ObjectProxies)
@@ -290,6 +293,7 @@ struct AUDIOEXTENSIONS_API FAudioParameter
 		ArrayObjectParam = MoveTemp(InParameter.ArrayObjectParam);
 		ArrayStringParam = MoveTemp(InParameter.ArrayStringParam);
 		ParamType = MoveTemp(InParameter.ParamType);
+		TypeName = MoveTemp(InParameter.TypeName);
 
 		ObjectProxies = MoveTemp(InParameter.ObjectProxies);
 
@@ -343,6 +347,10 @@ struct AUDIOEXTENSIONS_API FAudioParameter
 	UPROPERTY(EditAnywhere, BlueprintReadWrite, meta = (DisplayName = "Type"), Category = AudioParameter)
 	EAudioParameterType ParamType = EAudioParameterType::None;
 
+	// Optional TypeName used to describe what constructed type this parameter should be initializing.
+	UPROPERTY()
+	FName TypeName;
+
 	// Object proxies to be generated when parameter is passed to the AudioThread to represent ObjectParam/ArrayObjectParam safely
 	TArray<Audio::IProxyDataPtr> ObjectProxies;
 
diff --git a/Engine/Source/Runtime/MovieScene/Private/Channels/MovieSceneAudioTriggerChannel.cpp b/Engine/Source/Runtime/MovieScene/Private/Channels/MovieSceneAudioTriggerChannel.cpp
new file mode 100644
--- /dev/null
+++ b/Engine/Source/Runtime/MovieScene/Private/Channels/MovieSceneAudioTriggerChannel.cpp
@@ -0,0 +1,101 @@
+// Copyright Epic Games, Inc. All Rights Reserved.
+
+#include "Channels/MovieSceneAudioTriggerChannel.h"
+#include "Channels/MovieSceneChannelProxy.h"
+#include "Curves/IntegralCurve.h"
+#include "Misc/FrameRate.h"
+#include "MovieSceneFwd.h"
+#include "MovieSceneFrameMigration.h"
+
+bool FMovieSceneAudioTriggerChannel::EvaluatePossibleTriggers(const FMovieSceneContext& InContext, FMoveSceneAudioTriggerState& InState, bool& OutTriggered) const
+{
+	FFrameTime CurrentTime = InContext.GetTime();
+
+	if (Times.Num())
+	{
+		bool bHasGoneBackwards = InState.PreviousUpdateTime.IsSet() && *InState.PreviousUpdateTime > CurrentTime;
+
+		// If we've somehow gone backwards, reset out state.
+		if (bHasGoneBackwards || InContext.GetDirection() == EPlayDirection::Backwards)
+		{
+			InState.PreviousIndex.Reset();
+			InState.PreviousUpdateTime.Reset();
+		}
+
+		if (InContext.GetDirection() == EPlayDirection::Forwards)
+		{
+			const int32 Index = FMath::Max(0, Algo::UpperBound(Times, CurrentTime.FrameNumber) - 1);
+			if (Times.IsValidIndex(Index) && CurrentTime.FrameNumber > Times[Index] )
+			{
+				bool bNewIndex = InState.PreviousIndex.IsSet() ? Index > *InState.PreviousIndex : true;
+				OutTriggered = bNewIndex;
+				InState.PreviousIndex = Index;
+				InState.PreviousUpdateTime = CurrentTime;
+				return true;
+			}
+		}	
+	}
+	return false;
+}
+
+void FMovieSceneAudioTriggerChannel::GetKeys(const TRange<FFrameNumber>& WithinRange, TArray<FFrameNumber>* OutKeyTimes, TArray<FKeyHandle>* OutKeyHandles)
+{
+	GetData().GetKeys(WithinRange, OutKeyTimes, OutKeyHandles);
+}
+
+void FMovieSceneAudioTriggerChannel::GetKeyTimes(TArrayView<const FKeyHandle> InHandles, TArrayView<FFrameNumber> OutKeyTimes)
+{
+	GetData().GetKeyTimes(InHandles, OutKeyTimes);
+}
+
+void FMovieSceneAudioTriggerChannel::SetKeyTimes(TArrayView<const FKeyHandle> InHandles, TArrayView<const FFrameNumber> InKeyTimes)
+{
+	GetData().SetKeyTimes(InHandles, InKeyTimes);
+}
+
+void FMovieSceneAudioTriggerChannel::DuplicateKeys(TArrayView<const FKeyHandle> InHandles, TArrayView<FKeyHandle> OutNewHandles)
+{
+	GetData().DuplicateKeys(InHandles, OutNewHandles);
+}
+
+void FMovieSceneAudioTriggerChannel::DeleteKeys(TArrayView<const FKeyHandle> InHandles)
+{
+	GetData().DeleteKeys(InHandles);
+}
+
+void FMovieSceneAudioTriggerChannel::DeleteKeysFrom(FFrameNumber InTime, bool bDeleteKeysBefore)
+{
+	GetData().DeleteKeysFrom(InTime, bDeleteKeysBefore);
+}
+
+void FMovieSceneAudioTriggerChannel::ChangeFrameResolution(FFrameRate SourceRate, FFrameRate DestinationRate)
+{
+	GetData().ChangeFrameResolution(SourceRate, DestinationRate);
+}
+
+TRange<FFrameNumber> FMovieSceneAudioTriggerChannel::ComputeEffectiveRange() const
+{
+	return GetData().GetTotalRange();
+}
+
+int32 FMovieSceneAudioTriggerChannel::GetNumKeys() const
+{
+	return Times.Num();
+}
+
+void FMovieSceneAudioTriggerChannel::Reset()
+{
+	Times.Reset();
+	Values.Reset();
+	KeyHandles.Reset();
+}
+
+void FMovieSceneAudioTriggerChannel::Optimize(const FKeyDataOptimizationParams& InParameters)
+{
+	UE::MovieScene::Optimize(this, InParameters);
+}
+
+void FMovieSceneAudioTriggerChannel::Offset(FFrameNumber DeltaPosition)
+{
+	GetData().Offset(DeltaPosition);
+}
\ No newline at end of file
diff --git a/Engine/Source/Runtime/MovieScene/Public/Channels/MovieSceneAudioTriggerChannel.h b/Engine/Source/Runtime/MovieScene/Public/Channels/MovieSceneAudioTriggerChannel.h
new file mode 100644
--- /dev/null
+++ b/Engine/Source/Runtime/MovieScene/Public/Channels/MovieSceneAudioTriggerChannel.h
@@ -0,0 +1,115 @@
+// Copyright Epic Games, Inc. All Rights Reserved.
+
+#pragma once
+
+#include "UObject/ObjectMacros.h"
+#include "Misc/FrameNumber.h"
+#include "MovieSceneChannelData.h"
+#include "MovieSceneChannelTraits.h"
+#include "MovieSceneAudioTriggerChannel.generated.h"
+
+struct FMoveSceneAudioTriggerState
+{	
+	TOptional<int32> PreviousIndex;
+	TOptional<FFrameTime> PreviousUpdateTime;
+};
+
+USTRUCT()
+struct MOVIESCENE_API FMovieSceneAudioTriggerChannel : public FMovieSceneChannel
+{
+	GENERATED_BODY()
+
+	FMovieSceneAudioTriggerChannel() = default;
+
+	/**
+	 * Access a mutable interface for this channel's data
+	 *
+	 * @return An object that is able to manipulate this channel's data
+	 */
+	FORCEINLINE TMovieSceneChannelData<bool> GetData()
+	{
+		return TMovieSceneChannelData<bool>(&Times, &Values, &KeyHandles);
+	}
+
+	/**
+	 * Access a constant interface for this channel's data
+	 *
+	 * @return An object that is able to interrogate this channel's data
+	 */
+	FORCEINLINE TMovieSceneChannelData<const bool> GetData() const
+	{
+		return TMovieSceneChannelData<const bool>(&Times, &Values);
+	}
+
+	/**
+	 * Const access to this channel's times
+	 */
+	FORCEINLINE TArrayView<const FFrameNumber> GetTimes() const
+	{
+		return Times;
+	}
+
+	/**
+	 * Check whether this channel has any data
+	 */
+	FORCEINLINE bool HasAnyData() const
+	{
+		return Times.Num() != 0;
+	}
+
+	/**
+	 * Evaluate this channel
+	 *
+	 * @param InContext   Scene Context with time and directionality necessary to compute triggers
+	 * @param InState    A persistent state regard if the trigger has fired
+	 * @param OutTriggered   A Value that if set to true combined with the function returning true means there's a trigger
+	 * @return true if the channel was evaluated successfully, false otherwise
+	 */
+;
+	bool EvaluatePossibleTriggers(const FMovieSceneContext& InContext, FMoveSceneAudioTriggerState& InState, bool& OutTriggered) const;
+
+public:
+
+	// ~ FMovieSceneChannel Interface
+	virtual void GetKeys(const TRange<FFrameNumber>& WithinRange, TArray<FFrameNumber>* OutKeyTimes, TArray<FKeyHandle>* OutKeyHandles) override;
+	virtual void GetKeyTimes(TArrayView<const FKeyHandle> InHandles, TArrayView<FFrameNumber> OutKeyTimes) override;
+	virtual void SetKeyTimes(TArrayView<const FKeyHandle> InHandles, TArrayView<const FFrameNumber> InKeyTimes) override;
+	virtual void DuplicateKeys(TArrayView<const FKeyHandle> InHandles, TArrayView<FKeyHandle> OutNewHandles) override;
+	virtual void DeleteKeys(TArrayView<const FKeyHandle> InHandles) override;
+	virtual void DeleteKeysFrom(FFrameNumber InTime, bool bDeleteKeysBefore) override;
+	virtual void ChangeFrameResolution(FFrameRate SourceRate, FFrameRate DestinationRate) override;
+	virtual TRange<FFrameNumber> ComputeEffectiveRange() const override;
+	virtual int32 GetNumKeys() const override;
+	virtual void Reset() override;
+	virtual void Offset(FFrameNumber DeltaPosition) override;
+	virtual void Optimize(const FKeyDataOptimizationParams& InParameters) override;
+
+private:
+
+	UPROPERTY(meta = (KeyTimes))
+	TArray<FFrameNumber> Times;
+
+	// These are all 1s, but are required for the templates to bind correctly
+	UPROPERTY(meta = (KeyValues))
+	TArray<bool> Values;
+
+	FMovieSceneKeyHandleMap KeyHandles;
+};
+
+inline bool EvaluateChannel(const FMovieSceneAudioTriggerChannel* InChannel, FFrameTime InTime, bool& OutValue)
+{
+	return false;
+}
+
+template<>
+struct TMovieSceneChannelTraits<FMovieSceneAudioTriggerChannel> : TMovieSceneChannelTraitsBase<FMovieSceneAudioTriggerChannel>
+{
+	enum { SupportsDefaults = false };
+
+#if WITH_EDITORONLY_DATA
+
+	/** Bool channels can have external values (ie, they can get their values from external objects for UI purposes) */
+	typedef TMovieSceneExternalValue<bool> ExtendedEditorDataType;
+
+#endif
+};
diff --git a/Engine/Source/Runtime/MovieSceneTracks/Private/Evaluation/MovieSceneAudioTemplate.cpp b/Engine/Source/Runtime/MovieSceneTracks/Private/Evaluation/MovieSceneAudioTemplate.cpp
--- a/Engine/Source/Runtime/MovieSceneTracks/Private/Evaluation/MovieSceneAudioTemplate.cpp
+++ b/Engine/Source/Runtime/MovieSceneTracks/Private/Evaluation/MovieSceneAudioTemplate.cpp
@@ -17,6 +17,7 @@
 #include "Evaluation/MovieSceneEvaluation.h"
 #include "IMovieScenePlayer.h"
 #include "GameFramework/WorldSettings.h"
+#include "Channels/MovieSceneAudioTriggerChannel.h"
 
 
 DECLARE_CYCLE_STAT(TEXT("Audio Track Evaluate"), MovieSceneEval_AudioTrack_Evaluate, STATGROUP_MovieSceneEval);
@@ -100,6 +101,8 @@ struct FDestroyAudioPreAnimatedToken : IMovieScenePreAnimatedToken
 
 struct FCachedAudioTrackData : IPersistentEvaluationData
 {
+	TMap<FName, FMoveSceneAudioTriggerState> TriggerStateMap;
+
 	TMap<FObjectKey, TMap<FObjectKey, TWeakObjectPtr<UAudioComponent>>> AudioComponentsByActorKey;
 	
 	FCachedAudioTrackData()
@@ -395,6 +398,22 @@ struct FAudioSectionExecutionToken : IMovieSceneExecutionToken
 		});
 	}
 
+	void EvaluateAllAndFireTriggers(IAudioParameterControllerInterface& InParamaterInterface, const FMovieSceneContext& InContext, FCachedAudioTrackData& InPersistentData) const
+	{
+		AudioSection->ForEachInput([&InParamaterInterface, &InContext, &InPersistentData](FName InName, const FMovieSceneAudioTriggerChannel& InChannel)
+		{
+			bool OutValue = false;
+			FMoveSceneAudioTriggerState& TriggerState = InPersistentData.TriggerStateMap.FindOrAdd(InName);
+			if(InChannel.EvaluatePossibleTriggers(InContext, TriggerState, OutValue))
+			{
+				if(OutValue)
+				{
+					InParamaterInterface.SetTriggerParameter(InName);
+				}
+			}
+		});
+	}
+
 	void EnsureAudioIsPlaying(UAudioComponent& AudioComponent, FPersistentEvaluationData& PersistentData, const FMovieSceneContext& Context, bool bAllowSpatialization, IMovieScenePlayer& Player) const
 	{
 		Player.SavePreAnimatedState(AudioComponent, FStopAudioPreAnimatedToken::GetAnimTypeID(), FStopAudioPreAnimatedToken::FProducer());
@@ -555,6 +574,11 @@ struct FAudioSectionExecutionToken : IMovieSceneExecutionToken
 			AudioComponent.StopDelayed(AudioTrackConstants::ScrubDuration);
 		}
 
+		if(AudioComponent.IsPlaying())
+		{
+			EvaluateAllAndFireTriggers(AudioComponent, Context, TrackData);
+		}
+
 		if (bAllowSpatialization)
 		{
 			if (FAudioDevice* AudioDevice = AudioComponent.GetAudioDevice())
diff --git a/Engine/Source/Runtime/MovieSceneTracks/Private/Sections/MovieSceneAudioSection.cpp b/Engine/Source/Runtime/MovieSceneTracks/Private/Sections/MovieSceneAudioSection.cpp
--- a/Engine/Source/Runtime/MovieSceneTracks/Private/Sections/MovieSceneAudioSection.cpp
+++ b/Engine/Source/Runtime/MovieSceneTracks/Private/Sections/MovieSceneAudioSection.cpp
@@ -68,7 +68,7 @@ namespace MovieSceneAudioSectionPrivate
 			FMovieSceneChannelMetaData Data;
 			FText TextName = FText::FromName(InName);	
 			Data.SetIdentifiers(InName, TextName, TextName);
-			InChannelProxyData.Add(const_cast<ChannelType&>(InChannel), Data, TMovieSceneExternalValue<ValueType>());
+			InChannelProxyData.Add(const_cast<ChannelType&>(InChannel), Data, TMovieSceneExternalValue<ValueType>::Make());
 #else //WITH_EDITOR
 			InChannelProxyData.Add(const_cast<ChannelType&>(InChannel));
 #endif //WITH_EDITOR
@@ -129,6 +129,7 @@ EMovieSceneChannelProxyType  UMovieSceneAudioSection::CacheChannelProxy()
 	AddInputChannels<FMovieSceneBoolChannel, bool>(this, Channels);
 	AddInputChannels<FMovieSceneIntegerChannel, int32>(this, Channels);
 	AddInputChannels<FMovieSceneStringChannel, FString>(this, Channels);
+	AddInputChannels<FMovieSceneAudioTriggerChannel, bool>(this, Channels);
 
 	ChannelProxy = MakeShared<FMovieSceneChannelProxy>(MoveTemp(Channels));
 
@@ -148,17 +149,34 @@ void UMovieSceneAudioSection::SetupSoundInputParameters(const USoundBase* InSoun
 			switch(Param.ParamType)
 			{
 			case EAudioParameterType::Float:
+			{
 				Inputs_Float.FindOrAdd(Param.ParamName, FMovieSceneFloatChannel{}).SetDefault(Param.FloatParam);
 				break;
+			}
 			case EAudioParameterType::Boolean:
-				Inputs_Bool.FindOrAdd(Param.ParamName, FMovieSceneBoolChannel{}).SetDefault(Param.BoolParam);
+			{
+				// Triggers are fundamentally just booleans outside of Metasound.
+				static const FName TriggerName = FName(TEXT("Trigger")); // MOVE ME.
+				if (Param.TypeName == TriggerName)
+				{
+					Inputs_Trigger.FindOrAdd(Param.ParamName, FMovieSceneAudioTriggerChannel{});
+				}
+				else
+				{
+					Inputs_Bool.FindOrAdd(Param.ParamName, FMovieSceneBoolChannel{}).SetDefault(Param.BoolParam);
+				}
 				break;
+			}
 			case EAudioParameterType::Integer:
+			{
 				Inputs_Int.FindOrAdd(Param.ParamName, FMovieSceneIntegerChannel{}).SetDefault(Param.IntParam);
 				break;
+			}
 			case EAudioParameterType::String:
+			{
 				Inputs_String.FindOrAdd(Param.ParamName, FMovieSceneStringChannel{}).SetDefault(Param.StringParam);
 				break;
+			}
 			default:
 				// Not supported yet.
 				break;
diff --git a/Engine/Source/Runtime/MovieSceneTracks/Public/Sections/MovieSceneAudioSection.h b/Engine/Source/Runtime/MovieSceneTracks/Public/Sections/MovieSceneAudioSection.h
--- a/Engine/Source/Runtime/MovieSceneTracks/Public/Sections/MovieSceneAudioSection.h
+++ b/Engine/Source/Runtime/MovieSceneTracks/Public/Sections/MovieSceneAudioSection.h
@@ -12,6 +12,7 @@
 #include "Channels/MovieSceneBoolChannel.h"
 #include "Channels/MovieSceneIntegerChannel.h"
 #include "Sections/MovieSceneActorReferenceSection.h"
+#include "Channels/MovieSceneAudioTriggerChannel.h"
 #include "MovieSceneAudioSection.generated.h"
 
 class USoundBase;
@@ -167,6 +168,7 @@ public:
 	void ForEachInput(TFunction<void(FName, const FMovieSceneStringChannel&)> InFunction) const { ForEachInternal(InFunction, Inputs_String); }
 	void ForEachInput(TFunction<void(FName, const FMovieSceneIntegerChannel&)> InFunction) const  { ForEachInternal(InFunction, Inputs_Int); }
 	void ForEachInput(TFunction<void(FName, const FMovieSceneFloatChannel&)> InFunction) const { ForEachInternal(InFunction, Inputs_Float); }
+	void ForEachInput(TFunction<void(FName, const FMovieSceneAudioTriggerChannel&)> InFunction) const { ForEachInternal(InFunction, Inputs_Trigger); }
 
 public:
 
@@ -233,6 +235,8 @@ private:
 	TMap<FName, FMovieSceneBoolChannel> Inputs_Bool;
 	UPROPERTY()
 	TMap<FName, FMovieSceneIntegerChannel> Inputs_Int;
+	UPROPERTY()
+	TMap<FName, FMovieSceneAudioTriggerChannel> Inputs_Trigger;
 
 	UPROPERTY()
 	FMovieSceneActorReferenceData AttachActorData;
