Sequencer: Root instance handles are now type-enforced
  - This prevents passing an FInstanceHandle to something that specifically wants a root instance handle.
  - FRootInstanceHandle is implicitly castable to an FInstanceHandle, but not visa-versa.

diff --git a/Engine/Source/Runtime/MovieScene/Private/EntitySystem/MovieSceneInstanceRegistry.cpp b/Engine/Source/Runtime/MovieScene/Private/EntitySystem/MovieSceneInstanceRegistry.cpp
--- a/Engine/Source/Runtime/MovieScene/Private/EntitySystem/MovieSceneInstanceRegistry.cpp
+++ b/Engine/Source/Runtime/MovieScene/Private/EntitySystem/MovieSceneInstanceRegistry.cpp
@@ -54,21 +54,21 @@ FInstanceHandle FInstanceRegistry::FindRelatedInstanceHandle(FInstanceHandle Ins
 	return RootInstance->FindSubInstance(SequenceID);
 }
 
-FInstanceHandle FInstanceRegistry::AllocateRootInstance(IMovieScenePlayer* Player)
+FRootInstanceHandle FInstanceRegistry::AllocateRootInstance(IMovieScenePlayer* Player)
 {
 	check(Instances.Num() < 65535);
 
 	const uint16 InstanceSerial = InstanceSerialNumber++;
 
 	FSparseArrayAllocationInfo NewAllocation = Instances.AddUninitialized();
-	FInstanceHandle InstanceHandle { (uint16)NewAllocation.Index, InstanceSerial };
+	FRootInstanceHandle InstanceHandle { (uint16)NewAllocation.Index, InstanceSerial };
 
 	new (NewAllocation) FSequenceInstance(Linker, Player, InstanceHandle);
 
 	return InstanceHandle;
 }
 
-FInstanceHandle FInstanceRegistry::AllocateSubInstance(IMovieScenePlayer* Player, FMovieSceneSequenceID SequenceID, FInstanceHandle RootInstanceHandle)
+FInstanceHandle FInstanceRegistry::AllocateSubInstance(IMovieScenePlayer* Player, FMovieSceneSequenceID SequenceID, FRootInstanceHandle RootInstanceHandle)
 {
 	check(Instances.Num() < 65535 && SequenceID != MovieSceneSequenceID::Root);
 
diff --git a/Engine/Source/Runtime/MovieScene/Private/EntitySystem/MovieSceneSequenceInstance.cpp b/Engine/Source/Runtime/MovieScene/Private/EntitySystem/MovieSceneSequenceInstance.cpp
--- a/Engine/Source/Runtime/MovieScene/Private/EntitySystem/MovieSceneSequenceInstance.cpp
+++ b/Engine/Source/Runtime/MovieScene/Private/EntitySystem/MovieSceneSequenceInstance.cpp
@@ -28,12 +28,12 @@ namespace MovieScene
 DECLARE_CYCLE_STAT(TEXT("Sequence Instance Update"), MovieSceneEval_SequenceInstanceUpdate, STATGROUP_MovieSceneEval);
 DECLARE_CYCLE_STAT(TEXT("[External] Sequence Instance Post-Update"), MovieSceneEval_SequenceInstancePostUpdate, STATGROUP_MovieSceneEval);
 
-FSequenceInstance::FSequenceInstance(UMovieSceneEntitySystemLinker* Linker, IMovieScenePlayer* Player, FInstanceHandle InInstanceHandle)
+FSequenceInstance::FSequenceInstance(UMovieSceneEntitySystemLinker* Linker, IMovieScenePlayer* Player, FRootInstanceHandle InInstanceHandle)
 	: SequenceID(MovieSceneSequenceID::Root)
 	, RootOverrideSequenceID(MovieSceneSequenceID::Root)
 	, PlayerIndex(Player->GetUniqueIndex())
 	, InstanceHandle(InInstanceHandle)
-	, RootInstanceHandle(InstanceHandle)
+	, RootInstanceHandle(InInstanceHandle)
 {
 	// Root instances always start in a finished state in order to ensure that 'Start'
 	// is called correctly for the top level instance. This is subtly different from
@@ -49,7 +49,7 @@ FSequenceInstance::FSequenceInstance(UMovieSceneEntitySystemLinker* Linker, IMov
 	InvalidateCachedData(Linker);
 }
 
-FSequenceInstance::FSequenceInstance(UMovieSceneEntitySystemLinker* Linker, IMovieScenePlayer* Player, FInstanceHandle InInstanceHandle, FInstanceHandle InRootInstanceHandle, FMovieSceneSequenceID InSequenceID, FMovieSceneCompiledDataID InCompiledDataID)
+FSequenceInstance::FSequenceInstance(UMovieSceneEntitySystemLinker* Linker, IMovieScenePlayer* Player, FInstanceHandle InInstanceHandle, FRootInstanceHandle InRootInstanceHandle, FMovieSceneSequenceID InSequenceID, FMovieSceneCompiledDataID InCompiledDataID)
 	: CompiledDataID(InCompiledDataID)
 	, SequenceID(InSequenceID)
 	, RootOverrideSequenceID(MovieSceneSequenceID::Invalid)
@@ -165,8 +165,10 @@ void FSequenceInstance::Start(UMovieSceneEntitySystemLinker* Linker, const FMovi
 	bFinished = false;
 	bHasEverUpdated = true;
 
+	check(RootInstanceHandle == InstanceHandle);
+
 	IMovieScenePlayer* Player = GetPlayer();
-	SequenceUpdater->Start(Linker, InstanceHandle, Player, InContext);
+	SequenceUpdater->Start(Linker, RootInstanceHandle, Player, InContext);
 }
 
 void FSequenceInstance::Update(UMovieSceneEntitySystemLinker* Linker, const FMovieSceneContext& InContext)
@@ -185,8 +187,10 @@ void FSequenceInstance::Update(UMovieSceneEntitySystemLinker* Linker, const FMov
 		Start(Linker, InContext);
 	}
 
+	check(RootInstanceHandle == InstanceHandle);
+
 	Context = InContext;
-	SequenceUpdater->Update(Linker, InstanceHandle, GetPlayer(), InContext);
+	SequenceUpdater->Update(Linker, RootInstanceHandle, GetPlayer(), InContext);
 }
 
 void FSequenceInstance::Finish(UMovieSceneEntitySystemLinker* Linker)
@@ -210,7 +214,8 @@ void FSequenceInstance::Finish(UMovieSceneEntitySystemLinker* Linker)
 
 	if (SequenceUpdater)
 	{
-		SequenceUpdater->Finish(Linker, InstanceHandle, Player);
+		check(RootInstanceHandle == InstanceHandle);
+		SequenceUpdater->Finish(Linker, RootInstanceHandle, Player);
 	}
 
 	if (LegacyEvaluator)
@@ -296,7 +301,8 @@ void FSequenceInstance::OverrideRootSequence(UMovieSceneEntitySystemLinker* Link
 {
 	if (SequenceUpdater)
 	{
-		SequenceUpdater->OverrideRootSequence(Linker, InstanceHandle, NewRootSequenceID);
+		check(RootInstanceHandle == InstanceHandle);
+		SequenceUpdater->OverrideRootSequence(Linker, RootInstanceHandle, NewRootSequenceID);
 	}
 
 	RootOverrideSequenceID = NewRootSequenceID;
diff --git a/Engine/Source/Runtime/MovieScene/Private/EntitySystem/MovieSceneSequenceUpdaters.cpp b/Engine/Source/Runtime/MovieScene/Private/EntitySystem/MovieSceneSequenceUpdaters.cpp
--- a/Engine/Source/Runtime/MovieScene/Private/EntitySystem/MovieSceneSequenceUpdaters.cpp
+++ b/Engine/Source/Runtime/MovieScene/Private/EntitySystem/MovieSceneSequenceUpdaters.cpp
@@ -36,14 +36,14 @@ struct FSequenceUpdater_Flat : ISequenceUpdater
 	~FSequenceUpdater_Flat();
 
 	virtual void DissectContext(UMovieSceneEntitySystemLinker* Linker, IMovieScenePlayer* InPlayer, const FMovieSceneContext& Context, TArray<TRange<FFrameTime>>& OutDissections) override;
-	virtual void Start(UMovieSceneEntitySystemLinker* Linker, FInstanceHandle InstanceHandle, IMovieScenePlayer* InPlayer, const FMovieSceneContext& InContext) override;
-	virtual void Update(UMovieSceneEntitySystemLinker* Linker, FInstanceHandle InstanceHandle, IMovieScenePlayer* InPlayer, const FMovieSceneContext& Context) override;
-	virtual void Finish(UMovieSceneEntitySystemLinker* Linker, FInstanceHandle InstanceHandle, IMovieScenePlayer* InPlayer) override;
+	virtual void Start(UMovieSceneEntitySystemLinker* Linker, FRootInstanceHandle InstanceHandle, IMovieScenePlayer* InPlayer, const FMovieSceneContext& InContext) override;
+	virtual void Update(UMovieSceneEntitySystemLinker* Linker, FRootInstanceHandle InstanceHandle, IMovieScenePlayer* InPlayer, const FMovieSceneContext& Context) override;
+	virtual void Finish(UMovieSceneEntitySystemLinker* Linker, FRootInstanceHandle InstanceHandle, IMovieScenePlayer* InPlayer) override;
 	virtual void InvalidateCachedData(UMovieSceneEntitySystemLinker* Linker) override;
 	virtual void Destroy(UMovieSceneEntitySystemLinker* Linker) override;
 	virtual TUniquePtr<ISequenceUpdater> MigrateToHierarchical() override;
 	virtual FInstanceHandle FindSubInstance(FMovieSceneSequenceID SubSequenceID) const override { return FInstanceHandle(); }
-	virtual void OverrideRootSequence(UMovieSceneEntitySystemLinker* InLinker, FInstanceHandle InstanceHandle, FMovieSceneSequenceID NewRootOverrideSequenceID) override {}
+	virtual void OverrideRootSequence(UMovieSceneEntitySystemLinker* InLinker, FRootInstanceHandle InstanceHandle, FMovieSceneSequenceID NewRootOverrideSequenceID) override {}
 
 private:
 
@@ -61,20 +61,20 @@ struct FSequenceUpdater_Hierarchical : ISequenceUpdater
 	~FSequenceUpdater_Hierarchical();
 
 	virtual void DissectContext(UMovieSceneEntitySystemLinker* Linker, IMovieScenePlayer* InPlayer, const FMovieSceneContext& Context, TArray<TRange<FFrameTime>>& OutDissections) override;
-	virtual void Start(UMovieSceneEntitySystemLinker* Linker, FInstanceHandle InstanceHandle, IMovieScenePlayer* InPlayer, const FMovieSceneContext& InContext) override;
-	virtual void Update(UMovieSceneEntitySystemLinker* Linker, FInstanceHandle InstanceHandle, IMovieScenePlayer* InPlayer, const FMovieSceneContext& Context) override;
-	virtual void Finish(UMovieSceneEntitySystemLinker* Linker, FInstanceHandle InstanceHandle, IMovieScenePlayer* InPlayer) override;
+	virtual void Start(UMovieSceneEntitySystemLinker* Linker, FRootInstanceHandle InstanceHandle, IMovieScenePlayer* InPlayer, const FMovieSceneContext& InContext) override;
+	virtual void Update(UMovieSceneEntitySystemLinker* Linker, FRootInstanceHandle InstanceHandle, IMovieScenePlayer* InPlayer, const FMovieSceneContext& Context) override;
+	virtual void Finish(UMovieSceneEntitySystemLinker* Linker, FRootInstanceHandle InstanceHandle, IMovieScenePlayer* InPlayer) override;
 	virtual void InvalidateCachedData(UMovieSceneEntitySystemLinker* Linker) override;
 	virtual void Destroy(UMovieSceneEntitySystemLinker* Linker) override;
 	virtual TUniquePtr<ISequenceUpdater> MigrateToHierarchical() override { return nullptr; }
 	virtual FInstanceHandle FindSubInstance(FMovieSceneSequenceID SubSequenceID) const override { return SequenceInstances.FindRef(SubSequenceID).Handle; }
-	virtual void OverrideRootSequence(UMovieSceneEntitySystemLinker* InLinker, FInstanceHandle InstanceHandle, FMovieSceneSequenceID NewRootOverrideSequenceID) override;
+	virtual void OverrideRootSequence(UMovieSceneEntitySystemLinker* InLinker, FRootInstanceHandle InstanceHandle, FMovieSceneSequenceID NewRootOverrideSequenceID) override;
 
 private:
 
 	TRange<FFrameNumber> UpdateEntitiesForSequence(const FMovieSceneEntityComponentField* ComponentField, FFrameTime SequenceTime, FMovieSceneEvaluationFieldEntitySet& OutEntities);
 
-	FInstanceHandle GetOrCreateSequenceInstance(IMovieScenePlayer* InPlayer, FInstanceRegistry* InstanceRegistry, FInstanceHandle RootInstanceHandle, FMovieSceneSequenceID SequenceID);
+	FInstanceHandle GetOrCreateSequenceInstance(IMovieScenePlayer* InPlayer, FInstanceRegistry* InstanceRegistry, FRootInstanceHandle RootInstanceHandle, FMovieSceneSequenceID SequenceID);
 
 private:
 
@@ -227,11 +227,11 @@ void FSequenceUpdater_Flat::DissectContext(UMovieSceneEntitySystemLinker* Linker
 	}
 }
 
-void FSequenceUpdater_Flat::Start(UMovieSceneEntitySystemLinker* Linker, FInstanceHandle InstanceHandle, IMovieScenePlayer* InPlayer, const FMovieSceneContext& InContext)
+void FSequenceUpdater_Flat::Start(UMovieSceneEntitySystemLinker* Linker, FRootInstanceHandle InstanceHandle, IMovieScenePlayer* InPlayer, const FMovieSceneContext& InContext)
 {
 }
 
-void FSequenceUpdater_Flat::Update(UMovieSceneEntitySystemLinker* Linker, FInstanceHandle InstanceHandle, IMovieScenePlayer* InPlayer, const FMovieSceneContext& Context)
+void FSequenceUpdater_Flat::Update(UMovieSceneEntitySystemLinker* Linker, FRootInstanceHandle InstanceHandle, IMovieScenePlayer* InPlayer, const FMovieSceneContext& Context)
 {
 	FSequenceInstance& SequenceInstance = Linker->GetInstanceRegistry()->MutateInstance(InstanceHandle);
 	SequenceInstance.SetContext(Context);
@@ -288,7 +288,7 @@ void FSequenceUpdater_Flat::Update(UMovieSceneEntitySystemLinker* Linker, FInsta
 	}
 }
 
-void FSequenceUpdater_Flat::Finish(UMovieSceneEntitySystemLinker* Linker, FInstanceHandle InstanceHandle, IMovieScenePlayer* InPlayer)
+void FSequenceUpdater_Flat::Finish(UMovieSceneEntitySystemLinker* Linker, FRootInstanceHandle InstanceHandle, IMovieScenePlayer* InPlayer)
 {
 	InvalidateCachedData(Linker);
 }
@@ -388,7 +388,7 @@ void FSequenceUpdater_Hierarchical::DissectContext(UMovieSceneEntitySystemLinker
 	}
 }
 
-FInstanceHandle FSequenceUpdater_Hierarchical::GetOrCreateSequenceInstance(IMovieScenePlayer* InPlayer, FInstanceRegistry* InstanceRegistry, FInstanceHandle RootInstanceHandle, FMovieSceneSequenceID SequenceID)
+FInstanceHandle FSequenceUpdater_Hierarchical::GetOrCreateSequenceInstance(IMovieScenePlayer* InPlayer, FInstanceRegistry* InstanceRegistry, FRootInstanceHandle RootInstanceHandle, FMovieSceneSequenceID SequenceID)
 {
 	if (FSubInstanceData* Existing = SequenceInstances.Find(SequenceID))
 	{
@@ -402,7 +402,7 @@ FInstanceHandle FSequenceUpdater_Hierarchical::GetOrCreateSequenceInstance(IMovi
 	return InstanceHandle;
 }
 
-void FSequenceUpdater_Hierarchical::OverrideRootSequence(UMovieSceneEntitySystemLinker* InLinker, FInstanceHandle MasterInstanceHandle, FMovieSceneSequenceID NewRootOverrideSequenceID)
+void FSequenceUpdater_Hierarchical::OverrideRootSequence(UMovieSceneEntitySystemLinker* InLinker, FRootInstanceHandle MasterInstanceHandle, FMovieSceneSequenceID NewRootOverrideSequenceID)
 {
 	if (RootOverrideSequenceID != NewRootOverrideSequenceID)
 	{
@@ -421,11 +421,11 @@ void FSequenceUpdater_Hierarchical::OverrideRootSequence(UMovieSceneEntitySystem
 	}
 }
 
-void FSequenceUpdater_Hierarchical::Start(UMovieSceneEntitySystemLinker* Linker, FInstanceHandle InstanceHandle, IMovieScenePlayer* InPlayer, const FMovieSceneContext& InContext)
+void FSequenceUpdater_Hierarchical::Start(UMovieSceneEntitySystemLinker* Linker, FRootInstanceHandle InstanceHandle, IMovieScenePlayer* InPlayer, const FMovieSceneContext& InContext)
 {
 }
 
-void FSequenceUpdater_Hierarchical::Update(UMovieSceneEntitySystemLinker* Linker, FInstanceHandle InstanceHandle, IMovieScenePlayer* InPlayer, const FMovieSceneContext& Context)
+void FSequenceUpdater_Hierarchical::Update(UMovieSceneEntitySystemLinker* Linker, FRootInstanceHandle InstanceHandle, IMovieScenePlayer* InPlayer, const FMovieSceneContext& Context)
 {
 	FInstanceRegistry* InstanceRegistry = Linker->GetInstanceRegistry();
 	UMovieSceneCompiledDataManager* CompiledDataManager = InPlayer->GetEvaluationTemplate().GetCompiledDataManager();
@@ -482,7 +482,7 @@ void FSequenceUpdater_Hierarchical::Update(UMovieSceneEntitySystemLinker* Linker
 
 				FEntityImportSequenceParams Params;
 				Params.InstanceHandle = RootInstanceHandle;
-				Params.RootInstanceHandle = RootInstanceHandle;
+				Params.RootInstanceHandle = InstanceHandle;
 				Params.DefaultCompletionMode = RootSequence->DefaultCompletionMode;
 				Params.HierarchicalBias = 0;
 
@@ -499,7 +499,7 @@ void FSequenceUpdater_Hierarchical::Update(UMovieSceneEntitySystemLinker* Linker
 				{
 					FEntityImportSequenceParams Params;
 					Params.InstanceHandle = RootInstanceHandle;
-					Params.RootInstanceHandle = RootInstanceHandle;
+					Params.RootInstanceHandle = InstanceHandle;
 					Params.DefaultCompletionMode = RootSequence->DefaultCompletionMode;
 					Params.HierarchicalBias = 0;
 
@@ -582,7 +582,7 @@ void FSequenceUpdater_Hierarchical::Update(UMovieSceneEntitySystemLinker* Linker
 
 				FEntityImportSequenceParams Params;
 				Params.InstanceHandle = SubSequenceHandle;
-				Params.RootInstanceHandle = RootInstanceHandle;
+				Params.RootInstanceHandle = InstanceHandle;
 				Params.DefaultCompletionMode = SubSequence->DefaultCompletionMode;
 				Params.HierarchicalBias = SubData->HierarchicalBias;
 				Params.bPreRoll  = bIsPreRoll;
@@ -653,7 +653,7 @@ void FSequenceUpdater_Hierarchical::Update(UMovieSceneEntitySystemLinker* Linker
 	}
 }
 
-void FSequenceUpdater_Hierarchical::Finish(UMovieSceneEntitySystemLinker* Linker, FInstanceHandle InstanceHandle, IMovieScenePlayer* InPlayer)
+void FSequenceUpdater_Hierarchical::Finish(UMovieSceneEntitySystemLinker* Linker, FRootInstanceHandle InstanceHandle, IMovieScenePlayer* InPlayer)
 {
 	FInstanceRegistry* InstanceRegistry = Linker->GetInstanceRegistry();
 
diff --git a/Engine/Source/Runtime/MovieScene/Private/Evaluation/MovieScenePreAnimatedState.cpp b/Engine/Source/Runtime/MovieScene/Private/Evaluation/MovieScenePreAnimatedState.cpp
--- a/Engine/Source/Runtime/MovieScene/Private/Evaluation/MovieScenePreAnimatedState.cpp
+++ b/Engine/Source/Runtime/MovieScene/Private/Evaluation/MovieScenePreAnimatedState.cpp
@@ -45,7 +45,7 @@ FMovieScenePreAnimatedState::~FMovieScenePreAnimatedState()
 	}
 }
 
-void FMovieScenePreAnimatedState::Initialize(UMovieSceneEntitySystemLinker* Linker, UE::MovieScene::FInstanceHandle InInstanceHandle)
+void FMovieScenePreAnimatedState::Initialize(UMovieSceneEntitySystemLinker* Linker, UE::MovieScene::FRootInstanceHandle InInstanceHandle)
 {
 	// If we're re-using a pre-animated state class and it was previously
 	// capturing global state make sure to decrement that request
diff --git a/Engine/Source/Runtime/MovieScene/Private/Evaluation/PreAnimatedState/MovieScenePreAnimatedCaptureSource.cpp b/Engine/Source/Runtime/MovieScene/Private/Evaluation/PreAnimatedState/MovieScenePreAnimatedCaptureSource.cpp
--- a/Engine/Source/Runtime/MovieScene/Private/Evaluation/PreAnimatedState/MovieScenePreAnimatedCaptureSource.cpp
+++ b/Engine/Source/Runtime/MovieScene/Private/Evaluation/PreAnimatedState/MovieScenePreAnimatedCaptureSource.cpp
@@ -63,7 +63,7 @@ FScopedPreAnimatedCaptureSource*& FScopedPreAnimatedCaptureSource::GetCaptureSou
 	return GCaptureSource;
 }
 
-UE::MovieScene::FInstanceHandle FScopedPreAnimatedCaptureSource::GetRootInstanceHandle(UMovieSceneEntitySystemLinker* Linker) const
+UE::MovieScene::FRootInstanceHandle FScopedPreAnimatedCaptureSource::GetRootInstanceHandle(UMovieSceneEntitySystemLinker* Linker) const
 {
 	if (OptionalSequencePreAnimatedState)
 	{
@@ -74,7 +74,7 @@ UE::MovieScene::FInstanceHandle FScopedPreAnimatedCaptureSource::GetRootInstance
 	{
 		return Linker->GetInstanceRegistry()->GetInstance(TrackInstanceInput->InstanceHandle).GetRootInstanceHandle();
 	}
-	return UE::MovieScene::FInstanceHandle();
+	return UE::MovieScene::FRootInstanceHandle();
 }
 
 void FScopedPreAnimatedCaptureSource::BeginTracking(const UE::MovieScene::FPreAnimatedStateMetaData& MetaData, UMovieSceneEntitySystemLinker* Linker)
diff --git a/Engine/Source/Runtime/MovieScene/Private/Evaluation/PreAnimatedState/MovieScenePreAnimatedCaptureSources.inl b/Engine/Source/Runtime/MovieScene/Private/Evaluation/PreAnimatedState/MovieScenePreAnimatedCaptureSources.inl
--- a/Engine/Source/Runtime/MovieScene/Private/Evaluation/PreAnimatedState/MovieScenePreAnimatedCaptureSources.inl
+++ b/Engine/Source/Runtime/MovieScene/Private/Evaluation/PreAnimatedState/MovieScenePreAnimatedCaptureSources.inl
@@ -116,7 +116,7 @@ void TPreAnimatedCaptureSources<KeyType>::GatherAndRemoveMetaDataForGroup(FPreAn
 }
 
 template<typename KeyType>
-bool TPreAnimatedCaptureSources<KeyType>::ContainsInstanceHandle(FInstanceHandle RootInstanceHandle) const
+bool TPreAnimatedCaptureSources<KeyType>::ContainsInstanceHandle(FRootInstanceHandle RootInstanceHandle) const
 {
 	for (const TPair<KeyType, FPreAnimatedStateMetaDataArray>& Pair : KeyToMetaData)
 	{
diff --git a/Engine/Source/Runtime/MovieScene/Private/Evaluation/PreAnimatedState/MovieScenePreAnimatedEntityCaptureSource.cpp b/Engine/Source/Runtime/MovieScene/Private/Evaluation/PreAnimatedState/MovieScenePreAnimatedEntityCaptureSource.cpp
--- a/Engine/Source/Runtime/MovieScene/Private/Evaluation/PreAnimatedState/MovieScenePreAnimatedEntityCaptureSource.cpp
+++ b/Engine/Source/Runtime/MovieScene/Private/Evaluation/PreAnimatedState/MovieScenePreAnimatedEntityCaptureSource.cpp
@@ -20,7 +20,7 @@ void FPreAnimatedEntityCaptureSource::Reset()
 	KeyToMetaData.Empty();
 }
 
-void FPreAnimatedEntityCaptureSource::BeginTrackingEntity(const FPreAnimatedStateEntry& Entry, FMovieSceneEntityID EntityID, FInstanceHandle InstanceHandle, bool bWantsRestoreState)
+void FPreAnimatedEntityCaptureSource::BeginTrackingEntity(const FPreAnimatedStateEntry& Entry, FMovieSceneEntityID EntityID, FRootInstanceHandle InstanceHandle, bool bWantsRestoreState)
 {
 	FPreAnimatedStateMetaDataArray& Array = KeyToMetaData.FindOrAdd(EntityID);
 	if (FPreAnimatedStateMetaData* Existing = Algo::FindBy(Array, Entry, &FPreAnimatedStateMetaData::Entry))
@@ -130,7 +130,7 @@ void FPreAnimatedEntityCaptureSource::GatherAndRemoveMetaDataForGroup(FPreAnimat
 	}
 }
 
-bool FPreAnimatedEntityCaptureSource::ContainsInstanceHandle(FInstanceHandle RootInstanceHandle) const
+bool FPreAnimatedEntityCaptureSource::ContainsInstanceHandle(FRootInstanceHandle RootInstanceHandle) const
 {
 	for (const TPair<FMovieSceneEntityID, FPreAnimatedStateMetaDataArray>& Pair : KeyToMetaData)
 	{
diff --git a/Engine/Source/Runtime/MovieScene/Private/Evaluation/PreAnimatedState/MovieScenePreAnimatedStateExtension.cpp b/Engine/Source/Runtime/MovieScene/Private/Evaluation/PreAnimatedState/MovieScenePreAnimatedStateExtension.cpp
--- a/Engine/Source/Runtime/MovieScene/Private/Evaluation/PreAnimatedState/MovieScenePreAnimatedStateExtension.cpp
+++ b/Engine/Source/Runtime/MovieScene/Private/Evaluation/PreAnimatedState/MovieScenePreAnimatedStateExtension.cpp
@@ -459,7 +459,7 @@ void FPreAnimatedStateExtension::DiscardStateForGroup(FPreAnimatedStorageGroupHa
 	bEntriesInvalidated = true;
 }
 
-bool FPreAnimatedStateExtension::ContainsAnyStateForInstanceHandle(FInstanceHandle RootInstanceHandle) const
+bool FPreAnimatedStateExtension::ContainsAnyStateForInstanceHandle(FRootInstanceHandle RootInstanceHandle) const
 {
 	if (FPreAnimatedEntityCaptureSource* EntityMetaData = GetEntityMetaData())
 	{
diff --git a/Engine/Source/Runtime/MovieScene/Private/Tests/MovieScenePreAnimatedStateTests.cpp b/Engine/Source/Runtime/MovieScene/Private/Tests/MovieScenePreAnimatedStateTests.cpp
--- a/Engine/Source/Runtime/MovieScene/Private/Tests/MovieScenePreAnimatedStateTests.cpp
+++ b/Engine/Source/Runtime/MovieScene/Private/Tests/MovieScenePreAnimatedStateTests.cpp
@@ -150,7 +150,7 @@ bool FMovieScenePreAnimatedStateGlobalTest::RunTest(const FString& Parameters)
 	ResetValues();
 
 	FMovieScenePreAnimatedState State;
-	State.Initialize(Linker, FInstanceHandle());
+	State.Initialize(Linker, FRootInstanceHandle());
 	State.EnableGlobalPreAnimatedStateCapture();
 
 	FPreAnimatedTokenProducer Producer(&TestValue1);
@@ -184,7 +184,7 @@ bool FMovieScenePreAnimatedStateEntityTest::RunTest(const FString& Parameters)
 	UMovieSceneEntitySystemLinker* Linker = GetTestLinker();
 
 	FMovieScenePreAnimatedState State;
-	State.Initialize(Linker, FInstanceHandle());
+	State.Initialize(Linker, FRootInstanceHandle());
 	State.EnableGlobalPreAnimatedStateCapture();
 
 	FPreAnimatedTokenProducer Producer(&TestValue1);
@@ -221,7 +221,7 @@ bool FMovieScenePreAnimatedStateOverlappingEntitiesTest::RunTest(const FString&
 	UMovieSceneEntitySystemLinker* Linker = GetTestLinker();
 
 	FMovieScenePreAnimatedState State;
-	State.Initialize(Linker, FInstanceHandle());
+	State.Initialize(Linker, FRootInstanceHandle());
 	State.EnableGlobalPreAnimatedStateCapture();
 
 	FPreAnimatedTokenProducer Producer(&TestValue1);
@@ -293,7 +293,7 @@ bool FMovieScenePreAnimatedStateKeepThenRestoreEntityTest::RunTest(const FString
 	UMovieSceneEntitySystemLinker* Linker = GetTestLinker();
 
 	FMovieScenePreAnimatedState State;
-	State.Initialize(Linker, FInstanceHandle());
+	State.Initialize(Linker, FRootInstanceHandle());
 	State.EnableGlobalPreAnimatedStateCapture();
 
 	FPreAnimatedTokenProducer Producer(&TestValue1);
@@ -517,7 +517,7 @@ bool FMovieScenePreAnimatedStatePerformanceTest::RunTest(const FString& Paramete
 	UMovieSceneEntitySystemLinker* Linker = GetTestLinker();
 
 	FMovieScenePreAnimatedState State;
-	State.Initialize(Linker, FInstanceHandle());
+	State.Initialize(Linker, FRootInstanceHandle());
 	State.EnableGlobalPreAnimatedStateCapture();
 
 	FPreAnimatedTokenProducer Producer(&TestValue1);
diff --git a/Engine/Source/Runtime/MovieScene/Public/EntitySystem/BuiltInComponentTypes.h b/Engine/Source/Runtime/MovieScene/Public/EntitySystem/BuiltInComponentTypes.h
--- a/Engine/Source/Runtime/MovieScene/Public/EntitySystem/BuiltInComponentTypes.h
+++ b/Engine/Source/Runtime/MovieScene/Public/EntitySystem/BuiltInComponentTypes.h
@@ -189,7 +189,7 @@ public:
 
 	TComponentTypeID<FInstanceHandle>     InstanceHandle;
 
-	TComponentTypeID<FInstanceHandle>     RootInstanceHandle;
+	TComponentTypeID<FRootInstanceHandle> RootInstanceHandle;
 
 	TComponentTypeID<FFrameTime>          EvalTime;
 
diff --git a/Engine/Source/Runtime/MovieScene/Public/EntitySystem/IMovieSceneEntityProvider.h b/Engine/Source/Runtime/MovieScene/Public/EntitySystem/IMovieSceneEntityProvider.h
--- a/Engine/Source/Runtime/MovieScene/Public/EntitySystem/IMovieSceneEntityProvider.h
+++ b/Engine/Source/Runtime/MovieScene/Public/EntitySystem/IMovieSceneEntityProvider.h
@@ -66,7 +66,7 @@ struct FEntityImportSequenceParams
 	int32 HierarchicalBias;
 
 	FInstanceHandle InstanceHandle;
-	FInstanceHandle RootInstanceHandle;
+	FRootInstanceHandle RootInstanceHandle;
 
 	EMovieSceneCompletionMode DefaultCompletionMode;
 
diff --git a/Engine/Source/Runtime/MovieScene/Public/EntitySystem/MovieSceneInstanceRegistry.h b/Engine/Source/Runtime/MovieScene/Public/EntitySystem/MovieSceneInstanceRegistry.h
--- a/Engine/Source/Runtime/MovieScene/Public/EntitySystem/MovieSceneInstanceRegistry.h
+++ b/Engine/Source/Runtime/MovieScene/Public/EntitySystem/MovieSceneInstanceRegistry.h
@@ -81,9 +81,9 @@ struct FInstanceRegistry
 		return GetInstance(InstanceHandle).GetContext();
 	}
 
-	MOVIESCENE_API FInstanceHandle AllocateRootInstance(IMovieScenePlayer* Player);
+	MOVIESCENE_API FRootInstanceHandle AllocateRootInstance(IMovieScenePlayer* Player);
 
-	MOVIESCENE_API FInstanceHandle AllocateSubInstance(IMovieScenePlayer* Player, FMovieSceneSequenceID SequenceID, FInstanceHandle RootInstance);
+	MOVIESCENE_API FInstanceHandle AllocateSubInstance(IMovieScenePlayer* Player, FMovieSceneSequenceID SequenceID, FRootInstanceHandle RootInstance);
 
 	MOVIESCENE_API void DestroyInstance(FInstanceHandle InstanceHandle);
 
diff --git a/Engine/Source/Runtime/MovieScene/Public/EntitySystem/MovieSceneSequenceInstance.h b/Engine/Source/Runtime/MovieScene/Public/EntitySystem/MovieSceneSequenceInstance.h
--- a/Engine/Source/Runtime/MovieScene/Public/EntitySystem/MovieSceneSequenceInstance.h
+++ b/Engine/Source/Runtime/MovieScene/Public/EntitySystem/MovieSceneSequenceInstance.h
@@ -135,7 +135,7 @@ public:
 	 *
 	 * @return A handle to the root sequence instance for this sub sequence, or a handle to this sequence instance if it is not a sub sequence
 	 */
-	FInstanceHandle GetRootInstanceHandle() const
+	FRootInstanceHandle GetRootInstanceHandle() const
 	{
 		return RootInstanceHandle;
 	}
@@ -245,10 +245,10 @@ public:
 public:
 
 	/** Constructor for top level sequences */
-	explicit FSequenceInstance(UMovieSceneEntitySystemLinker* Linker, IMovieScenePlayer* Player, FInstanceHandle ThisInstanceHandle);
+	explicit FSequenceInstance(UMovieSceneEntitySystemLinker* Linker, IMovieScenePlayer* Player, FRootInstanceHandle ThisInstanceHandle);
 
 	/** Constructor for sub sequences */
-	explicit FSequenceInstance(UMovieSceneEntitySystemLinker* Linker, IMovieScenePlayer* Player, FInstanceHandle ThisInstanceHandle, FInstanceHandle RootInstanceHandle, FMovieSceneSequenceID InSequenceID, FMovieSceneCompiledDataID InCompiledDataID);
+	explicit FSequenceInstance(UMovieSceneEntitySystemLinker* Linker, IMovieScenePlayer* Player, FInstanceHandle ThisInstanceHandle, FRootInstanceHandle RootInstanceHandle, FMovieSceneSequenceID InSequenceID, FMovieSceneCompiledDataID InCompiledDataID);
 
 	/** Destructor */
 	~FSequenceInstance();
@@ -290,7 +290,7 @@ private:
 	/** This instance's handle. */
 	FInstanceHandle InstanceHandle;
 	/** This instance's root handle, if it is a sub sequence. */
-	FInstanceHandle RootInstanceHandle;
+	FRootInstanceHandle RootInstanceHandle;
 	/** Flag that is set when this sequence has or (will be) finished. */
 	bool bFinished : 1;
 	/** Flag that is set if this sequence has ever updated. */
diff --git a/Engine/Source/Runtime/MovieScene/Public/EntitySystem/MovieSceneSequenceInstanceHandle.h b/Engine/Source/Runtime/MovieScene/Public/EntitySystem/MovieSceneSequenceInstanceHandle.h
--- a/Engine/Source/Runtime/MovieScene/Public/EntitySystem/MovieSceneSequenceInstanceHandle.h
+++ b/Engine/Source/Runtime/MovieScene/Public/EntitySystem/MovieSceneSequenceInstanceHandle.h
@@ -49,6 +49,21 @@ struct FInstanceHandle
 	}
 };
 
+/**
+ * A handle to a root instance of a UMovieSceneSequence.
+ * This type should be used where type-safety needs to be enforced for root instances.
+ * FRootInstanceHandle is implicitly castable to FInstanceHandle but not visa-versa (even explicitly so).
+ */
+struct FRootInstanceHandle : FInstanceHandle
+{
+	FRootInstanceHandle()
+	{}
+
+	FRootInstanceHandle(uint16 InInstanceID, uint16 InInstanceSerial)
+		: FInstanceHandle(InInstanceID, InInstanceSerial)
+	{}
+};
+
 #if UE_MOVIESCENE_ENTITY_DEBUG
 template<> struct TComponentDebugType<FInstanceHandle> { static const EComponentDebugType Type = EComponentDebugType::InstanceHandle; };
 #endif
diff --git a/Engine/Source/Runtime/MovieScene/Public/EntitySystem/MovieSceneSequenceUpdaters.h b/Engine/Source/Runtime/MovieScene/Public/EntitySystem/MovieSceneSequenceUpdaters.h
--- a/Engine/Source/Runtime/MovieScene/Public/EntitySystem/MovieSceneSequenceUpdaters.h
+++ b/Engine/Source/Runtime/MovieScene/Public/EntitySystem/MovieSceneSequenceUpdaters.h
@@ -20,6 +20,7 @@ namespace MovieScene
 {
 
 struct FInstanceHandle;
+struct FRootInstanceHandle;
 
 
 /**
@@ -68,7 +69,7 @@ public:
 	 * @param InPlayer         The movie scene player instance playing this sequence
 	 * @param InContext        The root-level context for the current evaluation.
 	 */
-	virtual void Start(UMovieSceneEntitySystemLinker* InLinker, FInstanceHandle InInstanceHandle, IMovieScenePlayer* InPlayer, const FMovieSceneContext& InContext) = 0;
+	virtual void Start(UMovieSceneEntitySystemLinker* InLinker, FRootInstanceHandle InInstanceHandle, IMovieScenePlayer* InPlayer, const FMovieSceneContext& InContext) = 0;
 
 
 	/**
@@ -79,7 +80,7 @@ public:
 	 * @param InPlayer         The movie scene player instance playing this sequence
 	 * @param InContext        The root-level context for the current evaluation.
 	 */
-	virtual void Update(UMovieSceneEntitySystemLinker* InLinker, FInstanceHandle InInstanceHandle, IMovieScenePlayer* InPlayer, const FMovieSceneContext& InContext) = 0;
+	virtual void Update(UMovieSceneEntitySystemLinker* InLinker, FRootInstanceHandle InInstanceHandle, IMovieScenePlayer* InPlayer, const FMovieSceneContext& InContext) = 0;
 
 
 	/**
@@ -89,7 +90,7 @@ public:
 	 * @param InInstanceHandle The instance handle for the top level sequence instance that this updater belongs to
 	 * @param InPlayer         The movie scene player instance playing this sequence
 	 */
-	virtual void Finish(UMovieSceneEntitySystemLinker* InLinker, FInstanceHandle InInstanceHandle, IMovieScenePlayer* InPlayer) = 0;
+	virtual void Finish(UMovieSceneEntitySystemLinker* InLinker, FRootInstanceHandle InInstanceHandle, IMovieScenePlayer* InPlayer) = 0;
 
 
 	/**
@@ -115,7 +116,7 @@ public:
 	 * @param InstanceHandle              The instance handle for the master sequence
 	 * @param NewRootOverrideSequenceID   The new sequence ID to treat as the root
 	 */
-	virtual void OverrideRootSequence(UMovieSceneEntitySystemLinker* InLinker, FInstanceHandle InstanceHandle, FMovieSceneSequenceID NewRootOverrideSequenceID) = 0;
+	virtual void OverrideRootSequence(UMovieSceneEntitySystemLinker* InLinker, FRootInstanceHandle InstanceHandle, FMovieSceneSequenceID NewRootOverrideSequenceID) = 0;
 
 
 	/**
diff --git a/Engine/Source/Runtime/MovieScene/Public/Evaluation/MovieSceneEvaluationTemplateInstance.h b/Engine/Source/Runtime/MovieScene/Public/Evaluation/MovieSceneEvaluationTemplateInstance.h
--- a/Engine/Source/Runtime/MovieScene/Public/Evaluation/MovieSceneEvaluationTemplateInstance.h
+++ b/Engine/Source/Runtime/MovieScene/Public/Evaluation/MovieSceneEvaluationTemplateInstance.h
@@ -68,7 +68,7 @@ public:
 		return CompiledDataManager && WeakRootSequence.Get();
 	}
 
-	UE::MovieScene::FInstanceHandle GetRootInstanceHandle() const
+	UE::MovieScene::FRootInstanceHandle GetRootInstanceHandle() const
 	{
 		return RootInstanceHandle;
 	}
@@ -135,7 +135,7 @@ private:
 	UPROPERTY()
 	TObjectPtr<UMovieSceneCompiledDataManager> CompiledDataManager;
 
-	UE::MovieScene::FInstanceHandle RootInstanceHandle;
+	UE::MovieScene::FRootInstanceHandle RootInstanceHandle;
 
 	UPROPERTY()
 	TObjectPtr<UMovieSceneEntitySystemLinker> EntitySystemLinker;
diff --git a/Engine/Source/Runtime/MovieScene/Public/Evaluation/MovieScenePreAnimatedState.h b/Engine/Source/Runtime/MovieScene/Public/Evaluation/MovieScenePreAnimatedState.h
--- a/Engine/Source/Runtime/MovieScene/Public/Evaluation/MovieScenePreAnimatedState.h
+++ b/Engine/Source/Runtime/MovieScene/Public/Evaluation/MovieScenePreAnimatedState.h
@@ -51,7 +51,7 @@ public:
 
 	MOVIESCENE_API ~FMovieScenePreAnimatedState();
 
-	MOVIESCENE_API void Initialize(UMovieSceneEntitySystemLinker* Linker, UE::MovieScene::FInstanceHandle InstanceHandle);
+	MOVIESCENE_API void Initialize(UMovieSceneEntitySystemLinker* Linker, UE::MovieScene::FRootInstanceHandle InstanceHandle);
 
 	/**
 	 * Check whether this sequence instance is capturing any and all changes of state so they can be restored later
@@ -149,7 +149,7 @@ private:
 	TSharedPtr<UE::MovieScene::FPreAnimatedEvaluationHookCaptureSources> EvaluationHookMetaData;
 
 	/** The instance handle for the root sequence instance */
-	UE::MovieScene::FInstanceHandle InstanceHandle;
+	UE::MovieScene::FRootInstanceHandle InstanceHandle;
 
 	bool bCapturingGlobalPreAnimatedState;
 };
diff --git a/Engine/Source/Runtime/MovieScene/Public/Evaluation/PreAnimatedState/IMovieScenePreAnimatedCaptureSource.h b/Engine/Source/Runtime/MovieScene/Public/Evaluation/PreAnimatedState/IMovieScenePreAnimatedCaptureSource.h
--- a/Engine/Source/Runtime/MovieScene/Public/Evaluation/PreAnimatedState/IMovieScenePreAnimatedCaptureSource.h
+++ b/Engine/Source/Runtime/MovieScene/Public/Evaluation/PreAnimatedState/IMovieScenePreAnimatedCaptureSource.h
@@ -11,7 +11,7 @@ namespace UE
 namespace MovieScene
 {
 
-struct FInstanceHandle;
+struct FRootInstanceHandle;
 struct FRestoreStateParams;
 struct FPreAnimatedStateMetaData;
 struct FPreAnimatedStorageGroupHandle;
@@ -21,7 +21,7 @@ struct IPreAnimatedCaptureSource
 	virtual ~IPreAnimatedCaptureSource(){}
 
 	virtual void Reset() = 0;
-	virtual bool ContainsInstanceHandle(FInstanceHandle RootInstanceHandle) const = 0;
+	virtual bool ContainsInstanceHandle(FRootInstanceHandle RootInstanceHandle) const = 0;
 	virtual void GatherAndRemoveExpiredMetaData(const FRestoreStateParams& Params, TArray<FPreAnimatedStateMetaData>& OutExpiredMetaData) = 0;
 	virtual void GatherAndRemoveMetaDataForGroup(FPreAnimatedStorageGroupHandle Group, TArray<FPreAnimatedStateMetaData>& OutExpiredMetaData) = 0;
 };
diff --git a/Engine/Source/Runtime/MovieScene/Public/Evaluation/PreAnimatedState/IMovieScenePreAnimatedStorage.h b/Engine/Source/Runtime/MovieScene/Public/Evaluation/PreAnimatedState/IMovieScenePreAnimatedStorage.h
--- a/Engine/Source/Runtime/MovieScene/Public/Evaluation/PreAnimatedState/IMovieScenePreAnimatedStorage.h
+++ b/Engine/Source/Runtime/MovieScene/Public/Evaluation/PreAnimatedState/IMovieScenePreAnimatedStorage.h
@@ -47,7 +47,7 @@ struct IPreAnimatedStorage : TSharedFromThis<IPreAnimatedStorage>
 
 struct IPreAnimatedObjectEntityStorage
 {
-	virtual void BeginTrackingEntities(const FPreAnimatedTrackerParams& Params, TRead<FMovieSceneEntityID> EntityIDs, TRead<FInstanceHandle> InstanceHandles, TRead<UObject*> BoundObjects) = 0;
+	virtual void BeginTrackingEntities(const FPreAnimatedTrackerParams& Params, TRead<FMovieSceneEntityID> EntityIDs, TRead<FRootInstanceHandle> InstanceHandles, TRead<UObject*> BoundObjects) = 0;
 	virtual void CachePreAnimatedValues(const FCachePreAnimatedValueParams& Params, TArrayView<UObject* const> BoundObjects) = 0;
 };
 
@@ -55,7 +55,7 @@ struct IPreAnimatedObjectEntityStorage
 
 struct IPreAnimatedStateTokenStorage
 {
-	virtual void RestoreState(UMovieSceneEntitySystemLinker* Linker, const FMovieSceneEvaluationKey& Key, FInstanceHandle InstanceHandle) = 0;
+	virtual void RestoreState(UMovieSceneEntitySystemLinker* Linker, const FMovieSceneEvaluationKey& Key, FRootInstanceHandle InstanceHandle) = 0;
 };
 
 
@@ -64,7 +64,7 @@ struct IPreAnimatedObjectPropertyStorage
 {
 	using FThreeWayAccessor = TMultiReadOptional<FCustomPropertyIndex, uint16, TSharedPtr<FTrackInstancePropertyBindings>>;
 
-	virtual void BeginTrackingEntities(const FPreAnimatedTrackerParams& Params, TRead<FMovieSceneEntityID> EntityIDs, TRead<FInstanceHandle> InstanceHandles, TRead<UObject*> BoundObjects, TRead<FMovieScenePropertyBinding> PropertyBindings) = 0;
+	virtual void BeginTrackingEntities(const FPreAnimatedTrackerParams& Params, TRead<FMovieSceneEntityID> EntityIDs, TRead<FRootInstanceHandle> InstanceHandles, TRead<UObject*> BoundObjects, TRead<FMovieScenePropertyBinding> PropertyBindings) = 0;
 	virtual void CachePreAnimatedValues(const FCachePreAnimatedValueParams& Params, FEntityAllocationProxy Item, TRead<UObject*> Objects, TRead<FMovieScenePropertyBinding> PropertyBindings, FThreeWayAccessor Properties) = 0;
 };
 
diff --git a/Engine/Source/Runtime/MovieScene/Public/Evaluation/PreAnimatedState/MovieScenePreAnimatedCaptureSource.h b/Engine/Source/Runtime/MovieScene/Public/Evaluation/PreAnimatedState/MovieScenePreAnimatedCaptureSource.h
--- a/Engine/Source/Runtime/MovieScene/Public/Evaluation/PreAnimatedState/MovieScenePreAnimatedCaptureSource.h
+++ b/Engine/Source/Runtime/MovieScene/Public/Evaluation/PreAnimatedState/MovieScenePreAnimatedCaptureSource.h
@@ -72,7 +72,7 @@ private:
 	static FScopedPreAnimatedCaptureSource*& GetCaptureSourcePtr();
 
 	void BeginTracking(const UE::MovieScene::FPreAnimatedStateMetaData& MetaData, UMovieSceneEntitySystemLinker* Linker);
-	UE::MovieScene::FInstanceHandle GetRootInstanceHandle(UMovieSceneEntitySystemLinker* Linker) const;
+	UE::MovieScene::FRootInstanceHandle GetRootInstanceHandle(UMovieSceneEntitySystemLinker* Linker) const;
 
 	friend UE::MovieScene::FPreAnimatedStateExtension;
 	struct FEvalHookType
diff --git a/Engine/Source/Runtime/MovieScene/Public/Evaluation/PreAnimatedState/MovieScenePreAnimatedCaptureSources.h b/Engine/Source/Runtime/MovieScene/Public/Evaluation/PreAnimatedState/MovieScenePreAnimatedCaptureSources.h
--- a/Engine/Source/Runtime/MovieScene/Public/Evaluation/PreAnimatedState/MovieScenePreAnimatedCaptureSources.h
+++ b/Engine/Source/Runtime/MovieScene/Public/Evaluation/PreAnimatedState/MovieScenePreAnimatedCaptureSources.h
@@ -64,7 +64,7 @@ struct TPreAnimatedCaptureSources : IPreAnimatedCaptureSource
 	void StopTrackingCaptureSource(const KeyType& InKey);
 
 	void Reset() override;
-	bool ContainsInstanceHandle(FInstanceHandle RootInstanceHandle) const override;
+	bool ContainsInstanceHandle(FRootInstanceHandle RootInstanceHandle) const override;
 	void GatherAndRemoveExpiredMetaData(const FRestoreStateParams& Params, TArray<FPreAnimatedStateMetaData>& OutExpiredMetaData) override;
 	void GatherAndRemoveMetaDataForGroup(FPreAnimatedStorageGroupHandle Group, TArray<FPreAnimatedStateMetaData>& OutExpiredMetaData) override;
 
diff --git a/Engine/Source/Runtime/MovieScene/Public/Evaluation/PreAnimatedState/MovieScenePreAnimatedEntityCaptureSource.h b/Engine/Source/Runtime/MovieScene/Public/Evaluation/PreAnimatedState/MovieScenePreAnimatedEntityCaptureSource.h
--- a/Engine/Source/Runtime/MovieScene/Public/Evaluation/PreAnimatedState/MovieScenePreAnimatedEntityCaptureSource.h
+++ b/Engine/Source/Runtime/MovieScene/Public/Evaluation/PreAnimatedState/MovieScenePreAnimatedEntityCaptureSource.h
@@ -29,11 +29,11 @@ struct FPreAnimatedEntityCaptureSource
 
 	void Reset();
 
-	MOVIESCENE_API void BeginTrackingEntity(const FPreAnimatedStateEntry& Entry, FMovieSceneEntityID EntityID, FInstanceHandle RootInstanceHandle, bool bWantsRestoreState);
+	MOVIESCENE_API void BeginTrackingEntity(const FPreAnimatedStateEntry& Entry, FMovieSceneEntityID EntityID, FRootInstanceHandle RootInstanceHandle, bool bWantsRestoreState);
 	MOVIESCENE_API void StopTrackingEntity(FMovieSceneEntityID EntityID, FPreAnimatedStorageID StorageID);
 	MOVIESCENE_API void StopTrackingEntity(FMovieSceneEntityID EntityID);
 
-	bool ContainsInstanceHandle(FInstanceHandle RootInstanceHandle) const;
+	bool ContainsInstanceHandle(FRootInstanceHandle RootInstanceHandle) const;
 	void GatherAndRemoveExpiredMetaData(const FRestoreStateParams& Params, TArray<FPreAnimatedStateMetaData>& OutExpiredMetaData);
 	void GatherAndRemoveMetaDataForGroup(FPreAnimatedStorageGroupHandle Group, TArray<FPreAnimatedStateMetaData>& OutExpiredMetaData);
 
diff --git a/Engine/Source/Runtime/MovieScene/Public/Evaluation/PreAnimatedState/MovieScenePreAnimatedObjectStorage.h b/Engine/Source/Runtime/MovieScene/Public/Evaluation/PreAnimatedState/MovieScenePreAnimatedObjectStorage.h
--- a/Engine/Source/Runtime/MovieScene/Public/Evaluation/PreAnimatedState/MovieScenePreAnimatedObjectStorage.h
+++ b/Engine/Source/Runtime/MovieScene/Public/Evaluation/PreAnimatedState/MovieScenePreAnimatedObjectStorage.h
@@ -58,7 +58,7 @@ public:
 		this->ReplaceKey(StorageIndex, ExistingKey);
 	}
 
-	void BeginTrackingEntities(const FPreAnimatedTrackerParams& Params, TRead<FMovieSceneEntityID> EntityIDs, TRead<FInstanceHandle> InstanceHandles, TRead<UObject*> BoundObjects) override
+	void BeginTrackingEntities(const FPreAnimatedTrackerParams& Params, TRead<FMovieSceneEntityID> EntityIDs, TRead<FRootInstanceHandle> InstanceHandles, TRead<UObject*> BoundObjects) override
 	{
 		const int32 Num = Params.Num;
 		const bool  bWantsRestore = Params.bWantsRestoreState;
diff --git a/Engine/Source/Runtime/MovieScene/Public/Evaluation/PreAnimatedState/MovieScenePreAnimatedPropertyStorage.h b/Engine/Source/Runtime/MovieScene/Public/Evaluation/PreAnimatedState/MovieScenePreAnimatedPropertyStorage.h
--- a/Engine/Source/Runtime/MovieScene/Public/Evaluation/PreAnimatedState/MovieScenePreAnimatedPropertyStorage.h
+++ b/Engine/Source/Runtime/MovieScene/Public/Evaluation/PreAnimatedState/MovieScenePreAnimatedPropertyStorage.h
@@ -83,7 +83,7 @@ struct TPreAnimatedPropertyStorageImpl<PropertyTraits, TPropertyMetaData<MetaDat
 		return this;
 	}
 
-	void BeginTrackingEntities(const FPreAnimatedTrackerParams& Params, TRead<FMovieSceneEntityID> EntityIDs, TRead<FInstanceHandle> InstanceHandles, TRead<UObject*> BoundObjects, TRead<FMovieScenePropertyBinding> PropertyBindings) override
+	void BeginTrackingEntities(const FPreAnimatedTrackerParams& Params, TRead<FMovieSceneEntityID> EntityIDs, TRead<FRootInstanceHandle> InstanceHandles, TRead<UObject*> BoundObjects, TRead<FMovieScenePropertyBinding> PropertyBindings) override
 	{
 		const int32 Num = Params.Num;
 		const bool  bWantsRestore = Params.bWantsRestoreState;
diff --git a/Engine/Source/Runtime/MovieScene/Public/Evaluation/PreAnimatedState/MovieScenePreAnimatedStateExtension.h b/Engine/Source/Runtime/MovieScene/Public/Evaluation/PreAnimatedState/MovieScenePreAnimatedStateExtension.h
--- a/Engine/Source/Runtime/MovieScene/Public/Evaluation/PreAnimatedState/MovieScenePreAnimatedStateExtension.h
+++ b/Engine/Source/Runtime/MovieScene/Public/Evaluation/PreAnimatedState/MovieScenePreAnimatedStateExtension.h
@@ -248,7 +248,7 @@ public:
 	 * Search for any captured state that originated from the specified root instance handle
 	 * WARNING: This is a linear search across all state, and so is potentially very slow
 	 */
-	bool ContainsAnyStateForInstanceHandle(FInstanceHandle RootInstanceHandle) const;
+	bool ContainsAnyStateForInstanceHandle(FRootInstanceHandle RootInstanceHandle) const;
 
 
 	// Use FScopedPreAnimatedCaptureSource to capture from a specific source rather than globally
diff --git a/Engine/Source/Runtime/MovieScene/Public/Evaluation/PreAnimatedState/MovieScenePreAnimatedStateTypes.h b/Engine/Source/Runtime/MovieScene/Public/Evaluation/PreAnimatedState/MovieScenePreAnimatedStateTypes.h
--- a/Engine/Source/Runtime/MovieScene/Public/Evaluation/PreAnimatedState/MovieScenePreAnimatedStateTypes.h
+++ b/Engine/Source/Runtime/MovieScene/Public/Evaluation/PreAnimatedState/MovieScenePreAnimatedStateTypes.h
@@ -181,7 +181,7 @@ struct FPreAnimatedStateMetaData
 	/** The pre-animated state entry */
 	FPreAnimatedStateEntry Entry;
 	/** The instance handle for the root sequence */
-	FInstanceHandle RootInstanceHandle;
+	FRootInstanceHandle RootInstanceHandle;
 	/** True if this entry should be restored when it is removed */
 	bool bWantsRestoreState;
 };
diff --git a/Engine/Source/Runtime/MovieSceneTracks/Private/Systems/MovieSceneComponentAttachmentSystem.cpp b/Engine/Source/Runtime/MovieSceneTracks/Private/Systems/MovieSceneComponentAttachmentSystem.cpp
--- a/Engine/Source/Runtime/MovieSceneTracks/Private/Systems/MovieSceneComponentAttachmentSystem.cpp
+++ b/Engine/Source/Runtime/MovieSceneTracks/Private/Systems/MovieSceneComponentAttachmentSystem.cpp
@@ -277,7 +277,7 @@ void UMovieSceneComponentAttachmentSystem::SavePreAnimatedState(const FPreAnimat
 	.Read(BuiltInComponents->BoundObject)
 	.FilterAll(FilterMask)
 	.Iterate_PerAllocation(&Linker->EntityManager,
-		[EntityMetaData, ComponentTransformStorage, ComponentAttachmentStorage](FEntityAllocationIteratorItem Item, TRead<FMovieSceneEntityID> EntityIDs, TRead<FInstanceHandle> InstanceHandles, TRead<UObject*> BoundObjects)
+		[EntityMetaData, ComponentTransformStorage, ComponentAttachmentStorage](FEntityAllocationIteratorItem Item, TRead<FMovieSceneEntityID> EntityIDs, TRead<FRootInstanceHandle> InstanceHandles, TRead<UObject*> BoundObjects)
 		{
 			TArrayView<UObject* const> BoundObjectArray = BoundObjects.AsArray(Item.GetAllocation()->Num());
 
diff --git a/Engine/Source/Runtime/MovieSceneTracks/Private/Systems/MovieSceneComponentMobilitySystem.cpp b/Engine/Source/Runtime/MovieSceneTracks/Private/Systems/MovieSceneComponentMobilitySystem.cpp
--- a/Engine/Source/Runtime/MovieSceneTracks/Private/Systems/MovieSceneComponentMobilitySystem.cpp
+++ b/Engine/Source/Runtime/MovieSceneTracks/Private/Systems/MovieSceneComponentMobilitySystem.cpp
@@ -179,7 +179,7 @@ void UMovieSceneComponentMobilitySystem::SavePreAnimatedState(const FPreAnimatio
 
 	FPreAnimatedStateExtension* Extension = InParameters.CacheExtension;
 
-	auto IterNewObjects = [Extension](FEntityAllocationIteratorItem Item, TRead<FMovieSceneEntityID> EntityIDs, TRead<FInstanceHandle> InstanceHandles, TRead<UObject*> BoundObjects)
+	auto IterNewObjects = [Extension](FEntityAllocationIteratorItem Item, TRead<FMovieSceneEntityID> EntityIDs, TRead<FRootInstanceHandle> InstanceHandles, TRead<UObject*> BoundObjects)
 	{
 		TSharedPtr<FPreAnimatedComponentMobilityStorage> ComponentMobilityStorage
 			= Extension->GetOrCreateStorage<FPreAnimatedComponentMobilityStorage>();
@@ -201,7 +201,7 @@ void UMovieSceneComponentMobilitySystem::SavePreAnimatedState(const FPreAnimatio
 			}
 
 			FMovieSceneEntityID EntityID       = EntityIDs[Index];
-			FInstanceHandle     InstanceHandle = InstanceHandles[Index];
+			FRootInstanceHandle InstanceHandle = InstanceHandles[Index];
 
 			TArray<USceneComponent*, TInlineAllocator<4>> FlatHierarchy;
 			GetFlattenedHierarchy(SceneComponent, FlatHierarchy);
@@ -227,7 +227,7 @@ void UMovieSceneComponentMobilitySystem::SavePreAnimatedState(const FPreAnimatio
 
 	FEntityTaskBuilder()
 	.ReadEntityIDs()
-	.Read(BuiltInComponents->InstanceHandle)
+	.Read(BuiltInComponents->RootInstanceHandle)
 	.Read(BuiltInComponents->BoundObject)
 	.FilterAll(RestrictiveMask)
 	.FilterAny({ TrackComponents->ComponentTransform.PropertyTag, TrackComponents->AttachParent })
diff --git a/Engine/Source/Runtime/MovieSceneTracks/Private/Systems/MovieSceneDataLayerSystem.cpp b/Engine/Source/Runtime/MovieSceneTracks/Private/Systems/MovieSceneDataLayerSystem.cpp
--- a/Engine/Source/Runtime/MovieSceneTracks/Private/Systems/MovieSceneDataLayerSystem.cpp
+++ b/Engine/Source/Runtime/MovieSceneTracks/Private/Systems/MovieSceneDataLayerSystem.cpp
@@ -641,7 +641,7 @@ void UMovieSceneDataLayerSystem::BeginTrackingEntities()
 	auto GatherDataLayers = [this, EntityMetaData, DataLayerSubsystem, PreAnimatedStorage](
 		FEntityAllocationIteratorItem Item,
 		TRead<FMovieSceneEntityID> EntityIDs,
-		TRead<FInstanceHandle> RootInstanceHandles,
+		TRead<FRootInstanceHandle> RootInstanceHandles,
 		TRead<FMovieSceneDataLayerComponentData> ComponentData)
 	{
 		const bool bRestoreState = Item.GetAllocationType().Contains(FBuiltInComponentTypes::Get()->Tags.RestoreState);
@@ -651,7 +651,7 @@ void UMovieSceneDataLayerSystem::BeginTrackingEntities()
 			if (Section)
 			{
 				FMovieSceneEntityID EntityID     = EntityIDs[Index];
-				FInstanceHandle     RootInstance = RootInstanceHandles[Index];
+				FRootInstanceHandle RootInstance = RootInstanceHandles[Index];
 
 				for (const FActorDataLayer& ActorDataLayer : Section->GetDataLayers())
 				{
diff --git a/Engine/Source/Runtime/MovieSceneTracks/Private/Systems/MovieScenePropertySystem.cpp b/Engine/Source/Runtime/MovieSceneTracks/Private/Systems/MovieScenePropertySystem.cpp
--- a/Engine/Source/Runtime/MovieSceneTracks/Private/Systems/MovieScenePropertySystem.cpp
+++ b/Engine/Source/Runtime/MovieSceneTracks/Private/Systems/MovieScenePropertySystem.cpp
@@ -76,7 +76,7 @@ void UMovieScenePropertySystem::SavePreAnimatedState(const FPreAnimationParamete
 		.Read(BuiltInComponents->BoundObject)
 		.FilterAll(ComponentMask)
 		.Iterate_PerAllocation(&Linker->EntityManager,
-			[ObjectStorage](FEntityAllocationIteratorItem Item, TRead<FMovieSceneEntityID> EntityIDs, TRead<FInstanceHandle> RootInstanceHandles, TRead<UObject*> BoundObjects)
+			[ObjectStorage](FEntityAllocationIteratorItem Item, TRead<FMovieSceneEntityID> EntityIDs, TRead<FRootInstanceHandle> RootInstanceHandles, TRead<UObject*> BoundObjects)
 			{
 				ObjectStorage->BeginTrackingEntities(Item, EntityIDs, RootInstanceHandles, BoundObjects);
 			}
@@ -102,7 +102,7 @@ void UMovieScenePropertySystem::SavePreAnimatedState(const FPreAnimationParamete
 		.Read(BuiltInComponents->PropertyBinding)
 		.FilterAll(ComponentMask)
 		.Iterate_PerAllocation(&Linker->EntityManager,
-			[PropertyStorage](FEntityAllocationIteratorItem Item, TRead<FMovieSceneEntityID> EntityIDs, TRead<FInstanceHandle> RootInstanceHandles, TRead<UObject*> BoundObjects, TRead<FMovieScenePropertyBinding> PropertyBindings)
+			[PropertyStorage](FEntityAllocationIteratorItem Item, TRead<FMovieSceneEntityID> EntityIDs, TRead<FRootInstanceHandle> RootInstanceHandles, TRead<UObject*> BoundObjects, TRead<FMovieScenePropertyBinding> PropertyBindings)
 			{
 				PropertyStorage->BeginTrackingEntities(Item, EntityIDs, RootInstanceHandles, BoundObjects, PropertyBindings);
 			}
