Sequencer: Clamp sound duration to 0.

diff --git a/Engine/Source/Runtime/MovieScene/Private/MovieSceneCommonHelpers.cpp b/Engine/Source/Runtime/MovieScene/Private/MovieSceneCommonHelpers.cpp
--- a/Engine/Source/Runtime/MovieScene/Private/MovieSceneCommonHelpers.cpp
+++ b/Engine/Source/Runtime/MovieScene/Private/MovieSceneCommonHelpers.cpp
@@ -434,7 +434,7 @@ UCameraComponent* MovieSceneHelpers::CameraComponentFromRuntimeObject(UObject* R
 
 float MovieSceneHelpers::GetSoundDuration(USoundBase* Sound)
 {
-	return Sound ? Sound->GetDuration() : 0.0f;
+	return Sound ? FMath::Max(0.0f, Sound->GetDuration()) : 0.0f;
 }
 
 
