Sequencer: Clear rows disabled if there aren't multiple rows

diff --git a/Engine/Source/Runtime/MovieScene/Private/MovieSceneTrack.cpp b/Engine/Source/Runtime/MovieScene/Private/MovieSceneTrack.cpp
--- a/Engine/Source/Runtime/MovieScene/Private/MovieSceneTrack.cpp
+++ b/Engine/Source/Runtime/MovieScene/Private/MovieSceneTrack.cpp
@@ -265,6 +265,13 @@ bool UMovieSceneTrack::FixRowIndices()
 				++NewIndex;
 			}
 		}
+
+		// If there aren't multiple rows (ie. max row is 0), there shouldn't be any disabled rows either
+		if (GetMaxRowIndex() == 0 && !RowsDisabled.IsEmpty())
+		{
+			Modify();
+			RowsDisabled.Empty();
+		}
 	}
 	else
 	{
