Fix looping subsequences clamping

We were previously not clamping compilation range by the containing section's range, so some of the sub-sequence's animation would be taken into account after the last loop and still play.

diff --git a/Engine/Source/Runtime/MovieScene/Private/Compilation/MovieSceneCompiledDataManager.cpp b/Engine/Source/Runtime/MovieScene/Private/Compilation/MovieSceneCompiledDataManager.cpp
--- a/Engine/Source/Runtime/MovieScene/Private/Compilation/MovieSceneCompiledDataManager.cpp
+++ b/Engine/Source/Runtime/MovieScene/Private/Compilation/MovieSceneCompiledDataManager.cpp
@@ -1546,7 +1546,8 @@ void UMovieSceneCompiledDataManager::PopulateSubSequenceTree(UMovieSceneSubTrack
 			continue;
 		}
 
-		TRange<FFrameNumber> EffectiveRange = Params.ClampRoot(Entry.Range * Params.RootToSequenceTransform.InverseFromWarp(Params.RootToSequenceWarpCounter));
+		const FMovieSceneTimeTransform SequenceToRootTransform = Params.RootToSequenceTransform.InverseFromWarp(Params.RootToSequenceWarpCounter);
+		TRange<FFrameNumber> EffectiveRange = Params.ClampRoot(Entry.Range * SequenceToRootTransform);
 		if (EffectiveRange.IsEmpty())
 		{
 			continue;
@@ -1583,8 +1584,6 @@ void UMovieSceneCompiledDataManager::PopulateSubSequenceTree(UMovieSceneSubTrack
 		else
 		{
 			// The section is looping so we need to add its contents to the tree as many times as it has loops.
-			const FMovieSceneSequenceTransform SequenceToRootTransform = Params.RootToSequenceTransform.InverseFromWarp(Params.RootToSequenceWarpCounter);
-
 			const float RootToSubSequenceTimeScale = SubData->RootToSequenceTransform.GetTimeScale();
 			const float SubSequenceToRootTimeScale = (RootToSubSequenceTimeScale != 0.f) ? 1.0f / RootToSubSequenceTimeScale : 1.f;
 
@@ -1611,14 +1610,16 @@ void UMovieSceneCompiledDataManager::PopulateSubSequenceTree(UMovieSceneSubTrack
 				{
 					if (CurRootRange.Overlaps(Params.RootClampRange))
 					{
+						// Clamp the sub-sequence's range by the containing section's range and the current compilation range.
+						const TRange<FFrameNumber> ClampedCurRootRange = TRange<FFrameNumber>::Intersection(CurRootRange, EffectiveRange);
+
 						FGatherParameters CurLoopParams = Params.CreateForSubData(*SubData, SubSequenceID, Params.RootToSequenceWarpCounter);
-						CurLoopParams.SetClampRange(EffectiveRange);
+						CurLoopParams.SetClampRange(ClampedCurRootRange);
 						CurLoopParams.Flags |= Entry.Flags;
 						CurLoopParams.NetworkMask = NewMask;
 						CurLoopParams.RootToSequenceWarpCounter.AddWarpingLevel(LoopCount);
 
 						// Add the section to the tree for the current loop.
-						const TRange<FFrameNumber> ClampedCurRootRange = TRange<FFrameNumber>::Intersection(CurRootRange, Params.RootClampRange);
 						InOutHierarchy->AddRange(ClampedCurRootRange, SubSequenceID, SubEntryFlags, CurLoopParams.RootToSequenceWarpCounter);
 
 						// Recurse into this loop's sub sequence.
@@ -1629,6 +1630,7 @@ void UMovieSceneCompiledDataManager::PopulateSubSequenceTree(UMovieSceneSubTrack
 						RootPath->PopGenerations(1);
 					}
 
+					// Move on to the next loop.
 					CurRootRangeStart = CurRootRange.GetUpperBoundValue();
 					CurRootRange = TRange<FFrameNumber>(CurRootRangeStart.FloorToFrame(), (CurRootRangeStart + RootLoopLength).FloorToFrame());
 					if (CurRootRange.GetUpperBoundValue() > RootSectionEndTime)
@@ -1636,7 +1638,7 @@ void UMovieSceneCompiledDataManager::PopulateSubSequenceTree(UMovieSceneSubTrack
 					++LoopCount;
 				}
 			}
-			// Faced with the cosmic horror or infinites, we choose to shield our sanity and skip this sub-section.
+			// Faced with the cosmic horror of infinites, we choose to shield our sanity and skip this sub-section.
 			// (it either has an open-ended start time, which means we needed to loop since before time began, which means
 			//  we don't know where loops are in the present... or it means the section and root sequence have open-ended
 			//  end times, which means we would need to compile loops forever)
