Cleaning up dependencies with select heavy-hitter headers in MovieScene and addressing the fallout from that.
MovieScene was targeted as its headers were at the top of profiling ShooterGame builds.
Tested with -disableunity and -nopch. Current timings take ShooterGameEditor rebuild from 480s to 440s

diff --git a/Engine/Plugins/Importers/USDImporter/Source/USDUtilities/Private/USDPrimConversion.cpp b/Engine/Plugins/Importers/USDImporter/Source/USDUtilities/Private/USDPrimConversion.cpp
--- a/Engine/Plugins/Importers/USDImporter/Source/USDUtilities/Private/USDPrimConversion.cpp
+++ b/Engine/Plugins/Importers/USDImporter/Source/USDUtilities/Private/USDPrimConversion.cpp
@@ -38,6 +38,7 @@
 #include "GeometryCache.h"
 #include "GeometryCacheComponent.h"
 #include "InstancedFoliageActor.h"
+#include "MovieScene.h"
 #include "MovieSceneTimeHelpers.h"
 #include "Rendering/SkeletalMeshRenderData.h"
 #include "Sections/MovieScene3DTransformSection.h"
diff --git a/Engine/Plugins/Media/MediaCompositing/Source/MediaCompositingEditor/Private/Sequencer/MovieSceneMediaPlayerSectionRecorder.cpp b/Engine/Plugins/Media/MediaCompositing/Source/MediaCompositingEditor/Private/Sequencer/MovieSceneMediaPlayerSectionRecorder.cpp
--- a/Engine/Plugins/Media/MediaCompositing/Source/MediaCompositingEditor/Private/Sequencer/MovieSceneMediaPlayerSectionRecorder.cpp
+++ b/Engine/Plugins/Media/MediaCompositing/Source/MediaCompositingEditor/Private/Sequencer/MovieSceneMediaPlayerSectionRecorder.cpp
@@ -14,6 +14,7 @@
 #include "MovieSceneMediaTrack.h"
 #include "ImgMediaSource.h"
 #include "UObject/Package.h"
+#include "MovieScene.h"
 
 
 
diff --git a/Engine/Plugins/MovieScene/LevelSequenceEditor/Source/LevelSequenceEditor/Public/Misc/LevelSequenceEditorSpawnRegister.h b/Engine/Plugins/MovieScene/LevelSequenceEditor/Source/LevelSequenceEditor/Public/Misc/LevelSequenceEditorSpawnRegister.h
--- a/Engine/Plugins/MovieScene/LevelSequenceEditor/Source/LevelSequenceEditor/Public/Misc/LevelSequenceEditorSpawnRegister.h
+++ b/Engine/Plugins/MovieScene/LevelSequenceEditor/Source/LevelSequenceEditor/Public/Misc/LevelSequenceEditorSpawnRegister.h
@@ -12,6 +12,7 @@
 class IMovieScenePlayer;
 class ISequencer;
 class UMovieScene;
+class UMovieSceneSequence;
 class FObjectPreSaveContext;
 
 /**
diff --git a/Engine/Plugins/MovieScene/MovieRenderPipeline/Source/MovieRenderPipelineCore/Private/MoviePipelineBlueprintLibrary.cpp b/Engine/Plugins/MovieScene/MovieRenderPipeline/Source/MovieRenderPipelineCore/Private/MoviePipelineBlueprintLibrary.cpp
--- a/Engine/Plugins/MovieScene/MovieRenderPipeline/Source/MovieRenderPipelineCore/Private/MoviePipelineBlueprintLibrary.cpp
+++ b/Engine/Plugins/MovieScene/MovieRenderPipeline/Source/MovieRenderPipelineCore/Private/MoviePipelineBlueprintLibrary.cpp
@@ -3,6 +3,7 @@
 #include "MoviePipelineBlueprintLibrary.h"
 #include "MoviePipeline.h"
 #include "MovieRenderPipelineDataTypes.h"
+#include "MovieScene.h"
 #include "MovieSceneSequence.h"
 #include "MovieSceneSpawnable.h"
 #include "MovieScenePossessable.h"
diff --git a/Engine/Plugins/MovieScene/MovieRenderPipeline/Source/MovieRenderPipelineCore/Private/MoviePipelineUtils.cpp b/Engine/Plugins/MovieScene/MovieRenderPipeline/Source/MovieRenderPipelineCore/Private/MoviePipelineUtils.cpp
--- a/Engine/Plugins/MovieScene/MovieRenderPipeline/Source/MovieRenderPipelineCore/Private/MoviePipelineUtils.cpp
+++ b/Engine/Plugins/MovieScene/MovieRenderPipeline/Source/MovieRenderPipelineCore/Private/MoviePipelineUtils.cpp
@@ -14,6 +14,7 @@
 #include "Tracks/MovieSceneCinematicShotTrack.h"
 #include "Sections/MovieSceneCinematicShotSection.h"
 #include "Sections/MovieSceneCameraCutSection.h"
+#include "MovieScene.h"
 #include "MovieSceneSequence.h"
 #include "MovieRenderPipelineCoreModule.h"
 #include "Misc/Paths.h"
diff --git a/Engine/Plugins/MovieScene/ReplayTracks/Source/ReplayTracks/Private/Tracks/MovieSceneReplayTrack.cpp b/Engine/Plugins/MovieScene/ReplayTracks/Source/ReplayTracks/Private/Tracks/MovieSceneReplayTrack.cpp
--- a/Engine/Plugins/MovieScene/ReplayTracks/Source/ReplayTracks/Private/Tracks/MovieSceneReplayTrack.cpp
+++ b/Engine/Plugins/MovieScene/ReplayTracks/Source/ReplayTracks/Private/Tracks/MovieSceneReplayTrack.cpp
@@ -3,6 +3,7 @@
 #include "Tracks/MovieSceneReplayTrack.h"
 #include "IMovieSceneTracksModule.h"
 #include "Sections/MovieSceneReplaySection.h"
+#include "MovieScene.h"
 
 #define LOCTEXT_NAMESPACE "MovieSceneReplayTrack"
 
diff --git a/Engine/Plugins/MovieScene/SequencerScripting/Source/SequencerScripting/Private/ExtensionLibraries/MovieSceneFolderExtensions.cpp b/Engine/Plugins/MovieScene/SequencerScripting/Source/SequencerScripting/Private/ExtensionLibraries/MovieSceneFolderExtensions.cpp
--- a/Engine/Plugins/MovieScene/SequencerScripting/Source/SequencerScripting/Private/ExtensionLibraries/MovieSceneFolderExtensions.cpp
+++ b/Engine/Plugins/MovieScene/SequencerScripting/Source/SequencerScripting/Private/ExtensionLibraries/MovieSceneFolderExtensions.cpp
@@ -3,7 +3,7 @@
 #include "ExtensionLibraries/MovieSceneFolderExtensions.h"
 #include "MovieSceneFolder.h"
 #include "MovieSceneSequence.h"
-
+#include "MovieScene.h"
 
 FName UMovieSceneFolderExtensions::GetFolderName(UMovieSceneFolder* Folder)
 {
diff --git a/Engine/Plugins/MovieScene/SequencerScripting/Source/SequencerScripting/Private/KeysAndChannels/MovieSceneScriptingChannel.h b/Engine/Plugins/MovieScene/SequencerScripting/Source/SequencerScripting/Private/KeysAndChannels/MovieSceneScriptingChannel.h
--- a/Engine/Plugins/MovieScene/SequencerScripting/Source/SequencerScripting/Private/KeysAndChannels/MovieSceneScriptingChannel.h
+++ b/Engine/Plugins/MovieScene/SequencerScripting/Source/SequencerScripting/Private/KeysAndChannels/MovieSceneScriptingChannel.h
@@ -13,9 +13,11 @@
 #include "SequencerScriptingRange.h"
 #include "SequenceTimeUnit.h"
 #include "ExtensionLibraries/MovieSceneSequenceExtensions.h"
+#include "MovieSceneSection.h"
 #include "MovieSceneSequence.h"
 #include "MovieScene.h"
 #include "MovieSceneScriptingChannel.generated.h"
+
 UCLASS(abstract, BlueprintType)
 class UMovieSceneScriptingKey : public UObject
 {
diff --git a/Engine/Plugins/MovieScene/SequencerScripting/Source/SequencerScripting/Private/KeysAndChannels/MovieSceneScriptingObjectPath.h b/Engine/Plugins/MovieScene/SequencerScripting/Source/SequencerScripting/Private/KeysAndChannels/MovieSceneScriptingObjectPath.h
--- a/Engine/Plugins/MovieScene/SequencerScripting/Source/SequencerScripting/Private/KeysAndChannels/MovieSceneScriptingObjectPath.h
+++ b/Engine/Plugins/MovieScene/SequencerScripting/Source/SequencerScripting/Private/KeysAndChannels/MovieSceneScriptingObjectPath.h
@@ -11,6 +11,7 @@
 #include "KeysAndChannels/MovieSceneScriptingChannel.h"
 #include "KeyParams.h"
 #include "MovieScene.h"
+#include "MovieSceneSection.h"
 #include "MovieSceneSequence.h"
 #include "SequencerBindingProxy.h"
 
diff --git a/Engine/Plugins/MovieScene/SequencerScripting/Source/SequencerScripting/Public/ExtensionLibraries/MovieSceneSequenceExtensions.h b/Engine/Plugins/MovieScene/SequencerScripting/Source/SequencerScripting/Public/ExtensionLibraries/MovieSceneSequenceExtensions.h
--- a/Engine/Plugins/MovieScene/SequencerScripting/Source/SequencerScripting/Public/ExtensionLibraries/MovieSceneSequenceExtensions.h
+++ b/Engine/Plugins/MovieScene/SequencerScripting/Source/SequencerScripting/Public/ExtensionLibraries/MovieSceneSequenceExtensions.h
@@ -10,6 +10,8 @@
 #include "SequencerScriptingRange.h"
 #include "Templates/SubclassOf.h"
 #include "MovieSceneTrack.h"
+#include "MovieSceneObjectBindingID.h" // for EMovieSceneObjectBindingSpace
+#include "MovieScene.h" // only for FMovieSceneMarkedFrame in the .gen.cpp
 
 #include "MovieSceneSequenceExtensions.generated.h"
 
diff --git a/Engine/Plugins/MovieScene/SequencerScripting/Source/SequencerScriptingEditor/Public/SequencerTools.h b/Engine/Plugins/MovieScene/SequencerScripting/Source/SequencerScriptingEditor/Public/SequencerTools.h
--- a/Engine/Plugins/MovieScene/SequencerScripting/Source/SequencerScriptingEditor/Public/SequencerTools.h
+++ b/Engine/Plugins/MovieScene/SequencerScripting/Source/SequencerScriptingEditor/Public/SequencerTools.h
@@ -127,9 +127,6 @@ public:
 	UFUNCTION(BlueprintCallable, Category = "Editor Scripting | Sequencer Tools | FBX")
 	static bool ExportLevelSequenceFBX(UWorld* InWorld, ULevelSequence* InSequence, const TArray<FSequencerBindingProxy>& InBindings, const TArray<UMovieSceneTrack*>& InMasterTracks, UFbxExportOption* OverrideOptions, const FString& InFBXFileName);
 
-	UE_DEPRECATED(4.27, "Please use ExportLevelSequenceFBX instead")
-	static bool ExportFBX(UWorld* InWorld, ULevelSequence* InSequence, const TArray<FSequencerBindingProxy>& InBindings, UFbxExportOption* OverrideOptions,const FString& InFBXFileName) { return ExportLevelSequenceFBX(InWorld, InSequence, InBindings, InSequence->GetMovieScene()->GetMasterTracks(), OverrideOptions, InFBXFileName); }
-
 	/*
 	 * Export Passed in Binding as an Anim Seqquence.
 	 *
@@ -156,9 +153,6 @@ public:
 	UFUNCTION(BlueprintCallable, Category = "Editor Scripting | Sequencer Tools | FBX")
 	static bool ImportLevelSequenceFBX(UWorld* InWorld, ULevelSequence* InSequence, const TArray<FSequencerBindingProxy>& InBindings, UMovieSceneUserImportFBXSettings* InImportFBXSettings, const FString& InImportFilename);
 
-	UE_DEPRECATED(4.27, "Please use ImportLevelSequenceFBX instead")
-	static bool ImportFBX(UWorld* InWorld, ULevelSequence* InSequence, const TArray<FSequencerBindingProxy>& InBindings, UMovieSceneUserImportFBXSettings* InImportFBXSettings, const FString& InImportFilename) { return ImportLevelSequenceFBX(InWorld, InSequence, InBindings, InImportFBXSettings, InImportFilename); }
-
 public:
 	/**
 	 * Create an event from a previously created blueprint endpoint and a payload. The resulting event should be added only
diff --git a/Engine/Plugins/Runtime/GameplayAbilities/Source/GameplayAbilities/Public/Sequencer/MovieSceneGameplayCueSections.h b/Engine/Plugins/Runtime/GameplayAbilities/Source/GameplayAbilities/Public/Sequencer/MovieSceneGameplayCueSections.h
--- a/Engine/Plugins/Runtime/GameplayAbilities/Source/GameplayAbilities/Public/Sequencer/MovieSceneGameplayCueSections.h
+++ b/Engine/Plugins/Runtime/GameplayAbilities/Source/GameplayAbilities/Public/Sequencer/MovieSceneGameplayCueSections.h
@@ -11,6 +11,7 @@
 #include "Channels/MovieSceneChannelTraits.h"
 #include "MovieSceneClipboard.h"
 #include "GameplayCueInterface.h"
+#include "MovieSceneObjectBindingID.h"
 #include "MovieSceneGameplayCueSections.generated.h"
 
 
diff --git a/Engine/Plugins/Runtime/HairStrands/Source/HairStrandsCore/Private/MovieSceneGroomCacheSection.cpp b/Engine/Plugins/Runtime/HairStrands/Source/HairStrandsCore/Private/MovieSceneGroomCacheSection.cpp
--- a/Engine/Plugins/Runtime/HairStrands/Source/HairStrandsCore/Private/MovieSceneGroomCacheSection.cpp
+++ b/Engine/Plugins/Runtime/HairStrands/Source/HairStrandsCore/Private/MovieSceneGroomCacheSection.cpp
@@ -2,8 +2,10 @@
 
 #include "MovieSceneGroomCacheSection.h"
 #include "GroomCache.h"
+#include "MovieScene.h"
 #include "MovieSceneGroomCacheTemplate.h"
 #include "MovieSceneTimeHelpers.h"
+#include "Misc/QualifiedFrameTime.h"
 
 FMovieSceneGroomCacheParams::FMovieSceneGroomCacheParams()
 {
diff --git a/Engine/Plugins/Runtime/HairStrands/Source/HairStrandsCore/Private/MovieSceneGroomCacheTrack.cpp b/Engine/Plugins/Runtime/HairStrands/Source/HairStrandsCore/Private/MovieSceneGroomCacheTrack.cpp
--- a/Engine/Plugins/Runtime/HairStrands/Source/HairStrandsCore/Private/MovieSceneGroomCacheTrack.cpp
+++ b/Engine/Plugins/Runtime/HairStrands/Source/HairStrandsCore/Private/MovieSceneGroomCacheTrack.cpp
@@ -2,6 +2,7 @@
 
 #include "MovieSceneGroomCacheTrack.h"
 #include "GroomComponent.h"
+#include "MovieScene.h"
 #include "MovieSceneGroomCacheSection.h"
 #include "MovieSceneGroomCacheTemplate.h"
 
diff --git a/Engine/Plugins/VirtualProduction/Takes/Source/TakesCore/Private/TakePreset.cpp b/Engine/Plugins/VirtualProduction/Takes/Source/TakesCore/Private/TakePreset.cpp
--- a/Engine/Plugins/VirtualProduction/Takes/Source/TakesCore/Private/TakePreset.cpp
+++ b/Engine/Plugins/VirtualProduction/Takes/Source/TakesCore/Private/TakePreset.cpp
@@ -2,6 +2,7 @@
 
 #include "TakePreset.h"
 #include "LevelSequence.h"
+#include "MovieScene.h"
 #include "MovieSceneToolsProjectSettings.h"
 #include "UObject/Package.h"
 
diff --git a/Engine/Source/Editor/MovieSceneTools/Private/Channels/DoubleChannelCurveModel.cpp b/Engine/Source/Editor/MovieSceneTools/Private/Channels/DoubleChannelCurveModel.cpp
--- a/Engine/Source/Editor/MovieSceneTools/Private/Channels/DoubleChannelCurveModel.cpp
+++ b/Engine/Source/Editor/MovieSceneTools/Private/Channels/DoubleChannelCurveModel.cpp
@@ -3,6 +3,7 @@
 #include "Channels/DoubleChannelCurveModel.h"
 #include "Channels/DoubleChannelKeyProxy.h"
 #include "CurveEditorScreenSpace.h"
+#include "MovieScene.h"
 
 /**
  * Buffered curve implementation for a double channel curve model, stores a copy of the double channel in order to draw itself.
diff --git a/Engine/Source/Editor/MovieSceneTools/Private/Channels/FloatChannelCurveModel.cpp b/Engine/Source/Editor/MovieSceneTools/Private/Channels/FloatChannelCurveModel.cpp
--- a/Engine/Source/Editor/MovieSceneTools/Private/Channels/FloatChannelCurveModel.cpp
+++ b/Engine/Source/Editor/MovieSceneTools/Private/Channels/FloatChannelCurveModel.cpp
@@ -3,6 +3,7 @@
 #include "Channels/FloatChannelCurveModel.h"
 #include "Channels/FloatChannelKeyProxy.h"
 #include "CurveEditorScreenSpace.h"
+#include "MovieScene.h"
 
 /**
  * Buffered curve implementation for a float channel curve model, stores a copy of the float channel in order to draw itself.
diff --git a/Engine/Source/Editor/MovieSceneTools/Private/MovieSceneEventUtils.cpp b/Engine/Source/Editor/MovieSceneTools/Private/MovieSceneEventUtils.cpp
--- a/Engine/Source/Editor/MovieSceneTools/Private/MovieSceneEventUtils.cpp
+++ b/Engine/Source/Editor/MovieSceneTools/Private/MovieSceneEventUtils.cpp
@@ -7,6 +7,7 @@
 #include "UObject/StrongObjectPtr.h"
 #include "UObject/Package.h"
 
+#include "MovieScene.h"
 #include "MovieSceneSequence.h"
 #include "Sections/MovieSceneEventSectionBase.h"
 #include "Tracks/MovieSceneEventTrack.h"
diff --git a/Engine/Source/Editor/Sequencer/Private/SSequencerSectionAreaView.cpp b/Engine/Source/Editor/Sequencer/Private/SSequencerSectionAreaView.cpp
--- a/Engine/Source/Editor/Sequencer/Private/SSequencerSectionAreaView.cpp
+++ b/Engine/Source/Editor/Sequencer/Private/SSequencerSectionAreaView.cpp
@@ -7,6 +7,7 @@
 #include "CommonMovieSceneTools.h"
 #include "MovieSceneTrack.h"
 #include "MovieScene.h"
+#include "Generators/MovieSceneEasingFunction.h"
 
 namespace SequencerSectionAreaConstants
 {
diff --git a/Engine/Source/Editor/Sequencer/Private/SequencerSectionPainter.cpp b/Engine/Source/Editor/Sequencer/Private/SequencerSectionPainter.cpp
--- a/Engine/Source/Editor/Sequencer/Private/SequencerSectionPainter.cpp
+++ b/Engine/Source/Editor/Sequencer/Private/SequencerSectionPainter.cpp
@@ -1,6 +1,7 @@
 // Copyright Epic Games, Inc. All Rights Reserved.
 
 #include "SequencerSectionPainter.h"
+#include "MovieSceneSection.h"
 #include "MovieSceneTrack.h"
 
 FSequencerSectionPainter::FSequencerSectionPainter(FSlateWindowElementList& OutDrawElements, const FGeometry& InSectionGeometry, UMovieSceneSection& InSection)
diff --git a/Engine/Source/Editor/Sequencer/Public/SequencerChannelTraits.h b/Engine/Source/Editor/Sequencer/Public/SequencerChannelTraits.h
--- a/Engine/Source/Editor/Sequencer/Public/SequencerChannelTraits.h
+++ b/Engine/Source/Editor/Sequencer/Public/SequencerChannelTraits.h
@@ -11,6 +11,7 @@
 #include "SequencerKeyStructGenerator.h"
 #include "Widgets/SNullWidget.h"
 #include "ISequencer.h"
+#include "MovieSceneSection.h"
 
 /** Utility struct representing a number of selected keys on a single channel */
 template<typename ChannelType>
diff --git a/Engine/Source/Runtime/LevelSequence/Private/LevelSequencePlayer.cpp b/Engine/Source/Runtime/LevelSequence/Private/LevelSequencePlayer.cpp
--- a/Engine/Source/Runtime/LevelSequence/Private/LevelSequencePlayer.cpp
+++ b/Engine/Source/Runtime/LevelSequence/Private/LevelSequencePlayer.cpp
@@ -26,6 +26,7 @@
 #include "LevelUtils.h"
 #include "Core/Public/ProfilingDebugging/CsvProfiler.h"
 #include "LevelSequenceModule.h"
+#include "Generators/MovieSceneEasingCurves.h"
 
 /* ULevelSequencePlayer structors
  *****************************************************************************/
diff --git a/Engine/Source/Runtime/LevelSequence/Public/LevelSequencePlayer.h b/Engine/Source/Runtime/LevelSequence/Public/LevelSequencePlayer.h
--- a/Engine/Source/Runtime/LevelSequence/Public/LevelSequencePlayer.h
+++ b/Engine/Source/Runtime/LevelSequence/Public/LevelSequencePlayer.h
@@ -7,7 +7,9 @@
 #include "UObject/ObjectMacros.h"
 #include "UObject/Object.h"
 #include "UObject/ScriptMacros.h"
+#if UE_ENABLE_INCLUDE_ORDER_DEPRECATED_IN_5_1
 #include "IMovieScenePlayer.h"
+#endif
 #include "Evaluation/MovieScenePlayback.h"
 #include "Evaluation/MovieSceneEvaluationTemplateInstance.h"
 #include "Evaluation/PersistentEvaluationData.h"
diff --git a/Engine/Source/Runtime/MovieScene/Private/Compilation/MovieSceneCompiledDataManager.cpp b/Engine/Source/Runtime/MovieScene/Private/Compilation/MovieSceneCompiledDataManager.cpp
--- a/Engine/Source/Runtime/MovieScene/Private/Compilation/MovieSceneCompiledDataManager.cpp
+++ b/Engine/Source/Runtime/MovieScene/Private/Compilation/MovieSceneCompiledDataManager.cpp
@@ -7,6 +7,7 @@
 #include "EntitySystem/IMovieSceneEntityProvider.h"
 #include "Evaluation/MovieSceneEvaluationCustomVersion.h"
 #include "Evaluation/MovieSceneRootOverridePath.h"
+#include "MovieScene.h"
 #include "MovieSceneSequence.h"
 #include "Sections/MovieSceneSubSection.h"
 #include "Tracks/MovieSceneSubTrack.h"
diff --git a/Engine/Source/Runtime/MovieScene/Private/Compilation/MovieSceneCompilerRules.cpp b/Engine/Source/Runtime/MovieScene/Private/Compilation/MovieSceneCompilerRules.cpp
new file mode 100644
--- /dev/null
+++ b/Engine/Source/Runtime/MovieScene/Private/Compilation/MovieSceneCompilerRules.cpp
@@ -0,0 +1,103 @@
+// Copyright Epic Games, Inc. All Rights Reserved.
+
+#include "Compilation/MovieSceneCompilerRules.h"
+#include "MovieSceneSection.h"
+
+TOptional<FMovieSceneSegment> MovieSceneSegmentCompiler::EvaluateNearestSegment(const TRange<FFrameNumber>& Range, const FMovieSceneSegment* PreviousSegment, const FMovieSceneSegment* NextSegment)
+{
+	if (PreviousSegment)
+	{
+		// There is a preceeding segment
+		FFrameNumber PreviousSegmentRangeBound = PreviousSegment->Range.GetUpperBoundValue();
+
+		FMovieSceneSegment EmptySpace(Range);
+		for (FSectionEvaluationData Data : PreviousSegment->Impls)
+		{
+			EmptySpace.Impls.Add(FSectionEvaluationData(Data.ImplIndex, PreviousSegmentRangeBound));
+		}
+		return EmptySpace;
+	}
+	else if (NextSegment)
+	{
+		// Before any sections
+		FFrameNumber NextSegmentRangeBound = NextSegment->Range.GetLowerBoundValue();
+
+		FMovieSceneSegment EmptySpace(Range);
+		for (FSectionEvaluationData Data : NextSegment->Impls)
+		{
+			EmptySpace.Impls.Add(FSectionEvaluationData(Data.ImplIndex, NextSegmentRangeBound));
+		}
+		return EmptySpace;
+	}
+
+	return TOptional<FMovieSceneSegment>();
+}
+
+bool MovieSceneSegmentCompiler::AlwaysEvaluateSection(const FMovieSceneSectionData& InSectionData)
+{
+	return InSectionData.Section->GetBlendType().IsValid() || EnumHasAnyFlags(InSectionData.Flags, ESectionEvaluationFlags::PreRoll | ESectionEvaluationFlags::PostRoll);
+}
+
+void MovieSceneSegmentCompiler::FilterOutUnderlappingSections(FSegmentBlendData& BlendData)
+{
+	if (!BlendData.Num())
+	{
+		return;
+	}
+
+	int32 HighestOverlap = TNumericLimits<int32>::Lowest();
+	for (const FMovieSceneSectionData& SectionData : BlendData)
+	{
+		if (!AlwaysEvaluateSection(SectionData))
+		{
+			HighestOverlap = FMath::Max(HighestOverlap, SectionData.Section->GetOverlapPriority());
+		}
+	}
+
+	// Remove anything that's not the highest priority, (excluding pre/postroll sections)
+	for (int32 RemoveAtIndex = BlendData.Num() - 1; RemoveAtIndex >= 0; --RemoveAtIndex)
+	{
+		const FMovieSceneSectionData& SectionData = BlendData[RemoveAtIndex];
+		if (SectionData.Section->GetOverlapPriority() != HighestOverlap && !AlwaysEvaluateSection(SectionData))
+		{
+			BlendData.RemoveAt(RemoveAtIndex, 1, false);
+		}
+	}
+}
+
+void MovieSceneSegmentCompiler::ChooseLowestRowIndex(FSegmentBlendData& BlendData)
+{
+	if (!BlendData.Num())
+	{
+		return;
+	}
+
+	int32 LowestRowIndex = TNumericLimits<int32>::Max();
+	for (const FMovieSceneSectionData& SectionData : BlendData)
+	{
+		if (!AlwaysEvaluateSection(SectionData))
+		{
+			LowestRowIndex = FMath::Min(LowestRowIndex, SectionData.Section->GetRowIndex());
+		}
+	}
+
+	// Remove anything that's not the highest priority, (excluding pre/postroll sections)
+	for (int32 RemoveAtIndex = BlendData.Num() - 1; RemoveAtIndex >= 0; --RemoveAtIndex)
+	{
+		const FMovieSceneSectionData& SectionData = BlendData[RemoveAtIndex];
+		if (SectionData.Section->GetRowIndex() > LowestRowIndex && !AlwaysEvaluateSection(SectionData))
+		{
+			BlendData.RemoveAt(RemoveAtIndex, 1, false);
+		}
+	}
+}
+
+// Reduces the evaluated sections to only the section that resides last in the source data. Legacy behaviour from various track instances.
+void MovieSceneSegmentCompiler::BlendSegmentLegacySectionOrder(FSegmentBlendData& BlendData)
+{
+	if (BlendData.Num() > 1)
+	{
+		Algo::SortBy(BlendData, &FMovieSceneSectionData::TemplateIndex);
+		BlendData.RemoveAt(1, BlendData.Num() - 1, false);
+	}
+}
\ No newline at end of file
diff --git a/Engine/Source/Runtime/MovieScene/Private/Compilation/MovieSceneSegmentCompiler.cpp b/Engine/Source/Runtime/MovieScene/Private/Compilation/MovieSceneSegmentCompiler.cpp
--- a/Engine/Source/Runtime/MovieScene/Private/Compilation/MovieSceneSegmentCompiler.cpp
+++ b/Engine/Source/Runtime/MovieScene/Private/Compilation/MovieSceneSegmentCompiler.cpp
@@ -3,6 +3,7 @@
 #include "Compilation/MovieSceneSegmentCompiler.h"
 #include "Compilation/MovieSceneCompilerRules.h"
 #include "Evaluation/MovieSceneEvaluationTree.h"
+#include "MovieSceneSection.h"
 
 bool FMovieSceneAdditiveCameraTrackBlender::SortByStartTime(const FMovieSceneSectionData& A, const FMovieSceneSectionData& B)
 {
diff --git a/Engine/Source/Runtime/MovieScene/Private/EntitySystem/MovieSceneSpawnablesSystem.cpp b/Engine/Source/Runtime/MovieScene/Private/EntitySystem/MovieSceneSpawnablesSystem.cpp
--- a/Engine/Source/Runtime/MovieScene/Private/EntitySystem/MovieSceneSpawnablesSystem.cpp
+++ b/Engine/Source/Runtime/MovieScene/Private/EntitySystem/MovieSceneSpawnablesSystem.cpp
@@ -9,6 +9,7 @@
 #include "EntitySystem/BuiltInComponentTypes.h"
 #include "Evaluation/MovieSceneEvaluationOperand.h"
 
+#include "MovieScene.h"
 #include "MovieSceneSequence.h"
 #include "MovieSceneExecutionToken.h"
 #include "IMovieScenePlayer.h"
diff --git a/Engine/Source/Runtime/MovieScene/Private/MovieSceneSection.cpp b/Engine/Source/Runtime/MovieScene/Private/MovieSceneSection.cpp
--- a/Engine/Source/Runtime/MovieScene/Private/MovieSceneSection.cpp
+++ b/Engine/Source/Runtime/MovieScene/Private/MovieSceneSection.cpp
@@ -1,6 +1,7 @@
 // Copyright Epic Games, Inc. All Rights Reserved.
 
 #include "MovieSceneSection.h"
+#include "MovieScene.h"
 #include "MovieSceneTrack.h"
 #include "MovieSceneSequence.h"
 #include "MovieSceneCommonHelpers.h"
diff --git a/Engine/Source/Runtime/MovieScene/Private/Sections/MovieSceneSpawnSection.cpp b/Engine/Source/Runtime/MovieScene/Private/Sections/MovieSceneSpawnSection.cpp
--- a/Engine/Source/Runtime/MovieScene/Private/Sections/MovieSceneSpawnSection.cpp
+++ b/Engine/Source/Runtime/MovieScene/Private/Sections/MovieSceneSpawnSection.cpp
@@ -3,6 +3,7 @@
 #include "Sections/MovieSceneSpawnSection.h"
 #include "UObject/SequencerObjectVersion.h"
 
+#include "MovieScene.h"
 #include "MovieSceneTimeHelpers.h"
 #include "Evaluation/MovieSceneEvaluationField.h"
 #include "EntitySystem/MovieSceneEntityManager.h"
diff --git a/Engine/Source/Runtime/MovieScene/Public/Compilation/MovieSceneCompilerRules.h b/Engine/Source/Runtime/MovieScene/Public/Compilation/MovieSceneCompilerRules.h
--- a/Engine/Source/Runtime/MovieScene/Public/Compilation/MovieSceneCompilerRules.h
+++ b/Engine/Source/Runtime/MovieScene/Public/Compilation/MovieSceneCompilerRules.h
@@ -3,110 +3,25 @@
 #pragma once
 
 #include "CoreMinimal.h"
-#include "MovieSceneTrack.h"
-#include "Evaluation/MovieSceneSegment.h"
 #include "Compilation/MovieSceneSegmentCompiler.h"
 
+#if UE_ENABLE_INCLUDE_ORDER_DEPRECATED_IN_5_1
+#include "Evaluation/MovieSceneSegment.h"
+#include "MovieSceneTrack.h"
+#endif
+
 namespace MovieSceneSegmentCompiler
 {
-	static TOptional<FMovieSceneSegment> EvaluateNearestSegment(const TRange<FFrameNumber>& Range, const FMovieSceneSegment* PreviousSegment, const FMovieSceneSegment* NextSegment)
-	{
-		if (PreviousSegment)
-		{
-			// There is a preceeding segment
-			FFrameNumber PreviousSegmentRangeBound = PreviousSegment->Range.GetUpperBoundValue();
-
-			FMovieSceneSegment EmptySpace(Range);
-			for (FSectionEvaluationData Data : PreviousSegment->Impls)
-			{
-				EmptySpace.Impls.Add(FSectionEvaluationData(Data.ImplIndex, PreviousSegmentRangeBound));
-			}
-			return EmptySpace;
-		}
-		else if (NextSegment)
-		{
-			// Before any sections
-			FFrameNumber NextSegmentRangeBound = NextSegment->Range.GetLowerBoundValue();
-
-			FMovieSceneSegment EmptySpace(Range);
-			for (FSectionEvaluationData Data : NextSegment->Impls)
-			{
-				EmptySpace.Impls.Add(FSectionEvaluationData(Data.ImplIndex, NextSegmentRangeBound));
-			}
-			return EmptySpace;
-		}
-
-		return TOptional<FMovieSceneSegment>();
-	}
-
-	static bool AlwaysEvaluateSection(const FMovieSceneSectionData& InSectionData)
-	{
-		return InSectionData.Section->GetBlendType().IsValid() || EnumHasAnyFlags(InSectionData.Flags, ESectionEvaluationFlags::PreRoll | ESectionEvaluationFlags::PostRoll);
-	}
-
-	static void FilterOutUnderlappingSections(FSegmentBlendData& BlendData)
-	{
-		if (!BlendData.Num())
-		{
-			return;
-		}
+	MOVIESCENE_API TOptional<FMovieSceneSegment> EvaluateNearestSegment(const TRange<FFrameNumber>& Range, const FMovieSceneSegment* PreviousSegment, const FMovieSceneSegment* NextSegment);
 
-		int32 HighestOverlap = TNumericLimits<int32>::Lowest();
-		for (const FMovieSceneSectionData& SectionData : BlendData)
-		{
-			if (!AlwaysEvaluateSection(SectionData))
-			{
-				HighestOverlap = FMath::Max(HighestOverlap, SectionData.Section->GetOverlapPriority());
-			}
-		}
+	MOVIESCENE_API bool AlwaysEvaluateSection(const FMovieSceneSectionData& InSectionData);
 
-		// Remove anything that's not the highest priority, (excluding pre/postroll sections)
-		for (int32 RemoveAtIndex = BlendData.Num() - 1; RemoveAtIndex >= 0; --RemoveAtIndex)
-		{
-			const FMovieSceneSectionData& SectionData = BlendData[RemoveAtIndex];
-			if (SectionData.Section->GetOverlapPriority() != HighestOverlap && !AlwaysEvaluateSection(SectionData))
-			{
-				BlendData.RemoveAt(RemoveAtIndex, 1, false);
-			}
-		}
-	}
+	MOVIESCENE_API void FilterOutUnderlappingSections(FSegmentBlendData& BlendData);
 
-	static void ChooseLowestRowIndex(FSegmentBlendData& BlendData)
-	{
-		if (!BlendData.Num())
-		{
-			return;
-		}
-
-		int32 LowestRowIndex = TNumericLimits<int32>::Max();
-		for (const FMovieSceneSectionData& SectionData : BlendData)
-		{
-			if (!AlwaysEvaluateSection(SectionData))
-			{
-				LowestRowIndex = FMath::Min(LowestRowIndex, SectionData.Section->GetRowIndex());
-			}
-		}
-
-		// Remove anything that's not the highest priority, (excluding pre/postroll sections)
-		for (int32 RemoveAtIndex = BlendData.Num() - 1; RemoveAtIndex >= 0; --RemoveAtIndex)
-		{
-			const FMovieSceneSectionData& SectionData = BlendData[RemoveAtIndex];
-			if (SectionData.Section->GetRowIndex() > LowestRowIndex && !AlwaysEvaluateSection(SectionData))
-			{
-				BlendData.RemoveAt(RemoveAtIndex, 1, false);
-			}
-		}
-	}
+	MOVIESCENE_API void ChooseLowestRowIndex(FSegmentBlendData& BlendData);
 
 	// Reduces the evaluated sections to only the section that resides last in the source data. Legacy behaviour from various track instances.
-	static void BlendSegmentLegacySectionOrder(FSegmentBlendData& BlendData)
-	{
-		if (BlendData.Num() > 1)
-		{
-			Algo::SortBy(BlendData, &FMovieSceneSectionData::TemplateIndex);
-			BlendData.RemoveAt(1, BlendData.Num() - 1, false);
-		}
-	}
+	MOVIESCENE_API void BlendSegmentLegacySectionOrder(FSegmentBlendData& BlendData);
 }
 
 /** Default track row segment blender for all tracks */
diff --git a/Engine/Source/Runtime/MovieScene/Public/EntitySystem/IMovieSceneEntityProvider.h b/Engine/Source/Runtime/MovieScene/Public/EntitySystem/IMovieSceneEntityProvider.h
--- a/Engine/Source/Runtime/MovieScene/Public/EntitySystem/IMovieSceneEntityProvider.h
+++ b/Engine/Source/Runtime/MovieScene/Public/EntitySystem/IMovieSceneEntityProvider.h
@@ -7,6 +7,7 @@
 #include "Misc/InlineValue.h"
 #include "EntitySystem/MovieSceneEntityIDs.h"
 #include "EntitySystem/MovieSceneSequenceInstanceHandle.h"
+#include "EntitySystem/MovieSceneEntityBuilder.h"
 #include "Math/Range.h"
 #include "Evaluation/MovieSceneCompletionMode.h"
 
diff --git a/Engine/Source/Runtime/MovieScene/Public/IMovieScenePlayer.h b/Engine/Source/Runtime/MovieScene/Public/IMovieScenePlayer.h
--- a/Engine/Source/Runtime/MovieScene/Public/IMovieScenePlayer.h
+++ b/Engine/Source/Runtime/MovieScene/Public/IMovieScenePlayer.h
@@ -4,17 +4,22 @@
 
 #include "CoreMinimal.h"
 #include "Misc/Guid.h"
-#include "Evaluation/MovieSceneAnimTypeID.h"
-#include "Evaluation/MovieSceneEvaluationKey.h"
-#include "Evaluation/MovieScenePreAnimatedState.h"
-#include "MovieSceneFwd.h"
-#include "MovieSceneSpawnRegister.h"
 #include "Containers/ArrayView.h"
-#include "Evaluation/MovieSceneEvaluationState.h"
 #include "Misc/InlineValue.h"
+
+#include "MovieSceneSpawnRegister.h"
+#include "Evaluation/MovieScenePreAnimatedState.h"
+#include "Evaluation/MovieSceneEvaluationState.h"
 #include "Evaluation/MovieSceneEvaluationOperand.h"
+
+#if UE_ENABLE_INCLUDE_ORDER_DEPRECATED_IN_5_1
+#include "Evaluation/MovieSceneAnimTypeID.h"
+#include "Evaluation/MovieSceneEvaluationKey.h"
+#include "MovieSceneFwd.h"
 #include "Generators/MovieSceneEasingCurves.h"
+#endif
 
+enum class EMovieSceneBuiltInEasing : uint8;
 
 class UMovieSceneSequence;
 class FViewportClient;
diff --git a/Engine/Source/Runtime/MovieScene/Public/MovieScenePossessable.h b/Engine/Source/Runtime/MovieScene/Public/MovieScenePossessable.h
--- a/Engine/Source/Runtime/MovieScene/Public/MovieScenePossessable.h
+++ b/Engine/Source/Runtime/MovieScene/Public/MovieScenePossessable.h
@@ -5,11 +5,17 @@
 #include "CoreMinimal.h"
 #include "UObject/ObjectMacros.h"
 #include "Misc/Guid.h"
-#include "IMovieScenePlayer.h"
 #include "MovieSceneObjectBindingID.h"
 #include "MovieSceneSequenceID.h"
+
+#if UE_ENABLE_INCLUDE_ORDER_DEPRECATED_IN_5_1
+#include "IMovieScenePlayer.h"
+#endif
+
 #include "MovieScenePossessable.generated.h"
 
+class IMovieScenePlayer;
+
 /**
  * MovieScenePossessable is a "typed slot" used to allow the MovieScene to control an already-existing object
  */
diff --git a/Engine/Source/Runtime/MovieScene/Public/MovieSceneSection.h b/Engine/Source/Runtime/MovieScene/Public/MovieSceneSection.h
--- a/Engine/Source/Runtime/MovieScene/Public/MovieSceneSection.h
+++ b/Engine/Source/Runtime/MovieScene/Public/MovieSceneSection.h
@@ -4,23 +4,33 @@
 
 #include "CoreMinimal.h"
 #include "Containers/ContainersFwd.h"
+#include "Misc/Timecode.h"
+#include "Evaluation/MovieSceneCompletionMode.h"
+#include "Evaluation/MovieSceneEvaluationCustomVersion.h"
+#include "MovieSceneFrameMigration.h"
+#include "MovieSceneSignedObject.h"
+#include "Evaluation/Blending/MovieSceneBlendType.h"
+
+#if UE_ENABLE_INCLUDE_ORDER_DEPRECATED_IN_5_1
 #include "Misc/FrameTime.h"
 #include "Misc/QualifiedFrameTime.h"
-#include "Misc/Timecode.h"
 #include "UObject/ObjectMacros.h"
 #include "MovieSceneFwd.h"
 #include "KeyParams.h"
 #include "MovieScene.h"
-#include "MovieSceneSignedObject.h"
-#include "Evaluation/Blending/MovieSceneBlendType.h"
-#include "Evaluation/MovieSceneCompletionMode.h"
 #include "Generators/MovieSceneEasingFunction.h"
 #include "Evaluation/MovieSceneSequenceHierarchy.h"
 #include "MovieSceneFrameMigration.h"
 #include "Evaluation/MovieSceneEvaluationCustomVersion.h"
 #include "EntitySystem/MovieSceneEntityBuilder.h"
+#endif
+
 #include "MovieSceneSection.generated.h"
 
+class IMovieSceneEasingFunction;
+struct FMovieSceneBlendTypeField;
+struct FQualifiedFrameTime;
+struct FMovieSceneSequenceHierarchy;
 
 class FStructOnScope;
 
@@ -412,9 +422,6 @@ public:
 	 */
 	MOVIESCENE_API virtual UMovieSceneSection* SplitSection(FQualifiedFrameTime SplitTime, bool bDeleteKeys);
 
-	UE_DEPRECATED(4.23, "Please use SplitSection(SplitTime, bDeleteKeys) instead.")
-	virtual UMovieSceneSection* SplitSection(FQualifiedFrameTime SplitTime) { return SplitSection(SplitTime, false); }
-
 	/**
 	 * Trim a section at the trim time
 	 *
@@ -424,9 +431,6 @@ public:
 	 */
 	MOVIESCENE_API virtual void TrimSection(FQualifiedFrameTime TrimTime, bool bTrimLeft, bool bDeleteKeys);
 
-	UE_DEPRECATED(4.23, "Please use TrimSection(SplitTime, bTrimLeft, bDeleteKeys) instead.")
-	virtual void TrimSection(FQualifiedFrameTime SplitTime, bool bTrimLeft) { TrimSection(SplitTime, bTrimLeft, false); }
-
 	/**
 	 * Get the data structure representing the specified keys.
 	 *
diff --git a/Engine/Source/Runtime/MovieScene/Public/MovieSceneSequence.h b/Engine/Source/Runtime/MovieScene/Public/MovieSceneSequence.h
--- a/Engine/Source/Runtime/MovieScene/Public/MovieSceneSequence.h
+++ b/Engine/Source/Runtime/MovieScene/Public/MovieSceneSequence.h
@@ -3,17 +3,23 @@
 #pragma once
 
 #include "CoreMinimal.h"
+#include "Templates/SubclassOf.h"
 #include "UObject/ObjectMacros.h"
 #include "Misc/Guid.h"
 #include "MovieSceneSignedObject.h"
 #include "MovieSceneTrack.h"
 #include "MovieSceneSequenceID.h"
+#include "Evaluation/MovieSceneCompletionMode.h"
+#include "MovieSceneSection.h" // only for FMovieSceneTimecodeSource in .gen.cpp
+#include "MovieSceneObjectBindingID.h" // only for FMovieSceneObjectBindingID in .gen.cpp
 #include "MovieSceneSequence.generated.h"
 
+class IMovieScenePlayer;
 class ITargetPlatform;
 class UMovieScene;
 class UMovieSceneCompiledData;
 class UMovieSceneEntitySystemLinker;
+class UMovieSceneTrack;
 
 struct FMovieScenePossessable;
 struct FMovieSceneObjectCache;
@@ -331,6 +337,6 @@ public:
 	/*
 	 * Sequences can determine whether they support a particular track type
 	 */
-	virtual ETrackSupport IsTrackSupported(TSubclassOf<class UMovieSceneTrack> InTrackClass) const { return ETrackSupport::Default; }
+	virtual ETrackSupport IsTrackSupported(TSubclassOf<UMovieSceneTrack> InTrackClass) const { return ETrackSupport::Default; }
 #endif
 };
diff --git a/Engine/Source/Runtime/MovieScene/Public/MovieSceneSpawnRegister.h b/Engine/Source/Runtime/MovieScene/Public/MovieSceneSpawnRegister.h
--- a/Engine/Source/Runtime/MovieScene/Public/MovieSceneSpawnRegister.h
+++ b/Engine/Source/Runtime/MovieScene/Public/MovieSceneSpawnRegister.h
@@ -6,14 +6,21 @@
 #include "Misc/Guid.h"
 #include "MovieSceneSequenceID.h"
 #include "MovieSceneFwd.h"
-#include "MovieSceneSpawnable.h"
 #include "Templates/ValueOrError.h"
 #include "UObject/WeakObjectPtr.h"
 #include "IMovieSceneObjectSpawner.h"
 
+#if UE_ENABLE_INCLUDE_ORDER_DEPRECATED_IN_5_1
+#include "MovieSceneSpawnable.h"
+#endif
+
+enum class ESpawnOwnership : uint8;
+
+class UActorFactory;
 class IMovieScenePlayer;
 class UMovieScene;
 struct FMovieSceneEvaluationState;
+struct FMovieSceneSpawnable;
 struct FTransformData;
 class ISequencer;
 class USequencerSettings;
diff --git a/Engine/Source/Runtime/MovieScene/Public/MovieSceneTrack.h b/Engine/Source/Runtime/MovieScene/Public/MovieSceneTrack.h
--- a/Engine/Source/Runtime/MovieScene/Public/MovieSceneTrack.h
+++ b/Engine/Source/Runtime/MovieScene/Public/MovieSceneTrack.h
@@ -4,17 +4,24 @@
 
 #include "Compilation/MovieSceneSegmentCompiler.h"
 #include "CoreMinimal.h"
+#include "Templates/SubclassOf.h"
 #include "Evaluation/MovieSceneEvaluationField.h"
+#include "Evaluation/Blending/MovieSceneBlendType.h"
 #include "Misc/CoreMiscDefines.h"
 #include "Misc/EnumClassFlags.h"
 #include "Misc/Guid.h"
 #include "Misc/InlineValue.h"
-#include "MovieSceneSection.h"
 #include "MovieSceneSignedObject.h"
 #include "MovieSceneTrackEvaluationField.h"
 #include "UObject/ObjectMacros.h"
+
+#if UE_ENABLE_INCLUDE_ORDER_DEPRECATED_IN_5_1
+#include "MovieSceneSection.h"
+#endif
+
 #include "MovieSceneTrack.generated.h"
 
+class UMovieSceneSection;
 struct FMovieSceneEvaluationTrack;
 struct FMovieSceneTrackSegmentBlender;
 struct FMovieSceneTrackRowSegmentBlender;
@@ -257,14 +264,8 @@ private:
 	/** Sub-classes can override this method to perform custom pre-compilation logic. */
 	virtual void PreCompileImpl(FMovieSceneTrackPreCompileResult& OutPreCompileResult)
 	{
-		PRAGMA_DISABLE_DEPRECATION_WARNINGS
-		PreCompileImpl();
-		PRAGMA_ENABLE_DEPRECATION_WARNINGS
 	}
 
-	UE_DEPRECATED(4.27, "Please override the PreCompileImpl method that takes a parameter")
-	virtual void PreCompileImpl() {}
-
 private:
 
 	void UpdateEvaluationTree();
@@ -361,7 +362,7 @@ public:
 	 *
 	 * @param Section The section to add.
 	 */
-	virtual void AddSection(UMovieSceneSection& Section) PURE_VIRTUAL(UMovieSceneSection::AddSection,);
+	virtual void AddSection(UMovieSceneSection& Section) PURE_VIRTUAL(UMovieSceneTrack::AddSection,);
 
 	/**
 	 * Generates a new section suitable for use with this track.
@@ -383,21 +384,21 @@ public:
 	 * @param Section The section to query for.
 	 * @return True if the section is in this track.
 	 */
-	virtual bool HasSection(const UMovieSceneSection& Section) const PURE_VIRTUAL(UMovieSceneSection::HasSection, return false;);
+	virtual bool HasSection(const UMovieSceneSection& Section) const PURE_VIRTUAL(UMovieSceneTrack::HasSection, return false;);
 
 	/**
 	 * Removes a section from this track.
 	 *
 	 * @param Section The section to remove.
 	 */
-	virtual void RemoveSection(UMovieSceneSection& Section) PURE_VIRTUAL(UMovieSceneSection::RemoveSection, );
+	virtual void RemoveSection(UMovieSceneSection& Section) PURE_VIRTUAL(UMovieSceneTrack::RemoveSection, );
 
 	/**
 	 * Removes a section from this track at a particular index
 	 *
 	 * @param SectionIndex The section index to remove.
 	 */
-	virtual void RemoveSectionAt(int32 SectionIndex) PURE_VIRTUAL(UMovieSceneSection::RemoveSectionAt, );
+	virtual void RemoveSectionAt(int32 SectionIndex) PURE_VIRTUAL(UMovieSceneTrack::RemoveSectionAt, );
 
 #if WITH_EDITOR
 
diff --git a/Engine/Source/Runtime/MovieScene/Public/Sections/MovieSceneSubSection.h b/Engine/Source/Runtime/MovieScene/Public/Sections/MovieSceneSubSection.h
--- a/Engine/Source/Runtime/MovieScene/Public/Sections/MovieSceneSubSection.h
+++ b/Engine/Source/Runtime/MovieScene/Public/Sections/MovieSceneSubSection.h
@@ -4,6 +4,7 @@
 
 #include "CoreMinimal.h"
 #include "UObject/ObjectMacros.h"
+#include "Templates/SubclassOf.h"
 #include "Engine/EngineTypes.h"
 #include "GameFramework/Actor.h"
 #include "Curves/KeyHandle.h"
@@ -15,9 +16,11 @@
 #include "Evaluation/MovieSceneSequenceTransform.h"
 #include "MovieSceneSubSection.generated.h"
 
+class UMovieScene;
 class UMovieSceneSequence;
 struct FMovieSceneEvaluationTemplate;
 struct FMovieSceneTrackCompilerArgs;
+struct FMovieSceneSectionParameters;
 
 DECLARE_DELEGATE_OneParam(FOnSequenceChanged, UMovieSceneSequence* /*Sequence*/);
 
diff --git a/Engine/Source/Runtime/MovieSceneTracks/Private/EntitySystem/Interrogation/MovieSceneInterrogationLinker.cpp b/Engine/Source/Runtime/MovieSceneTracks/Private/EntitySystem/Interrogation/MovieSceneInterrogationLinker.cpp
--- a/Engine/Source/Runtime/MovieSceneTracks/Private/EntitySystem/Interrogation/MovieSceneInterrogationLinker.cpp
+++ b/Engine/Source/Runtime/MovieSceneTracks/Private/EntitySystem/Interrogation/MovieSceneInterrogationLinker.cpp
@@ -9,6 +9,7 @@
 #include "EntitySystem/MovieSceneEntitySystemLinker.h"
 #include "EntitySystem/MovieSceneInitialValueCache.h"
 
+#include "MovieScene.h"
 #include "MovieSceneTimeHelpers.h"
 #include "MovieSceneSequence.h"
 #include "IMovieScenePlayer.h"
diff --git a/Engine/Source/Runtime/MovieSceneTracks/Private/Tracks/MovieSceneCameraShakeSourceShakeTrack.cpp b/Engine/Source/Runtime/MovieSceneTracks/Private/Tracks/MovieSceneCameraShakeSourceShakeTrack.cpp
--- a/Engine/Source/Runtime/MovieSceneTracks/Private/Tracks/MovieSceneCameraShakeSourceShakeTrack.cpp
+++ b/Engine/Source/Runtime/MovieSceneTracks/Private/Tracks/MovieSceneCameraShakeSourceShakeTrack.cpp
@@ -5,6 +5,7 @@
 #include "Compilation/MovieSceneCompilerRules.h"
 #include "Sections/MovieSceneCameraShakeSourceShakeSection.h"
 #include "Evaluation/MovieSceneCameraShakeSourceShakeTemplate.h"
+#include "MovieScene.h"
 
 #define LOCTEXT_NAMESPACE "MovieSceneCameraShakeSourceShakeTrack"
 
diff --git a/Engine/Source/Runtime/MovieSceneTracks/Private/Tracks/MovieSceneMaterialTrack.cpp b/Engine/Source/Runtime/MovieSceneTracks/Private/Tracks/MovieSceneMaterialTrack.cpp
--- a/Engine/Source/Runtime/MovieSceneTracks/Private/Tracks/MovieSceneMaterialTrack.cpp
+++ b/Engine/Source/Runtime/MovieSceneTracks/Private/Tracks/MovieSceneMaterialTrack.cpp
@@ -1,6 +1,7 @@
 // Copyright Epic Games, Inc. All Rights Reserved.
 
 #include "Tracks/MovieSceneMaterialTrack.h"
+#include "MovieScene.h"
 #include "MovieSceneCommonHelpers.h"
 #include "Evaluation/MovieSceneParameterTemplate.h"
 #include "Evaluation/MovieSceneEvaluationTrack.h"
diff --git a/Engine/Source/Runtime/MovieSceneTracks/Private/Tracks/MovieSceneParticleParameterTrack.cpp b/Engine/Source/Runtime/MovieSceneTracks/Private/Tracks/MovieSceneParticleParameterTrack.cpp
--- a/Engine/Source/Runtime/MovieSceneTracks/Private/Tracks/MovieSceneParticleParameterTrack.cpp
+++ b/Engine/Source/Runtime/MovieSceneTracks/Private/Tracks/MovieSceneParticleParameterTrack.cpp
@@ -1,6 +1,7 @@
 // Copyright Epic Games, Inc. All Rights Reserved.
 
 #include "Tracks/MovieSceneParticleParameterTrack.h"
+#include "MovieScene.h"
 #include "MovieSceneCommonHelpers.h"
 #include "Evaluation/MovieSceneParticleParameterTemplate.h"
 
diff --git a/Engine/Source/Runtime/MovieSceneTracks/Private/Tracks/MovieScenePropertyTrack.cpp b/Engine/Source/Runtime/MovieSceneTracks/Private/Tracks/MovieScenePropertyTrack.cpp
--- a/Engine/Source/Runtime/MovieSceneTracks/Private/Tracks/MovieScenePropertyTrack.cpp
+++ b/Engine/Source/Runtime/MovieSceneTracks/Private/Tracks/MovieScenePropertyTrack.cpp
@@ -3,6 +3,7 @@
 #include "Tracks/MovieScenePropertyTrack.h"
 #include "Channels/MovieSceneSectionChannelOverrideRegistry.h"
 #include "Algo/Sort.h"
+#include "MovieScene.h"
 #include "MovieSceneCommonHelpers.h"
 #include "MovieSceneTracksComponentTypes.h"
 #include "PropertyPathHelpers.h"
