Sequencer: Fix stale track templates not being purged on evaluation

Previously stale tracks were always purged inside FSequencer::EvaluateInternal, but this code was lost during the ECS migration and the segragation of the legacy evaluator from the main evaluation templates. Rather than put the code back in there, I'm opting to put it in the sequence instance's post-evaluate function, meaning it will apply to volatile sequences inside cooked builds too, rather than just the editor UI.

Ideally this code would exist as part of the FMovieSceneTrackEvaluator, but in order for it to only be applied for volatile sequences (and thus not affect runtime performance) we'd have to call through from the FSequenceInstance which would involve changing headers. I'd like to get this change into a patch release if possible, so I've opted for the hotfixable change - we can tidy it up in main.

diff --git a/Engine/Source/Runtime/MovieScene/Private/EntitySystem/MovieSceneSequenceInstance.cpp b/Engine/Source/Runtime/MovieScene/Private/EntitySystem/MovieSceneSequenceInstance.cpp
--- a/Engine/Source/Runtime/MovieScene/Private/EntitySystem/MovieSceneSequenceInstance.cpp
+++ b/Engine/Source/Runtime/MovieScene/Private/EntitySystem/MovieSceneSequenceInstance.cpp
@@ -279,6 +279,45 @@ void FSequenceInstance::PostEvaluation(UMovieSceneEntitySystemLinker* Linker)
 			SCOPE_CYCLE_COUNTER(MovieSceneEval_SequenceInstancePostUpdate);
 
 			Player->PostEvaluation(Context);
+
+			if (VolatilityManager && LegacyEvaluator)
+			{
+				FMovieSceneRootEvaluationTemplateInstance& RootTemplateInstance = Player->GetEvaluationTemplate();
+
+				UMovieSceneCompiledDataManager* CompiledDataManager = RootTemplateInstance.GetCompiledDataManager();
+
+				// If this sequence is volatile, purge any stale track templates from the compiled data
+				FMovieSceneEvaluationTemplate* EvalTemplate = const_cast<FMovieSceneEvaluationTemplate*>(CompiledDataManager->FindTrackTemplate(CompiledDataID));
+				if (EvalTemplate)
+				{
+					EvalTemplate->PurgeStaleTracks();
+				}
+
+				// Do the same for all subsequences
+				const FMovieSceneSequenceHierarchy* Hierarchy = CompiledDataManager->FindHierarchy(CompiledDataID);
+				if (Hierarchy)
+				{
+					for (const TTuple<FMovieSceneSequenceID, FMovieSceneSubSequenceData>& Pair : Hierarchy->AllSubSequenceData())
+					{
+						UMovieSceneSequence* SubSequence = Pair.Value.GetLoadedSequence();
+						if (!SubSequence)
+						{
+							continue;
+						}
+						FMovieSceneCompiledDataID SubCompiledDataID = CompiledDataManager->FindDataID(SubSequence);
+						if (!SubCompiledDataID.IsValid())
+						{
+							continue;
+						}
+
+						FMovieSceneEvaluationTemplate* SubEvalTemplate = const_cast<FMovieSceneEvaluationTemplate*>(CompiledDataManager->FindTrackTemplate(SubCompiledDataID));
+						if (SubEvalTemplate)
+						{
+							SubEvalTemplate->PurgeStaleTracks();
+						}
+					}
+				}
+			}
 		}
 	}
 }
