Sequencer: When moving contents from one binding to another, if a possessable already exists, don't move it, just remove it.

Spawnable Cine Camera Actor
    Possessable Cine Camera Component <---
        Tracks

moving to:

Spawnable Camera Actor
    Possessable Camera Component
        Tracks

Existing code swaps the parent of the original cine camera actor's possesable cine camera component to the new spawnable camera actor. However, if spawnable camera actor already has a possessable camera component, it shouldn't do this because there would be duplicate component bindings. Instead, just remove it. This happens when changing the binding for a template sequence.

diff --git a/Engine/Source/Runtime/MovieScene/Private/MovieScene.cpp b/Engine/Source/Runtime/MovieScene/Private/MovieScene.cpp
--- a/Engine/Source/Runtime/MovieScene/Private/MovieScene.cpp
+++ b/Engine/Source/Runtime/MovieScene/Private/MovieScene.cpp
@@ -1392,14 +1392,40 @@ void UMovieScene::MoveBindingContents(const FGuid& SourceBindingId, const FGuid&
 
 	FMovieSceneSpawnable* DestinationSpawnable = FindSpawnable(DestinationBindingId);
 
-	for (FMovieScenePossessable& Possessable : Possessables)
+	for( auto PossesableIter( Possessables.CreateIterator() ); PossesableIter; ++PossesableIter )
 	{
-		if (Possessable.GetParent() == SourceBindingId)
+		FMovieScenePossessable& SourcePossessable = *PossesableIter;
+		
+		// If there is a possessable whose parent is the binding we're moving contents for, 
+		// that possessable needs to be remapped to the new destination parent
+		if (SourcePossessable.GetParent() == SourceBindingId)
 		{
-			Possessable.SetParent(DestinationBindingId);
-			if (DestinationSpawnable)
+			// But if there is already a possessable for that destination binding, don't keep the source possessable around.
+			bool bAlreadyExists = false;
+			for (FMovieScenePossessable& DestinationPossessable : Possessables)
+			{
+				if (DestinationPossessable.GetName() == SourcePossessable.GetName() &&
+					DestinationPossessable.GetParent() == DestinationBindingId)
+				{
+					FGuid CurGuid = SourcePossessable.GetGuid();
+
+					// Found it!
+					Possessables.RemoveAt( PossesableIter.GetIndex() );
+
+					RemoveBinding( CurGuid );
+
+					bAlreadyExists = true;
+					break;
+				}
+			}
+
+			if (!bAlreadyExists)
 			{
-				DestinationSpawnable->AddChildPossessable(Possessable.GetGuid());
+				SourcePossessable.SetParent(DestinationBindingId);
+				if (DestinationSpawnable)
+				{
+					DestinationSpawnable->AddChildPossessable(SourcePossessable.GetGuid());
+				}
 			}
 		}
 	}
