Fix FSequenceInstance::SetFinished assigning to parameter instead of member.

diff --git a/Engine/Source/Runtime/MovieScene/Public/EntitySystem/MovieSceneSequenceInstance.h b/Engine/Source/Runtime/MovieScene/Public/EntitySystem/MovieSceneSequenceInstance.h
--- a/Engine/Source/Runtime/MovieScene/Public/EntitySystem/MovieSceneSequenceInstance.h
+++ b/Engine/Source/Runtime/MovieScene/Public/EntitySystem/MovieSceneSequenceInstance.h
@@ -224,7 +224,7 @@ public:
 	 */
 	void SetFinished(bool bInFinished)
 	{
-		bInFinished = bFinished;
+		bFinished = bInFinished;
 	}
 
 	/**
