Sequencer: Key only unmuted and active sections

diff --git a/Engine/Plugins/Animation/ControlRig/Source/ControlRig/Private/Sequencer/MovieSceneControlRigParameterTrack.cpp b/Engine/Plugins/Animation/ControlRig/Source/ControlRig/Private/Sequencer/MovieSceneControlRigParameterTrack.cpp
--- a/Engine/Plugins/Animation/ControlRig/Source/ControlRig/Private/Sequencer/MovieSceneControlRigParameterTrack.cpp
+++ b/Engine/Plugins/Animation/ControlRig/Source/ControlRig/Private/Sequencer/MovieSceneControlRigParameterTrack.cpp
@@ -193,7 +193,7 @@ TArray<UMovieSceneSection*, TInlineAllocator<4>> UMovieSceneControlRigParameterT
 
 	for (UMovieSceneSection* Section : Sections)
 	{
-		if (Section->GetRange().Contains(Time))
+		if (MovieSceneHelpers::IsSectionKeyable(Section) && Section->GetRange().Contains(Time))
 		{
 			OverlappingSections.Add(Section);
 		}
@@ -229,10 +229,10 @@ UMovieSceneSection* UMovieSceneControlRigParameterTrack::FindOrExtendSection(FFr
 {
 	Weight = 1.0f;
 	TArray<UMovieSceneSection*, TInlineAllocator<4>> OverlappingSections = FindAllSections(Time);
-	if (SectionToKey)
+	if (SectionToKey && MovieSceneHelpers::IsSectionKeyable(SectionToKey))
 	{
 		bool bCalculateWeight = false;
-		if (SectionToKey && !OverlappingSections.Contains(SectionToKey))
+		if (!OverlappingSections.Contains(SectionToKey))
 		{
 			if (SectionToKey->HasEndFrame() && SectionToKey->GetExclusiveEndFrame() <= Time)
 			{
diff --git a/Engine/Source/Runtime/MovieScene/Private/MovieSceneCommonHelpers.cpp b/Engine/Source/Runtime/MovieScene/Private/MovieSceneCommonHelpers.cpp
--- a/Engine/Source/Runtime/MovieScene/Private/MovieSceneCommonHelpers.cpp
+++ b/Engine/Source/Runtime/MovieScene/Private/MovieSceneCommonHelpers.cpp
@@ -17,6 +17,22 @@
 #include "MovieSceneTrack.h"
 #include "Engine/Engine.h"
 
+bool MovieSceneHelpers::IsSectionKeyable(UMovieSceneSection* Section)
+{
+	if (!Section)
+	{
+		return false;
+	}
+
+	UMovieSceneTrack* Track = Section->GetTypedOuter<UMovieSceneTrack>();
+	if (!Track)
+	{
+		return false;
+	}
+
+	return !Track->IsRowEvalDisabled(Section->GetRowIndex()) && !Track->IsEvalDisabled() && Section->IsActive();
+}
+
 UMovieSceneSection* MovieSceneHelpers::FindSectionAtTime( TArrayView<UMovieSceneSection* const> Sections, FFrameNumber Time )
 {
 	for( int32 SectionIndex = 0; SectionIndex < Sections.Num(); ++SectionIndex )
@@ -24,7 +40,7 @@ UMovieSceneSection* MovieSceneHelpers::FindSectionAtTime( TArrayView<UMovieScene
 		UMovieSceneSection* Section = Sections[SectionIndex];
 
 		//@todo sequencer: There can be multiple sections overlapping in time. Returning instantly does not account for that.
-		if( Section->IsTimeWithinSection( Time ) && Section->IsActive() )
+		if( Section->IsTimeWithinSection( Time ) && IsSectionKeyable(Section) )
 		{
 			return Section;
 		}
@@ -38,13 +54,16 @@ UMovieSceneSection* MovieSceneHelpers::FindNearestSectionAtTime( TArrayView<UMov
 	TArray<UMovieSceneSection*> OverlappingSections, NonOverlappingSections;
 	for (UMovieSceneSection* Section : Sections)
 	{
-		if (Section->GetRange().Contains(Time))
+		if (IsSectionKeyable(Section))
 		{
-			OverlappingSections.Add(Section);
-		}
-		else
-		{
-			NonOverlappingSections.Add(Section);
+			if (Section->GetRange().Contains(Time))
+			{
+				OverlappingSections.Add(Section);
+			}
+			else
+			{
+				NonOverlappingSections.Add(Section);
+			}
 		}
 	}
 
diff --git a/Engine/Source/Runtime/MovieScene/Public/MovieSceneCommonHelpers.h b/Engine/Source/Runtime/MovieScene/Public/MovieSceneCommonHelpers.h
--- a/Engine/Source/Runtime/MovieScene/Public/MovieSceneCommonHelpers.h
+++ b/Engine/Source/Runtime/MovieScene/Public/MovieSceneCommonHelpers.h
@@ -25,6 +25,11 @@ class MOVIESCENE_API MovieSceneHelpers
 {
 public:
 
+	/** 
+	 * @return Whether the section is keyable (active, on a track that is not muted, etc 
+	 */
+	static bool IsSectionKeyable(UMovieSceneSection*);
+
 	/**
 	 * Finds a section that exists at a given time
 	 *
diff --git a/Engine/Source/Runtime/MovieSceneTracks/Private/Tracks/MovieScenePropertyTrack.cpp b/Engine/Source/Runtime/MovieSceneTracks/Private/Tracks/MovieScenePropertyTrack.cpp
--- a/Engine/Source/Runtime/MovieSceneTracks/Private/Tracks/MovieScenePropertyTrack.cpp
+++ b/Engine/Source/Runtime/MovieSceneTracks/Private/Tracks/MovieScenePropertyTrack.cpp
@@ -146,14 +146,13 @@ bool UMovieScenePropertyTrack::IsEmpty() const
 	return Sections.Num() == 0;
 }
 
-
 TArray<UMovieSceneSection*, TInlineAllocator<4>> UMovieScenePropertyTrack::FindAllSections(FFrameNumber Time)
 {
 	TArray<UMovieSceneSection*, TInlineAllocator<4>> OverlappingSections;
 
 	for (UMovieSceneSection* Section : Sections)
 	{
-		if (Section->GetRange().Contains(Time))
+		if (MovieSceneHelpers::IsSectionKeyable(Section) && Section->GetRange().Contains(Time))
 		{
 			OverlappingSections.Add(Section);
 		}
@@ -189,10 +188,10 @@ UMovieSceneSection* UMovieScenePropertyTrack::FindOrExtendSection(FFrameNumber T
 {
 	Weight = 1.0f;
 	TArray<UMovieSceneSection*, TInlineAllocator<4>> OverlappingSections = FindAllSections(Time);
-	if (SectionToKey)
+	if (SectionToKey && MovieSceneHelpers::IsSectionKeyable(SectionToKey))
 	{
 		bool bCalculateWeight = false;
-		if (SectionToKey && !OverlappingSections.Contains(SectionToKey))
+		if (!OverlappingSections.Contains(SectionToKey))
 		{
 			if (SectionToKey->HasEndFrame() && SectionToKey->GetExclusiveEndFrame() <= Time)
 			{
