Clean up obsolete code from the legacy pre animated state class

diff --git a/Engine/Source/Runtime/MovieScene/Private/Evaluation/MovieScenePreAnimatedState.cpp b/Engine/Source/Runtime/MovieScene/Private/Evaluation/MovieScenePreAnimatedState.cpp
--- a/Engine/Source/Runtime/MovieScene/Private/Evaluation/MovieScenePreAnimatedState.cpp
+++ b/Engine/Source/Runtime/MovieScene/Private/Evaluation/MovieScenePreAnimatedState.cpp
@@ -59,8 +59,6 @@ void FMovieScenePreAnimatedState::Initialize(UMovieSceneEntitySystemLinker* Link
 
 	bCapturingGlobalPreAnimatedState = false;
 
-	WeakObjectStorage = nullptr;
-	WeakMasterStorage = nullptr;
 	TemplateMetaData = nullptr;
 	EvaluationHookMetaData = nullptr;
 
@@ -207,7 +205,13 @@ void FMovieScenePreAnimatedState::RestorePreAnimatedState(UObject& Object, TFunc
 {
 	using namespace UE::MovieScene;
 
-	TSharedPtr<FAnimTypePreAnimatedStateObjectStorage> ObjectStorage = WeakObjectStorage.Pin();
+	UMovieSceneEntitySystemLinker* Linker = WeakLinker.Get();
+	if (!Linker)
+	{
+		return;
+	}
+
+	TSharedPtr<FAnimTypePreAnimatedStateObjectStorage> ObjectStorage = Linker->PreAnimatedState.FindStorage(FAnimTypePreAnimatedStateObjectStorage::StorageID);
 	if (!ObjectStorage)
 	{
 		return;
@@ -290,4 +294,4 @@ bool FMovieScenePreAnimatedState::ContainsAnyStateForSequence() const
 
 	UMovieSceneEntitySystemLinker* Linker = WeakLinker.Get();
 	return Linker && InstanceHandle.IsValid() && Linker->PreAnimatedState.ContainsAnyStateForInstanceHandle(InstanceHandle);
-}
\ No newline at end of file
+}
diff --git a/Engine/Source/Runtime/MovieScene/Public/Evaluation/MovieScenePreAnimatedState.h b/Engine/Source/Runtime/MovieScene/Public/Evaluation/MovieScenePreAnimatedState.h
--- a/Engine/Source/Runtime/MovieScene/Public/Evaluation/MovieScenePreAnimatedState.h
+++ b/Engine/Source/Runtime/MovieScene/Public/Evaluation/MovieScenePreAnimatedState.h
@@ -27,11 +27,7 @@ namespace MovieScene
 {
 
 struct FPreAnimatedStateEntry;
-struct FPreAnimatedStateExtension;
 struct FPreAnimatedTemplateCaptureSources;
-struct FAnimTypePreAnimatedStateObjectStorage;
-struct FAnimTypePreAnimatedStateMasterStorage;
-struct FPreAnimatedTrackInstanceCaptureSources;
 struct FPreAnimatedEvaluationHookCaptureSources;
 
 }
@@ -127,8 +123,6 @@ private:
 
 	void ConditionalInitializeEntityStorage(bool bOverrideWantsRestoreState);
 
-	void InitializeStorage(TSharedPtr<UE::MovieScene::FPreAnimatedStateExtension> Extension);
-
 	void AddSourceMetaData(const UE::MovieScene::FPreAnimatedStateEntry& Entry);
 
 private:
@@ -138,11 +132,6 @@ private:
 	/** Weak pointer to the linker that we're associated with */
 	TWeakObjectPtr<UMovieSceneEntitySystemLinker> WeakLinker;
 
-	/** Pointers to the storage for state bound to objects, organized by FMovieSceneAnimTypeID */
-	TWeakPtr<UE::MovieScene::FAnimTypePreAnimatedStateObjectStorage> WeakObjectStorage;
-	/** Pointers to the storage for state created from master tracks, or otherwise not bound to objects */
-	TWeakPtr<UE::MovieScene::FAnimTypePreAnimatedStateMasterStorage> WeakMasterStorage;
-
 	/** Meta-data ledger for any pre-animated state that originates from track templates */
 	TSharedPtr<UE::MovieScene::FPreAnimatedTemplateCaptureSources> TemplateMetaData;
 	/** Meta-data ledger for any pre-animated state that originates from evaluation hooks */
