Start all linker evaluations with a clean capture source in case there are nested linkers running.
Add some ensures to detect invalid use of capture sources.

diff --git a/Engine/Source/Runtime/MovieScene/Private/EntitySystem/MovieSceneEntitySystemRunner.cpp b/Engine/Source/Runtime/MovieScene/Private/EntitySystem/MovieSceneEntitySystemRunner.cpp
--- a/Engine/Source/Runtime/MovieScene/Private/EntitySystem/MovieSceneEntitySystemRunner.cpp
+++ b/Engine/Source/Runtime/MovieScene/Private/EntitySystem/MovieSceneEntitySystemRunner.cpp
@@ -4,6 +4,7 @@
 #include "EntitySystem/MovieSceneEntitySystemLinker.h"
 #include "EntitySystem/MovieSceneEntityMutations.h"
 #include "EntitySystem/BuiltInComponentTypes.h"
+#include "Evaluation/PreAnimatedState/MovieScenePreAnimatedCaptureSource.h"
 #include "ProfilingDebugging/CountersTrace.h"
 
 DECLARE_CYCLE_STAT(TEXT("ECS System Cost"), 			MovieSceneEval_TotalGTCost, 				STATGROUP_MovieSceneEval);
@@ -200,6 +201,10 @@ void FMovieSceneEntitySystemRunner::DoFlushUpdateQueueOnce()
 	// again.
 	TGuardValue<FEntityManager*> DebugVizGuard(GEntityManagerForDebuggingVisualizers, GetEntityManager());
 
+	// Also reset the capture source scope so that each group of sequences tied to a given linker starts
+	// with a clean slate.
+	TGuardValue<FScopedPreAnimatedCaptureSource*> CaptureSourceGuard(FScopedPreAnimatedCaptureSource::GetCaptureSourcePtr(), nullptr);
+
 	// Entry point to the whole ECS loop... this will either unroll in the current thread's call stack
 	// if there's not much to do, or it will start queuing up tasks on the task graph.
 	// We immediately wait for the completion task to be executed.
diff --git a/Engine/Source/Runtime/MovieScene/Private/Evaluation/PreAnimatedState/MovieScenePreAnimatedCaptureSource.cpp b/Engine/Source/Runtime/MovieScene/Private/Evaluation/PreAnimatedState/MovieScenePreAnimatedCaptureSource.cpp
--- a/Engine/Source/Runtime/MovieScene/Private/Evaluation/PreAnimatedState/MovieScenePreAnimatedCaptureSource.cpp
+++ b/Engine/Source/Runtime/MovieScene/Private/Evaluation/PreAnimatedState/MovieScenePreAnimatedCaptureSource.cpp
@@ -9,6 +9,7 @@
 
 FScopedPreAnimatedCaptureSource::FScopedPreAnimatedCaptureSource(FMovieScenePreAnimatedState* InPreAnimatedState, const FMovieSceneEvaluationKey& InEvalKey, bool bInWantsRestoreState)
 	: Variant(TInPlaceType<FMovieSceneEvaluationKey>(), InEvalKey)
+	, WeakLinker(InPreAnimatedState ? InPreAnimatedState->GetLinker() : nullptr)
 	, OptionalSequencePreAnimatedState(InPreAnimatedState)
 	, bWantsRestoreState(bInWantsRestoreState)
 {
@@ -18,6 +19,7 @@ FScopedPreAnimatedCaptureSource::FScopedPreAnimatedCaptureSource(FMovieScenePreA
 }
 FScopedPreAnimatedCaptureSource::FScopedPreAnimatedCaptureSource(FMovieScenePreAnimatedState* InPreAnimatedState, const UObject* InEvalHook, FMovieSceneSequenceID InSequenceID, bool bInWantsRestoreState)
 	: Variant(TInPlaceType<FEvalHookType>(), FEvalHookType{InEvalHook, InSequenceID} )
+	, WeakLinker(InPreAnimatedState ? InPreAnimatedState->GetLinker() : nullptr)
 	, OptionalSequencePreAnimatedState(InPreAnimatedState)
 	, bWantsRestoreState(bInWantsRestoreState)
 {
@@ -27,6 +29,7 @@ FScopedPreAnimatedCaptureSource::FScopedPreAnimatedCaptureSource(FMovieScenePreA
 }
 FScopedPreAnimatedCaptureSource::FScopedPreAnimatedCaptureSource(UMovieSceneEntitySystemLinker* InLinker, UMovieSceneTrackInstance* InTrackInstance, bool bInWantsRestoreState)
 	: Variant(TInPlaceType<UMovieSceneTrackInstance*>(), InTrackInstance)
+	, WeakLinker(InLinker)
 	, OptionalSequencePreAnimatedState(nullptr)
 	, bWantsRestoreState(bInWantsRestoreState)
 {
@@ -36,6 +39,7 @@ FScopedPreAnimatedCaptureSource::FScopedPreAnimatedCaptureSource(UMovieSceneEnti
 }
 FScopedPreAnimatedCaptureSource::FScopedPreAnimatedCaptureSource(UMovieSceneEntitySystemLinker* InLinker, const FMovieSceneTrackInstanceInput& TrackInstanceInput)
 	: Variant(TInPlaceType<FMovieSceneTrackInstanceInput>(), TrackInstanceInput)
+	, WeakLinker(InLinker)
 	, OptionalSequencePreAnimatedState(nullptr)
 {
 	EMovieSceneCompletionMode CompletionMode = TrackInstanceInput.Section->GetCompletionMode();
@@ -81,6 +85,9 @@ void FScopedPreAnimatedCaptureSource::BeginTracking(const UE::MovieScene::FPreAn
 {
 	using namespace UE::MovieScene;
 
+	ensureMsgf(!WeakLinker.IsValid() || WeakLinker.Get() == Linker, 
+			TEXT("Attempting to track state on a capture source related to a different linker. Are you missing setting a scoped capture source?"));
+
 	if (FMovieSceneEvaluationKey* EvalKey = Variant.TryGet<FMovieSceneEvaluationKey>())
 	{
 		check(OptionalSequencePreAnimatedState);
diff --git a/Engine/Source/Runtime/MovieScene/Private/Evaluation/PreAnimatedState/MovieScenePreAnimatedStateExtension.cpp b/Engine/Source/Runtime/MovieScene/Private/Evaluation/PreAnimatedState/MovieScenePreAnimatedStateExtension.cpp
--- a/Engine/Source/Runtime/MovieScene/Private/Evaluation/PreAnimatedState/MovieScenePreAnimatedStateExtension.cpp
+++ b/Engine/Source/Runtime/MovieScene/Private/Evaluation/PreAnimatedState/MovieScenePreAnimatedStateExtension.cpp
@@ -484,12 +484,16 @@ bool FPreAnimatedStateExtension::ContainsAnyStateForInstanceHandle(FRootInstance
 bool FPreAnimatedStateExtension::HasActiveCaptureSource() const
 {
 	FScopedPreAnimatedCaptureSource* CaptureSource = FScopedPreAnimatedCaptureSource::GetCaptureSourcePtr();
+	ensureMsgf(!CaptureSource || CaptureSource->WeakLinker.Get() == Linker,
+			TEXT("The current capture source is related to a different linker. Are you missing setting a scope capture source?"));
 	return (CaptureSource && CaptureSource->bWantsRestoreState);
 }
 
 bool FPreAnimatedStateExtension::ShouldCaptureAnyState() const
 {
 	FScopedPreAnimatedCaptureSource* CaptureSource = FScopedPreAnimatedCaptureSource::GetCaptureSourcePtr();
+	ensureMsgf(!CaptureSource || CaptureSource->WeakLinker.Get() == Linker,
+			TEXT("The current capture source is related to a different linker. Are you missing setting a scope capture source?"));
 	return (CaptureSource && CaptureSource->bWantsRestoreState) || IsCapturingGlobalState();
 }
 
@@ -498,6 +502,8 @@ void FPreAnimatedStateExtension::AddSourceMetaData(const UE::MovieScene::FPreAni
 	using namespace UE::MovieScene;
 
 	FScopedPreAnimatedCaptureSource* CaptureSource = FScopedPreAnimatedCaptureSource::GetCaptureSourcePtr();
+	ensureMsgf(!CaptureSource || CaptureSource->WeakLinker.Get() == Linker,
+			TEXT("The current capture source is related to a different linker. Are you missing setting a scope capture source?"));
 	if (!CaptureSource)
 	{
 		EnsureMetaData(Entry);
diff --git a/Engine/Source/Runtime/MovieScene/Public/Evaluation/PreAnimatedState/MovieScenePreAnimatedCaptureSource.h b/Engine/Source/Runtime/MovieScene/Public/Evaluation/PreAnimatedState/MovieScenePreAnimatedCaptureSource.h
--- a/Engine/Source/Runtime/MovieScene/Public/Evaluation/PreAnimatedState/MovieScenePreAnimatedCaptureSource.h
+++ b/Engine/Source/Runtime/MovieScene/Public/Evaluation/PreAnimatedState/MovieScenePreAnimatedCaptureSource.h
@@ -69,12 +69,14 @@ struct FScopedPreAnimatedCaptureSource
 
 private:
 
+	friend class FMovieSceneEntitySystemRunner;
+	friend struct UE::MovieScene::FPreAnimatedStateExtension;
+
 	static FScopedPreAnimatedCaptureSource*& GetCaptureSourcePtr();
 
 	void BeginTracking(const UE::MovieScene::FPreAnimatedStateMetaData& MetaData, UMovieSceneEntitySystemLinker* Linker);
 	UE::MovieScene::FRootInstanceHandle GetRootInstanceHandle(UMovieSceneEntitySystemLinker* Linker) const;
 
-	friend UE::MovieScene::FPreAnimatedStateExtension;
 	struct FEvalHookType
 	{
 		const UObject* EvalHook;
@@ -83,7 +85,8 @@ private:
 	using CaptureSourceType = TVariant<FMovieSceneEvaluationKey, FEvalHookType, UMovieSceneTrackInstance*, FMovieSceneTrackInstanceInput>;
 
 	CaptureSourceType Variant;
+	TWeakObjectPtr<UMovieSceneEntitySystemLinker> WeakLinker;
 	FMovieScenePreAnimatedState* OptionalSequencePreAnimatedState;
 	FScopedPreAnimatedCaptureSource* PrevCaptureSource;
 	bool bWantsRestoreState;
-};
\ No newline at end of file
+};
