Sequencer: Remove duplicates from bindings when first resolved

diff --git a/Engine/Source/Runtime/MovieScene/Private/Evaluation/MovieSceneEvaluationState.cpp b/Engine/Source/Runtime/MovieScene/Private/Evaluation/MovieSceneEvaluationState.cpp
--- a/Engine/Source/Runtime/MovieScene/Private/Evaluation/MovieSceneEvaluationState.cpp
+++ b/Engine/Source/Runtime/MovieScene/Private/Evaluation/MovieSceneEvaluationState.cpp
@@ -7,6 +7,9 @@
 #include "IMovieScenePlaybackClient.h"
 #include "MovieSceneObjectBindingID.h"
 
+#include "Algo/Sort.h"
+#include "Algo/Unique.h"
+
 DECLARE_CYCLE_STAT(TEXT("Find Bound Objects"), MovieSceneEval_FindBoundObjects, STATGROUP_MovieSceneEval);
 DECLARE_CYCLE_STAT(TEXT("Iterate Bound Objects"), MovieSceneEval_IterateBoundObjects, STATGROUP_MovieSceneEval);
 
@@ -424,8 +427,27 @@
 		}
 	}
 
-	if (Bindings->Objects.Num())
+	const int32 NumBoundObjects = Bindings->Objects.Num();
+
+	// Remove duplicates from bound objects
+	if (NumBoundObjects > 1)
 	{
+		Algo::Sort(Bindings->Objects, [](const TWeakObjectPtr<>& A, const TWeakObjectPtr<>& B) { return A.Get() < B.Get(); });
+		const int32 EndIndex = Algo::Unique(Bindings->Objects);
+		if (EndIndex < NumBoundObjects)
+		{
+			FMovieSceneBinding* Binding = Sequence->GetMovieScene()->FindBinding(InGuid);
+			FString BindingName = Binding ? Binding->GetName() : TEXT("<invalid binding>");
+			UE_LOG(LogMovieScene, Warning, TEXT("Found %d duplicate object(s) while resolving binding %s (%s) in %s"),
+					NumBoundObjects - EndIndex, 
+					*BindingName, *LexToString(InGuid), 
+					*Sequence->GetPathName());
+			Bindings->Objects.SetNum(EndIndex);
+		}
+	}
+
+	if (NumBoundObjects > 0)
+	{
 		Bindings->bUpToDate = true;
 		Player.NotifyBindingUpdate(InGuid, SequenceID, Bindings->Objects);
 
