Sequencer: Fixed usage of garbage this ptr

diff --git a/Engine/Source/Runtime/MovieScene/Private/EntitySystem/MovieSceneSequenceInstance.cpp b/Engine/Source/Runtime/MovieScene/Private/EntitySystem/MovieSceneSequenceInstance.cpp
--- a/Engine/Source/Runtime/MovieScene/Private/EntitySystem/MovieSceneSequenceInstance.cpp
+++ b/Engine/Source/Runtime/MovieScene/Private/EntitySystem/MovieSceneSequenceInstance.cpp
@@ -28,6 +28,44 @@ namespace MovieScene
 DECLARE_CYCLE_STAT(TEXT("Sequence Instance Update"), MovieSceneEval_SequenceInstanceUpdate, STATGROUP_MovieSceneEval);
 DECLARE_CYCLE_STAT(TEXT("[External] Sequence Instance Post-Update"), MovieSceneEval_SequenceInstancePostUpdate, STATGROUP_MovieSceneEval);
 
+
+void PurgeStaleTrackTemplates(UMovieSceneCompiledDataManager* CompiledDataManager, FMovieSceneCompiledDataID CompiledDataID)
+{
+	FMovieSceneEvaluationTemplate* EvalTemplate = const_cast<FMovieSceneEvaluationTemplate*>(CompiledDataManager->FindTrackTemplate(CompiledDataID));
+	if (EvalTemplate)
+	{
+		EvalTemplate->PurgeStaleTracks();
+	}
+
+	// Do the same for all subsequences
+	const FMovieSceneSequenceHierarchy* Hierarchy = CompiledDataManager->FindHierarchy(CompiledDataID);
+	if (Hierarchy)
+	{
+		for (const TTuple<FMovieSceneSequenceID, FMovieSceneSubSequenceData>& Pair : Hierarchy->AllSubSequenceData())
+		{
+			UMovieSceneSequence* SubSequence = Pair.Value.GetLoadedSequence();
+			if (!SubSequence)
+			{
+				continue;
+			}
+			FMovieSceneCompiledDataID SubCompiledDataID = CompiledDataManager->FindDataID(SubSequence);
+			if (!SubCompiledDataID.IsValid())
+			{
+				continue;
+			}
+
+			FMovieSceneEvaluationTemplate* SubEvalTemplate = const_cast<FMovieSceneEvaluationTemplate*>(CompiledDataManager->FindTrackTemplate(SubCompiledDataID));
+			if (SubEvalTemplate)
+			{
+				SubEvalTemplate->PurgeStaleTracks();
+			}
+		}
+	}
+}
+
+
+
+
 FSequenceInstance::FSequenceInstance(UMovieSceneEntitySystemLinker* Linker, IMovieScenePlayer* Player, FRootInstanceHandle InInstanceHandle)
 	: SequenceID(MovieSceneSequenceID::Root)
 	, RootOverrideSequenceID(MovieSceneSequenceID::Root)
@@ -278,45 +316,28 @@ void FSequenceInstance::PostEvaluation(UMovieSceneEntitySystemLinker* Linker)
 		{
 			SCOPE_CYCLE_COUNTER(MovieSceneEval_SequenceInstancePostUpdate);
 
+
+			// DANGER: This function is highly fragile due to the nature of IMovieScenePlayer::PostEvaluation
+			//         being able to re-evaluate sequences. Ultimately this can lead to FSequenceInstances being
+			//         created, destroyed, or reallocated. As such
+			//
+			//                  ***** the current this ptr can become invalid at any point ***** 
+			//
+			//         Any code which needs to run after PostEvaluate must cache any member variables it needs on
+			//         the stack _before_ Player->PostEvaluation is called.
+
+
+			// If this sequence is volatile and has legacy track templates, purge any stale track templates from the compiled data after evaluation
+			const bool bShouldPurgeTemplates = VolatilityManager && LegacyEvaluator;
+
+			UMovieSceneCompiledDataManager* LocalCompiledDataManager = bShouldPurgeTemplates ? Player->GetEvaluationTemplate().GetCompiledDataManager() : nullptr;
+			FMovieSceneCompiledDataID       LocalCompiledDataID      = CompiledDataID;
+
 			Player->PostEvaluation(Context);
 
-			if (VolatilityManager && LegacyEvaluator)
+			if (LocalCompiledDataManager)
 			{
-				FMovieSceneRootEvaluationTemplateInstance& RootTemplateInstance = Player->GetEvaluationTemplate();
-
-				UMovieSceneCompiledDataManager* CompiledDataManager = RootTemplateInstance.GetCompiledDataManager();
-
-				// If this sequence is volatile, purge any stale track templates from the compiled data
-				FMovieSceneEvaluationTemplate* EvalTemplate = const_cast<FMovieSceneEvaluationTemplate*>(CompiledDataManager->FindTrackTemplate(CompiledDataID));
-				if (EvalTemplate)
-				{
-					EvalTemplate->PurgeStaleTracks();
-				}
-
-				// Do the same for all subsequences
-				const FMovieSceneSequenceHierarchy* Hierarchy = CompiledDataManager->FindHierarchy(CompiledDataID);
-				if (Hierarchy)
-				{
-					for (const TTuple<FMovieSceneSequenceID, FMovieSceneSubSequenceData>& Pair : Hierarchy->AllSubSequenceData())
-					{
-						UMovieSceneSequence* SubSequence = Pair.Value.GetLoadedSequence();
-						if (!SubSequence)
-						{
-							continue;
-						}
-						FMovieSceneCompiledDataID SubCompiledDataID = CompiledDataManager->FindDataID(SubSequence);
-						if (!SubCompiledDataID.IsValid())
-						{
-							continue;
-						}
-
-						FMovieSceneEvaluationTemplate* SubEvalTemplate = const_cast<FMovieSceneEvaluationTemplate*>(CompiledDataManager->FindTrackTemplate(SubCompiledDataID));
-						if (SubEvalTemplate)
-						{
-							SubEvalTemplate->PurgeStaleTracks();
-						}
-					}
-				}
+				PurgeStaleTrackTemplates(LocalCompiledDataManager, LocalCompiledDataID);
 			}
 		}
 	}
