Add root sequence name to sequence instance class for easier debugging

diff --git a/Engine/Source/Runtime/MovieScene/Private/EntitySystem/MovieSceneSequenceInstance.cpp b/Engine/Source/Runtime/MovieScene/Private/EntitySystem/MovieSceneSequenceInstance.cpp
--- a/Engine/Source/Runtime/MovieScene/Private/EntitySystem/MovieSceneSequenceInstance.cpp
+++ b/Engine/Source/Runtime/MovieScene/Private/EntitySystem/MovieSceneSequenceInstance.cpp
@@ -14,6 +14,7 @@
 #include "Evaluation/MovieSceneRootOverridePath.h"
 
 #include "IMovieScenePlayer.h"
+#include "MovieSceneSequence.h"
 #include "MovieSceneSequencePlayer.h"
 #include "MovieSceneTimeHelpers.h"
 
@@ -79,6 +80,11 @@ FSequenceInstance::FSequenceInstance(UMovieSceneEntitySystemLinker* Linker, IMov
 	bFinished = true;
 	bHasEverUpdated = false;
 
+#if !UE_BUILD_SHIPPING && !UE_BUILD_TEST
+	UMovieSceneSequence* RootSequence = Player->GetEvaluationTemplate().GetRootSequence();
+	RootSequenceName = RootSequence->GetPathName();
+#endif
+
 	CompiledDataID = Player->GetEvaluationTemplate().GetCompiledDataID();
 
 	FMovieSceneObjectCache& ObjectCache = Player->State.GetObjectCache(SequenceID);
diff --git a/Engine/Source/Runtime/MovieScene/Public/EntitySystem/MovieSceneSequenceInstance.h b/Engine/Source/Runtime/MovieScene/Public/EntitySystem/MovieSceneSequenceInstance.h
--- a/Engine/Source/Runtime/MovieScene/Public/EntitySystem/MovieSceneSequenceInstance.h
+++ b/Engine/Source/Runtime/MovieScene/Public/EntitySystem/MovieSceneSequenceInstance.h
@@ -269,6 +269,11 @@ private:
 
 	FMovieSceneContext Context;
 
+#if !UE_BUILD_SHIPPING && !UE_BUILD_TEST
+	/** Name of the root sequence */
+	FString RootSequenceName;
+#endif
+
 	/** For top-level sequences only - legacy track template evaluator for the entire sequence */
 	TUniquePtr<FMovieSceneTrackEvaluator> LegacyEvaluator;
 	/** For top-level sequences only - interface for either a flat or hierarchical entity updater */
