Add a "role" to sequencer ECS linkers to make it easier to set debug breakpoints in linker code

diff --git a/Engine/Plugins/MovieScene/TemplateSequence/Source/TemplateSequence/Private/CameraAnimationSequencePlayer.cpp b/Engine/Plugins/MovieScene/TemplateSequence/Source/TemplateSequence/Private/CameraAnimationSequencePlayer.cpp
--- a/Engine/Plugins/MovieScene/TemplateSequence/Source/TemplateSequence/Private/CameraAnimationSequencePlayer.cpp
+++ b/Engine/Plugins/MovieScene/TemplateSequence/Source/TemplateSequence/Private/CameraAnimationSequencePlayer.cpp
@@ -250,7 +250,9 @@ UMovieSceneEntitySystemLinker* UCameraAnimationSequencePlayer::ConstructEntitySy
 			return Linker;
 		}
 	}
-	return NewObject<UMovieSceneEntitySystemLinker>(GetTransientPackage());
+	UMovieSceneEntitySystemLinker* NewLinker = NewObject<UMovieSceneEntitySystemLinker>(GetTransientPackage());
+	NewLinker->SetLinkerRole(UE::MovieScene::EEntitySystemLinkerRole::Standalone);
+	return NewLinker;
 }
 
 EMovieScenePlayerStatus::Type UCameraAnimationSequencePlayer::GetPlaybackStatus() const
diff --git a/Engine/Plugins/MovieScene/TemplateSequence/Source/TemplateSequence/Private/CameraAnimationSequenceSubsystem.cpp b/Engine/Plugins/MovieScene/TemplateSequence/Source/TemplateSequence/Private/CameraAnimationSequenceSubsystem.cpp
--- a/Engine/Plugins/MovieScene/TemplateSequence/Source/TemplateSequence/Private/CameraAnimationSequenceSubsystem.cpp
+++ b/Engine/Plugins/MovieScene/TemplateSequence/Source/TemplateSequence/Private/CameraAnimationSequenceSubsystem.cpp
@@ -42,6 +42,7 @@
 	if (!Linker && bAutoCreate)
 	{
 		Linker = NewObject<UMovieSceneEntitySystemLinker>(this, TEXT("CameraAnimationSequenceSubsystemLinker"));
+		Linker->SetLinkerRole(UE::MovieScene::EEntitySystemLinkerRole::CameraAnimations);
 		Runner.AttachToLinker(Linker);
 	}
 	return Linker;
diff --git a/Engine/Source/Runtime/MovieScene/Private/EntitySystem/MovieSceneEntitySystemLinker.cpp b/Engine/Source/Runtime/MovieScene/Private/EntitySystem/MovieSceneEntitySystemLinker.cpp
--- a/Engine/Source/Runtime/MovieScene/Private/EntitySystem/MovieSceneEntitySystemLinker.cpp
+++ b/Engine/Source/Runtime/MovieScene/Private/EntitySystem/MovieSceneEntitySystemLinker.cpp
@@ -37,6 +37,16 @@ struct FCustomEventDeleter
 
 static FComponentRegistry GComponentRegistry;
 
+EEntitySystemLinkerRole RegisterCustomEntitySystemLinkerRole()
+{
+	static EEntitySystemLinkerRole NextCustom = EEntitySystemLinkerRole::Custom;
+
+	EEntitySystemLinkerRole Result = NextCustom;
+	NextCustom = (EEntitySystemLinkerRole)((uint32)NextCustom + 1);
+	check((uint32)NextCustom != TNumericLimits<uint32>::Max());
+	return Result;
+}
+
 } // namespace MovieScene
 } // namespace UE
 
@@ -46,6 +56,7 @@ UMovieSceneEntitySystemLinker::UMovieSceneEntitySystemLinker(const FObjectInitia
 {
 	using namespace UE::MovieScene;
 
+	Role = EEntitySystemLinkerRole::Unknown;
 	LastSystemLinkVersion = 0;
 	LastInstantiationVersion = 0;
 	AutoLinkMode = EAutoLinkRelevantSystems::Enabled;
@@ -92,7 +103,7 @@ void UMovieSceneEntitySystemLinker::Reset()
 	EntityManager.Destroy();
 }
 
-UMovieSceneEntitySystemLinker* UMovieSceneEntitySystemLinker::FindOrCreateLinker(UObject* PreferredOuter, const TCHAR* Name)
+UMovieSceneEntitySystemLinker* UMovieSceneEntitySystemLinker::FindOrCreateLinker(UObject* PreferredOuter, UE::MovieScene::EEntitySystemLinkerRole LinkerRole, const TCHAR* Name)
 {
 	if (!PreferredOuter)
 	{
@@ -103,18 +114,22 @@ UMovieSceneEntitySystemLinker* UMovieSceneEntitySystemLinker::FindOrCreateLinker
 	if (!Existing)
 	{
 		Existing = NewObject<UMovieSceneEntitySystemLinker>(PreferredOuter, Name);
+		Existing->SetLinkerRole(LinkerRole);
 	}
+	ensure(Existing->Role == LinkerRole);
 	return Existing;
 }
 
-UMovieSceneEntitySystemLinker* UMovieSceneEntitySystemLinker::CreateLinker(UObject* PreferredOuter)
+UMovieSceneEntitySystemLinker* UMovieSceneEntitySystemLinker::CreateLinker(UObject* PreferredOuter, UE::MovieScene::EEntitySystemLinkerRole LinkerRole)
 {
 	if (!PreferredOuter)
 	{
 		PreferredOuter = GetTransientPackage();
 	}
 
-	return NewObject<UMovieSceneEntitySystemLinker>(PreferredOuter);
+	UMovieSceneEntitySystemLinker* NewLinker = NewObject<UMovieSceneEntitySystemLinker>(PreferredOuter);
+	NewLinker->Role = LinkerRole;
+	return NewLinker;
 }
 
 UE::MovieScene::FComponentRegistry* UMovieSceneEntitySystemLinker::GetComponents()
diff --git a/Engine/Source/Runtime/MovieScene/Private/Evaluation/MovieSceneEvaluationTemplateInstance.cpp b/Engine/Source/Runtime/MovieScene/Private/Evaluation/MovieSceneEvaluationTemplateInstance.cpp
--- a/Engine/Source/Runtime/MovieScene/Private/Evaluation/MovieSceneEvaluationTemplateInstance.cpp
+++ b/Engine/Source/Runtime/MovieScene/Private/Evaluation/MovieSceneEvaluationTemplateInstance.cpp
@@ -53,7 +53,7 @@ UMovieSceneEntitySystemLinker* FMovieSceneRootEvaluationTemplateInstance::Constr
 	}
 
 	UObject* PlaybackContext = Player.GetPlaybackContext();
-	return UMovieSceneEntitySystemLinker::FindOrCreateLinker(PlaybackContext, TEXT("DefaultEntitySystemLinker"));
+	return UMovieSceneEntitySystemLinker::FindOrCreateLinker(PlaybackContext, UE::MovieScene::EEntitySystemLinkerRole::Standalone, TEXT("DefaultEntitySystemLinker"));
 }
 
 void FMovieSceneRootEvaluationTemplateInstance::Initialize(UMovieSceneSequence& InRootSequence, IMovieScenePlayer& Player, UMovieSceneCompiledDataManager* InCompiledDataManager)
diff --git a/Engine/Source/Runtime/MovieScene/Private/MovieSceneSequencePlayer.cpp b/Engine/Source/Runtime/MovieScene/Private/MovieSceneSequencePlayer.cpp
--- a/Engine/Source/Runtime/MovieScene/Private/MovieSceneSequencePlayer.cpp
+++ b/Engine/Source/Runtime/MovieScene/Private/MovieSceneSequencePlayer.cpp
@@ -734,7 +734,7 @@ UMovieSceneEntitySystemLinker* UMovieSceneSequencePlayer::ConstructEntitySystemL
 		return TickManager->GetLinker();
 	}
 
-	return UMovieSceneEntitySystemLinker::CreateLinker(GetPlaybackContext());
+	return UMovieSceneEntitySystemLinker::CreateLinker(GetPlaybackContext(), UE::MovieScene::EEntitySystemLinkerRole::Standalone);
 }
 
 void UMovieSceneSequencePlayer::InitializeForTick(UObject* Context)
diff --git a/Engine/Source/Runtime/MovieScene/Private/MovieSceneSequenceTickManager.cpp b/Engine/Source/Runtime/MovieScene/Private/MovieSceneSequenceTickManager.cpp
--- a/Engine/Source/Runtime/MovieScene/Private/MovieSceneSequenceTickManager.cpp
+++ b/Engine/Source/Runtime/MovieScene/Private/MovieSceneSequenceTickManager.cpp
@@ -143,7 +143,7 @@ UMovieSceneSequenceTickManager* UMovieSceneSequenceTickManager::Get(UObject* Pla
 	{
 		TickManager = NewObject<UMovieSceneSequenceTickManager>(World, TEXT("GlobalMovieSceneSequenceTickManager"));
 
-		TickManager->Linker = UMovieSceneEntitySystemLinker::FindOrCreateLinker(World, TEXT("MovieSceneSequencePlayerEntityLinker"));
+		TickManager->Linker = UMovieSceneEntitySystemLinker::FindOrCreateLinker(World, UE::MovieScene::EEntitySystemLinkerRole::LevelSequences, TEXT("MovieSceneSequencePlayerEntityLinker"));
 		check(TickManager->Linker);
 		TickManager->Runner.AttachToLinker(TickManager->Linker);
 
diff --git a/Engine/Source/Runtime/MovieScene/Public/EntitySystem/MovieSceneEntitySystemLinker.h b/Engine/Source/Runtime/MovieScene/Public/EntitySystem/MovieSceneEntitySystemLinker.h
--- a/Engine/Source/Runtime/MovieScene/Public/EntitySystem/MovieSceneEntitySystemLinker.h
+++ b/Engine/Source/Runtime/MovieScene/Public/EntitySystem/MovieSceneEntitySystemLinker.h
@@ -34,6 +34,26 @@ namespace MovieScene
 		Enabled,
 		Disable,
 	};
+
+	/** Enum that describes what a sequencer ECS linker is meant for (only used for debugging reasons) */
+	enum class EEntitySystemLinkerRole : uint32
+	{
+		/** The linker's role is unknown */
+		Unknown = 0,
+		/** The linker is handling level sequences */
+		LevelSequences = 1,
+		/** The linker is handling camera animations */
+		CameraAnimations,
+		/** The linker is handling UMG animations */
+		UMG,
+		/** The linker is handling a standalone sequence, such as those with a blocking evaluation flag */
+		Standalone,
+		/** This value and any greater values are for other custom roles */
+		Custom
+	};
+
+	/** Register a new custom linker role */
+	MOVIESCENE_API EEntitySystemLinkerRole RegisterCustomEntitySystemLinkerRole();
 }
 }
 
@@ -71,8 +91,8 @@ public:
 
 	static FComponentRegistry* GetComponents();
 
-	static UMovieSceneEntitySystemLinker* FindOrCreateLinker(UObject* PreferredOuter, const TCHAR* Name = TEXT("DefaultMovieSceneEntitySystemLinker"));
-	static UMovieSceneEntitySystemLinker* CreateLinker(UObject* PreferredOuter);
+	static UMovieSceneEntitySystemLinker* FindOrCreateLinker(UObject* PreferredOuter, UE::MovieScene::EEntitySystemLinkerRole LinkerRole, const TCHAR* Name = TEXT("DefaultMovieSceneEntitySystemLinker"));
+	static UMovieSceneEntitySystemLinker* CreateLinker(UObject* PreferredOuter, UE::MovieScene::EEntitySystemLinkerRole LinkerRole);
 
 	FInstanceRegistry* GetInstanceRegistry()
 	{
@@ -113,7 +133,6 @@ public:
 		return SystemContext;
 	}
 
-
 	/**
 	 * Set the system context for this linker allowing some systems to be excluded based on the context.
 	 *
@@ -124,6 +143,21 @@ public:
 		SystemContext = InSystemContext;
 	}
 
+	/**
+	 * Gets the role of this linker
+	 */
+	UE::MovieScene::EEntitySystemLinkerRole GetLinkerRole() const
+	{
+		return Role;
+	}
+
+	/**
+	 * Gets the role of this linker
+	 */
+	void SetLinkerRole(UE::MovieScene::EEntitySystemLinkerRole InRole)
+	{
+		Role = InRole;
+	}
 
 	/**
 	 * Completely reset this linker back to its default state, abandoning all systems and destroying all entities
@@ -299,6 +333,7 @@ private:
 
 protected:
 
+	UE::MovieScene::EEntitySystemLinkerRole Role;
 	UE::MovieScene::EAutoLinkRelevantSystems AutoLinkMode;
 	UE::MovieScene::EEntitySystemContext SystemContext;
 };
diff --git a/Engine/Source/Runtime/UMG/Private/Animation/UMGSequencePlayer.cpp b/Engine/Source/Runtime/UMG/Private/Animation/UMGSequencePlayer.cpp
--- a/Engine/Source/Runtime/UMG/Private/Animation/UMGSequencePlayer.cpp
+++ b/Engine/Source/Runtime/UMG/Private/Animation/UMGSequencePlayer.cpp
@@ -53,7 +53,7 @@ UMovieSceneEntitySystemLinker* UUMGSequencePlayer::ConstructEntitySystemLinker()
 		return Widget->AnimationTickManager->GetLinker();
 	}
 
-	return UMovieSceneEntitySystemLinker::CreateLinker(Widget ? Widget->GetWorld() : nullptr);
+	return UMovieSceneEntitySystemLinker::CreateLinker(Widget ? Widget->GetWorld() : nullptr, UE::MovieScene::EEntitySystemLinkerRole::UMG);
 }
 
 void UUMGSequencePlayer::Tick(float DeltaTime)
diff --git a/Engine/Source/Runtime/UMG/Private/Animation/UMGSequenceTickManager.cpp b/Engine/Source/Runtime/UMG/Private/Animation/UMGSequenceTickManager.cpp
--- a/Engine/Source/Runtime/UMG/Private/Animation/UMGSequenceTickManager.cpp
+++ b/Engine/Source/Runtime/UMG/Private/Animation/UMGSequenceTickManager.cpp
@@ -257,7 +257,7 @@ UUMGSequenceTickManager* UUMGSequenceTickManager::Get(UObject* PlaybackContext)
 	{
 		TickManager = NewObject<UUMGSequenceTickManager>(Owner, TickManagerName);
 
-		TickManager->Linker = UMovieSceneEntitySystemLinker::FindOrCreateLinker(Owner, TEXT("UMGAnimationEntitySystemLinker"));
+		TickManager->Linker = UMovieSceneEntitySystemLinker::FindOrCreateLinker(Owner, UE::MovieScene::EEntitySystemLinkerRole::UMG, TEXT("UMGAnimationEntitySystemLinker"));
 		check(TickManager->Linker);
 		TickManager->Runner.AttachToLinker(TickManager->Linker);
 
