Sequencer: Added support for tick intervals to UMovieSceneSequenceTickManager

The sequence tick manager is now able to manage updating sequences at various different tick intervals. To do so it creates new linkers and runners for each unique tick interval, and groups all updates for that tick interval into their own linker.
As such, sequences running at different tick intervals are unable to blend with each other, since the data is separated, but this is an acceptable trade-off for the flexibility it brings.

Usage of Actor nomenclature in UMovieSceneSequenceTickManager has been deprecated and replaced with 'client' since it is not necessarily an actor. In fact, all existing implementations of this interface have been unified into UMovieSceneSequencePlayer so it can be updated directly without having to go through an actor or component indirection first. This does introduce an issue whereby the player could continue to tick _after_ its owner EndPlay has been called, which forces the owner to call TearDown on EndPlay. The SequencePlayer itself is only registered inside ::Initialize so should not begin ticking until it is ready to play.

A new struct type, FMovieSceneSequenceTickInterval (and FMovieSceneSequenceResolvedTickInterval), has been added to unify places where we deal with tick intervals and provide a nice edit interface in the playback settings allowing users to manually override it without having to manipulate actor ticks.

diff --git a/Engine/Plugins/MovieScene/ActorSequence/Source/ActorSequence/Private/ActorSequenceComponent.cpp b/Engine/Plugins/MovieScene/ActorSequence/Source/ActorSequence/Private/ActorSequenceComponent.cpp
--- a/Engine/Plugins/MovieScene/ActorSequence/Source/ActorSequence/Private/ActorSequenceComponent.cpp
+++ b/Engine/Plugins/MovieScene/ActorSequence/Source/ActorSequence/Private/ActorSequenceComponent.cpp
@@ -41,12 +41,6 @@ void UActorSequenceComponent::BeginPlay()
 		{
 			SequencePlayer->Play();
 		}
-
-		UMovieSceneSequenceTickManager* TickManager = SequencePlayer->GetTickManager();
-		if (ensure(TickManager))
-		{
-			TickManager->RegisterSequenceActor(GetOwner(), this);
-		}
 	}
 }
 
@@ -57,11 +51,7 @@ void UActorSequenceComponent::EndPlay(const EEndPlayReason::Type EndPlayReason)
 		// Stop the internal sequence player during EndPlay when it's safer
 		// to modify other actor's state during state restoration.
 		SequencePlayer->Stop();
-
-		if (UMovieSceneSequenceTickManager* TickManager = SequencePlayer->GetTickManager())
-		{
-			TickManager->UnregisterSequenceActor(GetOwner(), this);
-		}
+		SequencePlayer->TearDown();
 	}
 	
 	Super::EndPlay(EndPlayReason);
@@ -72,11 +62,3 @@ void UActorSequenceComponent::TickComponent(float DeltaSeconds, enum ELevelTick
 	Super::TickComponent(DeltaSeconds, TickType, ThisTickFunction);
 }
 
-void UActorSequenceComponent::TickFromSequenceTickManager(float DeltaSeconds)
-{
-	if (SequencePlayer)
-	{
-		SequencePlayer->UpdateAsync(DeltaSeconds);
-	}
-}
-
diff --git a/Engine/Plugins/MovieScene/ActorSequence/Source/ActorSequence/Public/ActorSequenceComponent.h b/Engine/Plugins/MovieScene/ActorSequence/Source/ActorSequence/Public/ActorSequenceComponent.h
--- a/Engine/Plugins/MovieScene/ActorSequence/Source/ActorSequence/Public/ActorSequenceComponent.h
+++ b/Engine/Plugins/MovieScene/ActorSequence/Source/ActorSequence/Public/ActorSequenceComponent.h
@@ -4,7 +4,6 @@
 
 #include "Components/ActorComponent.h"
 #include "MovieSceneSequencePlayer.h"
-#include "MovieSceneSequenceTickManager.h"
 #include "ActorSequenceComponent.generated.h"
 
 
@@ -18,7 +17,6 @@ class UActorSequencePlayer;
 UCLASS(Blueprintable, Experimental, ClassGroup=Sequence, hidecategories=(Collision, Cooking, Activation), meta=(BlueprintSpawnableComponent))
 class ACTORSEQUENCE_API UActorSequenceComponent
 	: public UActorComponent
-	, public IMovieSceneSequenceActor
 {
 public:
 	GENERATED_BODY()
@@ -41,9 +39,6 @@ public:
 	virtual void EndPlay(const EEndPlayReason::Type EndPlayReason) override;
 	virtual void TickComponent(float DeltaTime, enum ELevelTick TickType, FActorComponentTickFunction *ThisTickFunction) override;
 
-	// IMovieSceneSequenceActor interface
-	virtual void TickFromSequenceTickManager(float DeltaSeconds) override;
-
 protected:
 
 	UPROPERTY(EditAnywhere, Category="Playback", meta=(ShowOnlyInnerProperties))
diff --git a/Engine/Plugins/MovieScene/TemplateSequence/Source/TemplateSequence/Private/TemplateSequenceActor.cpp b/Engine/Plugins/MovieScene/TemplateSequence/Source/TemplateSequence/Private/TemplateSequenceActor.cpp
--- a/Engine/Plugins/MovieScene/TemplateSequence/Source/TemplateSequence/Private/TemplateSequenceActor.cpp
+++ b/Engine/Plugins/MovieScene/TemplateSequence/Source/TemplateSequence/Private/TemplateSequenceActor.cpp
@@ -45,13 +45,13 @@
 {
 	Super::PostInitializeComponents();
 
+	SequencePlayer->InitializeForTick(this);
+
 	InitializePlayer();
 }
 
 void ATemplateSequenceActor::BeginPlay()
 {
-	UMovieSceneSequenceTickManager::Get(this)->RegisterSequenceActor(this);
-
 	Super::BeginPlay();
 	
 	if (PlaybackSettings.bAutoPlay)
@@ -66,21 +66,12 @@
 	{
 		// See comment in LevelSequenceActor.cpp
 		SequencePlayer->Stop();
+		SequencePlayer->TearDown();
 	}
 
-	UMovieSceneSequenceTickManager::Get(this)->UnregisterSequenceActor(this);
-
 	Super::EndPlay(EndPlayReason);
 }
 
-void ATemplateSequenceActor::TickFromSequenceTickManager(float DeltaSeconds)
-{
-	if (SequencePlayer)
-	{
-		SequencePlayer->UpdateAsync(DeltaSeconds);
-	}
-}
-
 UTemplateSequence* ATemplateSequenceActor::GetSequence() const
 {
 	return Cast<UTemplateSequence>(TemplateSequence.ResolveObject());
diff --git a/Engine/Plugins/MovieScene/TemplateSequence/Source/TemplateSequence/Public/TemplateSequenceActor.h b/Engine/Plugins/MovieScene/TemplateSequence/Source/TemplateSequence/Public/TemplateSequenceActor.h
--- a/Engine/Plugins/MovieScene/TemplateSequence/Source/TemplateSequence/Public/TemplateSequenceActor.h
+++ b/Engine/Plugins/MovieScene/TemplateSequence/Source/TemplateSequence/Public/TemplateSequenceActor.h
@@ -6,7 +6,6 @@
 #include "TemplateSequencePlayer.h"
 #include "IMovieScenePlaybackClient.h"
 #include "MovieSceneBindingOwnerInterface.h"
-#include "MovieSceneSequenceTickManager.h"
 #include "GameFramework/Actor.h"
 #include "UObject/ObjectMacros.h"
 #include "UObject/Object.h"
@@ -46,7 +45,6 @@ struct FTemplateSequenceBindingOverrideData
 UCLASS(hideCategories = (Rendering, Physics, LOD, Activation, Input))
 class TEMPLATESEQUENCE_API ATemplateSequenceActor
 	: public AActor
-	, public IMovieSceneSequenceActor
 	, public IMovieScenePlaybackClient
 {
 public:
@@ -109,10 +107,6 @@ public:
 
 protected:
 
-	//~ Begin IMovieSceneSequenceActor interface
-	virtual void TickFromSequenceTickManager(float DeltaSeconds) override;
-	//~ End IMovieSceneSequenceActor interface
-
 	//~ Begin IMovieScenePlaybackClient interface
 	virtual bool RetrieveBindingOverrides(const FGuid& InBindingId, FMovieSceneSequenceID InSequenceID, TArray<UObject*, TInlineAllocator<1>>& OutObjects) const override;
 	virtual UObject* GetInstanceData() const override;
diff --git a/Engine/Source/Runtime/LevelSequence/Private/LevelSequenceActor.cpp b/Engine/Source/Runtime/LevelSequence/Private/LevelSequenceActor.cpp
--- a/Engine/Source/Runtime/LevelSequence/Private/LevelSequenceActor.cpp
+++ b/Engine/Source/Runtime/LevelSequence/Private/LevelSequenceActor.cpp
@@ -171,12 +171,6 @@ void ALevelSequenceActor::PostInitializeComponents()
 
 void ALevelSequenceActor::BeginPlay()
 {
-	UMovieSceneSequenceTickManager* TickManager = SequencePlayer->GetTickManager();
-	if (ensure(TickManager))
-	{
-		TickManager->RegisterSequenceActor(this);
-	}
-
 	Super::BeginPlay();
 
 	if (PlaybackSettings.bAutoPlay)
@@ -198,23 +192,12 @@ void ALevelSequenceActor::EndPlay(const EEndPlayReason::Type EndPlayReason)
 		SequencePlayer->OnPlayReverse.RemoveAll(this);
 		SequencePlayer->OnStop.RemoveAll(this);
 
-		if (UMovieSceneSequenceTickManager* TickManager = SequencePlayer->GetTickManager())
-		{
-			TickManager->UnregisterSequenceActor(this);
-		}
+		SequencePlayer->TearDown();
 	}
 
 	Super::EndPlay(EndPlayReason);
 }
 
-void ALevelSequenceActor::TickFromSequenceTickManager(float DeltaSeconds)
-{
-	if (SequencePlayer)
-	{
-		SequencePlayer->UpdateAsync(DeltaSeconds);
-	}
-}
-
 void ALevelSequenceActor::PostLoad()
 {
 	Super::PostLoad();
diff --git a/Engine/Source/Runtime/LevelSequence/Public/LevelSequenceActor.h b/Engine/Source/Runtime/LevelSequence/Public/LevelSequenceActor.h
--- a/Engine/Source/Runtime/LevelSequence/Public/LevelSequenceActor.h
+++ b/Engine/Source/Runtime/LevelSequence/Public/LevelSequenceActor.h
@@ -10,7 +10,6 @@
 #include "LevelSequencePlayer.h"
 #include "MovieSceneBindingOwnerInterface.h"
 #include "MovieSceneBindingOverrides.h"
-#include "MovieSceneSequenceTickManager.h"
 #include "LevelSequenceActor.generated.h"
 
 class ULevelSequenceBurnIn;
@@ -61,7 +60,6 @@
 UCLASS(hideCategories=(Rendering, Physics, HLOD, Activation, Input))
 class LEVELSEQUENCE_API ALevelSequenceActor
 	: public AActor
-	, public IMovieSceneSequenceActor
 	, public IMovieScenePlaybackClient
 	, public IMovieSceneBindingOwnerInterface
 {
@@ -259,10 +257,6 @@
 
 protected:
 
-	//~ Begin IMovieSceneSequenceActor interface
-	virtual void TickFromSequenceTickManager(float DeltaSeconds) override;
-	//~ End IMovieSceneSequenceActor interface
-
 	//~ Begin IMovieScenePlaybackClient interface
 	virtual bool RetrieveBindingOverrides(const FGuid& InBindingId, FMovieSceneSequenceID InSequenceID, TArray<UObject*, TInlineAllocator<1>>& OutObjects) const override;
 	virtual UObject* GetInstanceData() const override;
diff --git a/Engine/Source/Runtime/MovieScene/Private/MovieSceneSequencePlayer.cpp b/Engine/Source/Runtime/MovieScene/Private/MovieSceneSequencePlayer.cpp
--- a/Engine/Source/Runtime/MovieScene/Private/MovieSceneSequencePlayer.cpp
+++ b/Engine/Source/Runtime/MovieScene/Private/MovieSceneSequencePlayer.cpp
@@ -117,6 +117,7 @@ UMovieSceneSequencePlayer::UMovieSceneSequencePlayer(const FObjectInitializer& I
 	, DurationFrames(0)
 	, DurationSubFrames(0.f)
 	, CurrentNumLoops(0)
+	, CurrentRunner(nullptr)
 {
 	PlayPosition.Reset(FFrameTime(0));
 
@@ -729,9 +730,9 @@ TOptional<TRange<FFrameTime>> UMovieSceneSequencePlayer::GetPauseRange(const FFr
 
 UMovieSceneEntitySystemLinker* UMovieSceneSequencePlayer::ConstructEntitySystemLinker()
 {
-	if (ensure(TickManager) && !EnumHasAnyFlags(Sequence->GetFlags(), EMovieSceneSequenceFlags::BlockingEvaluation))
+	if (ensure(TickManager && RegisteredTickInterval.IsSet()) && !EnumHasAnyFlags(Sequence->GetFlags(), EMovieSceneSequenceFlags::BlockingEvaluation))
 	{
-		return TickManager->GetLinker();
+		return TickManager->GetLinker(RegisteredTickInterval.GetValue());
 	}
 
 	return UMovieSceneEntitySystemLinker::CreateLinker(GetPlaybackContext(), UE::MovieScene::EEntitySystemLinkerRole::Standalone);
@@ -837,6 +838,18 @@ void UMovieSceneSequencePlayer::Initialize(UMovieSceneSequence* InSequence, cons
 		InitializeForTick(GetPlaybackContext());
 	}
 
+	FMovieSceneSequenceTickInterval TickInterval = PlaybackSettings.bInheritTickIntervalFromOwner
+		? FMovieSceneSequenceTickInterval::GetInheritedInterval(this)
+		: PlaybackSettings.TickInterval;
+
+	if (RegisteredTickInterval.IsSet() && RegisteredTickInterval.GetValue() != TickInterval)
+	{
+		TickManager->UnregisterTickClient(this);
+	}
+
+	RegisteredTickInterval = TickInterval;
+	TickManager->RegisterTickClient(RegisteredTickInterval.GetValue(), this);
+
 	RootTemplateInstance.Initialize(*Sequence, *this, nullptr);
 
 	LatentActionManager.ClearLatentActions();
@@ -889,6 +902,12 @@ void UMovieSceneSequencePlayer::Update(const float DeltaSeconds)
 	}
 }
 
+void UMovieSceneSequencePlayer::TickFromSequenceTickManager(float DeltaSeconds, FMovieSceneEntitySystemRunner* InRunner)
+{
+	TGuardValue<FMovieSceneEntitySystemRunner*> RunnerGuard(CurrentRunner, InRunner);
+	UpdateAsync(DeltaSeconds);
+}
+
 void UMovieSceneSequencePlayer::UpdateAsync(const float DeltaSeconds)
 {
 	check(!bIsAsyncUpdate);
@@ -1128,11 +1147,27 @@ void UMovieSceneSequencePlayer::UpdateMovieSceneInstance(FMovieSceneEvaluationRa
 		// Evaluate the sequence synchronously.
 		RootTemplateInstance.Evaluate(Context, *this);
 	}
+	else if (CurrentRunner)
+	{
+		CurrentRunner->QueueUpdate(Context, RootTemplateInstance.GetRootInstanceHandle());
+	}
 	else
 	{
-		// Queue an evaluation on the tick manager.
-		FMovieSceneEntitySystemRunner& Runner = TickManager->GetRunner();
-		Runner.QueueUpdate(Context, RootTemplateInstance.GetRootInstanceHandle());
+		// Look up the runner based on our tick interval
+		FMovieSceneEntitySystemRunner* Runner = TickManager->GetRunner(RegisteredTickInterval.GetValue());
+		if (ensure(Runner))
+		{
+			Runner->QueueUpdate(Context, RootTemplateInstance.GetRootInstanceHandle());
+		}
+	}
+}
+
+void UMovieSceneSequencePlayer::TearDown()
+{
+	if (TickManager)
+	{
+		TickManager->UnregisterTickClient(this);
+		TickManager = nullptr;
 	}
 }
 
@@ -1641,7 +1676,7 @@ void UMovieSceneSequencePlayer::BeginDestroy()
 {
 	RootTemplateInstance.BeginDestroy();
 
-	TickManager = nullptr;
+	TearDown();
 
 	Super::BeginDestroy();
 }
@@ -1700,7 +1735,10 @@ void UMovieSceneSequencePlayer::RunLatentActions()
 	}
 	else
 	{
-		LatentActionManager.RunLatentActions(RootTemplateInstance.GetEntitySystemRunner());
+		LatentActionManager.RunLatentActions([this]
+		{
+			this->RootTemplateInstance.GetEntitySystemRunner().Flush();
+		});
 	}
 }
 
diff --git a/Engine/Source/Runtime/MovieScene/Private/MovieSceneSequenceTickInterval.cpp b/Engine/Source/Runtime/MovieScene/Private/MovieSceneSequenceTickInterval.cpp
new file mode 100644
--- /dev/null
+++ b/Engine/Source/Runtime/MovieScene/Private/MovieSceneSequenceTickInterval.cpp
@@ -0,0 +1,71 @@
+// Copyright Epic Games, Inc. All Rights Reserved.
+
+#include "MovieSceneSequenceTickInterval.h"
+#include "GameFramework/Actor.h"
+#include "Components/ActorComponent.h"
+
+namespace UE::MovieScene
+{
+
+	int32 GMovieSceneTickIntervalResolutionMs = 50;
+	static FAutoConsoleVariableRef CVarMovieSceneTickIntervalResolutionMs(
+		TEXT("Sequencer.TickIntervalGroupingResolutionMs"),
+		GMovieSceneTickIntervalResolutionMs,
+		TEXT("Defines the maximum resolution for actor tick interval groupings. Bigger numbers will group more actors together when they have custom tick intervals, but will lead to less accurate intervals.\n"),
+		ECVF_Default
+	);
+
+} // namespace UE::MovieScene
+
+
+FMovieSceneSequenceTickInterval::FMovieSceneSequenceTickInterval(const AActor* InActor)
+{
+	check(InActor);
+
+	TickIntervalSeconds = InActor->PrimaryActorTick.TickInterval;
+	bTickWhenPaused = InActor->PrimaryActorTick.bTickEvenWhenPaused;
+	bAllowRounding = true;
+}
+
+FMovieSceneSequenceTickInterval::FMovieSceneSequenceTickInterval(const UActorComponent* InActorComponent)
+{
+	check(InActorComponent);
+
+	TickIntervalSeconds = InActorComponent->PrimaryComponentTick.TickInterval;
+	bTickWhenPaused = InActorComponent->PrimaryComponentTick.bTickEvenWhenPaused;
+	bAllowRounding = true;
+}
+
+int32 FMovieSceneSequenceTickInterval::RoundTickIntervalMs() const
+{
+	using namespace UE::MovieScene;
+
+	// We round actor tick intervals to the nearest GMovieSceneTickIntervalResolutionMs to prevent excessive fragmentation.
+	// TickInterval is in seconds, so we multiply by 1000/<rounding> to round it as ms in integer space
+	if (bAllowRounding)
+	{
+		return FMath::RoundToInt(TickIntervalSeconds * (1000.f/GMovieSceneTickIntervalResolutionMs)) * GMovieSceneTickIntervalResolutionMs;
+	}
+
+	return FMath::RoundToInt(TickIntervalSeconds * 1000.f);
+}
+
+FMovieSceneSequenceTickInterval FMovieSceneSequenceTickInterval::GetInheritedInterval(UObject* ContextObject)
+{
+	UObject* Current = ContextObject;
+	while (Current && !Current->IsA<ULevel>())
+	{
+		if (USceneComponent* SceneComponent = Cast<USceneComponent>(Current))
+		{
+			return FMovieSceneSequenceTickInterval(SceneComponent);
+		}
+		if (AActor* Actor = Cast<AActor>(Current))
+		{
+			return FMovieSceneSequenceTickInterval(Actor);
+		}
+
+		Current = Current->GetOuter();
+	}
+
+	return FMovieSceneSequenceTickInterval();
+}
\ No newline at end of file
diff --git a/Engine/Source/Runtime/MovieScene/Private/MovieSceneSequenceTickManager.cpp b/Engine/Source/Runtime/MovieScene/Private/MovieSceneSequenceTickManager.cpp
--- a/Engine/Source/Runtime/MovieScene/Private/MovieSceneSequenceTickManager.cpp
+++ b/Engine/Source/Runtime/MovieScene/Private/MovieSceneSequenceTickManager.cpp
@@ -5,20 +5,73 @@
 #include "EntitySystem/MovieSceneEntitySystemLinker.h"
 #include "ProfilingDebugging/CountersTrace.h"
 
+#include "Algo/IndexOf.h"
+
 DECLARE_CYCLE_STAT(TEXT("Sequence Tick Manager"), MovieSceneEval_SequenceTickManager, STATGROUP_MovieSceneEval);
 
-TRACE_DECLARE_INT_COUNTER(MovieSceneEval_SequenceTickManagerLatentActionRuns, TEXT("MovieScene/LatentActionRuns"));
+namespace UE::MovieScene
+{
 
-static TAutoConsoleVariable<int32> CVarMovieSceneMaxLatentActionLoops(
+int32 GMovieSceneMaxLatentActionLoops = 100;
+static FAutoConsoleVariableRef CVarMovieSceneMaxLatentActionLoops(
 	TEXT("Sequencer.MaxLatentActionLoops"),
-	100,
+	GMovieSceneMaxLatentActionLoops,
 	TEXT("Defines the maximum number of latent action loops that can be run in one frame.\n"),
 	ECVF_Default
 );
 
+} // namespace UE::MovieScene
+
+UMovieSceneSequenceTickManager* UMovieSceneSequenceTickManager::Get(UObject* PlaybackContext)
+{
+	check(PlaybackContext != nullptr && PlaybackContext->GetWorld() != nullptr);
+	UWorld* World = PlaybackContext->GetWorld();
+
+
+	UMovieSceneSequenceTickManager* TickManager = FindObject<UMovieSceneSequenceTickManager>(World, TEXT("GlobalMovieSceneSequenceTickManager"));
+	if (!TickManager)
+	{
+		TickManager = NewObject<UMovieSceneSequenceTickManager>(World, TEXT("GlobalMovieSceneSequenceTickManager"));
+
+		FDelegateHandle Handle = World->AddMovieSceneSequenceTickHandler(
+				FOnMovieSceneSequenceTick::FDelegate::CreateUObject(TickManager, &UMovieSceneSequenceTickManager::TickSequenceActors));
+		check(Handle.IsValid());
+		TickManager->WorldTickDelegateHandle = Handle;
+	}
+	return TickManager;
+}
+
+
 UMovieSceneSequenceTickManager::UMovieSceneSequenceTickManager(const FObjectInitializer& Init)
 	: Super(Init)
 {
+	PendingActorOperations = nullptr;
+}
+
+UMovieSceneEntitySystemLinker* UMovieSceneSequenceTickManager::GetLinker(const FMovieSceneSequenceTickInterval& TickInterval)
+{
+	const int32 RoundedTickIntervalMs = TickInterval.RoundTickIntervalMs();
+	for (const FLinkerGroup& Group : LinkerGroups)
+	{
+		if (Group.RoundedTickIntervalMs == RoundedTickIntervalMs)
+		{
+			return Group.Linker;
+		}
+	}
+	return nullptr;
+}
+
+FMovieSceneEntitySystemRunner* UMovieSceneSequenceTickManager::GetRunner(const FMovieSceneSequenceTickInterval& TickInterval)
+{
+	const int32 RoundedTickIntervalMs = TickInterval.RoundTickIntervalMs();
+	for (FLinkerGroup& Group : LinkerGroups)
+	{
+		if (Group.RoundedTickIntervalMs == RoundedTickIntervalMs)
+		{
+			return &Group.Runner;
+		}
+	}
+	return nullptr;
 }
 
 void UMovieSceneSequenceTickManager::BeginDestroy()
@@ -36,85 +89,269 @@ void UMovieSceneSequenceTickManager::BeginDestroy()
 	Super::BeginDestroy();
 }
 
-void UMovieSceneSequenceTickManager::RegisterSequenceActor(AActor* InActor)
+void UMovieSceneSequenceTickManager::RegisterTickClient(const FMovieSceneSequenceTickInterval& ResolvedTickInterval, TScriptInterface<IMovieSceneSequenceTickManagerClient> InTickInterface)
 {
-	TScriptInterface<IMovieSceneSequenceActor> SequenceActorInterface(InActor);
-	if (ensureMsgf(SequenceActorInterface, TEXT("The given actor doesn't implement the IMovieSceneSequenceActor interface!")))
+	using namespace UE::MovieScene;
+
+	UObject* InterfaceObject = InTickInterface.GetObject();
+	IMovieSceneSequenceTickManagerClient* ClientInterface = InTickInterface.GetInterface();
+	if (!ensure(InterfaceObject && ClientInterface))
 	{
-		SequenceActors.Add(FMovieSceneSequenceActorPointers{ InActor, SequenceActorInterface });
+		return;
 	}
-}
 
-void UMovieSceneSequenceTickManager::RegisterSequenceActor(AActor* InActor, TScriptInterface<IMovieSceneSequenceActor> InActorInterface)
-{
-	if (ensure(InActor && InActorInterface))
+	// If PendingActorOperations is non-null, it means we cannot mutate either TickableClients or LinkerGroups so we have to defer until afterwards
+	if (PendingActorOperations)
 	{
-		SequenceActors.Add(FMovieSceneSequenceActorPointers{ InActor, InActorInterface });
+		PendingActorOperations->Add(FPendingOperation{ InTickInterface, InterfaceObject, ResolvedTickInterval, FPendingOperation::EType::Register });
+		return;
 	}
-}
 
-void UMovieSceneSequenceTickManager::UnregisterSequenceActor(AActor* InActor)
-{
-	TScriptInterface<IMovieSceneSequenceActor> SequenceActorInterface(InActor);
-	if (ensureMsgf(SequenceActorInterface, TEXT("The given actor doesn't implement the IMovieSceneSequenceActor interface!")))
+#if DO_GUARD_SLOW
+	UWorld* ThisWorld   = GetWorld();
+	UWorld* ObjectWorld = InterfaceObject->GetWorld();
+	checkfSlow(ObjectWorld == ThisWorld, TEXT("Attempting to register object %s that is a part of world %s to a sequence tick manager contained within world %s."),
+		*InterfaceObject->GetName(), ObjectWorld ? *ObjectWorld->GetName() : TEXT("nullptr"), ThisWorld ? *ThisWorld->GetName() : TEXT("nullptr"));
+#endif
+
+	const int32 DesiredTickIntervalMs = ResolvedTickInterval.RoundTickIntervalMs();
+
+	// Find a linker group with the specified tick interval
+	// Would be nice to use Algo::IndexOfBy here but that cannot work on TSparseArray since it is a non-contiguous range
+	int32 LinkerIndex = 0;
+	for (; LinkerIndex < LinkerGroups.GetMaxIndex(); ++LinkerIndex)
+	{
+		if (LinkerGroups.IsAllocated(LinkerIndex) && LinkerGroups[LinkerIndex].RoundedTickIntervalMs == DesiredTickIntervalMs)
+		{
+			break;
+		}
+	}
+
+	if (!LinkerGroups.IsValidIndex(LinkerIndex))
+	{
+		LinkerIndex = LinkerGroups.Emplace();
+
+		TStringBuilder<64> LinkerName;
+		LinkerName.Append(TEXT("MovieSceneSequencePlayerEntityLinker"));
+		if (DesiredTickIntervalMs != 0)
+		{
+			LinkerName.Appendf(TEXT("_%i_ms"), DesiredTickIntervalMs);
+		}
+
+
+		UMovieSceneEntitySystemLinker* Linker = UMovieSceneEntitySystemLinker::FindOrCreateLinker(GetWorld(), EEntitySystemLinkerRole::LevelSequences, LinkerName.ToString());
+		check(Linker);
+
+		FLinkerGroup& NewGroup = LinkerGroups[LinkerIndex];
+		NewGroup.Linker = Linker;
+		NewGroup.Runner.AttachToLinker(Linker);
+		NewGroup.RoundedTickIntervalMs = DesiredTickIntervalMs;
+		NewGroup.NumClients = 1;
+	}
+	else
 	{
-		int32 NumRemoved = SequenceActors.RemoveAll([=](FMovieSceneSequenceActorPointers& Item) { return Item.SequenceActor == InActor; });
-		ensureMsgf(NumRemoved > 0, TEXT("The given sequence actor wasn't registered"));
+		++LinkerGroups[LinkerIndex].NumClients;
 	}
+
+	checkSlow(LinkerIndex >= 0 && LinkerIndex < uint16(-1));
+
+	TickableClients.Add(FTickableClientData{
+		ClientInterface,
+		InterfaceObject,
+		static_cast<uint16>(LinkerIndex),
+		ResolvedTickInterval.bTickWhenPaused
+	});
 }
 
-void UMovieSceneSequenceTickManager::UnregisterSequenceActor(AActor* InActor, TScriptInterface<IMovieSceneSequenceActor> InActorInterface)
+void UMovieSceneSequenceTickManager::UnregisterTickClient(TScriptInterface<IMovieSceneSequenceTickManagerClient> InTickInterface)
 {
-	if (ensure(InActor && InActorInterface))
+	using namespace UE::MovieScene;
+
+	IMovieSceneSequenceTickManagerClient* ClientInterface = InTickInterface.GetInterface();
+	if (!ensure(ClientInterface))
+	{
+		return;
+	}
+
+	// If PendingActorOperations is non-null, it means we cannot mutate either TickableClients or LinkerGroups so we have to defer until afterwards
+	if (PendingActorOperations)
+	{
+		PendingActorOperations->Add(FPendingOperation{ InTickInterface, InTickInterface.GetObject(), FMovieSceneSequenceTickInterval(),  FPendingOperation::EType::Unregister });
+		return;
+	}
+
+	const int32 ClientIndex = Algo::IndexOfBy(TickableClients, ClientInterface, &FTickableClientData::Interface);
+	if (!ensure(TickableClients.IsValidIndex(ClientIndex)))
 	{
-		int32 NumRemoved = SequenceActors.RemoveAll([=](const FMovieSceneSequenceActorPointers& Item) { return Item.SequenceActor == InActor && Item.SequenceActorInterface == InActorInterface; });
-		ensureMsgf(NumRemoved > 0, TEXT("The given sequence actor wasn't registered"));
+		return;
+	}
+
+	const int32 LinkerIndex = TickableClients[ClientIndex].LinkerIndex;
+	checkSlow(LinkerGroups.IsValidIndex(LinkerIndex));
+
+	TickableClients.RemoveAtSwap(ClientIndex);
+
+	if (--LinkerGroups[LinkerIndex].NumClients == 0)
+	{
+		LinkerGroups.RemoveAt(LinkerIndex);
 	}
 }
 
+// DEPRECATED
+void UMovieSceneSequenceTickManager::RegisterSequenceActor(AActor* InActor)
+{
+	TScriptInterface<IMovieSceneSequenceTickManagerClient> Interface(InActor);
+	check(InActor && Interface);
+
+	FMovieSceneSequenceTickInterval ResolvedTickInterval(InActor);
+	RegisterTickClient(ResolvedTickInterval, Interface);
+}
+// DEPRECATED
+void UMovieSceneSequenceTickManager::RegisterSequenceActor(AActor* InActor, TScriptInterface<IMovieSceneSequenceTickManagerClient> InActorInterface)
+{
+	check(InActor && InActorInterface);
+
+	FMovieSceneSequenceTickInterval ResolvedTickInterval(InActor);
+	RegisterTickClient(ResolvedTickInterval, InActorInterface);
+}
+// DEPRECATED
+void UMovieSceneSequenceTickManager::UnregisterSequenceActor(AActor* InActor)
+{
+	TScriptInterface<IMovieSceneSequenceTickManagerClient> Interface(InActor);
+	check(Interface);
+	UnregisterTickClient(Interface);
+}
+// DEPRECATED
+void UMovieSceneSequenceTickManager::UnregisterSequenceActor(AActor* InActor, TScriptInterface<IMovieSceneSequenceTickManagerClient> InActorInterface)
+{
+	check(InActorInterface);
+	UnregisterTickClient(InActorInterface);
+}
+
 void UMovieSceneSequenceTickManager::TickSequenceActors(float DeltaSeconds)
 {
 	SCOPE_CYCLE_COUNTER(MovieSceneEval_SequenceTickManager);
 
-	TRACE_COUNTER_SET(MovieSceneEval_SequenceTickManagerLatentActionRuns, 0);
-
-	// Let all level sequence actors update. Some of them won't do anything, others will do synchronous
+	// Let all tickable clients update. Some of them won't do anything, others will do synchronous
 	// things (e.g. start/stop, loop, etc.), but in 95% of cases, they will just queue up a normal evaluation
 	// request...
 	//
 	UWorld* World = GetTypedOuter<UWorld>();
-	check(World != nullptr);
+	checkSlow(World != nullptr);
+#if DO_GUARD_SLOW
 	ensure(LatentActionManager.IsEmpty());
-	
-	const bool bIsPaused = World->IsPaused();
+#endif
+
+	const bool  bIsPaused                  = World->IsPaused();
+	const float CurrentUnpausedTimeSeconds = World->GetUnpausedTimeSeconds();
+	const float CurrentTimeSeconds         = World->GetTimeSeconds();
 
-	for (int32 i = SequenceActors.Num() - 1; i >= 0; --i)
+	struct FUpdatedDeltaTimes
 	{
-		const FMovieSceneSequenceActorPointers& Pointers(SequenceActors[i]);
-		if (Pointers.SequenceActor)
+		float UnpausedDeltaTime;
+		float DeltaTime;
+	};
+
+	using InlineSparseAllocator = TSparseArrayAllocator<TInlineAllocator<16>, TInlineAllocator<4>>;
+
+	// Sparse array where an allocated entry maps to the FLinkerGroup within LinkerGroups
+	TSparseArray<FUpdatedDeltaTimes, InlineSparseAllocator> UpdatedDeltaTimes;
+
+	// -----------------------------------------------------------------------------
+	// Step 1: Check the tick intervals of all our linker groups, and mark any whose
+	//         interval has passed for update
+	for (int32 Index = 0; Index < LinkerGroups.GetMaxIndex(); ++Index)
+	{
+		if (LinkerGroups.IsAllocated(Index))
 		{
-			if (!bIsPaused || Pointers.SequenceActor->GetTickableWhenPaused())
+			// Check the delta time against the last update time for this group
+			FLinkerGroup& Group = LinkerGroups[Index];
+
+			float UnpausedDeltaTime = DeltaSeconds;
+			float DeltaTime = DeltaSeconds;
+
+			if (Group.LastUnpausedTimeSeconds >= 0.f)
 			{
-				check(Pointers.SequenceActorInterface);
-				check(Pointers.SequenceActor->GetWorld() == World);
-				Pointers.SequenceActorInterface->TickFromSequenceTickManager(DeltaSeconds);
+				UnpausedDeltaTime = CurrentUnpausedTimeSeconds - Group.LastUnpausedTimeSeconds;
+				DeltaTime = CurrentTimeSeconds - Group.LastTimeSeconds;
+
+				if (UnpausedDeltaTime < Group.RoundedTickIntervalMs * 0.001f)
+				{
+					// If the unpaused time is less than the required tick interval, leave this group alone this frame
+					// We don't need to check the paused delta-time because that will always be >= unpaused
+					continue;
+				}
 			}
+
+			// We know the unpaused delta-time is >= our interval, and thus so will the paused time
+			// Add this delta-time to the sparse array to indicate the group with the corresponding index needs updating
+			Group.LastUnpausedTimeSeconds = CurrentUnpausedTimeSeconds;
+			Group.LastTimeSeconds = CurrentTimeSeconds;
+
+			UpdatedDeltaTimes.Insert(Index, FUpdatedDeltaTimes{ UnpausedDeltaTime, DeltaTime });
 		}
 	}
 
-	// If we have nothing to do, we can early-out.
-	if (!Runner.HasQueuedUpdates())
+	// Skip any work if there are no updates scheduled
+	if (UpdatedDeltaTimes.Num() == 0)
 	{
 		return;
 	}
 
-	// Now we execute all those "normal evaluation requests" we mentioned above. All running level sequences
-	// will therefore be evaluated in a gloriously parallelized way.
-	//
-	if (ensure(Runner.IsAttachedToLinker()))
+	TArray<FPendingOperation> CurrentPendingActorOperations;
 	{
-		Runner.Flush();
-		LatentActionManager.RunLatentActions(Runner);
+		// Guard against any mutation while the entries are being ticked - anything added this tick cycle will have to wait until next tick
+		TGuardValue<TArray<FPendingOperation>*> Guard(PendingActorOperations, &CurrentPendingActorOperations);
+
+		// Step 2: Tick clients for any that need ticking
+		// 
+		for (const FTickableClientData& Client : TickableClients)
+		{
+			const bool bCanTick = (Client.bTickWhenPaused || !bIsPaused) && Client.WeakObject.Get() != nullptr;
+			if (bCanTick && UpdatedDeltaTimes.IsValidIndex(Client.LinkerIndex))
+			{
+				const FUpdatedDeltaTimes Delta = UpdatedDeltaTimes[Client.LinkerIndex];
+				const float DeltaTime = Client.bTickWhenPaused ? Delta.UnpausedDeltaTime : Delta.DeltaTime;
+
+				FMovieSceneEntitySystemRunner* Runner = &LinkerGroups[Client.LinkerIndex].Runner;
+				Client.Interface->TickFromSequenceTickManager(DeltaTime, Runner);
+			}
+		}
+
+		// Step 3: Update and flush linkers as needed
+		//
+		for (int32 Index = 0; Index < UpdatedDeltaTimes.GetMaxIndex(); ++Index)
+		{
+			if (UpdatedDeltaTimes.IsAllocated(Index))
+			{
+				LinkerGroups[Index].Runner.Flush();
+			}
+		}
+	}
+
+	// Process any pending operations that were added while we were updating
+	ProcessPendingOperations(CurrentPendingActorOperations);
+	// Run latent actions now we have finished flushing everything
+	RunLatentActions();
+}
+
+void UMovieSceneSequenceTickManager::ProcessPendingOperations(TArrayView<const FPendingOperation> InOperations)
+{
+	// Process any pending operations that were added while we were updating
+	for (const FPendingOperation& PendingOperation : InOperations)
+	{
+		// Check the object is still alive
+		if (PendingOperation.WeakObject.Get() != nullptr)
+		{
+			if (PendingOperation.Type == FPendingOperation::EType::Register)
+			{
+				RegisterTickClient(PendingOperation.TickInterval, PendingOperation.Interface);
+			}
+			else
+			{
+				UnregisterTickClient(PendingOperation.Interface);
+			}
+		}
 	}
 }
 
@@ -130,29 +367,26 @@ void UMovieSceneSequenceTickManager::AddLatentAction(FMovieSceneSequenceLatentAc
 
 void UMovieSceneSequenceTickManager::RunLatentActions()
 {
-	LatentActionManager.RunLatentActions(Runner);
+	LatentActionManager.RunLatentActions([this]
+	{
+		this->FlushRunners();
+	});
 }
 
-UMovieSceneSequenceTickManager* UMovieSceneSequenceTickManager::Get(UObject* PlaybackContext)
+void UMovieSceneSequenceTickManager::FlushRunners()
 {
-	check(PlaybackContext != nullptr && PlaybackContext->GetWorld() != nullptr);
-	UWorld* World = PlaybackContext->GetWorld();
-
-	UMovieSceneSequenceTickManager* TickManager = FindObject<UMovieSceneSequenceTickManager>(World, TEXT("GlobalMovieSceneSequenceTickManager"));
-	if (!TickManager)
+	TArray<FPendingOperation> CurrentPendingActorOperations;
 	{
-		TickManager = NewObject<UMovieSceneSequenceTickManager>(World, TEXT("GlobalMovieSceneSequenceTickManager"));
-
-		TickManager->Linker = UMovieSceneEntitySystemLinker::FindOrCreateLinker(World, UE::MovieScene::EEntitySystemLinkerRole::LevelSequences, TEXT("MovieSceneSequencePlayerEntityLinker"));
-		check(TickManager->Linker);
-		TickManager->Runner.AttachToLinker(TickManager->Linker);
-
-		FDelegateHandle Handle = World->AddMovieSceneSequenceTickHandler(
-				FOnMovieSceneSequenceTick::FDelegate::CreateUObject(TickManager, &UMovieSceneSequenceTickManager::TickSequenceActors));
-		check(Handle.IsValid());
-		TickManager->WorldTickDelegateHandle = Handle;
+		// Guard against any mutation while the entries are being ticked - anything added this tick cycle will have to wait until next tick
+		TGuardValue<TArray<FPendingOperation>*> Guard(PendingActorOperations, &CurrentPendingActorOperations);
+		for (FLinkerGroup& LinkerGroup : LinkerGroups)
+		{
+			LinkerGroup.Runner.Flush();
+		}
 	}
-	return TickManager;
+
+	// Process any pending operations that were added while we were updating
+	ProcessPendingOperations(CurrentPendingActorOperations);
 }
 
 void FMovieSceneLatentActionManager::AddLatentAction(FMovieSceneSequenceLatentActionDelegate Delegate)
@@ -184,7 +418,7 @@ void FMovieSceneLatentActionManager::ClearLatentActions()
 	}
 }
 
-void FMovieSceneLatentActionManager::RunLatentActions(FMovieSceneEntitySystemRunner& Runner)
+void FMovieSceneLatentActionManager::RunLatentActions(TFunctionRef<void()> FlushCallback)
 {
 	if (bIsRunningLatentActions)
 	{
@@ -196,11 +430,9 @@ void FMovieSceneLatentActionManager::RunLatentActions(FMovieSceneEntitySystemRun
 
 	TGuardValue<bool> IsRunningLatentActionsGuard(bIsRunningLatentActions, true);
 
-	int32 NumLoopsLeft = CVarMovieSceneMaxLatentActionLoops.GetValueOnGameThread();
+	int32 NumLoopsLeft = UE::MovieScene::GMovieSceneMaxLatentActionLoops;
 	while (LatentActions.Num() > 0)
 	{
-		TRACE_COUNTER_INCREMENT(MovieSceneEval_SequenceTickManagerLatentActionRuns);
-
 		// We can run *one* latent action per sequence player before having to flush the linker again.
 		// This way, if we have 42 sequence players with 2 latent actions each, we only flush the linker
 		// twice, as opposed to 42*2=84 times.
@@ -228,7 +460,7 @@ void FMovieSceneLatentActionManager::RunLatentActions(FMovieSceneEntitySystemRun
 			}
 		}
 
-		Runner.Flush();
+		FlushCallback();
 
 		--NumLoopsLeft;
 		if (!ensureMsgf(NumLoopsLeft > 0,
diff --git a/Engine/Source/Runtime/MovieScene/Public/MovieSceneSequencePlayer.h b/Engine/Source/Runtime/MovieScene/Public/MovieSceneSequencePlayer.h
--- a/Engine/Source/Runtime/MovieScene/Public/MovieSceneSequencePlayer.h
+++ b/Engine/Source/Runtime/MovieScene/Public/MovieSceneSequencePlayer.h
@@ -15,6 +15,7 @@
 #include "Evaluation/MovieScenePlayback.h"
 #include "MovieSceneSequenceTickManager.h"
 #include "Evaluation/MovieScenePlayback.h"
+#include "MovieSceneSequenceTickManager.h"
 #include "MovieSceneSequencePlayer.generated.h"
 
 class UMovieSceneSequenceTickManager;
@@ -109,6 +110,7 @@ struct FMovieSceneSequencePlaybackSettings
 		, bHideHud(false)
 		, bDisableCameraCuts(false)
 		, bPauseAtEnd(false)
+		, bInheritTickIntervalFromOwner(true)
 	{ }
 
 	GENERATED_BODY()
@@ -121,6 +123,10 @@ struct FMovieSceneSequencePlaybackSettings
 	UPROPERTY(EditAnywhere, BlueprintReadWrite, Category="Playback", meta=(UIMin=1, DisplayName="Loop"))
 	FMovieSceneSequenceLoopCount LoopCount;
 
+	/** Overridable tick interval for this sequence to update at. When not overridden, the owning actor or component's tick interval will be used */
+	UPROPERTY(EditAnywhere, Category="Playback", meta=(ShowOnlyInnerProperties, Units=s, EditCondition="!bInheritTickIntervalFromOwner"))
+	FMovieSceneSequenceTickInterval TickInterval;
+
 	/** The rate at which to playback the animation */
 	UPROPERTY(EditAnywhere, BlueprintReadWrite, Category="Playback", meta=(Units=Multiplier))
 	float PlayRate;
@@ -161,6 +167,10 @@ struct FMovieSceneSequencePlaybackSettings
 	UPROPERTY(EditAnywhere, BlueprintReadWrite, Category="Playback")
 	uint32 bPauseAtEnd : 1;
 
+	/** When checked, a custom tick interval can be provided to define how often to update this sequence */
+	UPROPERTY(EditAnywhere, Category="Playback", meta=(InlineEditConditionToggle))
+	uint32 bInheritTickIntervalFromOwner : 1;
+
 	MOVIESCENE_API bool SerializeFromMismatchedTag(const FPropertyTag& Tag, FStructuredArchive::FSlot Slot);
 };
 
@@ -250,6 +260,7 @@ UCLASS(Abstract, BlueprintType)
 class MOVIESCENE_API UMovieSceneSequencePlayer
 	: public UObject
 	, public IMovieScenePlayer
+	, public IMovieSceneSequenceTickManagerClient
 {
 public:
 	GENERATED_BODY()
@@ -508,6 +519,9 @@ public:
 	/** Update the sequence for the current time, if playing, asynchronously */
 	void UpdateAsync(const float DeltaSeconds);
 
+	/** Perform any tear-down work when this player is no longer (and will never) be needed */
+	void TearDown();
+
 public:
 
 	/**
@@ -604,6 +618,10 @@ protected:
 	virtual void BeginDestroy() override;
 	/*~ End UObject interface */
 
+	//~ Begin IMovieSceneSequenceTickManagerClient interface
+	virtual void TickFromSequenceTickManager(float DeltaSeconds, FMovieSceneEntitySystemRunner* Runner) override;
+	//~ End IMovieSceneSequenceTickManagerClient interface
+
 protected:
 
 	virtual bool CanPlay() const { return true; }
@@ -705,6 +723,9 @@ protected:
 	UPROPERTY(transient)
 	FMovieSceneRootEvaluationTemplateInstance RootTemplateInstance;
 
+	/** Usually nullptr, but will be set when we are updating inside a TickFromSequenceTickManager call */
+	FMovieSceneEntitySystemRunner* CurrentRunner;
+
 	/** Play position helper */
 	FMovieScenePlaybackPosition PlayPosition;
 
@@ -750,6 +771,9 @@ private:
 	/** The event that will be broadcast every time the sequence is updated */
 	mutable FOnMovieSceneSequencePlayerUpdated OnMovieSceneSequencePlayerUpdate;
 
+	/** The tick interval we are currently registered with (if any) */
+	TOptional<FMovieSceneSequenceTickInterval> RegisteredTickInterval;
+
 	/** The maximum tick rate prior to playing (used for overriding delta time during playback). */
 	TOptional<double> OldMaxTickRate;
 
diff --git a/Engine/Source/Runtime/MovieScene/Public/MovieSceneSequenceTickInterval.h b/Engine/Source/Runtime/MovieScene/Public/MovieSceneSequenceTickInterval.h
new file mode 100644
--- /dev/null
+++ b/Engine/Source/Runtime/MovieScene/Public/MovieSceneSequenceTickInterval.h
@@ -0,0 +1,76 @@
+// Copyright Epic Games, Inc. All Rights Reserved.
+
+#pragma once
+
+#include "CoreTypes.h"
+
+#include "MovieSceneSequenceTickInterval.generated.h"
+
+struct FMovieSceneSequenceTickInterval;
+
+
+/**
+ * Structure defining a concrete tick interval for a Sequencer based evaluation
+ */
+USTRUCT(BlueprintType)
+struct MOVIESCENE_API FMovieSceneSequenceTickInterval
+{
+	GENERATED_BODY()
+
+	FMovieSceneSequenceTickInterval() = default;
+
+	/**
+	 * Generate a tick interval from an actor's primary tick function
+	 */
+	FMovieSceneSequenceTickInterval(const AActor* InActor);
+
+	/**
+	 * Generate a tick interval from an component's primary tick function
+	 */
+	FMovieSceneSequenceTickInterval(const UActorComponent* InActorComponent);
+
+public:
+
+	/** Defines the rate at which the sequence should update, in seconds */
+	UPROPERTY(EditAnywhere, Category="Playback", meta=(DisplayName="Tick Interval", Units=s))
+	float TickIntervalSeconds = 0.f;
+
+	/** When true, the sequence will continue to tick and progress even when the world is paused */
+	UPROPERTY(EditAnywhere, Category="Playback")
+	bool bTickWhenPaused = false;
+
+	/** When true, allow the sequence to be grouped with other sequences based on Sequencer.TickIntervalGroupingResolutionMs. Otherwise the interval will be used precisely. */
+	UPROPERTY(EditAnywhere, Category="Playback")
+	bool bAllowRounding = true;
+
+public:
+
+	/**
+	 * Round this interval to the nearest Sequencer.TickIntervalGroupingResolutionMs milliseconds
+	 */
+	int32 RoundTickIntervalMs() const;
+
+	/**
+	 * Resolve this tick interval within the specified context object (usually a movie scene player)
+	 * inheriting properties from the first valid parent if possible
+	 */
+	static FMovieSceneSequenceTickInterval GetInheritedInterval(UObject* ContextObject);
+
+	/**
+	 * Equality comparison operator
+	 */
+	friend bool operator==(const FMovieSceneSequenceTickInterval& A, const FMovieSceneSequenceTickInterval& B)
+	{
+		return A.TickIntervalSeconds == B.TickIntervalSeconds
+			&& A.bTickWhenPaused == B.bTickWhenPaused
+			&& A.bAllowRounding == B.bAllowRounding;
+	}
+
+	/**
+	 * Inequality comparison operator
+	 */
+	friend bool operator!=(const FMovieSceneSequenceTickInterval& A, const FMovieSceneSequenceTickInterval& B)
+	{
+		return !(A == B);
+	}
+};
\ No newline at end of file
diff --git a/Engine/Source/Runtime/MovieScene/Public/MovieSceneSequenceTickManager.h b/Engine/Source/Runtime/MovieScene/Public/MovieSceneSequenceTickManager.h
--- a/Engine/Source/Runtime/MovieScene/Public/MovieSceneSequenceTickManager.h
+++ b/Engine/Source/Runtime/MovieScene/Public/MovieSceneSequenceTickManager.h
@@ -6,8 +6,10 @@
 #include "UObject/ObjectMacros.h"
 #include "UObject/Object.h"
 #include "EntitySystem/MovieSceneEntitySystemRunner.h"
+#include "MovieSceneSequenceTickInterval.h"
 #include "MovieSceneSequenceTickManager.generated.h"
 
+
 class UMovieSceneEntitySystemLinker;
 
 DECLARE_DELEGATE(FMovieSceneSequenceLatentActionDelegate);
@@ -22,7 +24,7 @@ public:
 	void ClearLatentActions(UObject* Object);
 	void ClearLatentActions();
 
-	void RunLatentActions(FMovieSceneEntitySystemRunner& Runner);
+	void RunLatentActions(TFunctionRef<void()> FlushCallback);
 
 	bool IsEmpty() const { return LatentActions.Num() == 0; }
 
@@ -33,41 +35,34 @@ private:
 };
 
 /**
- * Interface for sequence actors that are to be ticked by the tick manager.
+ * Interface for objects that are to be ticked by the tick manager.
  */
 UINTERFACE()
-class MOVIESCENE_API UMovieSceneSequenceActor
+class MOVIESCENE_API UMovieSceneSequenceTickManagerClient
 	: public UInterface
 {
 public:
 	GENERATED_BODY()
 };
 
-class MOVIESCENE_API IMovieSceneSequenceActor
+class MOVIESCENE_API IMovieSceneSequenceTickManagerClient
 {
 public:
 	GENERATED_BODY()
 
-	virtual void TickFromSequenceTickManager(float DeltaSeconds) = 0;
+	virtual void TickFromSequenceTickManager(float DeltaSeconds, FMovieSceneEntitySystemRunner* Runner) = 0;
 };
 
-/**
- * A structure for storing pointers to a sequence actor.
- */
-USTRUCT()
-struct FMovieSceneSequenceActorPointers
-{
-	GENERATED_BODY()
-
-	UPROPERTY()
-	TObjectPtr<AActor> SequenceActor = nullptr;
-
-	UPROPERTY()
-	TScriptInterface<IMovieSceneSequenceActor> SequenceActorInterface;
-};
 
 /**
- * An automatically created global object that will manage all level sequence actors' updates.
+ * Global (one per-UWorld) manager object that manages ticking and updating any and all Sequencer-based
+ * evaluations for the current frame, before any other actors are ticked.
+ *
+ * Ticking clients are registered based on their desired tick interval, and grouped together with other
+ * clients that tick with the same interval (based on Sequencer.TickIntervalGroupingResolutionMs).
+ * 
+ * Sequencer data is shared between all instances within the same group, allowing them to blend together.
+ * Clients ticking at different intervals do not support blending with each other.
  */
 UCLASS()
 class MOVIESCENE_API UMovieSceneSequenceTickManager : public UObject
@@ -77,37 +72,172 @@ public:
 
 	UMovieSceneSequenceTickManager(const FObjectInitializer& Init);
 
-	UMovieSceneEntitySystemLinker* GetLinker() { return Linker; }
-	FMovieSceneEntitySystemRunner& GetRunner() { return Runner; }
-
-	void RegisterSequenceActor(AActor* InActor);
-	void RegisterSequenceActor(AActor* InActor, TScriptInterface<IMovieSceneSequenceActor> InActorInterface);
-	void UnregisterSequenceActor(AActor* InActor);
-	void UnregisterSequenceActor(AActor* InActor, TScriptInterface<IMovieSceneSequenceActor> InActorInterface);
+	/**
+	 * Retrieve the tick manager for the specified playback context
+	 */
+	static UMovieSceneSequenceTickManager* Get(UObject* PlaybackContext);
 
+	/**
+	 * Register a new client to receive a tick at the specified tick interval.
+	 * This client will be grouped with all other clients that tick within Sequencer.TickIntervalGroupingResolutionMs of the specified interval.
+	 *
+	 * @param TickInterval    The interval at which to tick the client. 
+	 * @param InTickInterface The client that will receieve a tick
+	 */
+	void RegisterTickClient(const FMovieSceneSequenceTickInterval& TickInterval, TScriptInterface<IMovieSceneSequenceTickManagerClient> InTickInterface);
+
+	/**
+	 * Unregister a previously registered client. The client must have been registered or an assertion will fail.
+	 *
+	 * @param InTickInterface The client to unregister
+	 */
+	void UnregisterTickClient(TScriptInterface<IMovieSceneSequenceTickManagerClient> InTickInterface);
+
+	/**
+	 * Retrieve the linker associated with the specified (resolved) tick interval.
+	 * @return The linker that owns all the data for the specified tick interval, or nullptr if there is none
+	 */
+	UMovieSceneEntitySystemLinker* GetLinker(const FMovieSceneSequenceTickInterval& TickInterval);
+
+	/**
+	 * Retrieve the runner associated with the specified (resolved) tick interval.
+	 * @return The runner that flushes evaluations for the specified tick interval, or nullptr if there is none
+	 */
+	FMovieSceneEntitySystemRunner* GetRunner(const FMovieSceneSequenceTickInterval& TickInterval);
+
+	/**
+	 * Add an action that will be executed once all the pending evaluations in this tick manager have been flushed.
+	 * @note Latent actions are destroyed once they are called.
+	 *
+	 * @param Delegate A delegate to invoke when all evaluations have been flushed.
+	 */
 	void AddLatentAction(FMovieSceneSequenceLatentActionDelegate Delegate);
+
+	/**
+	 * Execute any pending latent actions
+	 */
 	void RunLatentActions();
 
+	/**
+	 * Clear latent actions pertaining to the specified object
+	 */
 	void ClearLatentActions(UObject* Object);
 
-	static UMovieSceneSequenceTickManager* Get(UObject* PlaybackContext);
+public:
+
+	/*~ -------------------------------------------------------------- */
+	/*~ Deprecated functions that previously used 'Actor' nomenclature */
+	UE_DEPRECATED(5.1, "Please use RegisterTickClient")
+	void RegisterSequenceActor(AActor* InActor);
+	UE_DEPRECATED(5.1, "Please use RegisterTickClient")
+	void RegisterSequenceActor(AActor* InActor, TScriptInterface<IMovieSceneSequenceTickManagerClient> InActorInterface);
+	UE_DEPRECATED(5.1, "Please use UnregisterTickClient")
+	void UnregisterSequenceActor(AActor* InActor);
+	UE_DEPRECATED(5.1, "Please use UnregisterTickClient")
+	void UnregisterSequenceActor(AActor* InActor, TScriptInterface<IMovieSceneSequenceTickManagerClient> InActorInterface);
+	/*~ -------------------------------------------------------------- */
 
 private:
+
+	/*~ UObject interface */
 	virtual void BeginDestroy() override;
 
-	void TickSequenceActors(float DeltaSeconds);
+private:
+
+	/** 24 bytes: Information pertaining to a tickable client */
+	struct FTickableClientData
+	{
+		/** 8 bytes: The actual interface ptr. Should only be used after validating against WeakObject. */
+		IMovieSceneSequenceTickManagerClient* Interface;
+		/** 8 bytes: A weak object ptr for the UObject implementing the interface. Only required for lifetime validation. */
+		TWeakObjectPtr<> WeakObject;
+		/** 4 bytes: The index into LinkerGroups for this client's linker (derived from its tick interval) */
+		int32 LinkerIndex;
+		/** 1 byte: Whether this client should tick when it is paused or not. Note that evaluation data will still be evaluated every frame, but the sequence will not progress. */
+		bool bTickWhenPaused;
+	};
+
+	/** Information pertaining to a group of tickable clients */
+	struct FLinkerGroup
+	{
+		/** The linker that owns all the entity data, systems and instances */
+		UMovieSceneEntitySystemLinker* Linker;
+		/** Runner responsible for evaluating each phase of the pipeline */
+		FMovieSceneEntitySystemRunner Runner;
+
+		/** The tick interval of this group, rounded to the nearest Sequencer.TickIntervalGroupingResolutionMs */
+		int32 RoundedTickIntervalMs = 0;
+
+		/** The value of UWorld::GetUnpausedTimeSeconds last time this group was evaluated */
+		float LastUnpausedTimeSeconds = -1;
+		/** The value of UWorld::GetTimeSeconds last time this group was evaluated */
+		float LastTimeSeconds = -1;
+
+		/** Total number of clients in this group */
+		uint16 NumClients = 0;
+	};
+
+	/** A pending registration operation */
+	struct FPendingOperation
+	{
+		enum class EType
+		{
+			Register, Unregister
+		};
+		/** The interface to (un)register */
+		TScriptInterface<IMovieSceneSequenceTickManagerClient> Interface;
+		/** The object that we're registering (only used to ensure the interface is still valid) */
+		TWeakObjectPtr<UObject> WeakObject;
+		/** Tick interval for registration (not used for UnregisterTickClient calls) */
+		FMovieSceneSequenceTickInterval TickInterval;
+		/** The type of operation (Register or Unregister) */
+		EType Type;
+	};
 
 private:
 
-	UPROPERTY(transient)
-	TArray<FMovieSceneSequenceActorPointers> SequenceActors;
+	/**
+	 * Main update loop called in response to a world tick
+	 */
+	void TickSequenceActors(float DeltaSeconds);
+
+	/**
+	 * Flush all runners maintained by this tick manager.
+	 * Should only be called when running latent actions.
+	 */
+	void FlushRunners();
 
-	UPROPERTY(transient)
-	TObjectPtr<UMovieSceneEntitySystemLinker> Linker;
+	/**
+	 * Process any pending registration operations
+	 */
+	void ProcessPendingOperations(TArrayView<const FPendingOperation> InOperations);
 
-	FMovieSceneEntitySystemRunner Runner;
+private:
 
+	/** Flat array of tickable client data with indices into LinkerGroups */
+	TArray<FTickableClientData> TickableClients;
+
+	/** Sparse array of group data - indices into this array are stable */
+	TSparseArray<FLinkerGroup> LinkerGroups;
+
+	/** Delegate handle for the world tick delegate */
 	FDelegateHandle WorldTickDelegateHandle;
 
+	/** Pointer to an array of pending operations - nullptr everywhere except where mutations to TickableClients or LinkerGroups is disallowed */
+	TArray<FPendingOperation>* PendingActorOperations;
+
+	/** Latent action manager for managing actions that need performing after evaluation */
 	FMovieSceneLatentActionManager LatentActionManager;
 };
+
+
+/** Deprecated IMovieSceneSequenceActor interface that has been replaced with IMovieSceneSequenceTickManagerClient */
+class UE_DEPRECATED(5.1, "Please use IMovieSceneSequenceTickManagerClient directly") IMovieSceneSequenceActor
+	: public IMovieSceneSequenceTickManagerClient
+{
+	virtual void TickFromSequenceTickManager(float DeltaSeconds) = 0;
+	virtual void TickFromSequenceTickManager(float DeltaSeconds, FMovieSceneEntitySystemRunner* Runner) override
+	{
+		TickFromSequenceTickManager(DeltaSeconds);
+	}
+};
diff --git a/Engine/Source/Runtime/UMG/Private/Animation/UMGSequenceTickManager.cpp b/Engine/Source/Runtime/UMG/Private/Animation/UMGSequenceTickManager.cpp
--- a/Engine/Source/Runtime/UMG/Private/Animation/UMGSequenceTickManager.cpp
+++ b/Engine/Source/Runtime/UMG/Private/Animation/UMGSequenceTickManager.cpp
@@ -199,7 +199,7 @@ void UUMGSequenceTickManager::ForceFlush()
 	if (Runner.IsAttachedToLinker())
 	{
 		Runner.Flush();
-		LatentActionManager.RunLatentActions(Runner);
+		RunLatentActions();
 	}
 }
 
@@ -216,7 +216,7 @@ void UUMGSequenceTickManager::HandleSlatePostTick(float DeltaSeconds)
 		SCOPE_CYCLE_COUNTER(MovieSceneEval_FlushEndOfFrameAnimations);
 
 		Runner.Flush();
-		LatentActionManager.RunLatentActions(Runner);
+		RunLatentActions();
 	}
 }
 
@@ -232,7 +232,10 @@ void UUMGSequenceTickManager::ClearLatentActions(UObject* Object)
 
 void UUMGSequenceTickManager::RunLatentActions()
 {
-	LatentActionManager.RunLatentActions(Runner);
+	LatentActionManager.RunLatentActions([this]
+	{
+		this->Runner.Flush();
+	});
 }
 
 UUMGSequenceTickManager* UUMGSequenceTickManager::Get(UObject* PlaybackContext)
