Fixed fwd declarations

diff --git a/Engine/Source/Runtime/MovieScene/Public/MovieSceneSequenceTickInterval.h b/Engine/Source/Runtime/MovieScene/Public/MovieSceneSequenceTickInterval.h
--- a/Engine/Source/Runtime/MovieScene/Public/MovieSceneSequenceTickInterval.h
+++ b/Engine/Source/Runtime/MovieScene/Public/MovieSceneSequenceTickInterval.h
@@ -6,8 +6,8 @@
 
 #include "MovieSceneSequenceTickInterval.generated.h"
 
-struct FMovieSceneSequenceTickInterval;
-
+class AActor;
+class UActorComponent;
 
 /**
  * Structure defining a concrete tick interval for a Sequencer based evaluation
