Sequencer: Added the ability to set per-frame budgets for sequencer evaluations.

Setting a budget will cause evaluation to become time-sliced across multiple frames, as soon as the supplied budget is exceeded.
By default this is turned off, but can be overridden by specifying an Evaluation Budget on the playback settings.
Mismatches of budgets will cause the smallest budget to be used, but will emit a warning if budgeted and non-budgeted evaluations are defined for the same tick interval.

Implementation details:

The evaluation loop is set up as a crude state machine which allows us to break out of the
loop and pick up where we left off next frame. Each bit within the FlushState bitmask defines
a particular task that needs to execute. In this way we can loop the evaluation while there
are updates pending by simply setting the necessary flags for each task. For convenience, the
required tasks to perform another iteration are defined in ERunnerFlushState::LoopEval.

diff --git a/Engine/Source/Runtime/MovieScene/Private/EntitySystem/MovieSceneEntitySystemLinker.cpp b/Engine/Source/Runtime/MovieScene/Private/EntitySystem/MovieSceneEntitySystemLinker.cpp
--- a/Engine/Source/Runtime/MovieScene/Private/EntitySystem/MovieSceneEntitySystemLinker.cpp
+++ b/Engine/Source/Runtime/MovieScene/Private/EntitySystem/MovieSceneEntitySystemLinker.cpp
@@ -66,6 +66,7 @@ UMovieSceneEntitySystemLinker::UMovieSceneEntitySystemLinker(const FObjectInitia
 	{
 		PreAnimatedState.Initialize(this);
 
+		FCoreUObjectDelegates::GetPreGarbageCollectDelegate().AddUObject(this, &UMovieSceneEntitySystemLinker::HandlePreGarbageCollection);
 		FCoreUObjectDelegates::GetPostGarbageCollect().AddUObject(this, &UMovieSceneEntitySystemLinker::HandlePostGarbageCollection);
 
 		EntityManager.SetDebugName(GetName() + TEXT("[Entity Manager]"));
@@ -264,6 +265,21 @@ void UMovieSceneEntitySystemLinker::EndEvaluation(FMovieSceneEntitySystemRunner&
 	}
 }
 
+void UMovieSceneEntitySystemLinker::HandlePreGarbageCollection()
+{
+	using namespace UE::MovieScene;
+
+	// If we have any active runners part-way through an evaluation,
+	// they must be flushed before we run a garbage collection
+	for (int32 Index = ActiveRunners.Num()-1; Index >= 0; --Index)
+	{
+		if (ActiveRunners[Index].Runner->IsAttachedToLinker())
+		{
+			ActiveRunners[Index].Runner->Flush();
+		}
+	}
+}
+
 void UMovieSceneEntitySystemLinker::HandlePostGarbageCollection()
 {
 	using namespace UE::MovieScene;
@@ -353,6 +369,16 @@ void UMovieSceneEntitySystemLinker::OnObjectsReplaced(const TMap<UObject*, UObje
 
 void UMovieSceneEntitySystemLinker::OnWorldCleanup(UWorld* InWorld, bool bSessionEnded, bool bCleanupResources)
 {
+	// If we have any active runners part-way through an evaluation,
+	// they must be flushed before we cleanup
+	for (int32 Index = ActiveRunners.Num()-1; Index >= 0; --Index)
+	{
+		if (ActiveRunners[Index].Runner->IsAttachedToLinker())
+		{
+			ActiveRunners[Index].Runner->Flush();
+		}
+	}
+
 	Events.CleanUpWorld.Broadcast(this, InWorld);
 	InstanceRegistry->WorldCleanup(InWorld);
 
diff --git a/Engine/Source/Runtime/MovieScene/Private/EntitySystem/MovieSceneEntitySystemRunner.cpp b/Engine/Source/Runtime/MovieScene/Private/EntitySystem/MovieSceneEntitySystemRunner.cpp
--- a/Engine/Source/Runtime/MovieScene/Private/EntitySystem/MovieSceneEntitySystemRunner.cpp
+++ b/Engine/Source/Runtime/MovieScene/Private/EntitySystem/MovieSceneEntitySystemRunner.cpp
@@ -7,7 +7,7 @@
 #include "Evaluation/PreAnimatedState/MovieScenePreAnimatedCaptureSource.h"
 #include "ProfilingDebugging/CountersTrace.h"
 
-DECLARE_CYCLE_STAT(TEXT("ECS System Cost"), 			MovieSceneEval_TotalGTCost, 				STATGROUP_MovieSceneEval);
+DECLARE_CYCLE_STAT(TEXT("Runner Flush"), 				MovieSceneEval_RunnerFlush, 				STATGROUP_MovieSceneEval);
 
 DECLARE_CYCLE_STAT(TEXT("Spawn Phase"),                 MovieSceneEval_SpawnPhase,              	STATGROUP_MovieSceneECS);
 DECLARE_CYCLE_STAT(TEXT("Instantiation Phase"), 		MovieSceneEval_InstantiationPhase, 			STATGROUP_MovieSceneECS);
@@ -18,13 +18,10 @@ DECLARE_CYCLE_STAT(TEXT("Evaluation Phase"), 			MovieSceneEval_EvaluationPhase,
 DECLARE_CYCLE_STAT(TEXT("Finalization Phase"),          MovieSceneEval_FinalizationPhase,       	STATGROUP_MovieSceneECS);
 DECLARE_CYCLE_STAT(TEXT("Post Evaluation Phase"),       MovieSceneEval_PostEvaluationPhase,     	STATGROUP_MovieSceneECS);
 
-TRACE_DECLARE_INT_COUNTER(MovieSceneEntitySystemFlushes, TEXT("MovieScene/ECSFlushes"));
-TRACE_DECLARE_INT_COUNTER(MovieSceneEntitySystemEvaluations, TEXT("MovieScene/ECSEvaluations"));
-
 FMovieSceneEntitySystemRunner::FMovieSceneEntitySystemRunner()
-	: CompletionTask(nullptr)
-	, GameThread(ENamedThreads::GameThread_Local)
+	: GameThread(ENamedThreads::GameThread_Local)
 	, CurrentPhase(UE::MovieScene::ESystemPhase::None)
+	, FlushState(UE::MovieScene::ERunnerFlushState::None)
 {
 }
 
@@ -110,6 +107,11 @@ bool FMovieSceneEntitySystemRunner::HasQueuedUpdates() const
 	return false;
 }
 
+bool FMovieSceneEntitySystemRunner::IsCurrentlyEvaluating() const
+{
+	return FlushState != UE::MovieScene::ERunnerFlushState::None;
+}
+
 bool FMovieSceneEntitySystemRunner::HasQueuedUpdates(FInstanceHandle InInstanceHandle) const
 {
 	return Algo::FindBy(UpdateQueue, InInstanceHandle, &FMovieSceneUpdateRequest::InstanceHandle) != nullptr ||
@@ -133,28 +135,17 @@ void FMovieSceneEntitySystemRunner::Update(const FMovieSceneContext& Context, FI
 	Flush();
 }
 
-void FMovieSceneEntitySystemRunner::Flush()
+bool FMovieSceneEntitySystemRunner::StartEvaluation(UMovieSceneEntitySystemLinker* Linker)
 {
 	using namespace UE::MovieScene;
 
-	UMovieSceneEntitySystemLinker* Linker = GetLinker();
-
-	// Check that we are attached to a linker that allows starting a new evaluation.
-	if (!ensureMsgf(Linker, TEXT("Runner isn't attached to a valid linker")))
-	{
-		return;
-	}
-
 	if (!Linker->StartEvaluation(*this))
 	{
-		return;
+		return false;
 	}
 
-	SCOPE_CYCLE_COUNTER(MovieSceneEval_TotalGTCost);
-	TRACE_COUNTER_INCREMENT(MovieSceneEntitySystemFlushes);
-
-	// We need to run the system from the game thread so we know we can fire events and callbacks from here.
-	check(IsInGameThread());
+	// Don't run start again until explicitly told to
+	FlushState &= ~ERunnerFlushState::Start;
 
 	// Our entity manager cannot be locked down for us to continue. Something must have left it locked if 
 	// this check fails
@@ -166,68 +157,172 @@ void FMovieSceneEntitySystemRunner::Flush()
 
 	// We specifically only check whether the entity manager has changed since the last instantation once
 	// to ensure that we are not vulnerable to infinite loops where components are added/removed in post-evaluation
-	bool bStructureHadChanged = Linker->HasStructureChangedSinceLastRun();
+	const bool bStructureHadChanged = Linker->HasStructureChangedSinceLastRun();
 
-	// Start flushing the update queue... keep flushing as long as we have work to do.
-	while (UpdateQueue.Num() > 0 ||
-		DissectedUpdates.Num() > 0 ||
-		bStructureHadChanged)
+	if (!bStructureHadChanged && UpdateQueue.Num() == 0 && DissectedUpdates.Num() == 0)
 	{
-		DoFlushUpdateQueueOnce();
-
-		bStructureHadChanged = false;
+		// If nothing has changed, and we have no pending updates, we don't run any evaluation
+		// Just jump straight to the End to ensure we clean up our runner correctly
+		FlushState = ERunnerFlushState::End;
 	}
 
+	return true;
+}
+
+void FMovieSceneEntitySystemRunner::EndEvaluation(UMovieSceneEntitySystemLinker* Linker)
+{
+	using namespace UE::MovieScene;
+
 	Linker->EndEvaluation(*this);
+	FlushState &= ~ERunnerFlushState::End;
 }
 
-void FMovieSceneEntitySystemRunner::DoFlushUpdateQueueOnce()
+bool FMovieSceneEntitySystemRunner::FlushNext(UMovieSceneEntitySystemLinker* Linker)
 {
 	using namespace UE::MovieScene;
 
-	TRACE_COUNTER_INCREMENT(MovieSceneEntitySystemEvaluations);
-	TRACE_CPUPROFILER_EVENT_SCOPE(FMovieSceneEntitySystemRunner::DoFlushUpdateQueueOnce);
+	// The evaluation loop is set up as a crude state machine which allows us to break out of the
+	// loop and pick up where we left off next frame. Each bit within the FlushState bitmask defines
+	// a particular task that needs to execute. In this way we can loop the evaluation while there
+	// are updates pending by simply setting the necessary flags for each task. For convenience, the
+	// required tasks to perform another iteration are defined in ERunnerFlushState::LoopEval
 
-	// Setup the completion task that we can wait on.
-	CompletionTask = TGraphTask<FNullGraphTask>::CreateTask(nullptr, ENamedThreads::GameThread)
-		.ConstructAndHold(TStatId(), ENamedThreads::GameThread_Local);
+	// Step 1: Initialize all external flags and broadcast 'begin eval' events.
+	//         NOTE: This is only ever called once, regardless of how many iterations we perform
+	if (EnumHasAnyFlags(FlushState, ERunnerFlushState::Start))
+	{
+		return StartEvaluation(Linker);
+	}
 
-	// Set the debug vizualizer's entity manager pointer, so all debugging happening here will show
-	// relevant information. We need to set it here instead of higher up because we could have, say,
-	// a blocking sequence triggering another blocking sequence via an event track. The nested call stack
-	// of the second sequence needs to show debug information relevant to its private linker, but when
-	// we return back up to the first sequence, which might still have another update round (such as
-	// the other side of the dissected update range around the event), we need to set the pointer back
-	// again.
-	TGuardValue<FEntityManager*> DebugVizGuard(GEntityManagerForDebuggingVisualizers, GetEntityManager());
+	// Step 2: Update sequence instances and import entities
+	//         This is the entry-point for for each iteration
+	if (EnumHasAnyFlags(FlushState, ERunnerFlushState::Import))
+	{
+		return GameThread_UpdateSequenceInstances(Linker);
+	}
 
-	// Also reset the capture source scope so that each group of sequences tied to a given linker starts
-	// with a clean slate.
-	TGuardValue<FScopedPreAnimatedCaptureSource*> CaptureSourceGuard(FScopedPreAnimatedCaptureSource::GetCaptureSourcePtr(), nullptr);
+	// Step 3: Conditionally run the spawn phase of the system graph
+	//
+	if (EnumHasAnyFlags(FlushState, ERunnerFlushState::Spawn))
+	{
+		return GameThread_SpawnPhase(Linker);
+	}
+
+	// Step 4: Conditionally run the instantiation phase of the system graph
+	//
+	if (EnumHasAnyFlags(FlushState, ERunnerFlushState::Instantiation))
+	{
+		return GameThread_InstantiationPhase(Linker);
+	}
 
-	// Entry point to the whole ECS loop... this will either unroll in the current thread's call stack
-	// if there's not much to do, or it will start queuing up tasks on the task graph.
-	// We immediately wait for the completion task to be executed.
+	// Step 5: Run the evaluation phase of the system graph
+	//
+	if (EnumHasAnyFlags(FlushState, ERunnerFlushState::Evaluation))
 	{
-		GameThread_ProcessQueue();
+		return GameThread_EvaluationPhase(Linker);
 	}
 
-	FTaskGraphInterface::Get().WaitUntilTaskCompletes(CompletionTask->GetCompletionEvent(), ENamedThreads::GameThread_Local);
+	// Step 6: Run the finalization phase of the system graph, including legacy templates
+	//
+	if (EnumHasAnyFlags(FlushState, ERunnerFlushState::Finalization))
+	{
+		// Now run the post-evaluation logic, which contains stuff we don't want to run from inside a task
+		// graph call.
+		GameThread_EvaluationFinalizationPhase(Linker);
+		GameThread_PostEvaluationPhase(Linker);
 
-	CompletionTask = nullptr;
+		FlushState &= ~ERunnerFlushState::Finalization;
 
-	// Now run the post-evaluation logic, which contains stuff we don't want to run from inside a task
-	// graph call.
-	GameThread_EvaluationFinalizationPhase();
-	GameThread_PostEvaluationPhase();
+		// If we have any pending updates, we need to run another evaluation.
+		// This will cause us to effectively loop back to ERunnerFlushState::Import
+		// which will pick up the next updates
+		if (UpdateQueue.Num() > 0 || DissectedUpdates.Num() > 0)
+		{
+			FlushState = ERunnerFlushState::LoopEval;
+		}
+		return true;
+	}
+
+	// Step 7: Perform any clean up and broadcast 'end eval' events
+	//         NOTE: Only ever called once regardless of how many iterations we perform
+	if (EnumHasAnyFlags(FlushState, ERunnerFlushState::End))
+	{
+		EndEvaluation(Linker);
+		return true;
+	}
+
+	return false;
 }
 
-void FMovieSceneEntitySystemRunner::GameThread_ProcessQueue()
+void FMovieSceneEntitySystemRunner::Flush(double BudgetMs)
 {
 	using namespace UE::MovieScene;
 
 	UMovieSceneEntitySystemLinker* Linker = GetLinker();
-	check(Linker);
+
+	// Check that we are attached to a linker that allows starting a new evaluation.
+	if (!ensureMsgf(Linker, TEXT("Runner isn't attached to a valid linker")))
+	{
+		return;
+	}
+
+	SCOPE_CYCLE_COUNTER(MovieSceneEval_RunnerFlush);
+
+	// If we're not currently evaluating, start by flushing everything
+	if (FlushState == ERunnerFlushState::None)
+	{
+		FlushState = ERunnerFlushState::Everything;
+	}
+
+	// We need to run the system from the game thread so we know we can fire events and callbacks from here.
+	check(IsInGameThread());
+
+	const double BudgetSeconds = BudgetMs / 1000.f;
+	if (BudgetSeconds > 0.0)
+	{
+		double StartTime = FPlatformTime::Seconds();
+		do
+		{
+			if (!FlushNext(Linker))
+			{
+				break;
+			}
+		}
+		while (FPlatformTime::Seconds() - StartTime < BudgetSeconds && FlushState != ERunnerFlushState::None);
+	}
+	else
+	{
+		while (FlushState != ERunnerFlushState::None)
+		{
+			if (!FlushNext(Linker))
+			{
+				break;
+			}
+		}
+	}
+}
+
+bool FMovieSceneEntitySystemRunner::GameThread_UpdateSequenceInstances(UMovieSceneEntitySystemLinker* Linker)
+{
+	using namespace UE::MovieScene;
+
+	TRACE_CPUPROFILER_EVENT_SCOPE(FMovieSceneEntitySystemRunner::GameThread_UpdateSequenceInstances);
+
+	// Set the debug vizualizer's entity manager pointer, so all debugging happening here will show
+	// relevant information. We need to set it here instead of higher up because we could have, say,
+	// a blocking sequence triggering another blocking sequence via an event track. The nested call stack
+	// of the second sequence needs to show debug information relevant to its private linker, but when
+	// we return back up to the first sequence, which might still have another update round (such as
+	// the other side of the dissected update range around the event), we need to set the pointer back
+	// again.
+	TGuardValue<FEntityManager*> DebugVizGuard(GEntityManagerForDebuggingVisualizers, GetEntityManager());
+
+	// Also reset the capture source scope so that each group of sequences tied to a given linker starts
+	// with a clean slate.
+	TGuardValue<FScopedPreAnimatedCaptureSource*> CaptureSourceGuard(FScopedPreAnimatedCaptureSource::GetCaptureSourcePtr(), nullptr);
+
+	// Allow the next step to run
+	FlushState &= ~ERunnerFlushState::Import;
 
 	FInstanceRegistry* InstanceRegistry = GetInstanceRegistry();
 
@@ -317,25 +412,6 @@ void FMovieSceneEntitySystemRunner::GameThread_ProcessQueue()
 		InstanceRegistry->MutateInstance(UpdatedInstanceHandle).PreEvaluation(Linker);
 	}
 
-	// Process updates
-	GameThread_SpawnPhase();
-}
-
-void FMovieSceneEntitySystemRunner::GameThread_SpawnPhase()
-{
-	using namespace UE::MovieScene;
-
-	check(GameThread == ENamedThreads::GameThread || GameThread == ENamedThreads::GameThread_Local);
-
-	UMovieSceneEntitySystemLinker* Linker = GetLinker();
-	check(Linker);
-
-	Linker->EntityManager.IncrementSystemSerial();
-
-	CurrentPhase = ESystemPhase::Spawn;
-
-	FInstanceRegistry* InstanceRegistry = GetInstanceRegistry();
-
 	// Update all systems
 	if (DissectedUpdates.Num() != 0)
 	{
@@ -354,6 +430,24 @@ void FMovieSceneEntitySystemRunner::GameThread_SpawnPhase()
 		DissectedUpdates.RemoveAt(0, Index);
 	}
 
+	return true;
+}
+
+bool FMovieSceneEntitySystemRunner::GameThread_SpawnPhase(UMovieSceneEntitySystemLinker* Linker)
+{
+	using namespace UE::MovieScene;
+
+	check(GameThread == ENamedThreads::GameThread || GameThread == ENamedThreads::GameThread_Local);
+
+	// Allow the next step to run
+	FlushState &= ~ERunnerFlushState::Spawn;
+
+	Linker->EntityManager.IncrementSystemSerial();
+
+	CurrentPhase = ESystemPhase::Spawn;
+
+	FInstanceRegistry* InstanceRegistry = GetInstanceRegistry();
+
 	const bool bInstantiationDirty = Linker->HasStructureChangedSinceLastRun() || InstanceRegistry->HasInvalidatedBindings();
 
 	FGraphEventArray AllTasks;
@@ -361,177 +455,144 @@ void FMovieSceneEntitySystemRunner::GameThread_SpawnPhase()
 	Linker->AutoLinkRelevantSystems();
 
 	// --------------------------------------------------------------------------------------------------------------------------------------------
-	// Step 1: Run the spawn phase if there were any changes to the current entity instantiations
-	if (bInstantiationDirty)
+	// Run the spawn phase if there were any changes to the current entity instantiations, and wait on the result
 	{
-		SCOPE_CYCLE_COUNTER(MovieSceneEval_SpawnPhase);
 
-		// The spawn phase can queue events to trigger from the event tracks.
-		bCanQueueEventTriggers = true;
+		SCOPE_CYCLE_COUNTER(MovieSceneEval_SpawnPhase);
+		if (bInstantiationDirty)
 		{
-			Linker->SystemGraph.ExecutePhase(ESystemPhase::Spawn, Linker, AllTasks);
+			// The spawn phase can queue events to trigger from the event tracks.
+			bCanQueueEventTriggers = true;
+			{
+				Linker->SystemGraph.ExecutePhase(ESystemPhase::Spawn, Linker, AllTasks);
+			}
+			bCanQueueEventTriggers = false;
+
+			// We don't open a re-entrancy window, however, because there's no way we can recursively evaluate things at this point... too many
+			// things are in an intermediate state. So events triggered as PreSpawn/PostSpawn can't be wired to something that starts a sequence.
+			if (EventTriggers.IsBound())
+			{
+				EventTriggers.Broadcast();
+				EventTriggers.Clear();
+			}
 		}
-		bCanQueueEventTriggers = false;
 
-		// We don't open a re-entrancy window, however, because there's no way we can recursively evaluate things at this point... too many
-		// things are in an intermediate state. So events triggered as PreSpawn/PostSpawn can't be wired to something that starts a sequence.
-		if (EventTriggers.IsBound())
+		// If there were any tasks created, we need to wait on them before proceeding. This is rare for the spawn phase.
+		if (AllTasks.Num() != 0)
 		{
-			EventTriggers.Broadcast();
-			EventTriggers.Clear();
+			FGraphEventRef SpawnEvent = TGraphTask<FNullGraphTask>::CreateTask(&AllTasks, ENamedThreads::GameThread)
+			.ConstructAndDispatchWhenReady(TStatId(), GameThread);
+
+			FTaskGraphInterface::Get().WaitUntilTaskCompletes(SpawnEvent, ENamedThreads::GameThread_Local);
 		}
 	}
 
 	// --------------------------------------------------------------------------------------------------------------------------------------------
-	// Step 2: Run the instantiation phase if there is anything to instantiate. This must come after the spawn phase because new intatniations may
-	//         be created during the spawn phase
-	auto NextStep = [this, bInstantiationDirty, Linker]
-	{
-		const bool bAnyPending = Linker->EntityManager.ContainsComponent(FBuiltInComponentTypes::Get()->Tags.NeedsLink) || Linker->GetInstanceRegistry()->HasInvalidatedBindings();
-		if (bInstantiationDirty || bAnyPending)
-		{
-			this->GameThread_InstantiationPhase();
-		}
-		else
-		{
-			// Go straight to evaluation
-			this->GameThread_EvaluationPhase();
-		}
-	};
-
-	if (AllTasks.Num() != 0)
+	// Only run the instantiation phase if there is anything to instantiate. This must come after the spawn phase because new instantiations may
+	// be created during the spawn phase
+	const bool bAnyPending = Linker->EntityManager.ContainsComponent(FBuiltInComponentTypes::Get()->Tags.NeedsLink) || InstanceRegistry->HasInvalidatedBindings();
+	if (bInstantiationDirty == false && bAnyPending == false)
 	{
-		TGraphTask<TFunctionGraphTaskImpl<void(), ESubsequentsMode::TrackSubsequents>>::CreateTask(&AllTasks, ENamedThreads::GameThread)
-			.ConstructAndDispatchWhenReady(MoveTemp(NextStep), TStatId(), GameThread);
-	}
-	else
-	{
-		NextStep();
+		FlushState &= ~ERunnerFlushState::Instantiation;
 	}
+
+	return true;
 }
 
-void FMovieSceneEntitySystemRunner::GameThread_InstantiationPhase()
+bool FMovieSceneEntitySystemRunner::GameThread_InstantiationPhase(UMovieSceneEntitySystemLinker* Linker)
 {
 	using namespace UE::MovieScene;
 
-	check(GameThread == ENamedThreads::GameThread || GameThread == ENamedThreads::GameThread_Local);
+	SCOPE_CYCLE_COUNTER(MovieSceneEval_InstantiationPhase);
 
-	UMovieSceneEntitySystemLinker* Linker = GetLinker();
-	check(Linker);
+	check(GameThread == ENamedThreads::GameThread || GameThread == ENamedThreads::GameThread_Local);
 
+	// Allow the next step to run
+	FlushState &= ~ERunnerFlushState::Instantiation;
 	CurrentPhase = ESystemPhase::Instantiation;
 
 	FGraphEventArray AllTasks;
-	{
-		SCOPE_CYCLE_COUNTER(MovieSceneEval_InstantiationPhase);
-
-		Linker->SystemGraph.ExecutePhase(ESystemPhase::Instantiation, Linker, AllTasks);
-	}
+	Linker->SystemGraph.ExecutePhase(ESystemPhase::Instantiation, Linker, AllTasks);
 
+	// If there were any tasks created, we need to wait on them before proceeding. This is rare for the instantiation phase.
 	if (AllTasks.Num() != 0)
 	{
-		TGraphTask<TFunctionGraphTaskImpl<void(), ESubsequentsMode::TrackSubsequents>>::CreateTask(&AllTasks, ENamedThreads::GameThread)
-			.ConstructAndDispatchWhenReady(
-				[this]
-				{
-					this->GameThread_PostInstantiation();
-				}
-				, TStatId(), GameThread);
-	}
-	else
-	{
-		this->GameThread_PostInstantiation();
+		FGraphEventRef InstantiationEvent = TGraphTask<FNullGraphTask>::CreateTask(&AllTasks, ENamedThreads::GameThread)
+		.ConstructAndDispatchWhenReady(TStatId(), GameThread);
+
+		FTaskGraphInterface::Get().WaitUntilTaskCompletes(InstantiationEvent, ENamedThreads::GameThread_Local);
 	}
+
+	return GameThread_PostInstantiation(Linker);
 }
 
-void FMovieSceneEntitySystemRunner::GameThread_PostInstantiation()
+bool FMovieSceneEntitySystemRunner::GameThread_PostInstantiation(UMovieSceneEntitySystemLinker* Linker)
 {
 	using namespace UE::MovieScene;
 
-	{
-		SCOPE_CYCLE_COUNTER(MovieSceneEval_PostInstantiation);
-
-		check(GameThread == ENamedThreads::GameThread || GameThread == ENamedThreads::GameThread_Local);
+	SCOPE_CYCLE_COUNTER(MovieSceneEval_PostInstantiation);
 
-		UMovieSceneEntitySystemLinker* Linker = GetLinker();
-		check(Linker);
+	check(GameThread == ENamedThreads::GameThread || GameThread == ENamedThreads::GameThread_Local);
 
-		Linker->PostInstantation(*this);
+	Linker->PostInstantation(*this);
 
-		FEntityManager& EntityManager = Linker->EntityManager;
-		FBuiltInComponentTypes* BuiltInComponentTypes = FBuiltInComponentTypes::Get();
+	FEntityManager& EntityManager = Linker->EntityManager;
+	FBuiltInComponentTypes* BuiltInComponentTypes = FBuiltInComponentTypes::Get();
 
-		// Nothing needs linking, caching or restoring any more
-		FRemoveMultipleMutation Mutation;
-		Mutation.RemoveComponent(BuiltInComponentTypes->Tags.NeedsLink);
+	// Nothing needs linking, caching or restoring any more
+	FRemoveMultipleMutation Mutation;
+	Mutation.RemoveComponent(BuiltInComponentTypes->Tags.NeedsLink);
 
-		FEntityComponentFilter Filter = FEntityComponentFilter().Any({ BuiltInComponentTypes->Tags.NeedsLink });
-		EntityManager.MutateAll(Filter, Mutation);
+	FEntityComponentFilter Filter = FEntityComponentFilter().Any({ BuiltInComponentTypes->Tags.NeedsLink });
+	EntityManager.MutateAll(Filter, Mutation);
 
-		// Free anything that has been unlinked
-		EntityManager.FreeEntities(FEntityComponentFilter().All({ BuiltInComponentTypes->Tags.NeedsUnlink }));
+	// Free anything that has been unlinked
+	EntityManager.FreeEntities(FEntityComponentFilter().All({ BuiltInComponentTypes->Tags.NeedsUnlink }));
 
-		Linker->SystemGraph.RemoveIrrelevantSystems(Linker);
+	Linker->SystemGraph.RemoveIrrelevantSystems(Linker);
 
-		EntityManager.Compact();
-	}
-
-	GameThread_EvaluationPhase();
+	EntityManager.Compact();
+	return true;
 }
 
-void FMovieSceneEntitySystemRunner::GameThread_EvaluationPhase()
+bool FMovieSceneEntitySystemRunner::GameThread_EvaluationPhase(UMovieSceneEntitySystemLinker* Linker)
 {
 	using namespace UE::MovieScene;
 
-	UMovieSceneEntitySystemLinker* Linker = GetLinker();
-	check(Linker);
+	SCOPE_CYCLE_COUNTER(MovieSceneEval_EvaluationPhase);
 
+	FlushState &= ~ERunnerFlushState::Evaluation;
 	CurrentPhase = ESystemPhase::Evaluation;
 
 	FGraphEventArray AllTasks;
-	{
-		SCOPE_CYCLE_COUNTER(MovieSceneEval_EvaluationPhase);
-
-		// --------------------------------------------------------------------------------------------------------------------------------------------
-		// Step 2: Run the evaluation phase. The entity manager is locked down for this phase, meaning no changes to entity-component structure is allowed
-		//         This vastly simplifies the concurrent handling of entity component allocations
-		Linker->EntityManager.LockDown();
-
-		checkf(!Linker->EntityManager.ContainsComponent(FBuiltInComponentTypes::Get()->Tags.NeedsUnlink), TEXT("Stale entities remain in the entity manager during evaluation - these should have been destroyed during the instantiation phase. Did it run?"));
+	// --------------------------------------------------------------------------------------------------------------------------------------------
+	// Step 2: Run the evaluation phase. The entity manager is locked down for this phase, meaning no changes to entity-component structure is allowed
+	//         This vastly simplifies the concurrent handling of entity component allocations
+	Linker->EntityManager.LockDown();
 
-		Linker->SystemGraph.ExecutePhase(ESystemPhase::Evaluation, Linker, AllTasks);
-	}
+	checkf(!Linker->EntityManager.ContainsComponent(FBuiltInComponentTypes::Get()->Tags.NeedsUnlink), TEXT("Stale entities remain in the entity manager during evaluation - these should have been destroyed during the instantiation phase. Did it run?"));
 
-	auto Finish = [this]
-	{
-		// We are now done with the current update batch. Let's unlock the completion task to unblock
-		// the main thread, which is waiting on it inside Flush().
-		check(this->CompletionTask != nullptr);
-		this->CompletionTask->Unlock();
-	};
+	Linker->SystemGraph.ExecutePhase(ESystemPhase::Evaluation, Linker, AllTasks);
 
 	if (AllTasks.Num() != 0)
 	{
-		TGraphTask<TFunctionGraphTaskImpl<void(), ESubsequentsMode::TrackSubsequents>>::CreateTask(&AllTasks, ENamedThreads::GameThread)
-			.ConstructAndDispatchWhenReady(MoveTemp(Finish), TStatId(), GameThread);
-	}
-	else
-	{
-		Finish();
+		FGraphEventRef EvaluationEvent = TGraphTask<FNullGraphTask>::CreateTask(&AllTasks, ENamedThreads::GameThread)
+		.ConstructAndDispatchWhenReady(TStatId(), GameThread);
+
+		FTaskGraphInterface::Get().WaitUntilTaskCompletes(EvaluationEvent, ENamedThreads::GameThread_Local);
 	}
+
+	Linker->EntityManager.ReleaseLockDown();
+
+	return true;
 }
 
-void FMovieSceneEntitySystemRunner::GameThread_EvaluationFinalizationPhase()
+void FMovieSceneEntitySystemRunner::GameThread_EvaluationFinalizationPhase(UMovieSceneEntitySystemLinker* Linker)
 {
 	using namespace UE::MovieScene;
 
 	check(GameThread == ENamedThreads::GameThread || GameThread == ENamedThreads::GameThread_Local);
 
-	UMovieSceneEntitySystemLinker* Linker = GetLinker();
-	check(Linker);
-
-	Linker->EntityManager.ReleaseLockDown();
-
 	CurrentPhase = ESystemPhase::Finalization;
 
 	// Post-eval events can be queued during the finalization phase so let's open that up.
@@ -565,15 +626,12 @@ void FMovieSceneEntitySystemRunner::GameThread_EvaluationFinalizationPhase()
 	CurrentPhase = ESystemPhase::None;
 }
 
-void FMovieSceneEntitySystemRunner::GameThread_PostEvaluationPhase()
+void FMovieSceneEntitySystemRunner::GameThread_PostEvaluationPhase(UMovieSceneEntitySystemLinker* Linker)
 {
 	using namespace UE::MovieScene;
 
 	SCOPE_CYCLE_COUNTER(MovieSceneEval_PostEvaluationPhase);
 
-	UMovieSceneEntitySystemLinker* Linker = GetLinker();
-	check(Linker);
-
 	// Execute any queued events from the evaluation finalization phase.
 	if (EventTriggers.IsBound())
 	{
diff --git a/Engine/Source/Runtime/MovieScene/Private/MovieSceneSequenceTickInterval.cpp b/Engine/Source/Runtime/MovieScene/Private/MovieSceneSequenceTickInterval.cpp
--- a/Engine/Source/Runtime/MovieScene/Private/MovieSceneSequenceTickInterval.cpp
+++ b/Engine/Source/Runtime/MovieScene/Private/MovieSceneSequenceTickInterval.cpp
@@ -44,10 +44,10 @@ int32 FMovieSceneSequenceTickInterval::RoundTickIntervalMs() const
 	// TickInterval is in seconds, so we multiply by 1000/<rounding> to round it as ms in integer space
 	if (bAllowRounding)
 	{
-		return FMath::RoundToInt(TickIntervalSeconds * (1000.f/GMovieSceneTickIntervalResolutionMs)) * GMovieSceneTickIntervalResolutionMs;
+		return FMath::CeilToInt(TickIntervalSeconds * (1000.f/GMovieSceneTickIntervalResolutionMs)) * GMovieSceneTickIntervalResolutionMs;
 	}
 
-	return FMath::RoundToInt(TickIntervalSeconds * 1000.f);
+	return FMath::CeilToInt(TickIntervalSeconds * 1000.f);
 }
 
 FMovieSceneSequenceTickInterval FMovieSceneSequenceTickInterval::GetInheritedInterval(UObject* ContextObject)
diff --git a/Engine/Source/Runtime/MovieScene/Private/MovieSceneSequenceTickManager.cpp b/Engine/Source/Runtime/MovieScene/Private/MovieSceneSequenceTickManager.cpp
--- a/Engine/Source/Runtime/MovieScene/Private/MovieSceneSequenceTickManager.cpp
+++ b/Engine/Source/Runtime/MovieScene/Private/MovieSceneSequenceTickManager.cpp
@@ -5,9 +5,15 @@
 #include "EntitySystem/MovieSceneEntitySystemLinker.h"
 #include "ProfilingDebugging/CountersTrace.h"
 
+#include "Logging/MessageLog.h"
+#include "Logging/TokenizedMessage.h"
+#include "Misc/UObjectToken.h"
 #include "Algo/IndexOf.h"
 
+#define LOCTEXT_NAMESPACE "MovieSceneSequenceTickManager"
+
 DECLARE_CYCLE_STAT(TEXT("Sequence Tick Manager"), MovieSceneEval_SequenceTickManager, STATGROUP_MovieSceneEval);
+DECLARE_CYCLE_STAT(TEXT("Tick Clients"), MovieSceneEval_TickClients, STATGROUP_MovieSceneEval);
 
 namespace UE::MovieScene
 {
@@ -127,6 +133,7 @@ void UMovieSceneSequenceTickManager::RegisterTickClient(const FMovieSceneSequenc
 		}
 	}
 
+	const float DesiredBudgetMs = ResolvedTickInterval.EvaluationBudgetMicroseconds / 1000.f;
 	if (!LinkerGroups.IsValidIndex(LinkerIndex))
 	{
 		LinkerIndex = LinkerGroups.Emplace();
@@ -146,11 +153,50 @@ void UMovieSceneSequenceTickManager::RegisterTickClient(const FMovieSceneSequenc
 		NewGroup.Linker = Linker;
 		NewGroup.Runner.AttachToLinker(Linker);
 		NewGroup.RoundedTickIntervalMs = DesiredTickIntervalMs;
+		NewGroup.FrameBudgetMs = DesiredBudgetMs;
 		NewGroup.NumClients = 1;
 	}
 	else
 	{
-		++LinkerGroups[LinkerIndex].NumClients;
+		FLinkerGroup& ExistingGroup = LinkerGroups[LinkerIndex];
+		++ExistingGroup.NumClients;
+
+		// If we're attempting to play sequences at the same tick interval with different budgets,
+		// we have to reconcile those. We emit a warning for mismatches of budgeted vs. not-budgeted
+		// as it will actually change behavior (non-budgeted will always eval in one frame, budgeted
+		// will sometimes eval over multiple frames).
+		if (ExistingGroup.FrameBudgetMs != DesiredBudgetMs)
+		{
+			if (ExistingGroup.FrameBudgetMs == 0.f)
+			{
+				ExistingGroup.FrameBudgetMs = DesiredBudgetMs;
+
+#if !UE_BUILD_SHIPPING && !UE_BUILD_TEST
+				TSharedRef<FTokenizedMessage> Message = FMessageLog("PIE")
+				.Warning(LOCTEXT("BudgetMismatch", "Attempting to play multiple sequences with the same tick interval but different budgets. Lowest budget will be used, but may result in undesired multi-frame effects."))
+				->AddToken(FTextToken::Create(LOCTEXT("ClientList", "Initiating tick client ")))
+				->AddToken(FUObjectToken::Create(InterfaceObject))
+				->AddToken(FTextToken::Create(FText::Format(LOCTEXT("ExistingClients", " with budget of {0}μs into group with an interval of {1}ms and budget of {2}μs. Existing clients: "),
+					FMath::RoundToInt(DesiredBudgetMs/1000.f), ExistingGroup.RoundedTickIntervalMs, FMath::RoundToInt(ExistingGroup.FrameBudgetMs/1000.f))));
+				for (const FTickableClientData& Client : TickableClients)
+				{
+					if (Client.LinkerIndex == LinkerIndex)
+					{
+						UObject* Object = Client.WeakObject.Get();
+						while (Object && !(Object->IsA<AActor>() || Object->IsA<USceneComponent>()))
+						{
+							Object = Object->GetOuter();
+						}
+						Message->AddToken(FUObjectToken::Create(Object));
+					}
+				}
+#endif
+			}
+			else if (DesiredBudgetMs != 0.f)
+			{
+				ExistingGroup.FrameBudgetMs = FMath::Min(ExistingGroup.FrameBudgetMs, DesiredBudgetMs);
+			}
+		}
 	}
 
 	checkSlow(LinkerIndex >= 0 && LinkerIndex < uint16(-1));
@@ -257,6 +303,8 @@ void UMovieSceneSequenceTickManager::TickSequenceActors(float DeltaSeconds)
 	// Sparse array where an allocated entry maps to the FLinkerGroup within LinkerGroups
 	TSparseArray<FUpdatedDeltaTimes, InlineSparseAllocator> UpdatedDeltaTimes;
 
+	TBitArray<> OutstandingLinkers;
+
 	// -----------------------------------------------------------------------------
 	// Step 1: Check the tick intervals of all our linker groups, and mark any whose
 	//         interval has passed for update
@@ -270,6 +318,15 @@ void UMovieSceneSequenceTickManager::TickSequenceActors(float DeltaSeconds)
 			float UnpausedDeltaTime = DeltaSeconds;
 			float DeltaTime = DeltaSeconds;
 
+			// If we're currently evaluating this linker group,
+			// skip ticking it until we're completely finished
+			if (Group.Runner.IsCurrentlyEvaluating())
+			{
+				OutstandingLinkers.PadToNum(Index + 1, false);
+				OutstandingLinkers[Index] = true;
+				continue;
+			}
+
 			if (Group.LastUnpausedTimeSeconds >= 0.f)
 			{
 				UnpausedDeltaTime = CurrentUnpausedTimeSeconds - Group.LastUnpausedTimeSeconds;
@@ -293,7 +350,7 @@ void UMovieSceneSequenceTickManager::TickSequenceActors(float DeltaSeconds)
 	}
 
 	// Skip any work if there are no updates scheduled
-	if (UpdatedDeltaTimes.Num() == 0)
+	if (UpdatedDeltaTimes.Num() == 0 && OutstandingLinkers.Num() == 0)
 	{
 		return;
 	}
@@ -303,28 +360,46 @@ void UMovieSceneSequenceTickManager::TickSequenceActors(float DeltaSeconds)
 		// Guard against any mutation while the entries are being ticked - anything added this tick cycle will have to wait until next tick
 		TGuardValue<TArray<FPendingOperation>*> Guard(PendingActorOperations, &CurrentPendingActorOperations);
 
-		// Step 2: Tick clients for any that need ticking
+		// Step 2: Tick unbudgeted clients for any with updated delta times
 		// 
-		for (const FTickableClientData& Client : TickableClients)
+		if (UpdatedDeltaTimes.Num() != 0)
 		{
-			const bool bCanTick = (Client.bTickWhenPaused || !bIsPaused) && Client.WeakObject.Get() != nullptr;
-			if (bCanTick && UpdatedDeltaTimes.IsValidIndex(Client.LinkerIndex))
+			SCOPE_CYCLE_COUNTER(MovieSceneEval_TickClients);
+			for (const FTickableClientData& Client : TickableClients)
 			{
-				const FUpdatedDeltaTimes Delta = UpdatedDeltaTimes[Client.LinkerIndex];
-				const float DeltaTime = Client.bTickWhenPaused ? Delta.UnpausedDeltaTime : Delta.DeltaTime;
+				const bool bCanTick = (Client.bTickWhenPaused || !bIsPaused) && Client.WeakObject.Get() != nullptr;
+				if (bCanTick && UpdatedDeltaTimes.IsValidIndex(Client.LinkerIndex))
+				{
+					const FUpdatedDeltaTimes Delta = UpdatedDeltaTimes[Client.LinkerIndex];
+					const float DeltaTime = Client.bTickWhenPaused ? Delta.UnpausedDeltaTime : Delta.DeltaTime;
 
-				FMovieSceneEntitySystemRunner* Runner = &LinkerGroups[Client.LinkerIndex].Runner;
-				Client.Interface->TickFromSequenceTickManager(DeltaTime, Runner);
+					FMovieSceneEntitySystemRunner* Runner = &LinkerGroups[Client.LinkerIndex].Runner;
+					Client.Interface->TickFromSequenceTickManager(DeltaTime, Runner);
+				}
 			}
 		}
 
-		// Step 3: Update and flush linkers as needed
+		// Step 3: Flush any budgeted evaluations that are still pending
+		// 
+		for (TConstSetBitIterator<> LinkerIndex(OutstandingLinkers); LinkerIndex; ++LinkerIndex)
+		{
+			FLinkerGroup& Group = LinkerGroups[LinkerIndex.GetIndex()];
+
+			check(!UpdatedDeltaTimes.IsValidIndex(LinkerIndex.GetIndex()));
+
+			Group.Runner.Flush(Group.FrameBudgetMs);
+		}
+
+		// Step 4: Update and flush linkers as needed
 		//
 		for (int32 Index = 0; Index < UpdatedDeltaTimes.GetMaxIndex(); ++Index)
 		{
 			if (UpdatedDeltaTimes.IsAllocated(Index))
 			{
-				LinkerGroups[Index].Runner.Flush();
+				FLinkerGroup& Group = LinkerGroups[Index];
+				// Hitting this check would indicate that the loop above that processes OutstandingLinkers either failed, or some other partial flush happened between then and now.
+				ensureMsgf(!Group.Runner.IsCurrentlyEvaluating(), TEXT("Linker is part-way thorugh a flush when a new flush is being instigated. This is undefined behavior."));
+				Group.Runner.Flush(Group.FrameBudgetMs);
 			}
 		}
 	}
@@ -471,3 +546,4 @@ void FMovieSceneLatentActionManager::RunLatentActions(TFunctionRef<void()> Flush
 	}
 }
 
+#undef LOCTEXT_NAMESPACE
\ No newline at end of file
diff --git a/Engine/Source/Runtime/MovieScene/Public/EntitySystem/MovieSceneEntitySystemLinker.h b/Engine/Source/Runtime/MovieScene/Public/EntitySystem/MovieSceneEntitySystemLinker.h
--- a/Engine/Source/Runtime/MovieScene/Public/EntitySystem/MovieSceneEntitySystemLinker.h
+++ b/Engine/Source/Runtime/MovieScene/Public/EntitySystem/MovieSceneEntitySystemLinker.h
@@ -283,6 +283,7 @@ public:
 
 private:
 
+	void HandlePreGarbageCollection();
 	void HandlePostGarbageCollection();
 
 	void TagInvalidBoundObjects();
diff --git a/Engine/Source/Runtime/MovieScene/Public/EntitySystem/MovieSceneEntitySystemRunner.h b/Engine/Source/Runtime/MovieScene/Public/EntitySystem/MovieSceneEntitySystemRunner.h
--- a/Engine/Source/Runtime/MovieScene/Public/EntitySystem/MovieSceneEntitySystemRunner.h
+++ b/Engine/Source/Runtime/MovieScene/Public/EntitySystem/MovieSceneEntitySystemRunner.h
@@ -13,12 +13,33 @@
 
 class UMovieSceneEntitySystemLinker;
 
-namespace UE {
-	namespace MovieScene {
-		class FEntityManager;
-		struct FInstanceRegistry;
-	}
-}
+namespace UE::MovieScene
+{
+
+class FEntityManager;
+struct FInstanceRegistry;
+
+/** Bit-mask enumeration that defines tasks that need running */
+enum class ERunnerFlushState
+{
+	None          = 0,			// Signifies no evaluation is currently underway
+	Start         = 1 << 0,		// Sets up initial evaluation flags for external players and listeners
+	Import        = 1 << 1,		// Update sequence instances and import entities into the entity manager
+	Spawn         = 1 << 2,		// Perform the Spawn phase in the Entity System Graph
+	Instantiation = 1 << 3,		// Perform the Instantiation phase in the Entity System Graph
+	Evaluation    = 1 << 4,		// Perform the Evaluation phase in the Entity System Graph
+	Finalization  = 1 << 5,		// Perform the Finalization phase in the Entity System Graph and trigger any external events
+	End           = 1 << 6,		// Counterpart for Start - resets external players and listeners
+
+	// Signifies that, during the Finalization task, there were still outstanding tasks and we need to perform another iteration
+	LoopEval      = Import | Spawn | Instantiation | Evaluation | Finalization | End,
+	// Initial flush state
+	Everything    = Start | LoopEval,
+};
+ENUM_CLASS_FLAGS(ERunnerFlushState)
+
+
+} // namespace UE::MovieScene
 
 DECLARE_MULTICAST_DELEGATE(FMovieSceneEntitySystemEventTriggers);
 
@@ -53,8 +74,14 @@ public:
 	/** Updates the given instance with the given context immediately. This doesn't affect the current update queue. */
 	void Update(const FMovieSceneContext& Context, FInstanceHandle Instance);
 
-	/** Flushes the update queue, blocking the current thread until all queued updates have run. */
-	void Flush();
+	/**
+	 * Flushes the update queue and applies any outstanding evaluation logic
+	 * 
+	 * @param BudgetMs      A budget (in milliseconds) to use for evaluation. Evaluation will cease prematurely once this budget is spent
+	 *                      and will process the outstanding work on the next call to Flush. A value of 0.0 signifies no budget - the queue
+	 *                      will be fully processed without leaving any outstanding work
+	 */
+	void Flush(double BudgetMs = 0.f);
 
 	/** Finish updating an instance. */
 	void FinishInstance(FInstanceHandle InInstanceHandle);
@@ -65,6 +92,8 @@ public:
 		return CurrentPhase;
 	}
 
+	bool IsCurrentlyEvaluating() const;
+
 public:
 
 	UMovieSceneEntitySystemLinker* GetLinker() const;
@@ -84,17 +113,40 @@ private:
 
 private:
 
-	// Update loop
-	
-	void DoFlushUpdateQueueOnce();
-	
-	void GameThread_ProcessQueue();
-	void GameThread_SpawnPhase();
-	void GameThread_InstantiationPhase();
-	void GameThread_PostInstantiation();
-	void GameThread_EvaluationPhase();
-	void GameThread_EvaluationFinalizationPhase();
-	void GameThread_PostEvaluationPhase();
+	/**
+	 * Flush the next item in our update loop based off the contents of FlushState
+	 *
+	 * @param  Linker   The linker we are arrached to
+	 * @return True if the loop is allowed to continue, or false if we should not flush any more
+	 */
+	bool FlushNext(UMovieSceneEntitySystemLinker* Linker);
+
+	/**
+	 * Set up initial state before any evaluation runs. Only called once regardless of the number of pending updates we have to process
+	 * Primarily used for setting up external 'is evaluating' flags for re-entrancy and async checks.
+	 */
+	bool StartEvaluation(UMovieSceneEntitySystemLinker* Linker);
+
+	/** Update sequence instances based on currently queued update requests, or outstanding dissected updates */
+	bool GameThread_UpdateSequenceInstances(UMovieSceneEntitySystemLinker* Linker);
+	/** Execute the spawn phase of the entity system graph, if there is anything to do */
+	bool GameThread_SpawnPhase(UMovieSceneEntitySystemLinker* Linker);
+	/** Execute the instantiation phase of the entity system graph, if there is anything to do */
+	bool GameThread_InstantiationPhase(UMovieSceneEntitySystemLinker* Linker);
+	/** Called immediately after instantiation to execute cleanup and bookkeeping tasks. Skipped if instantiation is skipped.*/
+	bool GameThread_PostInstantiation(UMovieSceneEntitySystemLinker* Linker);
+	/** Main entity-system evaluation phase. Blocks this thread until completion. */
+	bool GameThread_EvaluationPhase(UMovieSceneEntitySystemLinker* Linker);
+	/** Finalization phase for triggering external events and other behavior. */
+	void GameThread_EvaluationFinalizationPhase(UMovieSceneEntitySystemLinker* Linker);
+	/** Post-evaluation phase for clean up. */
+	void GameThread_PostEvaluationPhase(UMovieSceneEntitySystemLinker* Linker);
+
+	/**
+	 * Counterpart for StartEvaluation.
+	 * Called only when our UpdateQueue and DissectedUpdates have been fully processed and there is nothing left to do.
+	 */
+	void EndEvaluation(UMovieSceneEntitySystemLinker* Linker);
 
 private:
 
@@ -123,12 +175,19 @@ private:
 	/** When an update is running, the list of sub-contexts for the requested update */
 	TArray<FDissectedUpdate> DissectedUpdates;
 
-	TGraphTask<FNullGraphTask>* CompletionTask;
-
 	bool bCanQueueEventTriggers;
 	FMovieSceneEntitySystemEventTriggers EventTriggers;
 
 	ENamedThreads::Type GameThread;
 
 	UE::MovieScene::ESystemPhase CurrentPhase;
+
+	/**
+	 * Defines a bitmask of outstanding tasks that need running.
+	 * When a task has completed its bit becomes unset.
+	 * Evaluation can trigger tasks to run again by setting previously cleared bits in this mask.
+	 * A state of ERunnerFlushState::Everything means an evaluation has not yet started.
+	 * A state of ERunnerFlushState::None means an evaluation has just finished.
+	 */
+	UE::MovieScene::ERunnerFlushState FlushState;
 };
diff --git a/Engine/Source/Runtime/MovieScene/Public/MovieSceneSequenceTickInterval.h b/Engine/Source/Runtime/MovieScene/Public/MovieSceneSequenceTickInterval.h
--- a/Engine/Source/Runtime/MovieScene/Public/MovieSceneSequenceTickInterval.h
+++ b/Engine/Source/Runtime/MovieScene/Public/MovieSceneSequenceTickInterval.h
@@ -35,6 +35,10 @@ public:
 	UPROPERTY(EditAnywhere, Category="Playback", meta=(DisplayName="Tick Interval", Units=s))
 	float TickIntervalSeconds = 0.f;
 
+	/** Defines an approximate budget for evaluation of this sequence (and any other sequences with the same tick interval) */
+	UPROPERTY(EditAnywhere, Category="Playback", meta=(DisplayName="Evaluation Budget", ForceUnits=us))
+	float EvaluationBudgetMicroseconds = 0.f;
+
 	/** When true, the sequence will continue to tick and progress even when the world is paused */
 	UPROPERTY(EditAnywhere, Category="Playback")
 	bool bTickWhenPaused = false;
@@ -62,6 +66,7 @@ public:
 	friend bool operator==(const FMovieSceneSequenceTickInterval& A, const FMovieSceneSequenceTickInterval& B)
 	{
 		return A.TickIntervalSeconds == B.TickIntervalSeconds
+			&& A.EvaluationBudgetMicroseconds == B.EvaluationBudgetMicroseconds
 			&& A.bTickWhenPaused == B.bTickWhenPaused
 			&& A.bAllowRounding == B.bAllowRounding;
 	}
diff --git a/Engine/Source/Runtime/MovieScene/Public/MovieSceneSequenceTickManager.h b/Engine/Source/Runtime/MovieScene/Public/MovieSceneSequenceTickManager.h
--- a/Engine/Source/Runtime/MovieScene/Public/MovieSceneSequenceTickManager.h
+++ b/Engine/Source/Runtime/MovieScene/Public/MovieSceneSequenceTickManager.h
@@ -168,11 +168,17 @@ private:
 		/** The tick interval of this group, rounded to the nearest Sequencer.TickIntervalGroupingResolutionMs */
 		int32 RoundedTickIntervalMs = 0;
 
+		/** The frame budget for all linkers within this linker group in milliseconds */
+		float FrameBudgetMs = 0.f;
+
 		/** The value of UWorld::GetUnpausedTimeSeconds last time this group was evaluated */
 		float LastUnpausedTimeSeconds = -1;
 		/** The value of UWorld::GetTimeSeconds last time this group was evaluated */
 		float LastTimeSeconds = -1;
 
+		/** The number of times this group has been starved of frames for budgeted evaluation */
+		uint16 BudgetedStarvationCount = 0;
+
 		/** Total number of clients in this group */
 		uint16 NumClients = 0;
 	};
