Sequencer: Added the ability for tracks to be entity providers

  - Doing so will bypass sections providing entiies unless implemented explicitly by the track

diff --git a/Engine/Source/Runtime/MovieScene/Private/Compilation/MovieSceneCompiledDataManager.cpp b/Engine/Source/Runtime/MovieScene/Private/Compilation/MovieSceneCompiledDataManager.cpp
--- a/Engine/Source/Runtime/MovieScene/Private/Compilation/MovieSceneCompiledDataManager.cpp
+++ b/Engine/Source/Runtime/MovieScene/Private/Compilation/MovieSceneCompiledDataManager.cpp
@@ -1247,7 +1247,18 @@ void UMovieSceneCompiledDataManager::GatherTrack(const FMovieSceneBinding* Objec
 			FieldBuilder.GetSharedMetaData().ObjectBindingID = ObjectBinding->GetObjectGuid();
 		}
 
-		for (const FMovieSceneTrackEvaluationFieldEntry& Entry : EvaluationField.Entries)
+		IMovieSceneEntityProvider* TrackEntityProvider = Cast<IMovieSceneEntityProvider>(Track);
+
+		// If the track is an entity provider, allow it to add entries first
+		if (TrackEntityProvider)
+		{
+			FMovieSceneEvaluationFieldEntityMetaData MetaData(PreCompileResult.DefaultMetaData);
+			MetaData.bEvaluateInSequencePreRoll  = Track->EvalOptions.bEvaluateInPreroll;
+			MetaData.bEvaluateInSequencePostRoll = Track->EvalOptions.bEvaluateInPostroll;
+
+			TrackEntityProvider->PopulateEvaluationField(Params.LocalClampRange, MetaData, &FieldBuilder);
+		}
+		else for (const FMovieSceneTrackEvaluationFieldEntry& Entry : EvaluationField.Entries)
 		{
 			if (Entry.Section && Track->IsRowEvalDisabled(Entry.Section->GetRowIndex()))
 			{
