Added missing include

diff --git a/Engine/Source/Runtime/MovieScene/Private/EntitySystem/MovieSceneEntitySystemLinker.cpp b/Engine/Source/Runtime/MovieScene/Private/EntitySystem/MovieSceneEntitySystemLinker.cpp
--- a/Engine/Source/Runtime/MovieScene/Private/EntitySystem/MovieSceneEntitySystemLinker.cpp
+++ b/Engine/Source/Runtime/MovieScene/Private/EntitySystem/MovieSceneEntitySystemLinker.cpp
@@ -8,6 +8,7 @@
 #include "EntitySystem/MovieSceneComponentRegistry.h"
 #include "EntitySystem/BuiltInComponentTypes.h"
 #include "EntitySystem/MovieScenePreAnimatedStateSystem.h"
+#include "EntitySystem/MovieSceneEntitySystemRunner.h"
 #include "MovieSceneFwd.h"
 #include "UObject/Package.h"
 #include "UObject/UObjectGlobals.h"
