Sequencer: Fixed legacy track evaluators not evaluating the correct group with exclusive upper bounds

  - The approach is now unified between the ECS and legacy evaluators

diff --git a/Engine/Source/Runtime/MovieScene/Private/EntitySystem/MovieSceneSequenceUpdaters.cpp b/Engine/Source/Runtime/MovieScene/Private/EntitySystem/MovieSceneSequenceUpdaters.cpp
--- a/Engine/Source/Runtime/MovieScene/Private/EntitySystem/MovieSceneSequenceUpdaters.cpp
+++ b/Engine/Source/Runtime/MovieScene/Private/EntitySystem/MovieSceneSequenceUpdaters.cpp
@@ -92,21 +92,6 @@ private:
 	FMovieSceneSequenceID RootOverrideSequenceID;
 };
 
-FFrameNumber ImportTimeFromContext(const FMovieSceneContext& Context)
-{
-	TRange<FFrameNumber> Range = Context.GetFrameNumberRange();
-
-	FFrameTime Epsilon(0, 1.f - FFrameTime::MaxSubframe);
-	if (Context.GetDirection() == EPlayDirection::Forwards)
-	{
-		return DiscreteExclusiveUpper(Range)-1;
-	}
-	else
-	{
-		return DiscreteInclusiveLower(Range);
-	}
-}
-
 void DissectRange(TArrayView<const FFrameTime> InDissectionTimes, const TRange<FFrameTime>& Bounds, TArray<TRange<FFrameTime>>& OutDissections)
 {
 	if (InDissectionTimes.Num() == 0)
@@ -246,7 +231,7 @@ void FSequenceUpdater_Flat::Update(UMovieSceneEntitySystemLinker* Linker, FRootI
 
 	FMovieSceneEvaluationFieldEntitySet EntitiesScratch;
 
-	FFrameNumber ImportTime = ImportTimeFromContext(Context);
+	FFrameNumber ImportTime = Context.GetEvaluationFieldTime();
 		
 	const bool bOutsideCachedRange = !CachedEntityRange.Contains(ImportTime);
 	if (bOutsideCachedRange)
@@ -456,7 +441,7 @@ void FSequenceUpdater_Hierarchical::Update(UMovieSceneEntitySystemLinker* Linker
 		}
 	}
 
-	FFrameNumber ImportTime = ImportTimeFromContext(RootContext);
+	FFrameNumber ImportTime = RootContext.GetEvaluationFieldTime();
 	const bool bGatherEntities = !CachedEntityRange.Contains(ImportTime);
 
 	// ------------------------------------------------------------------------------------------------
@@ -578,7 +563,7 @@ void FSequenceUpdater_Hierarchical::Update(UMovieSceneEntitySystemLinker* Linker
 				const FMovieSceneEntityComponentField* SubComponentField = CompiledDataManager->FindEntityComponentField(SubDataID);
 
 				// Update entities if necessary
-				const FFrameTime SubSequenceTime = ImportTimeFromContext(SubContext);
+				const FFrameTime SubSequenceTime = SubContext.GetEvaluationFieldTime();
 
 				FEntityImportSequenceParams Params;
 				Params.InstanceHandle = SubSequenceHandle;
diff --git a/Engine/Source/Runtime/MovieScene/Private/Evaluation/Instances/MovieSceneTrackEvaluator.cpp b/Engine/Source/Runtime/MovieScene/Private/Evaluation/Instances/MovieSceneTrackEvaluator.cpp
--- a/Engine/Source/Runtime/MovieScene/Private/Evaluation/Instances/MovieSceneTrackEvaluator.cpp
+++ b/Engine/Source/Runtime/MovieScene/Private/Evaluation/Instances/MovieSceneTrackEvaluator.cpp
@@ -179,7 +179,7 @@ const FMovieSceneEvaluationGroup* FMovieSceneTrackEvaluator::SetupFrame(UMovieSc
 	RootOverridePath.Reset(InOverrideRootID, RootHierarchy);
 
 	const FMovieSceneEvaluationField* OverrideRootField = nullptr;
-	FFrameTime RootTime = Context.GetTime();
+	FFrameTime RootTime = Context.GetEvaluationFieldTime();
 
 	if (InOverrideRootID == MovieSceneSequenceID::Root)
 	{
diff --git a/Engine/Source/Runtime/MovieScene/Public/Evaluation/MovieScenePlayback.h b/Engine/Source/Runtime/MovieScene/Public/Evaluation/MovieScenePlayback.h
--- a/Engine/Source/Runtime/MovieScene/Public/Evaluation/MovieScenePlayback.h
+++ b/Engine/Source/Runtime/MovieScene/Public/Evaluation/MovieScenePlayback.h
@@ -6,6 +6,7 @@
 #include "MovieSceneFwd.h"
 #include "Misc/FrameTime.h"
 #include "Evaluation/MovieSceneSequenceTransform.h"
+#include "MovieSceneTimeHelpers.h"
 
 
 /** Enumeration specifying whether we're playing forwards or backwards */
@@ -86,6 +87,24 @@ struct MOVIESCENE_API FMovieSceneEvaluationRange
 		return Direction == EPlayDirection::Forwards ? EvaluationRange.GetUpperBoundValue() : EvaluationRange.GetLowerBoundValue();
 	}
 
+	/**
+	 * Get the current time to use for looking up within an evaluation field.
+	 * Subtly different from GetTime in that it returns the previous tick for exclusive boundaries
+	 */
+	FORCEINLINE FFrameNumber GetEvaluationFieldTime() const
+	{
+		TRange<FFrameNumber> Range = GetFrameNumberRange();
+
+		if (Direction == EPlayDirection::Forwards)
+		{
+			return UE::MovieScene::DiscreteExclusiveUpper(Range)-1;
+		}
+		else
+		{
+			return UE::MovieScene::DiscreteInclusiveLower(Range);
+		}
+	}
+
 	/**
 	 * Get the absolute amount of time that has passed since the last update (will always be >= 0)
 	 */
