Sequencer: Updated MovieScene natvis to expand debug property types

diff --git a/Engine/Source/Runtime/MovieScene/Public/EntitySystem/MovieSceneComponentDebug.h b/Engine/Source/Runtime/MovieScene/Public/EntitySystem/MovieSceneComponentDebug.h
--- a/Engine/Source/Runtime/MovieScene/Public/EntitySystem/MovieSceneComponentDebug.h
+++ b/Engine/Source/Runtime/MovieScene/Public/EntitySystem/MovieSceneComponentDebug.h
@@ -7,6 +7,8 @@
 
 #if UE_MOVIESCENE_ENTITY_DEBUG
 
+class FName;
+
 namespace UE
 {
 namespace MovieScene
@@ -29,6 +31,7 @@ enum class EComponentDebugType
 	Property,
 	InstanceHandle,
 	EntityID,
+	Name,
 };
 
 /**
@@ -53,6 +56,7 @@ template<>           struct TComponentDebugType<FVector>             { static co
 template<>           struct TComponentDebugType<FVector4>            { static const EComponentDebugType Type = EComponentDebugType::Vector4;  };
 template<>           struct TComponentDebugType<UObject*>            { static const EComponentDebugType Type = EComponentDebugType::Object;   };
 template<>           struct TComponentDebugType<FMovieSceneEntityID> { static const EComponentDebugType Type = EComponentDebugType::EntityID; };
+template<>           struct TComponentDebugType<FName>               { static const EComponentDebugType Type = EComponentDebugType::Name;     };
 
 } // namespace MovieScene
 } // namespace UE
