Sequencer: Added object path and bool channels to the moviescene ecs

 - I added bool channels thinking that they were required because bool parameters exist on the parameter section, but ended up not needing them. There's no harm in checking in the system for such a time as when we need it though. Ultimately most hi-perf uses of bool channels will be implemented at compile time using entities with fixed results, but having bool evaluation built-in seems like a nice convenience feature for people who need it.

diff --git a/Engine/Source/Runtime/MovieScene/Private/EntitySystem/BuiltInComponentTypes.cpp b/Engine/Source/Runtime/MovieScene/Private/EntitySystem/BuiltInComponentTypes.cpp
--- a/Engine/Source/Runtime/MovieScene/Private/EntitySystem/BuiltInComponentTypes.cpp
+++ b/Engine/Source/Runtime/MovieScene/Private/EntitySystem/BuiltInComponentTypes.cpp
@@ -36,6 +36,7 @@ FBuiltInComponentTypes::FBuiltInComponentTypes()
 	ComponentRegistry->NewComponentType(&SceneComponentBinding,   TEXT("USceneComponent Binding ID"));
 	ComponentRegistry->NewComponentType(&SpawnableBinding,        TEXT("Spawnable Binding"));
 	ComponentRegistry->NewComponentType(&TrackInstance,           TEXT("Track Instance"));
+	ComponentRegistry->NewComponentType(&BoolChannel,             TEXT("Bool Channel"));
 	ComponentRegistry->NewComponentType(&ByteChannel,             TEXT("Byte Channel"));
 	ComponentRegistry->NewComponentType(&IntegerChannel,          TEXT("Integer Channel"));
 	ComponentRegistry->NewComponentType(&FloatChannel[0],         TEXT("Float Channel 0"));
@@ -57,6 +58,7 @@ FBuiltInComponentTypes::FBuiltInComponentTypes()
 	ComponentRegistry->NewComponentType(&DoubleChannel[7],        TEXT("Double Channel 7"));
 	ComponentRegistry->NewComponentType(&DoubleChannel[8],        TEXT("Double Channel 8"));
 	ComponentRegistry->NewComponentType(&WeightChannel,           TEXT("Weight Channel"));
+	ComponentRegistry->NewComponentType(&ObjectPathChannel,       TEXT("Object Path Channel"));
 
 	ComponentRegistry->NewComponentType(&FloatChannelFlags[0],    TEXT("Float Channel 0 Flags"));
 	ComponentRegistry->NewComponentType(&FloatChannelFlags[1],    TEXT("Float Channel 1 Flags"));
@@ -112,6 +114,7 @@ FBuiltInComponentTypes::FBuiltInComponentTypes()
 	ComponentRegistry->NewComponentType(&DoubleResult[6],       TEXT("Double Result 6"));
 	ComponentRegistry->NewComponentType(&DoubleResult[7],       TEXT("Double Result 7"));
 	ComponentRegistry->NewComponentType(&DoubleResult[8],       TEXT("Double Result 8"));
+	ComponentRegistry->NewComponentType(&ObjectResult,          TEXT("Object Result"));
 
 	ComponentRegistry->NewComponentType(&BaseInteger,			TEXT("Base Integer"));
 	ComponentRegistry->NewComponentType(&BaseFloat[0],          TEXT("Base Float 0"));
@@ -199,7 +202,9 @@ FBuiltInComponentTypes::FBuiltInComponentTypes()
 	
 	// Bool channel relationships
 	{
+		ComponentRegistry->Factories.DuplicateChildComponent(BoolChannel);
 		ComponentRegistry->Factories.DuplicateChildComponent(BoolResult);
+		ComponentRegistry->Factories.DefineMutuallyInclusiveComponent(BoolChannel, BoolResult);
 		ComponentRegistry->Factories.DefineMutuallyInclusiveComponent(BoolResult, EvalTime);
 	}
 
@@ -304,6 +309,16 @@ FBuiltInComponentTypes::FBuiltInComponentTypes()
 		ComponentRegistry->Factories.DefineMutuallyInclusiveComponent(WeightResult, WeightAndEasingResult);
 	}
 
+	// Object path channel relationships
+	{
+		// Object channel components should be duplicated to children
+		ComponentRegistry->Factories.DuplicateChildComponent(ObjectPathChannel);
+
+		// Object channel components need a time and result to evaluate
+		ComponentRegistry->Factories.DefineMutuallyInclusiveComponent(ObjectPathChannel, EvalTime);
+		ComponentRegistry->Factories.DefineMutuallyInclusiveComponent(ObjectPathChannel, ObjectResult);
+	}
+
 	// Track instances always produce inputs
 	{
 		auto InitInput = [](const FMovieSceneTrackInstanceComponent& InInstance, FTrackInstanceInputComponent& OutInput)
diff --git a/Engine/Source/Runtime/MovieScene/Public/EntitySystem/BuiltInComponentTypes.h b/Engine/Source/Runtime/MovieScene/Public/EntitySystem/BuiltInComponentTypes.h
--- a/Engine/Source/Runtime/MovieScene/Public/EntitySystem/BuiltInComponentTypes.h
+++ b/Engine/Source/Runtime/MovieScene/Public/EntitySystem/BuiltInComponentTypes.h
@@ -17,10 +17,12 @@
 #include "BuiltInComponentTypes.generated.h"
 
 enum class EMovieSceneBlendType : uint8;
+struct FMovieSceneBoolChannel;
 struct FMovieSceneByteChannel;
 struct FMovieSceneIntegerChannel;
 struct FMovieSceneFloatChannel;
 struct FMovieSceneDoubleChannel;
+struct FMovieSceneObjectPathChannel;
 struct FMovieScenePropertyBinding;
 
 class UMovieSceneSection;
@@ -91,6 +93,22 @@ namespace UE
 namespace MovieScene
 {
 
+/**
+ * The component data for evaluating a bool channel
+ */
+struct FSourceBoolChannel
+{
+	FSourceBoolChannel()
+		: Source(nullptr)
+	{}
+
+	FSourceBoolChannel(const FMovieSceneBoolChannel* InSource)
+		: Source(InSource)
+	{}
+
+	const FMovieSceneBoolChannel* Source;
+};
+
 /**
  * The component data for evaluating a byte channel
  */
@@ -170,6 +188,20 @@ struct FEvaluationHookFlags
 	bool bHasBegun = false;
 };
 
+struct FSourceObjectPathChannel
+{
+	FSourceObjectPathChannel()
+		: Source(nullptr)
+	{}
+
+	FSourceObjectPathChannel(const FMovieSceneObjectPathChannel* InSource)
+		: Source(InSource)
+	{}
+
+	const FMovieSceneObjectPathChannel* Source;
+};
+
+
 /**
  * Pre-defined built in component types
  */
@@ -236,6 +268,9 @@ public:
 	// An integer representing the base value for the integer channel for the purposes of "additive from base" blending.
 	TComponentTypeID<int32> BaseInteger;
 
+	// An FMovieSceneBoolChannel
+	TComponentTypeID<FSourceBoolChannel> BoolChannel;
+
 	// An FMovieSceneFloatChannel considered to be at index N within the source structure (ie 0 = Location.X, Vector.X, Color.R; 1 = Location.Y, Vector.Y, Color.G)
 	TComponentTypeID<FSourceFloatChannel> FloatChannel[9];
 	TComponentTypeID<FSourceFloatChannelFlags> FloatChannelFlags[9];
@@ -248,6 +283,9 @@ public:
 	TComponentTypeID<FSourceFloatChannel> WeightChannel;
 	TComponentTypeID<FSourceFloatChannelFlags> WeightChannelFlags;
 
+	// FMovieSceneObjectPathChannel that represents a changing object path over time
+	TComponentTypeID<FSourceObjectPathChannel> ObjectPathChannel;
+
 	// A float representing the output of the channel considered to be at index N within the source structure (ie 0 = Location.X, Vector.X, Color.R; 1 = Location.Y, Vector.Y, Color.G)
 	TComponentTypeID<float> FloatResult[9];
 
@@ -266,6 +304,9 @@ public:
 	// A float representing the evaluated output of a weight channel
 	TComponentTypeID<float> WeightResult;
 
+	// The result of an evaluated FMovieSceneObjectPathChannel
+	TComponentTypeID<UObject*> ObjectResult;
+
 public:
 
 	// An FEasingComponentData for computing easing curves
diff --git a/Engine/Source/Runtime/MovieSceneTracks/Private/Systems/BoolChannelEvaluatorSystem.cpp b/Engine/Source/Runtime/MovieSceneTracks/Private/Systems/BoolChannelEvaluatorSystem.cpp
new file mode 100644
--- /dev/null
+++ b/Engine/Source/Runtime/MovieSceneTracks/Private/Systems/BoolChannelEvaluatorSystem.cpp
@@ -0,0 +1,50 @@
+// Copyright Epic Games, Inc. All Rights Reserved.
+
+#include "Systems/BoolChannelEvaluatorSystem.h"
+#include "EntitySystem/BuiltInComponentTypes.h"
+#include "EntitySystem/MovieSceneEntitySystemTask.h"
+#include "EntitySystem/MovieSceneEntitySystemLinker.h"
+
+#include "EntitySystem/MovieSceneEvalTimeSystem.h"
+#include "Channels/MovieSceneBoolChannel.h"
+
+DECLARE_CYCLE_STAT(TEXT("MovieScene: Evaluate bool channels"), MovieSceneEval_EvaluateBoolChannelTask, STATGROUP_MovieSceneECS);
+
+UBoolChannelEvaluatorSystem::UBoolChannelEvaluatorSystem(const FObjectInitializer& ObjInit)
+	: Super(ObjInit)
+{
+	using namespace UE::MovieScene;
+
+	RelevantComponent = FBuiltInComponentTypes::Get()->BoolChannel;
+
+	if (HasAnyFlags(RF_ClassDefaultObject))
+	{
+		DefineImplicitPrerequisite(UMovieSceneEvalTimeSystem::StaticClass(), GetClass());
+	}
+}
+
+void UBoolChannelEvaluatorSystem::OnRun(FSystemTaskPrerequisites& InPrerequisites, FSystemSubsequentTasks& Subsequents)
+{
+	using namespace UE::MovieScene;
+
+	FBuiltInComponentTypes* BuiltInComponents = FBuiltInComponentTypes::Get();
+
+	struct FEvaluateBoolChannels
+	{
+		static void ForEachEntity(FSourceBoolChannel BoolChannel, FFrameTime FrameTime, bool& OutResult)
+		{
+			if (!BoolChannel.Source->Evaluate(FrameTime, OutResult))
+			{
+				OutResult = false;
+			}
+		}
+	};
+
+	FEntityTaskBuilder()
+	.Read(BuiltInComponents->BoolChannel)
+	.Read(BuiltInComponents->EvalTime)
+	.Write(BuiltInComponents->BoolResult)
+	.FilterNone({ BuiltInComponents->Tags.Ignored })
+	.SetStat(GET_STATID(MovieSceneEval_EvaluateBoolChannelTask))
+	.Dispatch_PerEntity<FEvaluateBoolChannels>(&Linker->EntityManager, InPrerequisites, &Subsequents);
+}
diff --git a/Engine/Source/Runtime/MovieSceneTracks/Private/Systems/ObjectPathChannelEvaluatorSystem.cpp b/Engine/Source/Runtime/MovieSceneTracks/Private/Systems/ObjectPathChannelEvaluatorSystem.cpp
new file mode 100644
--- /dev/null
+++ b/Engine/Source/Runtime/MovieSceneTracks/Private/Systems/ObjectPathChannelEvaluatorSystem.cpp
@@ -0,0 +1,51 @@
+// Copyright Epic Games, Inc. All Rights Reserved.
+
+#include "Systems/ObjectPathChannelEvaluatorSystem.h"
+#include "EntitySystem/BuiltInComponentTypes.h"
+#include "EntitySystem/MovieSceneEntitySystemTask.h"
+#include "EntitySystem/MovieSceneEntitySystemLinker.h"
+#include "EntitySystem/MovieSceneEvalTimeSystem.h"
+#include "Channels/MovieSceneObjectPathChannel.h"
+
+DECLARE_CYCLE_STAT(TEXT("MovieScene: Evaluate object path channels"), MovieSceneEval_EvaluateObjectPathChannelTask, STATGROUP_MovieSceneECS);
+
+UObjectPathChannelEvaluatorSystem::UObjectPathChannelEvaluatorSystem(const FObjectInitializer& ObjInit)
+	: Super(ObjInit)
+{
+	using namespace UE::MovieScene;
+
+	FBuiltInComponentTypes* BuiltInComponents = FBuiltInComponentTypes::Get();
+	RelevantComponent = BuiltInComponents->ObjectPathChannel;
+
+	if (HasAnyFlags(RF_ClassDefaultObject))
+	{
+		DefineComponentConsumer(GetClass(), BuiltInComponents->EvalTime);
+		DefineComponentProducer(GetClass(), BuiltInComponents->ObjectResult);
+	}
+}
+
+void UObjectPathChannelEvaluatorSystem::OnRun(FSystemTaskPrerequisites& InPrerequisites, FSystemSubsequentTasks& Subsequents)
+{
+	using namespace UE::MovieScene;
+
+	FBuiltInComponentTypes* BuiltInComponents = FBuiltInComponentTypes::Get();
+
+	struct FEvaluateObjectPathChannels
+	{
+		static void ForEachEntity(FSourceObjectPathChannel ObjectPathChannel, FFrameTime FrameTime, UObject*& OutResult)
+		{
+			if (!ObjectPathChannel.Source->Evaluate(FrameTime, OutResult))
+			{
+				OutResult = nullptr;
+			}
+		}
+	};
+
+	FEntityTaskBuilder()
+	.Read(BuiltInComponents->ObjectPathChannel)
+	.Read(BuiltInComponents->EvalTime)
+	.Write(BuiltInComponents->ObjectResult)
+	.FilterNone({ BuiltInComponents->Tags.Ignored })
+	.SetStat(GET_STATID(MovieSceneEval_EvaluateObjectPathChannelTask))
+	.Dispatch_PerEntity<FEvaluateObjectPathChannels>(&Linker->EntityManager, InPrerequisites, &Subsequents);
+}
diff --git a/Engine/Source/Runtime/MovieSceneTracks/Public/Systems/BoolChannelEvaluatorSystem.h b/Engine/Source/Runtime/MovieSceneTracks/Public/Systems/BoolChannelEvaluatorSystem.h
new file mode 100644
--- /dev/null
+++ b/Engine/Source/Runtime/MovieSceneTracks/Public/Systems/BoolChannelEvaluatorSystem.h
@@ -0,0 +1,21 @@
+// Copyright Epic Games, Inc. All Rights Reserved.
+
+#pragma once
+
+#include "EntitySystem/MovieSceneEntitySystem.h"
+#include "BoolChannelEvaluatorSystem.generated.h"
+
+/**
+ * System that is responsible for evaluating bool channels.
+ */
+UCLASS(MinimalAPI)
+class UBoolChannelEvaluatorSystem : public UMovieSceneEntitySystem
+{
+public:
+
+	GENERATED_BODY()
+
+	UBoolChannelEvaluatorSystem(const FObjectInitializer& ObjInit);
+
+	virtual void OnRun(FSystemTaskPrerequisites& InPrerequisites, FSystemSubsequentTasks& Subsequents) override;
+};
\ No newline at end of file
diff --git a/Engine/Source/Runtime/MovieSceneTracks/Public/Systems/ObjectPathChannelEvaluatorSystem.h b/Engine/Source/Runtime/MovieSceneTracks/Public/Systems/ObjectPathChannelEvaluatorSystem.h
new file mode 100644
--- /dev/null
+++ b/Engine/Source/Runtime/MovieSceneTracks/Public/Systems/ObjectPathChannelEvaluatorSystem.h
@@ -0,0 +1,21 @@
+// Copyright Epic Games, Inc. All Rights Reserved.
+
+#pragma once
+
+#include "EntitySystem/MovieSceneEntitySystem.h"
+#include "ObjectPathChannelEvaluatorSystem.generated.h"
+
+/**
+ * System that is responsible for evaluating object path channels.
+ */
+UCLASS(MinimalAPI)
+class UObjectPathChannelEvaluatorSystem : public UMovieSceneEntitySystem
+{
+public:
+
+	GENERATED_BODY()
+
+	UObjectPathChannelEvaluatorSystem(const FObjectInitializer& ObjInit);
+
+	virtual void OnRun(FSystemTaskPrerequisites& InPrerequisites, FSystemSubsequentTasks& Subsequents) override;
+};
\ No newline at end of file
