Sequencer: Foundational improvements to overlapping entity trackers and pre-animated state

  - Both overlapping trackers and pre-animated state now support reference collection through an AddReferencedObjectsForComponent override on the key or storage types.
  - Unified much of the pre-animated state templates and moved core BeginTracking and CachePreAnimatedValues functionality into a single class
  - Bound object mechanisms are now handled automatically by the base type through traits

diff --git a/Engine/Source/Runtime/MovieScene/Private/Evaluation/MovieScenePreAnimatedState.cpp b/Engine/Source/Runtime/MovieScene/Private/Evaluation/MovieScenePreAnimatedState.cpp
--- a/Engine/Source/Runtime/MovieScene/Private/Evaluation/MovieScenePreAnimatedState.cpp
+++ b/Engine/Source/Runtime/MovieScene/Private/Evaluation/MovieScenePreAnimatedState.cpp
@@ -221,9 +221,9 @@ void FMovieScenePreAnimatedState::RestorePreAnimatedState(UObject& Object, TFunc
 	{
 		TFunctionRef<bool(FMovieSceneAnimTypeID)>* Filter;
 
-		virtual bool CanRestore(const FPreAnimatedObjectTokenTraits::FAnimatedKey& InKey) const override
+		virtual bool CanRestore(const FPreAnimatedObjectTokenTraits::KeyType& InKey) const override
 		{
-			return (*Filter)(InKey.AnimTypeID);
+			return (*Filter)(InKey.Get<1>());
 		}
 	} RestoreMask;
 	RestoreMask.Filter = &InFilter;
diff --git a/Engine/Source/Runtime/MovieScene/Private/Evaluation/PreAnimatedState/MovieScenePreAnimatedObjectTokenStorage.cpp b/Engine/Source/Runtime/MovieScene/Private/Evaluation/PreAnimatedState/MovieScenePreAnimatedObjectTokenStorage.cpp
--- a/Engine/Source/Runtime/MovieScene/Private/Evaluation/PreAnimatedState/MovieScenePreAnimatedObjectTokenStorage.cpp
+++ b/Engine/Source/Runtime/MovieScene/Private/Evaluation/PreAnimatedState/MovieScenePreAnimatedObjectTokenStorage.cpp
@@ -11,27 +11,12 @@ namespace MovieScene
 
 TAutoRegisterPreAnimatedStorageID<FAnimTypePreAnimatedStateObjectStorage> FAnimTypePreAnimatedStateObjectStorage::StorageID;
 
-void FAnimTypePreAnimatedStateObjectStorage::Initialize(FPreAnimatedStorageID InStorageID, FPreAnimatedStateExtension* InParentExtension)
-{
-	TPreAnimatedStateStorage<FPreAnimatedObjectTokenTraits>::Initialize(InStorageID, InParentExtension);
-
-	ObjectGroupManager = InParentExtension->GetOrCreateGroupManager<FPreAnimatedObjectGroupManager>();
-}
-
-void FAnimTypePreAnimatedStateObjectStorage::OnObjectReplaced(FPreAnimatedStorageIndex StorageIndex, const FObjectKey& OldObject, const FObjectKey& NewObject)
-{
-	FPreAnimatedObjectTokenTraits::FAnimatedKey ExistingKey = GetKey(StorageIndex);
-	ExistingKey.BoundObject = NewObject;
-
-	ReplaceKey(StorageIndex, ExistingKey);
-}
-
 FPreAnimatedStateEntry FAnimTypePreAnimatedStateObjectStorage::MakeEntry(UObject* Object, FMovieSceneAnimTypeID AnimTypeID)
 {
-	FPreAnimatedObjectTokenTraits::FAnimatedKey Key{ Object, AnimTypeID };
+	FPreAnimatedObjectTokenTraits::KeyType Key{ Object, AnimTypeID };
 
 	// Begin by finding or creating a pre-animated state group for this bound object
-	FPreAnimatedStorageGroupHandle Group = ObjectGroupManager->MakeGroupForKey(Object);
+	FPreAnimatedStorageGroupHandle Group = this->Traits.MakeGroup(Object);
 
 	// Find the storage index for the specific anim-type and object we're animating
 	FPreAnimatedStorageIndex StorageIndex = GetOrCreateStorageIndex(Key);
diff --git a/Engine/Source/Runtime/MovieScene/Private/Evaluation/PreAnimatedState/MovieScenePreAnimatedStateExtension.cpp b/Engine/Source/Runtime/MovieScene/Private/Evaluation/PreAnimatedState/MovieScenePreAnimatedStateExtension.cpp
--- a/Engine/Source/Runtime/MovieScene/Private/Evaluation/PreAnimatedState/MovieScenePreAnimatedStateExtension.cpp
+++ b/Engine/Source/Runtime/MovieScene/Private/Evaluation/PreAnimatedState/MovieScenePreAnimatedStateExtension.cpp
@@ -28,6 +28,7 @@ FPreAnimatedStateExtension::FPreAnimatedStateExtension()
 void FPreAnimatedStateExtension::Initialize(UMovieSceneEntitySystemLinker* InLinker)
 {
 	Linker = InLinker;
+	InLinker->Events.AddReferencedObjects.AddRaw(this, &FPreAnimatedStateExtension::AddReferencedObjects);
 }
 
 FPreAnimatedStateExtension::~FPreAnimatedStateExtension()
@@ -605,6 +606,14 @@ void FPreAnimatedStateExtension::SavePreAnimatedStateDirectly(FMovieSceneAnimTyp
 	}
 }
 
+void FPreAnimatedStateExtension::AddReferencedObjects(UMovieSceneEntitySystemLinker*, FReferenceCollector& ReferenceCollector)
+{
+	for (TPair<FPreAnimatedStorageID, TSharedPtr<IPreAnimatedStorage>>& Pair : StorageImplementations)
+	{
+		Pair.Value->AddReferencedObjects(ReferenceCollector);
+	}
+}
+
 
 } // namespace MovieScene
 } // namespace UE
diff --git a/Engine/Source/Runtime/MovieScene/Private/Evaluation/PreAnimatedState/MovieScenePreAnimatedStateStorage.cpp b/Engine/Source/Runtime/MovieScene/Private/Evaluation/PreAnimatedState/MovieScenePreAnimatedStateStorage.cpp
--- a/Engine/Source/Runtime/MovieScene/Private/Evaluation/PreAnimatedState/MovieScenePreAnimatedStateStorage.cpp
+++ b/Engine/Source/Runtime/MovieScene/Private/Evaluation/PreAnimatedState/MovieScenePreAnimatedStateStorage.cpp
@@ -4,6 +4,8 @@
 #include "EntitySystem/BuiltInComponentTypes.h"
 #include "EntitySystem/EntityAllocationIterator.h"
 #include "EntitySystem/MovieSceneEntitySystemTypes.h"
+#include "Evaluation/PreAnimatedState/MovieScenePreAnimatedStateExtension.h"
+#include "Evaluation/PreAnimatedState/MovieScenePreAnimatedObjectGroupManager.h"
 
 namespace UE
 {
@@ -16,6 +18,15 @@ FPreAnimatedTrackerParams::FPreAnimatedTrackerParams(FEntityAllocationProxy Item
 	bWantsRestoreState = Item.GetAllocationType().Contains(FBuiltInComponentTypes::Get()->Tags.RestoreState);
 }
 
+void FBoundObjectPreAnimatedStateTraits::Initialize(FPreAnimatedStorageID InStorageID, FPreAnimatedStateExtension* InParentExtension)
+{
+	ObjectGroupManager = InParentExtension->GetOrCreateGroupManager<FPreAnimatedObjectGroupManager>();
+}
+
+FPreAnimatedStorageGroupHandle FBoundObjectPreAnimatedStateTraits::MakeGroupImpl(UObject* BoundObject)
+{
+	return ObjectGroupManager->MakeGroupForKey(BoundObject);
+}
 
 } // namespace MovieScene
 } // namespace UE
diff --git a/Engine/Source/Runtime/MovieScene/Public/EntitySystem/MovieSceneComponentTypeInfo.h b/Engine/Source/Runtime/MovieScene/Public/EntitySystem/MovieSceneComponentTypeInfo.h
--- a/Engine/Source/Runtime/MovieScene/Public/EntitySystem/MovieSceneComponentTypeInfo.h
+++ b/Engine/Source/Runtime/MovieScene/Public/EntitySystem/MovieSceneComponentTypeInfo.h
@@ -18,13 +18,15 @@ namespace MovieScene
 struct FNotImplemented;
 struct FComponentTypeDebugInfo;
 
+
 /**
  * Stub for components that do not need reference collection. Overload void AddReferencedObjectForComponent(FReferenceCollector* ReferenceCollector, T* Component) for your own implementation
  * NOTE: The return value is _very_ important - overloaded AddReferencedObjectForComponent functions should return void otherwise they will not get called by the entity manager.
  */
 inline FNotImplemented* AddReferencedObjectForComponent(...) { return nullptr; }
 
-inline void AddReferencedObjectForComponent(FReferenceCollector* ReferenceCollector, UObject** Component)
+template<typename T>
+typename TEnableIf<TPointerIsConvertibleFromTo<T, UObject>::Value>::Type AddReferencedObjectForComponent(FReferenceCollector* ReferenceCollector, T** Component)
 {
 	ReferenceCollector->AddReferencedObject(*Component);
 }
@@ -39,6 +41,12 @@ typename TEnableIf< TIsSame<decltype(T::StaticStruct()), decltype(T::StaticStruc
 	}
 }
 
+template<typename T>
+struct THasAddReferencedObjectForComponent
+{
+	static constexpr bool Value = !TIsSame< FNotImplemented*, decltype( AddReferencedObjectForComponent((FReferenceCollector*)0, (T*)0) ) >::Value;
+};
+
 /**
  * Class that is instantiated to specify custom behavior for interacting with component data
  */
diff --git a/Engine/Source/Runtime/MovieScene/Public/EntitySystem/MovieSceneEntityFactoryTemplates.h b/Engine/Source/Runtime/MovieScene/Public/EntitySystem/MovieSceneEntityFactoryTemplates.h
--- a/Engine/Source/Runtime/MovieScene/Public/EntitySystem/MovieSceneEntityFactoryTemplates.h
+++ b/Engine/Source/Runtime/MovieScene/Public/EntitySystem/MovieSceneEntityFactoryTemplates.h
@@ -258,7 +258,7 @@ TComponentTypeID<T> FComponentRegistry::NewComponentType(const TCHAR* const Debu
 	NewTypeInfo.bIsPreserved               = EnumHasAnyFlags(Flags, EComponentTypeFlags::Preserved);
 	NewTypeInfo.bIsCopiedToOutput          = EnumHasAnyFlags(Flags, EComponentTypeFlags::CopyToOutput);
 	NewTypeInfo.bIsMigratedToOutput        = EnumHasAnyFlags(Flags, EComponentTypeFlags::MigrateToOutput);
-	NewTypeInfo.bHasReferencedObjects      = !TIsSame< FNotImplemented*, decltype( AddReferencedObjectForComponent((FReferenceCollector*)0, (T*)0) ) >::Value;
+	NewTypeInfo.bHasReferencedObjects      = THasAddReferencedObjectForComponent<T>::Value;
 
 #if UE_MOVIESCENE_ENTITY_DEBUG
 	NewTypeInfo.DebugInfo                = MakeUnique<FComponentTypeDebugInfo>();
diff --git a/Engine/Source/Runtime/MovieScene/Public/EntitySystem/MovieSceneOverlappingEntityTracker.h b/Engine/Source/Runtime/MovieScene/Public/EntitySystem/MovieSceneOverlappingEntityTracker.h
--- a/Engine/Source/Runtime/MovieScene/Public/EntitySystem/MovieSceneOverlappingEntityTracker.h
+++ b/Engine/Source/Runtime/MovieScene/Public/EntitySystem/MovieSceneOverlappingEntityTracker.h
@@ -13,6 +13,7 @@
 #include "EntitySystem/MovieSceneEntitySystemTask.h"
 #include "EntitySystem/MovieSceneEntitySystemLinker.h"
 #include "EntitySystem/BuiltInComponentTypes.h"
+#include "EntitySystem/MovieSceneComponentTypeInfo.h"
 
 namespace UE
 {
@@ -28,61 +29,145 @@ struct FEntityOutputAggregate
 };
 
 
+struct FGarbageTraits
+{
+	FORCEINLINE static constexpr bool IsGarbage(...)
+	{
+		return false;
+	}
+	FORCEINLINE static bool IsGarbage(UObject* InObject)
+	{
+		return FBuiltInComponentTypes::IsBoundObjectGarbage(InObject);
+	}
+	FORCEINLINE static constexpr void AddReferencedObjects(FReferenceCollector& ReferenceCollector, ...)
+	{}
+	template<typename T>
+	FORCEINLINE static typename TEnableIf<TPointerIsConvertibleFromTo<T, volatile const UObject>::Value>::Type
+		AddReferencedObjects(FReferenceCollector& ReferenceCollector, T*& InObjectPtr)
+	{
+		ReferenceCollector.AddReferencedObject(InObjectPtr);
+	}
+};
+
+template<typename... T> struct TGarbageTraitsImpl;
+template<typename... T> struct TOverlappingEntityInput;
+
+template<typename... T, int... Indices>
+struct TGarbageTraitsImpl<TIntegerSequence<int, Indices...>, T...>
+{
+	static constexpr bool bCanBeGarbage = (TPointerIsConvertibleFromTo<typename TDecay<typename TRemovePointer<T>::Type>::Type, UObject>::Value || ...);
+
+	static bool IsGarbage(TOverlappingEntityInput<T...>& InParam)
+	{
+		return (FGarbageTraits::IsGarbage(InParam.Key.template Get<Indices>()) || ...);
+	}
+
+	static void AddReferencedObjects(FReferenceCollector& ReferenceCollector, TOverlappingEntityInput<T...>& InParam)
+	{
+		(FGarbageTraits::AddReferencedObjects(ReferenceCollector, InParam.Key.template Get<Indices>()), ...);
+	}
+
+	template<typename CallbackType>
+	static void Unpack(const TTuple<T...>& InTuple, CallbackType&& Callback)
+	{
+		Callback(InTuple.template Get<Indices>()...);
+	}
+};
+
+
+template<typename... T>
+struct TOverlappingEntityInput
+{
+	using GarbageTraits = TGarbageTraitsImpl<TMakeIntegerSequence<int, sizeof...(T)>, T...>;
+
+	TTuple<T...> Key;
+
+	template<typename... ArgTypes>
+	TOverlappingEntityInput(ArgTypes&&... InArgs)
+		: Key(Forward<ArgTypes>(InArgs)...)
+	{}
+
+	friend uint32 GetTypeHash(const TOverlappingEntityInput<T...>& In)
+	{
+		return GetTypeHash(In.Key);
+	}
+	friend bool operator==(const TOverlappingEntityInput<T...>& A, const TOverlappingEntityInput<T...>& B)
+	{
+		return A.Key == B.Key;
+	}
+
+	template<typename CallbackType>
+	void Unpack(CallbackType&& Callback)
+	{
+		GarbageTraits::Unpack(Key, MoveTemp(Callback));
+	}
+};
+
+
 /**
  * Templated utility class that assists in tracking the state of many -> one data relationships in an FEntityManager.
- * KeyType defines the component type which defines the key that determines whether an entity animates the same output.
+ * InputKeyTypes defines the component type(s) which defines the key that determines whether an entity animates the same output.
  * OutputType defines the user-specfied data to be associated with the multiple inputs (ie, its output)
- * NOTE: Where KeyType is a UObject* it is recommended TOverlappingEntityTracker_BoundObject is used instead, as it provides
- * garbage collection and reference counting functions.
+ * NOTE: Where any of InputKeyTypes is a UObject*, AddReferencedObjects and CleanupGarbage must be called
  */
-template<typename KeyType, typename OutputType>
-struct TOverlappingEntityTracker
+template<typename OutputType, typename... InputKeyTypes>
+struct TOverlappingEntityTrackerImpl
 {
+	using KeyType = TOverlappingEntityInput<InputKeyTypes...>;
 	using ParamType = typename TCallTraits<KeyType>::ParamType;
 
+	bool IsInitialized() const
+	{
+		return bIsInitialized;
+	}
+
 	/**
 	 * Update this tracker by iterating any entity that contains InKeyComponent, and matches the additional optional filter
 	 * Only entities tagged as NeedsLink or NeedsUnlink are iterated, invalidating their outputs
 	 */
-	void Update(UMovieSceneEntitySystemLinker* Linker, TComponentTypeID<KeyType> InKeyComponent, const FEntityComponentFilter& InFilter)
+	void Update(UMovieSceneEntitySystemLinker* Linker, TComponentTypeID<InputKeyTypes>... InKeyComponents, const FEntityComponentFilter& InFilter)
 	{
+		check(bIsInitialized);
+
 		FBuiltInComponentTypes* BuiltInComponents = FBuiltInComponentTypes::Get();
 
 		// Visit unlinked entities
 		TFilteredEntityTask<>(TEntityTaskComponents<>())
 		.CombineFilter(InFilter)
-		.FilterAll({ BuiltInComponents->Tags.NeedsUnlink, InKeyComponent })
+		.FilterAll({ BuiltInComponents->Tags.NeedsUnlink, FComponentTypeID(InKeyComponents)... })
 		.Iterate_PerAllocation(&Linker->EntityManager, [this](const FEntityAllocation* Allocation){ this->VisitUnlinkedAllocation(Allocation); });
 
 		// Visit newly or re-linked entities
 		FEntityTaskBuilder()
-		.Read(InKeyComponent)
+		.ReadAllOf(InKeyComponents...)
 		.CombineFilter(InFilter)
 		.FilterAll({ BuiltInComponents->Tags.NeedsLink })
-		.Iterate_PerAllocation(&Linker->EntityManager, [this](const FEntityAllocation* Allocation, TRead<KeyType> ReadKeys){ this->VisitLinkedAllocation(Allocation, ReadKeys); });
+		.Iterate_PerAllocation(&Linker->EntityManager, [this](const FEntityAllocation* Allocation, TRead<InputKeyTypes>... ReadKeys){ this->VisitLinkedAllocation(Allocation, ReadKeys...); });
 	}
 
 	/**
 	 * Update this tracker by iterating any entity that contains InKeyComponent, and matches the additional optional filter
 	 * Only entities tagged as NeedsUnlink are iterated, invalidating their outputs
 	 */
-	void UpdateUnlinkedOnly(UMovieSceneEntitySystemLinker* Linker, TComponentTypeID<KeyType> InKeyComponent, const FEntityComponentFilter& InFilter)
+	void UpdateUnlinkedOnly(UMovieSceneEntitySystemLinker* Linker, TComponentTypeID<InputKeyTypes>... InKeyComponent, const FEntityComponentFilter& InFilter)
 	{
+		check(bIsInitialized);
+
 		FBuiltInComponentTypes* BuiltInComponents = FBuiltInComponentTypes::Get();
 
 		// Visit unlinked entities
 		TFilteredEntityTask<>(TEntityTaskComponents<>())
 		.CombineFilter(InFilter)
-		.FilterAll({ BuiltInComponents->Tags.NeedsUnlink, InKeyComponent })
+		.FilterAll({ BuiltInComponents->Tags.NeedsUnlink, InKeyComponent... })
 		.Iterate_PerAllocation(&Linker->EntityManager, [this](const FEntityAllocation* Allocation){ this->VisitUnlinkedAllocation(Allocation); });
 	}
 
 	/**
 	 * Update this tracker by (re)linking the specified allocation
 	 */
-	void VisitLinkedAllocation(const FEntityAllocation* Allocation, TRead<KeyType> ReadKeys)
+	void VisitLinkedAllocation(const FEntityAllocation* Allocation, TRead<InputKeyTypes>... ReadKeys)
 	{
-		VisitLinkedAllocationImpl(Allocation, ReadKeys);
+		VisitLinkedAllocationImpl(Allocation, ReadKeys...);
 	}
 
 	/**
@@ -98,13 +183,13 @@ struct TOverlappingEntityTracker
 	 *
 	 * InHandler    Any user-defined handler type that contains the following named functions:
 	 *                  // Called when an output is first created
-	 *                  void InitializeOutput(KeyType Object, TArrayView<const FMovieSceneEntityID> Inputs, OutputType* Output, FEntityOutputAggregate Aggregate);
+	 *                  void InitializeOutput(InputKeyTypes... Inputs, TArrayView<const FMovieSceneEntityID> Inputs, OutputType* Output, FEntityOutputAggregate Aggregate);
 	 *
 	 *                  // Called when an output has been updated with new inputs
-	 *                  void UpdateOutput(KeyType Object, TArrayView<const FMovieSceneEntityID> Inputs, OutputType* Output, FEntityOutputAggregate Aggregate);
+	 *                  void UpdateOutput(InputKeyTypes... Inputs, TArrayView<const FMovieSceneEntityID> Inputs, OutputType* Output, FEntityOutputAggregate Aggregate);
 	 *
 	 *                  // Called when all an output's inputs are no longer relevant, and as such the output should be destroyed (or restored)
-	 *                  void DestroyOutput(KeyType Object, OutputType* Output, FEntityOutputAggregate Aggregate);
+	 *                  void DestroyOutput(InputKeyTypes... Inputs, OutputType* Output, FEntityOutputAggregate Aggregate);
 	 */
 	template<typename HandlerType>
 	void ProcessInvalidatedOutputs(UMovieSceneEntitySystemLinker* Linker, HandlerType&& InHandler)
@@ -136,16 +221,22 @@ struct TOverlappingEntityTracker
 
 					if (NewOutputs.IsValidIndex(OutputIndex) && NewOutputs[OutputIndex] == true)
 					{
-						InHandler.InitializeOutput(Output.Key, InputArray, &Output.OutputData, Output.Aggregate);
+						Output.Key.Unpack([&InHandler, &InputArray, &Output](typename TCallTraits<InputKeyTypes>::ParamType... InKeys){
+							InHandler.InitializeOutput(InKeys..., InputArray, &Output.OutputData, Output.Aggregate);
+						});
 					}
 					else
 					{
-						InHandler.UpdateOutput(Output.Key, InputArray, &Output.OutputData, Output.Aggregate);
+						Output.Key.Unpack([&InHandler, &InputArray, &Output](typename TCallTraits<InputKeyTypes>::ParamType... InKeys){
+							InHandler.UpdateOutput(InKeys..., InputArray, &Output.OutputData, Output.Aggregate);
+						});
 					}
 				}
 				else
 				{
-					InHandler.DestroyOutput(Output.Key, &Output.OutputData, Output.Aggregate);
+					Output.Key.Unpack([&InHandler, &Output](typename TCallTraits<InputKeyTypes>::ParamType... InKeys){
+						InHandler.DestroyOutput(InKeys..., &Output.OutputData, Output.Aggregate);
+					});
 				}
 
 				if (InputArray.Num() == 0)
@@ -173,7 +264,9 @@ struct TOverlappingEntityTracker
 	{
 		for (FOutput& Output : Outputs)
 		{
-			InHandler.DestroyOutput(Output.Key, &Output.OutputData, Output.Aggregate);
+			Output.Key.Unpack([&InHandler, &Output](typename TCallTraits<InputKeyTypes>::ParamType... InKeys){
+				InHandler.DestroyOutput(InKeys..., &Output.OutputData, Output.Aggregate);
+			});
 		}
 
 		EntityToOutput.Empty();
@@ -243,8 +336,10 @@ struct TOverlappingEntityTracker
 
 protected:
 
-	void VisitLinkedAllocationImpl(const FEntityAllocation* Allocation, TRead<KeyType> Keys)
+	void VisitLinkedAllocationImpl(const FEntityAllocation* Allocation, TRead<InputKeyTypes>... Keys)
 	{
+		check(bIsInitialized);
+
 		const int32 Num = Allocation->Num();
 
 		const FMovieSceneEntityID* EntityIDs = Allocation->GetRawEntityIDs();
@@ -253,7 +348,7 @@ protected:
 		{
 			for (int32 Index = 0; Index < Num; ++Index)
 			{
-				const uint16 OutputIndex = MakeOutput(EntityIDs[Index], Keys[Index]);
+				const uint16 OutputIndex = MakeOutput(EntityIDs[Index], TOverlappingEntityInput<InputKeyTypes...>(Keys[Index]...));
 				Outputs[OutputIndex].Aggregate.bNeedsRestoration = true;
 			}
 		}
@@ -261,13 +356,15 @@ protected:
 		{
 			for (int32 Index = 0; Index < Num; ++Index)
 			{
-				MakeOutput(EntityIDs[Index], Keys[Index]);
+				MakeOutput(EntityIDs[Index], TOverlappingEntityInput<InputKeyTypes...>(Keys[Index]...));
 			}
 		}
 	}
 
 	void VisitUnlinkedAllocationImpl(const FEntityAllocation* Allocation)
 	{
+		check(bIsInitialized);
+
 		const int32 Num = Allocation->Num();
 		const FMovieSceneEntityID* EntityIDs = Allocation->GetRawEntityIDs();
 
@@ -362,16 +459,55 @@ protected:
 
 	TBitArray<> InvalidatedOutputs, NewOutputs;
 
+	bool bIsInitialized = false;
+
 	static constexpr uint16 NO_OUTPUT = MAX_uint16;
 };
 
 
-template<typename OutputType>
-struct TOverlappingEntityTracker_BoundObject : TOverlappingEntityTracker<UObject*, OutputType>
+template<typename OutputType, typename... KeyType>
+struct TOverlappingEntityTracker_NoGarbage : TOverlappingEntityTrackerImpl<OutputType, KeyType...>
 {
-	using FOutput = typename TOverlappingEntityTracker<UObject*, OutputType>::FOutput;
+	void Initialize(UMovieSceneEntitySystem* OwningSystem)
+	{
+		this->bIsInitialized = true;
+	}
+};
 
-	void CleanupGarbage()
+template<typename OutputType, typename... KeyType>
+struct TOverlappingEntityTracker_WithGarbage : TOverlappingEntityTrackerImpl<OutputType, KeyType...>
+{
+	using ThisType = TOverlappingEntityTracker_WithGarbage<OutputType, KeyType...>;
+	using Super = TOverlappingEntityTrackerImpl<OutputType, KeyType...>;
+	using typename Super::FOutput;
+	using typename Super::KeyType;
+
+	~TOverlappingEntityTracker_WithGarbage()
+	{
+		UMovieSceneEntitySystem* OwningSystem = WeakOwningSystem.GetEvenIfUnreachable();
+		UMovieSceneEntitySystemLinker* Linker = OwningSystem ? OwningSystem->GetLinker() : nullptr;
+		if (Linker)
+		{
+			Linker->Events.TagGarbage.RemoveAll(this);
+			Linker->Events.AddReferencedObjects.RemoveAll(this);
+		}
+	}
+
+	void Initialize(UMovieSceneEntitySystem* OwningSystem)
+	{
+		if (this->bIsInitialized)
+		{
+			return;
+		}
+
+		this->bIsInitialized = true;
+		WeakOwningSystem = OwningSystem;
+
+		OwningSystem->GetLinker()->Events.TagGarbage.AddRaw(this, &ThisType::CleanupGarbage);
+		OwningSystem->GetLinker()->Events.AddReferencedObjects.AddRaw(this, &ThisType::AddReferencedObjects);
+	}
+
+	void CleanupGarbage(UMovieSceneEntitySystemLinker* Linker)
 	{
 		for (int32 Index = this->Outputs.GetMaxIndex()-1; Index >= 0; --Index)
 		{
@@ -382,7 +518,7 @@ struct TOverlappingEntityTracker_BoundObject : TOverlappingEntityTracker<UObject
 			const uint16 OutputIndex = static_cast<uint16>(Index);
 
 			FOutput& Output = this->Outputs[Index];
-			if (FBuiltInComponentTypes::IsBoundObjectGarbage(Output.Key))
+			if (KeyType::GarbageTraits::IsGarbage(Output.Key))
 			{
 				this->Outputs.RemoveAt(Index, 1);
 
@@ -396,26 +532,43 @@ struct TOverlappingEntityTracker_BoundObject : TOverlappingEntityTracker<UObject
 
 		for (auto It = this->KeyToOutput.CreateIterator(); It; ++It)
 		{
-			UObject* Key = It.Key();
-			if (FBuiltInComponentTypes::IsBoundObjectGarbage(Key))
+			if (KeyType::GarbageTraits::IsGarbage(It.Key()) || !this->Outputs.IsValidIndex(It.Value()))
 			{
 				It.RemoveCurrent();
 			}
 		}
 	}
 
-	void AddReferencedObjects(FReferenceCollector& ReferenceCollector)
+	void AddReferencedObjects(UMovieSceneEntitySystemLinker* Linker, FReferenceCollector& ReferenceCollector)
 	{
-		ReferenceCollector.AddReferencedObjects(this->KeyToOutput);
+		for (TPair<KeyType, uint16>& Pair : this->KeyToOutput)
+		{
+			KeyType::GarbageTraits::AddReferencedObjects(ReferenceCollector, Pair.Key);
+		}
 
-		for (FOutput& Output : this->Outputs)
+		if constexpr (THasAddReferencedObjectForComponent<KeyType>::Value || THasAddReferencedObjectForComponent<OutputType>::Value)
 		{
-			ReferenceCollector.AddReferencedObject(Output.Key);
+			for (FOutput& Output : this->Outputs)
+			{
+				KeyType::GarbageTraits::AddReferencedObjects(ReferenceCollector, Output.Key);
+				KeyType::GarbageTraits::AddReferencedObjects(ReferenceCollector, Output.OutputData);
+			}
 		}
 	}
+
+protected:
+
+	TWeakObjectPtr<UMovieSceneEntitySystem> WeakOwningSystem;
 };
 
 
 
+template<typename OutputType, typename... KeyType>
+using TOverlappingEntityTracker = typename TChooseClass<
+	(THasAddReferencedObjectForComponent<KeyType>::Value || ...) || THasAddReferencedObjectForComponent<OutputType>::Value,
+	TOverlappingEntityTracker_WithGarbage<OutputType, KeyType...>,
+	TOverlappingEntityTracker_NoGarbage<OutputType, KeyType...>
+>::Result;
+
 } // namespace MovieScene
 } // namespace UE
diff --git a/Engine/Source/Runtime/MovieScene/Public/Evaluation/PreAnimatedState/IMovieScenePreAnimatedStorage.h b/Engine/Source/Runtime/MovieScene/Public/Evaluation/PreAnimatedState/IMovieScenePreAnimatedStorage.h
--- a/Engine/Source/Runtime/MovieScene/Public/Evaluation/PreAnimatedState/IMovieScenePreAnimatedStorage.h
+++ b/Engine/Source/Runtime/MovieScene/Public/Evaluation/PreAnimatedState/IMovieScenePreAnimatedStorage.h
@@ -12,6 +12,7 @@
 struct FObjectKey;
 struct FMovieSceneEvaluationKey;
 
+class FReferenceCollector;
 class FTrackInstancePropertyBindings;
 class UMovieSceneEntitySystemLinker;
 
@@ -41,6 +42,8 @@ struct IPreAnimatedStorage : TSharedFromThis<IPreAnimatedStorage>
 
 	virtual IPreAnimatedObjectPropertyStorage* AsPropertyStorage() { return nullptr; }
 	virtual IPreAnimatedObjectEntityStorage*   AsObjectStorage()   { return nullptr; }
+
+	virtual void AddReferencedObjects(FReferenceCollector& ReferenceCollector) {}
 };
 
 
diff --git a/Engine/Source/Runtime/MovieScene/Public/Evaluation/PreAnimatedState/MovieScenePreAnimatedMasterTokenStorage.h b/Engine/Source/Runtime/MovieScene/Public/Evaluation/PreAnimatedState/MovieScenePreAnimatedMasterTokenStorage.h
--- a/Engine/Source/Runtime/MovieScene/Public/Evaluation/PreAnimatedState/MovieScenePreAnimatedMasterTokenStorage.h
+++ b/Engine/Source/Runtime/MovieScene/Public/Evaluation/PreAnimatedState/MovieScenePreAnimatedMasterTokenStorage.h
@@ -18,7 +18,7 @@ namespace UE
 namespace MovieScene
 {
 
-struct FPreAnimatedMasterTokenTraits
+struct FPreAnimatedMasterTokenTraits : FDefaultPreAnimatedStateTraits
 {
 	using KeyType     = FMovieSceneAnimTypeID;
 	using StorageType = IMovieScenePreAnimatedGlobalTokenPtr;
diff --git a/Engine/Source/Runtime/MovieScene/Public/Evaluation/PreAnimatedState/MovieScenePreAnimatedObjectStorage.h b/Engine/Source/Runtime/MovieScene/Public/Evaluation/PreAnimatedState/MovieScenePreAnimatedObjectStorage.h
--- a/Engine/Source/Runtime/MovieScene/Public/Evaluation/PreAnimatedState/MovieScenePreAnimatedObjectStorage.h
+++ b/Engine/Source/Runtime/MovieScene/Public/Evaluation/PreAnimatedState/MovieScenePreAnimatedObjectStorage.h
@@ -44,20 +44,6 @@ public:
 		return this;
 	}
 
-	void Initialize(FPreAnimatedStorageID InStorageID, FPreAnimatedStateExtension* InParentExtension) override
-	{
-		TPreAnimatedStateStorage<ObjectTraits>::Initialize(InStorageID, InParentExtension);
-
-		ObjectGroupManager = InParentExtension->GetOrCreateGroupManager<FPreAnimatedObjectGroupManager>();
-	}
-
-	void OnObjectReplaced(FPreAnimatedStorageIndex StorageIndex, const FObjectKey& OldObject, const FObjectKey& NewObject) override
-	{
-		KeyType ExistingKey = this->GetKey(StorageIndex);
-		ExistingKey = NewObject;
-		this->ReplaceKey(StorageIndex, ExistingKey);
-	}
-
 	void BeginTrackingEntities(const FPreAnimatedTrackerParams& Params, TRead<FMovieSceneEntityID> EntityIDs, TRead<FRootInstanceHandle> InstanceHandles, TRead<UObject*> BoundObjects) override
 	{
 		const int32 Num = Params.Num;
@@ -74,7 +60,7 @@ public:
 			UObject* BoundObject = BoundObjects[Index];
 			FObjectKey Key{ BoundObject };
 
-			FPreAnimatedStorageGroupHandle GroupHandle  = this->ObjectGroupManager->MakeGroupForKey(Key);
+			FPreAnimatedStorageGroupHandle GroupHandle  = this->Traits.MakeGroup(BoundObject);
 			FPreAnimatedStorageIndex       StorageIndex = this->GetOrCreateStorageIndex(Key);
 
 			FPreAnimatedStateEntry Entry{ GroupHandle, FPreAnimatedStateCachedValueHandle{ this->StorageID, StorageIndex } };
@@ -95,7 +81,7 @@ public:
 		FObjectKey Key{ BoundObject };
 
 		FPreAnimatedStorageIndex       StorageIndex = this->GetOrCreateStorageIndex(Key);
-		FPreAnimatedStorageGroupHandle GroupHandle  = this->ObjectGroupManager->MakeGroupForKey(Key);
+		FPreAnimatedStorageGroupHandle GroupHandle  = this->Traits.MakeGroup(BoundObject);
 		FPreAnimatedStateEntry         Entry        = FPreAnimatedStateEntry{ GroupHandle, FPreAnimatedStateCachedValueHandle{ this->StorageID, StorageIndex } };
 
 		if (this->ParentExtension->IsCapturingGlobalState())
@@ -128,10 +114,6 @@ public:
 			this->ForciblyPersistStorage(StorageIndex);
 		}
 	}
-
-protected:
-
-	TSharedPtr<FPreAnimatedObjectGroupManager> ObjectGroupManager;
 };
 
 
diff --git a/Engine/Source/Runtime/MovieScene/Public/Evaluation/PreAnimatedState/MovieScenePreAnimatedObjectTokenStorage.h b/Engine/Source/Runtime/MovieScene/Public/Evaluation/PreAnimatedState/MovieScenePreAnimatedObjectTokenStorage.h
--- a/Engine/Source/Runtime/MovieScene/Public/Evaluation/PreAnimatedState/MovieScenePreAnimatedObjectTokenStorage.h
+++ b/Engine/Source/Runtime/MovieScene/Public/Evaluation/PreAnimatedState/MovieScenePreAnimatedObjectTokenStorage.h
@@ -21,35 +21,18 @@ namespace MovieScene
 {
 
 
-struct FPreAnimatedObjectTokenTraits
+struct FPreAnimatedObjectTokenTraits : FBoundObjectPreAnimatedStateTraits
 {
-	struct FAnimatedKey
-	{
-		FObjectKey BoundObject;
-		FMovieSceneAnimTypeID AnimTypeID;
-
-		friend uint32 GetTypeHash(const FAnimatedKey& In)
-		{
-			return HashCombine(GetTypeHash(In.BoundObject), GetTypeHash(In.AnimTypeID));
-		}
-		friend bool operator==(const FAnimatedKey& A, const FAnimatedKey& B)
-		{
-			return A.BoundObject == B.BoundObject && A.AnimTypeID == B.AnimTypeID;
-		}
-	};
-
-	using KeyType     = FAnimatedKey;
+	using KeyType     = TTuple<FObjectKey, FMovieSceneAnimTypeID>;
 	using StorageType = IMovieScenePreAnimatedTokenPtr;
 
-	static void RestorePreAnimatedValue(const FAnimatedKey& InKey, IMovieScenePreAnimatedTokenPtr& Token, const FRestoreStateParams& Params)
+	static void RestorePreAnimatedValue(const KeyType& InKey, IMovieScenePreAnimatedTokenPtr& Token, const FRestoreStateParams& Params)
 	{
-		if (UObject* Object = InKey.BoundObject.ResolveObjectPtr())
+		if (UObject* Object = InKey.Get<0>().ResolveObjectPtr())
 		{
 			Token->RestoreState(*Object, Params);
 		}
 	}
-
-	TSharedPtr<FPreAnimatedObjectGroupManager> ObjectGroupManager;
 };
 
 
@@ -58,14 +41,6 @@ struct MOVIESCENE_API FAnimTypePreAnimatedStateObjectStorage : TPreAnimatedState
 	static TAutoRegisterPreAnimatedStorageID<FAnimTypePreAnimatedStateObjectStorage> StorageID;
 
 	FPreAnimatedStateEntry MakeEntry(UObject* Object, FMovieSceneAnimTypeID AnimTypeID);
-
-	void Initialize(FPreAnimatedStorageID InStorageID, FPreAnimatedStateExtension* InParentExtension) override;
-	void OnObjectReplaced(FPreAnimatedStorageIndex StorageIndex, const FObjectKey& OldObject, const FObjectKey& NewObject) override;
-
-
-private:
-
-	TSharedPtr<FPreAnimatedObjectGroupManager> ObjectGroupManager;
 };
 
 
diff --git a/Engine/Source/Runtime/MovieScene/Public/Evaluation/PreAnimatedState/MovieScenePreAnimatedPropertyStorage.h b/Engine/Source/Runtime/MovieScene/Public/Evaluation/PreAnimatedState/MovieScenePreAnimatedPropertyStorage.h
--- a/Engine/Source/Runtime/MovieScene/Public/Evaluation/PreAnimatedState/MovieScenePreAnimatedPropertyStorage.h
+++ b/Engine/Source/Runtime/MovieScene/Public/Evaluation/PreAnimatedState/MovieScenePreAnimatedPropertyStorage.h
@@ -15,7 +15,6 @@
 #include "Evaluation/PreAnimatedState/MovieSceneRestoreStateParams.h"
 #include "Evaluation/PreAnimatedState/MovieScenePreAnimatedStorageID.inl"
 #include "Evaluation/PreAnimatedState/MovieScenePreAnimatedStateStorage.h"
-#include "Evaluation/PreAnimatedState/MovieScenePreAnimatedObjectGroupManager.h"
 #include "Evaluation/PreAnimatedState/MovieScenePreAnimatedEntityCaptureSource.h"
 
 
@@ -25,57 +24,125 @@ namespace MovieScene
 {
 
 
-template<typename PropertyTraits, typename MetaDataTypes, typename MetaDataIndices>
-struct TPreAnimatedPropertyStorageImpl;
+/**
+ * Pre-animated property value, including cached meta-data
+ */
+template<typename StorageType, typename ...MetaDataTypes>
+struct TPreAnimatedPropertyValue
+{
+	TPreAnimatedPropertyValue()
+	{}
 
-template<typename PropertyTraits, typename ...MetaDataTypes, int ...MetaDataIndices>
-struct TPreAnimatedPropertyStorageImpl<PropertyTraits, TPropertyMetaData<MetaDataTypes...>, TIntegerSequence<int, MetaDataIndices...>>
-	: IPreAnimatedStorage
-	, IPreAnimatedObjectPropertyStorage
+	TPreAnimatedPropertyValue(typename TCallTraits<MetaDataTypes>::ParamType... InMetaData)
+		: MetaData(InMetaData...)
+	{}
+
+	StorageType Data;
+	TVariant<const FCustomPropertyAccessor*, uint16, TSharedPtr<FTrackInstancePropertyBindings>> Binding;
+	TTuple<MetaDataTypes...> MetaData;
+};
+
+
+/**
+ * Pre-animated property value, specialized for no meta-data
+ */
+template<typename StorageType>
+struct TPreAnimatedPropertyValue<StorageType>
+{
+	StorageType Data;
+	TVariant<const FCustomPropertyAccessor*, uint16, TSharedPtr<FTrackInstancePropertyBindings>> Binding;
+};
+
+
+/**
+ * Pre-Animated traits class that wraps a user-provided property trait that defines property accessors 
+ */
+template<typename PropertyTraits, typename MetaDataIndices, typename ...MetaDataTypes>
+struct TPreAnimatedPropertyTraits;
+
+template<typename PropertyTraits, int... MetaDataIndices, typename ...MetaDataTypes>
+struct TPreAnimatedPropertyTraits<PropertyTraits, TIntegerSequence<int, MetaDataIndices...>, MetaDataTypes...> : FBoundObjectPreAnimatedStateTraits
 {
-	using StorageType       = typename PropertyTraits::StorageType;
+	using KeyType = TTuple<FObjectKey, FName>;
+	using StorageType = TPreAnimatedPropertyValue<typename PropertyTraits::StorageType, MetaDataTypes...>;
 
-	TPreAnimatedPropertyStorageImpl(const FPropertyDefinition& InPropertyDefinition)
-		: MetaDataComponents(InPropertyDefinition.MetaDataTypes)
+	static void RestorePreAnimatedValue(const KeyType& InKey, StorageType& CachedValue, const FRestoreStateParams& Params)
 	{
-		check(MetaDataComponents.Num() == sizeof...(MetaDataIndices));
+		UObject* Object = InKey.Get<0>().ResolveObjectPtr();
+		if (!Object)
+		{
+			return;
+		}
 
-		if (InPropertyDefinition.CustomPropertyRegistration)
+		if (const uint16* FastOffset = CachedValue.Binding.template TryGet<uint16>())
 		{
-			CustomAccessors = InPropertyDefinition.CustomPropertyRegistration->GetAccessors();
+			PropertyTraits::SetObjectPropertyValue(Object, CachedValue.MetaData.template Get<MetaDataIndices>()..., *FastOffset, CachedValue.Data);
+		}
+		else  if (const TSharedPtr<FTrackInstancePropertyBindings>* Bindings = CachedValue.Binding.template TryGet<TSharedPtr<FTrackInstancePropertyBindings>>())
+		{
+			PropertyTraits::SetObjectPropertyValue(Object, CachedValue.MetaData.template Get<MetaDataIndices>()..., Bindings->Get(), CachedValue.Data);
+		}
+		else if (const FCustomPropertyAccessor* CustomAccessor = CachedValue.Binding.template Get<const FCustomPropertyAccessor*>())
+		{
+			PropertyTraits::SetObjectPropertyValue(Object, CachedValue.MetaData.template Get<MetaDataIndices>()..., *CustomAccessor, CachedValue.Data);
 		}
 	}
+};
 
-	FPreAnimatedStorageID GetStorageType() const override
-	{
-		return StorageID;
-	}
 
-	void Initialize(FPreAnimatedStorageID InStorageID, FPreAnimatedStateExtension* InParentExtension) override
-	{
-		this->Storage.Initialize(InStorageID, InParentExtension);
 
-		ObjectGroupManager = InParentExtension->GetOrCreateGroupManager<FPreAnimatedObjectGroupManager>();
-		ParentExtension = InParentExtension;
-		StorageID = InStorageID;
-	}
+/**
+ * Pre-Animated traits class that wraps a user-provided property trait that defines property accessors with no meta-data
+ */
+template<typename PropertyTraits>
+struct TPreAnimatedPropertyTraits<PropertyTraits, TPropertyMetaData<>, TIntegerSequence<int>> : FBoundObjectPreAnimatedStateTraits
+{
+	using KeyType = TTuple<FObjectKey, FName>;
+	using StorageType = TPreAnimatedPropertyValue<typename PropertyTraits::StorageType>;
 
-	void OnObjectReplaced(FPreAnimatedStorageIndex StorageIndex, const FObjectKey& OldObject, const FObjectKey& NewObject) override
+	static void RestorePreAnimatedValue(const KeyType& InKey, StorageType& CachedValue, const FRestoreStateParams& Params)
 	{
-		FAnimatedPropertyKey ExistingKey = this->Storage.GetKey(StorageIndex);
-		ExistingKey.BoundObject = NewObject;
+		UObject* Object = InKey.Get<0>().ResolveObjectPtr();
+		if (!Object)
+		{
+			return;
+		}
 
-		this->Storage.ReplaceKey(StorageIndex, ExistingKey);
+		if (const uint16* FastOffset = CachedValue.Binding.template TryGet<uint16>())
+		{
+			PropertyTraits::SetObjectPropertyValue(Object, *FastOffset, CachedValue.Data);
+		}
+		else  if (const TSharedPtr<FTrackInstancePropertyBindings>* Bindings = CachedValue.Binding.template TryGet<TSharedPtr<FTrackInstancePropertyBindings>>())
+		{
+			PropertyTraits::SetObjectPropertyValue(Object, Bindings->Get(), CachedValue.Data);
+		}
+		else if (const FCustomPropertyAccessor* CustomAccessor = CachedValue.Binding.template Get<const FCustomPropertyAccessor*>())
+		{
+			PropertyTraits::SetObjectPropertyValue(Object, *CustomAccessor, CachedValue.Data);
+		}
 	}
+};
 
-	EPreAnimatedStorageRequirement RestorePreAnimatedStateStorage(FPreAnimatedStorageIndex StorageIndex, EPreAnimatedStorageRequirement SourceRequirement, EPreAnimatedStorageRequirement TargetRequirement, const FRestoreStateParams& Params) override
-	{
-		return this->Storage.RestorePreAnimatedStateStorage(StorageIndex, SourceRequirement, TargetRequirement, Params);
-	}
 
-	EPreAnimatedStorageRequirement DiscardPreAnimatedStateStorage(FPreAnimatedStorageIndex StorageIndex, EPreAnimatedStorageRequirement SourceRequirement) override
+template<typename PropertyTraits, typename MetaDataTypes, typename MetaDataIndices>
+struct TPreAnimatedPropertyStorageImpl;
+
+template<typename PropertyTraits, typename ...MetaDataTypes, int ...MetaDataIndices>
+struct TPreAnimatedPropertyStorageImpl<PropertyTraits, TPropertyMetaData<MetaDataTypes...>, TIntegerSequence<int, MetaDataIndices...>>
+	: TPreAnimatedStateStorage<TPreAnimatedPropertyTraits<PropertyTraits, TIntegerSequence<int, MetaDataIndices...>, MetaDataTypes...>>
+	, IPreAnimatedObjectPropertyStorage
+{
+	using StorageType = typename TPreAnimatedPropertyTraits<PropertyTraits, TIntegerSequence<int, MetaDataIndices...>, MetaDataTypes...>::StorageType;
+
+	TPreAnimatedPropertyStorageImpl(const FPropertyDefinition& InPropertyDefinition)
+		: MetaDataComponents(InPropertyDefinition.MetaDataTypes)
 	{
-		return this->Storage.DiscardPreAnimatedStateStorage(StorageIndex, SourceRequirement);
+		check(MetaDataComponents.Num() == sizeof...(MetaDataIndices));
+
+		if (InPropertyDefinition.CustomPropertyRegistration)
+		{
+			CustomAccessors = InPropertyDefinition.CustomPropertyRegistration->GetAccessors();
+		}
 	}
 
 	IPreAnimatedObjectPropertyStorage* AsPropertyStorage() override
@@ -93,18 +160,18 @@ struct TPreAnimatedPropertyStorageImpl<PropertyTraits, TPropertyMetaData<MetaDat
 			return;
 		}
 
-		FPreAnimatedEntityCaptureSource* EntityMetaData = ParentExtension->GetOrCreateEntityMetaData();
+		FPreAnimatedEntityCaptureSource* EntityMetaData = this->ParentExtension->GetOrCreateEntityMetaData();
 		for (int32 Index = 0; Index < Num; ++Index)
 		{
 			UObject* BoundObject  = BoundObjects[Index];
 			FName    PropertyPath = PropertyBindings[Index].PropertyPath;
 
-			FAnimatedPropertyKey Key{ BoundObject, PropertyPath };
+			TTuple<FObjectKey, FName> Key{ BoundObject, PropertyPath };
 
-			FPreAnimatedStorageGroupHandle GroupHandle  = this->ObjectGroupManager->MakeGroupForKey(BoundObject);
-			FPreAnimatedStorageIndex       StorageIndex = this->Storage.GetOrCreateStorageIndex(Key);
+			FPreAnimatedStorageGroupHandle GroupHandle  = this->Traits.MakeGroup(BoundObject);
+			FPreAnimatedStorageIndex       StorageIndex = this->GetOrCreateStorageIndex(Key);
 
-			FPreAnimatedStateEntry Entry{ GroupHandle, FPreAnimatedStateCachedValueHandle{ StorageID, StorageIndex } };
+			FPreAnimatedStateEntry Entry{ GroupHandle, FPreAnimatedStateCachedValueHandle{ this->StorageID, StorageIndex } };
 			EntityMetaData->BeginTrackingEntity(Entry, EntityIDs[Index], InstanceHandles[Index], bWantsRestore);
 		}
 	}
@@ -127,12 +194,12 @@ struct TPreAnimatedPropertyStorageImpl<PropertyTraits, TPropertyMetaData<MetaDat
 			UObject* BoundObject  = BoundObjects[Index];
 			FName    PropertyPath = PropertyBindings[Index].PropertyPath;
 
-			FAnimatedPropertyKey Key{ BoundObject, PropertyPath };
+			TTuple<FObjectKey, FName> Key{ BoundObject, PropertyPath };
 
-			FPreAnimatedStorageGroupHandle GroupHandle  = this->ObjectGroupManager->MakeGroupForKey(Key.BoundObject);
-			FPreAnimatedStorageIndex       StorageIndex = this->Storage.GetOrCreateStorageIndex(Key);
+			FPreAnimatedStorageGroupHandle GroupHandle  = this->Traits.MakeGroup(BoundObject);
+			FPreAnimatedStorageIndex       StorageIndex = this->GetOrCreateStorageIndex(Key);
 
-			FPreAnimatedStateEntry Entry{ GroupHandle, FPreAnimatedStateCachedValueHandle{ StorageID, StorageIndex } };
+			FPreAnimatedStateEntry Entry{ GroupHandle, FPreAnimatedStateCachedValueHandle{ this->StorageID, StorageIndex } };
 
 			if (this->ParentExtension->IsCapturingGlobalState())
 			{
@@ -144,10 +211,9 @@ struct TPreAnimatedPropertyStorageImpl<PropertyTraits, TPropertyMetaData<MetaDat
 			}
 
 			EPreAnimatedStorageRequirement StorageRequirement = this->ParentExtension->GetStorageRequirement(Entry);
-			if (!Storage.IsStorageRequirementSatisfied(StorageIndex, StorageRequirement))
+			if (!this->IsStorageRequirementSatisfied(StorageIndex, StorageRequirement))
 			{
-				FPreAnimatedProperty NewValue;
-				NewValue.MetaData = MakeTuple(MetaData.template Get<MetaDataIndices>()[Index]...);
+				StorageType NewValue(MetaData.template Get<MetaDataIndices>()[Index]...);
 
 				if (Fast)
 				{
@@ -169,75 +235,20 @@ struct TPreAnimatedPropertyStorageImpl<PropertyTraits, TPropertyMetaData<MetaDat
 					PropertyTraits::GetObjectPropertyValue(BoundObjects[Index], MetaData.template Get<MetaDataIndices>()[Index]..., Bindings.Get(), NewValue.Data);
 				}
 
-				Storage.AssignPreAnimatedValue(StorageIndex, StorageRequirement, MoveTemp(NewValue));
+				this->AssignPreAnimatedValue(StorageIndex, StorageRequirement, MoveTemp(NewValue));
 			}
 
 			if (Params.bForcePersist)
 			{
-				this->Storage.ForciblyPersistStorage(StorageIndex);
+				this->ForciblyPersistStorage(StorageIndex);
 			}
 		}
 	}
 
 protected:
 
-	struct FAnimatedPropertyKey
-	{
-		FObjectKey BoundObject;
-		FName PropertyPath;
-
-		friend uint32 GetTypeHash(const FAnimatedPropertyKey& In)
-		{
-			return HashCombine(GetTypeHash(In.BoundObject), GetTypeHash(In.PropertyPath));
-		}
-		friend bool operator==(const FAnimatedPropertyKey& A, const FAnimatedPropertyKey& B)
-		{
-			return A.BoundObject == B.BoundObject && A.PropertyPath == B.PropertyPath;
-		}
-	};
-
-	struct FPreAnimatedProperty
-	{
-		StorageType Data;
-		TVariant<const FCustomPropertyAccessor*, uint16, TSharedPtr<FTrackInstancePropertyBindings>> Binding;
-		TTuple<MetaDataTypes...> MetaData;
-	};
-
-	struct FPropertyStorageTraits
-	{
-		using KeyType = FAnimatedPropertyKey;
-		using StorageType = FPreAnimatedProperty;
-
-		void RestorePreAnimatedValue(const FAnimatedPropertyKey& InKey, FPreAnimatedProperty& CachedValue, const FRestoreStateParams& Params)
-		{
-			UObject* Object = InKey.BoundObject.ResolveObjectPtr();
-			if (!Object)
-			{
-				return;
-			}
-
-			if (const uint16* FastOffset = CachedValue.Binding.template TryGet<uint16>())
-			{
-				PropertyTraits::SetObjectPropertyValue(Object, CachedValue.MetaData.template Get<MetaDataIndices>()..., *FastOffset, CachedValue.Data);
-			}
-			else  if (const TSharedPtr<FTrackInstancePropertyBindings>* Bindings = CachedValue.Binding.template TryGet<TSharedPtr<FTrackInstancePropertyBindings>>())
-			{
-				PropertyTraits::SetObjectPropertyValue(Object, CachedValue.MetaData.template Get<MetaDataIndices>()..., Bindings->Get(), CachedValue.Data);
-			}
-			else if (const FCustomPropertyAccessor* CustomAccessor = CachedValue.Binding.template Get<const FCustomPropertyAccessor*>())
-			{
-				PropertyTraits::SetObjectPropertyValue(Object, CachedValue.MetaData.template Get<MetaDataIndices>()..., *CustomAccessor, CachedValue.Data);
-			}
-		}
-	};
-
 	TArrayView<const FComponentTypeID> MetaDataComponents;
 	FCustomAccessorView CustomAccessors;
-	TPreAnimatedStateStorage<FPropertyStorageTraits> Storage;
-	FPreAnimatedStorageID StorageID;
-	FPreAnimatedStateExtension* ParentExtension;
-
-	TSharedPtr<FPreAnimatedObjectGroupManager> ObjectGroupManager;
 };
 
 
diff --git a/Engine/Source/Runtime/MovieScene/Public/Evaluation/PreAnimatedState/MovieScenePreAnimatedStateExtension.h b/Engine/Source/Runtime/MovieScene/Public/Evaluation/PreAnimatedState/MovieScenePreAnimatedStateExtension.h
--- a/Engine/Source/Runtime/MovieScene/Public/Evaluation/PreAnimatedState/MovieScenePreAnimatedStateExtension.h
+++ b/Engine/Source/Runtime/MovieScene/Public/Evaluation/PreAnimatedState/MovieScenePreAnimatedStateExtension.h
@@ -289,6 +289,8 @@ private:
 	void FreeGroupInternal(FPreAnimatedStorageGroupHandle Handle);
 	bool ShouldCaptureAnyState() const;
 
+	void AddReferencedObjects(UMovieSceneEntitySystemLinker*, FReferenceCollector& ReferenceCollector);
+
 public:
 
 	/** The number of requests that have been made to capture global state - only one should exist per playing sequence */
diff --git a/Engine/Source/Runtime/MovieScene/Public/Evaluation/PreAnimatedState/MovieScenePreAnimatedStateStorage.h b/Engine/Source/Runtime/MovieScene/Public/Evaluation/PreAnimatedState/MovieScenePreAnimatedStateStorage.h
--- a/Engine/Source/Runtime/MovieScene/Public/Evaluation/PreAnimatedState/MovieScenePreAnimatedStateStorage.h
+++ b/Engine/Source/Runtime/MovieScene/Public/Evaluation/PreAnimatedState/MovieScenePreAnimatedStateStorage.h
@@ -3,23 +3,146 @@
 #pragma once
 
 #include "CoreTypes.h"
+#include "Templates/UnrealTypeTraits.h"
 #include "Containers/SortedMap.h"
 #include "Evaluation/PreAnimatedState/MovieSceneRestoreStateParams.h"
 #include "Evaluation/PreAnimatedState/IMovieScenePreAnimatedStorage.h"
 #include "Evaluation/PreAnimatedState/MovieScenePreAnimatedStateExtension.h"
+#include "Evaluation/PreAnimatedState/MovieScenePreAnimatedEntityCaptureSource.h"
+#include "EntitySystem/BuiltInComponentTypes.h"
+#include "EntitySystem/MovieSceneEntitySystemTask.h"
+#include "EntitySystem/MovieSceneEntitySystemLinker.h"
 
+class UObject;
 
 namespace UE
 {
 namespace MovieScene
 {
 
+struct FPreAnimatedObjectGroupManager;
+
+/**
+ * Default task that is provided to BeginTrackingEntitiesTask or CachePreAnimatedValuesTask to constrain
+ * the task using an additional filter, or to provide custom short-circuit behavior that allows the client
+ * to skip saving pre-animated state for some entities.
+ */
+template<typename... InputTypes>
+struct TSavePreAnimatedStateParams
+{
+	FEntityComponentFilter AdditionalFilter;
+
+	static constexpr bool ShouldSavePreAnimatedState(typename TCallTraits<InputTypes>::ParamType... Inputs)
+	{
+		return true;
+	}
+};
+
+struct FDefaultPreAnimatedStateTraits
+{
+	void Initialize(FPreAnimatedStorageID InStorageID, FPreAnimatedStateExtension* InParentExtension)
+	{}
+
+	/* Defined as a template rather than a variadic function to prevent error C4840 */
+	template<typename... T>
+	FPreAnimatedStorageGroupHandle MakeGroup(T&&... Unused)
+	{
+		return FPreAnimatedStorageGroupHandle();
+	}
+
+	template<typename KeyType>
+	void ReplaceObject(KeyType& InOutKey, const FObjectKey& NewObject)
+	{}
+};
+
+/**
+ * Group state class that groups pre-animated storage together by bound object
+ */
+struct MOVIESCENE_API FBoundObjectPreAnimatedStateTraits
+{
+	void Initialize(FPreAnimatedStorageID InStorageID, FPreAnimatedStateExtension* InParentExtension);
+
+	/* Defined as a template rather than a variadic function to prevent error C4840 */
+	template<typename... T>
+	FPreAnimatedStorageGroupHandle MakeGroup(UObject* BoundObject, T&&... Unused)
+	{
+		return MakeGroupImpl(BoundObject);
+	}
+	FPreAnimatedStorageGroupHandle MakeGroupImpl(UObject* BoundObject);
+
+	template<typename ...T>
+	void ReplaceObject(TTuple<FObjectKey, T...>& InOutKey, const FObjectKey& NewObject)
+	{
+		InOutKey.template Get<0>() = NewObject;
+	}
+	template<typename KeyType>
+	void ReplaceObject(KeyType& InOutKey, const FObjectKey& NewObject)
+	{
+		InOutKey.Object = NewObject;
+	}
+	template<typename ObjectType>
+	void ReplaceObject(TObjectKey<ObjectType>& InOutKey, const FObjectKey& NewObject)
+	{
+		if (ObjectType* CastResult = Cast<ObjectType>(NewObject.ResolveObjectPtr()))
+		{
+			InOutKey = CastResult;
+		}
+	}
+	void ReplaceObject(FObjectKey& InOutKey, const FObjectKey& NewObject)
+	{
+		InOutKey = NewObject;
+	}
+	TSharedPtr<FPreAnimatedObjectGroupManager> ObjectGroupManager;
+};
+
+
+/**
+ * Storage container for a specific type of pre-animated state as defined by the specified traits.
+ *
+ * Reference collection for KeyType and StorageType is automatically provided by way of an optional AddReferencedObjectForComponent
+ * override
+ *
+ * Traits must include a type definition or using alias for its KeyType and StorageType, defining the type of key to use
+ * for associating the pre-animated value and the storage value type respectively.
+ * 
+ * Additionally, traits must define a RestorePreAnimatedValue function that will be used by the storage
+ * container to restore data back to its previous value.
+
+ * An example trait that maps an object and a name identifier to a string would look like this:
+ * struct FExampleTraits
+ * {
+ *   using KeyType = TTuple<FObjectKey, FName>;
+ *   using StorageType = FString;
+ *
+ *   static void RestorePreAnimatedValue(const KeyType& InKey, const FString& PreviousString, const FRestoreStateParams& Params)
+ *   {
+ *     if (UMyObjectType* Object = Cast<UMyObjectType>(InKey.Get<0>().ResolveObjectPtr()))
+ *     {
+ *       Object->SetStringValue(InKey.Get<1>(), PreviousString);
+ *     }
+ *   }
+ * }
+ *
+ * Furthermore, if the CachePreAnimatedValuesTask is used, traits must implement a CachePreAnimatedValue function that receives
+ * the contributor component types, and returns the cached value:
+ *
+ *   static FString CachePreAnimatedValue(UObject* InObject, const FName& StringName)
+ *   {
+ *     UMyObjectType* Object = CastChecked<UMyObjectType>(InObject);
+ *     return Object->GetStringValue(StringName);
+ *   }
+ *
+ * Traits may be stateful if desired. Stateful traits must be provided to the constructor in order to be valid.
+ */
 template<typename StorageTraits>
 struct TPreAnimatedStateStorage : IPreAnimatedStorage
 {
+	/** The key type to use as a key in the map that associates pre-animated state values to owners (ie, an object and property name)*/
 	using KeyType     = typename StorageTraits::KeyType;
+	/** The value type this storage should store (ie, the actual property type) */
 	using StorageType = typename StorageTraits::StorageType;
 
+	/** Interface used for restoring pre-animated state that allows an external system to control whether state should be restored or not */
 	struct IRestoreMask
 	{
 		virtual ~IRestoreMask(){}
@@ -34,48 +157,40 @@ struct TPreAnimatedStateStorage : IPreAnimatedStorage
 		: Traits(MoveTemp(InTraits))
 	{}
 
+	/** Pre-Animated storage is not copyable */
 	TPreAnimatedStateStorage(const TPreAnimatedStateStorage&) = delete;
 	TPreAnimatedStateStorage& operator=(const TPreAnimatedStateStorage&) = delete;
 
 public:
 
+	/** Called when this storage is created inside the pre-animated state extension on a UMovieSceneEntitySystemLinker */
 	void Initialize(FPreAnimatedStorageID InStorageID, FPreAnimatedStateExtension* InParentExtension) override
 	{
 		ParentExtension = InParentExtension;
 		StorageID = InStorageID;
-	}
 
-	void SetRestoreMask(const IRestoreMask* InRestoreMask)
-	{
-		RestoreMask = InRestoreMask;
+		this->Traits.Initialize(InStorageID, InParentExtension);
 	}
 
+	/** Retrieve the ID that uniquely identifies this storage container */
 	FPreAnimatedStorageID GetStorageType() const override
 	{
 		return StorageID;
 	}
 
-	template<typename OnCacheValue /* StorageType(const KeyType&) */>
-	void SavePreAnimatedState(FPreAnimatedStorageGroupHandle GroupHandle, const KeyType& InKey, OnCacheValue&& CacheCallback)
-	{
-		// Find the storage index for the specific key we're animating
-		FPreAnimatedStorageIndex StorageIndex = GetOrCreateStorageIndex(InKey);
-
-		FPreAnimatedStateEntry Entry{ GroupHandle, FPreAnimatedStateCachedValueHandle{ StorageID, StorageIndex } };
-
-		// Begin tracking the entry from its capture source
-		ParentExtension->AddSourceMetaData(Entry);
-
-		EPreAnimatedStorageRequirement Requirement = ParentExtension->HasActiveCaptureSource()
-			? EPreAnimatedStorageRequirement::Transient
-			: EPreAnimatedStorageRequirement::Persistent;
-
-		if (!IsStorageRequirementSatisfied(Entry.ValueHandle.StorageIndex, Requirement))
-		{
-			AssignPreAnimatedValue(StorageIndex, Requirement, CacheCallback(InKey));
-		}
-	}
-
+	/**
+	 * Restore a specified index within this storage container.
+	 * Usually called when all the things contributing to this storage index have finished evaluating.
+	 *
+	 * @param StorageIndex      The unique index for the stored state - either an index within PreAnimatedStorage or TransientPreAnimatedStorage.
+	 * @param SourceRequirement The source requirement that is requesting to restore state:
+	 *                              Persistent - indicates that the storage should be completely restored as a result of a sequence finishing or explicitly restoring state
+	 *                              Transient - indicates that all 'Restore State' tracks have finished evaluating for this index, but persistent state may still be cached
+	 * @param TargetRequirement The target requirement for storage - None implies that no state needs to remain cached for this index, Persistent implies that keep state entities have finished, but the state still needs to be cached.
+	 * @param Params            Additional restore parameters defining the instigator context
+	 *
+	 * @return The resulting storage requirement for the supplied index
+	 */
 	EPreAnimatedStorageRequirement RestorePreAnimatedStateStorage(FPreAnimatedStorageIndex StorageIndex, EPreAnimatedStorageRequirement SourceRequirement, EPreAnimatedStorageRequirement TargetRequirement, const FRestoreStateParams& Params) override
 	{
 		if (RestoreMask)
@@ -170,6 +285,14 @@ public:
 		return EPreAnimatedStorageRequirement::Persistent;
 	}
 
+	/**
+	 * Discard a specified index within this storage container.
+	 *
+	 * @param StorageIndex      The unique index for the stored state - either an index within PreAnimatedStorage or TransientPreAnimatedStorage.
+	 * @param SourceRequirement The storage requirement to discard - Persistent will discard all cached state for the object, Transient may leave persistent storage around, if possible.
+	 * 
+	 * @return The resulting storage requirement for the supplied index
+	 */
 	EPreAnimatedStorageRequirement DiscardPreAnimatedStateStorage(FPreAnimatedStorageIndex StorageIndex, EPreAnimatedStorageRequirement SourceRequirement) override
 	{
 		if (RestoreMask)
@@ -202,6 +325,72 @@ public:
 		}
 	}
 
+	/**
+	 * Called prior to restoring pre-animated state to control whether this storage should restore state or not
+	 */
+	void SetRestoreMask(const IRestoreMask* InRestoreMask)
+	{
+		RestoreMask = InRestoreMask;
+	}
+
+	/**
+	 * Called by the owning extension to add reference collection tracking for pre-animated state
+	 */
+	void AddReferencedObjects(FReferenceCollector& ReferenceCollector) override
+	{
+		if constexpr (THasAddReferencedObjectForComponent<KeyType>::Value)
+		{
+			for (auto It = KeyToStorageIndex.CreateIterator(); It; ++It)
+			{
+				AddReferencedObjectForComponent(&ReferenceCollector, &It.Key());
+			}
+		}
+		
+		if constexpr (THasAddReferencedObjectForComponent<KeyType>::Value || THasAddReferencedObjectForComponent<StorageType>::Value)
+		{
+			for (FCachedData& CachedData : PreAnimatedStorage)
+			{
+				AddReferencedObjectForComponent(&ReferenceCollector, &CachedData.Key);
+				AddReferencedObjectForComponent(&ReferenceCollector, &CachedData.Value);
+			}
+		}
+		if constexpr (THasAddReferencedObjectForComponent<StorageType>::Value)
+		{
+			for (auto It = TransientPreAnimatedStorage.CreateIterator(); It; ++It)
+			{
+				AddReferencedObjectForComponent(&ReferenceCollector, &It.Value());
+			}
+		}
+	}
+
+	/**
+	 * Save pre-aniamted state for the specified group and key using a callback.
+	 * Callback will only be invoked if state has not already been saved.
+	 */
+	template<typename OnCacheValue /* StorageType(const KeyType&) */>
+	void SavePreAnimatedState(FPreAnimatedStorageGroupHandle GroupHandle, const KeyType& InKey, OnCacheValue&& CacheCallback)
+	{
+		// Find the storage index for the specific key we're animating
+		FPreAnimatedStorageIndex StorageIndex = GetOrCreateStorageIndex(InKey);
+
+		FPreAnimatedStateEntry Entry{ GroupHandle, FPreAnimatedStateCachedValueHandle{ StorageID, StorageIndex } };
+
+		// Begin tracking the entry from its capture source
+		ParentExtension->AddSourceMetaData(Entry);
+
+		EPreAnimatedStorageRequirement Requirement = ParentExtension->HasActiveCaptureSource()
+			? EPreAnimatedStorageRequirement::Transient
+			: EPreAnimatedStorageRequirement::Persistent;
+
+		if (!IsStorageRequirementSatisfied(Entry.ValueHandle.StorageIndex, Requirement))
+		{
+			AssignPreAnimatedValue(StorageIndex, Requirement, CacheCallback(InKey));
+		}
+	}
+
+	/**
+	 * Attempt to find a storage index for the specified key, creating a new one if it doesn't exist
+	 */
 	FPreAnimatedStorageIndex GetOrCreateStorageIndex(const KeyType& InKey)
 	{
 		FPreAnimatedStorageIndex Index = KeyToStorageIndex.FindRef(InKey);
@@ -214,11 +403,21 @@ public:
 		return Index;
 	}
 
+	/**
+	 * Attempt to find a storage index for the specified key
+	 */
 	FPreAnimatedStorageIndex FindStorageIndex(const KeyType& InKey) const
 	{
 		return KeyToStorageIndex.FindRef(InKey);
 	}
 
+	/**
+	 * Assign the value for a specific storage index. It is an error to re-assign an already cached value.
+	 *
+	 * @param StorageIndex       The storage index to assign a value for
+	 * @param StorageRequirement Whether to assign the value to this index's persistent, or transient value. Transient should be used when this value is directly associated with a track evaluating.
+	 * @param InNewValue         The value to assign
+	 */
 	void AssignPreAnimatedValue(FPreAnimatedStorageIndex StorageIndex, EPreAnimatedStorageRequirement StorageRequirement, StorageType&& InNewValue)
 	{
 		check(StorageIndex);
@@ -249,6 +448,9 @@ public:
 		}
 	}
 
+	/**
+	 * Check whether the storage for a given index and requirement has already been initialized
+	 */
 	bool IsStorageRequirementSatisfied(FPreAnimatedStorageIndex StorageIndex, EPreAnimatedStorageRequirement StorageRequirement) const
 	{
 		check(StorageIndex);
@@ -268,27 +470,47 @@ public:
 		return true;
 	}
 
+	/**
+	 * Cause a previously cached value to always outlive any actively animating sources.
+	 * This is called when a Restore State track overlaps a Keep State track. The Restore State track will initially
+	 * save state using the Transient requirement, but the Keep State track may need to keep this cached state alive
+	 * if it is capturing _global_ state. As such, we take the previously cached state and make it persistent.
+	 */
 	void ForciblyPersistStorage(FPreAnimatedStorageIndex StorageIndex)
 	{
 		check(StorageIndex);
 		PreAnimatedStorage[StorageIndex.Value].bPersistent = true;
 	}
 
+	/**
+	 * Check whether a piece of cached storage has been initialized yet. This will be false where multiple entities contribute to an object's state
+	 * but a different entity actually assigns the value (eg multiple blended entiies). Saving state for these requires 2 passes: firstly we gather
+	 * all the entities that contribute to the state (which may not be able to even know _how_ to cache the state), then we capture the actual value.
+	 */
 	bool IsStorageInitialized(FPreAnimatedStorageIndex StorageIndex) const
 	{
 		return StorageIndex && (PreAnimatedStorage[StorageIndex.Value].bInitialized || TransientPreAnimatedStorage.Contains(StorageIndex));
 	}
 
+	/**
+	 * Check whether we have ever animated the specified storage index.
+	 */
 	bool HasEverAnimated(FPreAnimatedStorageIndex StorageIndex) const
 	{
 		return StorageIndex && PreAnimatedStorage[StorageIndex.Value].bInitialized;
 	}
 
+	/**
+	 * Get the key associated with a particular storage index
+	 */
 	const KeyType& GetKey(FPreAnimatedStorageIndex StorageIndex) const
 	{
 		return PreAnimatedStorage[StorageIndex].Key;
 	}
 
+	/**
+	 * Replace the key associated with a particular storage index
+	 */
 	void ReplaceKey(FPreAnimatedStorageIndex StorageIndex, const KeyType& NewKey)
 	{
 		KeyType OldKey = PreAnimatedStorage[StorageIndex].Key;
@@ -298,6 +520,186 @@ public:
 		KeyToStorageIndex.Add(NewKey, StorageIndex);
 	}
 
+	/**
+	 * Look at any entity with the specified component types, and set up new associations with storage indices for those entities
+	 * WARNING: Does not cache actual pre-animated values
+	 */
+	template<typename... ContributorTypes>
+	void BeginTrackingEntities(UMovieSceneEntitySystemLinker* Linker, TComponentTypeID<ContributorTypes>... InComponentTypes)
+	{
+		BeginTrackingEntitiesTask(Linker, TSavePreAnimatedStateParams<ContributorTypes...>(), InComponentTypes...);
+	}
+
+	/**
+	 * Look at any entity with the specified component types, and set up new associations with storage indices for those entities
+	 * WARNING: Does not cache actual pre-animated values
+	 */
+	template<typename TaskType, typename... ContributorTypes>
+	void BeginTrackingEntitiesTask(UMovieSceneEntitySystemLinker* Linker, const TaskType& InParams, TComponentTypeID<ContributorTypes>... InComponentTypes)
+	{
+		auto VisitAllocation = [this, InParams](FEntityAllocationIteratorItem Item, TRead<FMovieSceneEntityID> EntityIDs, TRead<FRootInstanceHandle> RootInstanceHandles, TRead<ContributorTypes>... Inputs)
+		{
+			FPreAnimatedTrackerParams Params(Item);
+
+			const int32 Num = Params.Num;
+			const bool  bWantsRestore = Params.bWantsRestoreState;
+
+			if (!this->ParentExtension->IsCapturingGlobalState() && !bWantsRestore)
+			{
+				return;
+			}
+
+			FPreAnimatedEntityCaptureSource* EntityMetaData = this->ParentExtension->GetOrCreateEntityMetaData();
+			for (int32 Index = 0; Index < Num; ++Index)
+			{
+				if (!InParams.ShouldSavePreAnimatedState(Inputs[Index]...))
+				{
+					continue;
+				}
+
+				KeyType Key{ Inputs[Index]... };
+
+				FPreAnimatedStorageGroupHandle GroupHandle  = this->Traits.MakeGroup(Inputs[Index]...);
+				FPreAnimatedStorageIndex       StorageIndex = this->GetOrCreateStorageIndex(Key);
+
+				FPreAnimatedStateEntry Entry{ GroupHandle, FPreAnimatedStateCachedValueHandle{ this->StorageID, StorageIndex } };
+				EntityMetaData->BeginTrackingEntity(Entry, EntityIDs[Index], RootInstanceHandles[Index], bWantsRestore);
+			}
+		};
+
+		FBuiltInComponentTypes* BuiltInComponents = FBuiltInComponentTypes::Get();
+
+		FEntityTaskBuilder()
+		.ReadEntityIDs()
+		.Read(BuiltInComponents->RootInstanceHandle)
+		.ReadAllOf(InComponentTypes...)
+		.CombineFilter(InParams.AdditionalFilter)
+		.FilterAll({ BuiltInComponents->Tags.NeedsLink })
+		.Iterate_PerAllocation(&Linker->EntityManager, VisitAllocation);
+	}
+
+	/**
+	 * Cache pre-animated values for entities with the specified component types
+	 */
+	template<typename... CacheValueInputs>
+	void CachePreAnimatedValues(UMovieSceneEntitySystemLinker* Linker, TComponentTypeID<CacheValueInputs>... InComponentTypes)
+	{
+		CachePreAnimatedValuesTask(Linker, TSavePreAnimatedStateParams<CacheValueInputs...>(), InComponentTypes...);
+	}
+
+	/**
+	 * Cache pre-animated values for entities with the specified component types
+	 */
+	template<typename TaskType, typename... CacheValueInputs>
+	void CachePreAnimatedValuesTask(UMovieSceneEntitySystemLinker* Linker, const TaskType& InParams, TComponentTypeID<CacheValueInputs>... InComponentTypes)
+	{
+		FBuiltInComponentTypes* BuiltInComponents = FBuiltInComponentTypes::Get();
+
+		auto VisitAllocation = [this, &InParams](FEntityAllocationIteratorItem Item, TRead<CacheValueInputs>... Values)
+		{
+			const int32 Num = Item.GetAllocation()->Num();
+			for (int32 Index = 0; Index < Num; ++Index)
+			{
+				if (!InParams.ShouldSavePreAnimatedState(Values[Index]...))
+				{
+					continue;
+				}
+
+				KeyType Key{ Values[Index]... };
+
+				FPreAnimatedStorageGroupHandle GroupHandle  = this->Traits.MakeGroup(Values[Index]...);
+				FPreAnimatedStorageIndex       StorageIndex = this->GetOrCreateStorageIndex(Key);
+
+				FPreAnimatedStateEntry Entry{ GroupHandle, FPreAnimatedStateCachedValueHandle{ this->StorageID, StorageIndex } };
+
+				this->ParentExtension->EnsureMetaData(Entry);
+
+				EPreAnimatedStorageRequirement StorageRequirement = this->ParentExtension->GetStorageRequirement(Entry);
+				if (!this->IsStorageRequirementSatisfied(StorageIndex, StorageRequirement))
+				{
+					this->AssignPreAnimatedValue(StorageIndex, StorageRequirement, this->Traits.CachePreAnimatedValue(Values[Index]...));
+				}
+			}
+		};
+
+		FEntityTaskBuilder()
+		.ReadAllOf(InComponentTypes...)
+		.CombineFilter(InParams.AdditionalFilter)
+		.FilterAll({ BuiltInComponents->Tags.NeedsLink })
+		.Iterate_PerAllocation(&Linker->EntityManager, VisitAllocation);
+	}
+
+
+	/**
+	 * Look at any entity with the specified component types, and set up new associations with storage indices for those entities, whilst also caching pre-animated values at the same time.
+	 */
+	template<typename... ContributorTypes>
+	void BeginTrackingAndCachePreAnimatedValues(UMovieSceneEntitySystemLinker* Linker, TComponentTypeID<ContributorTypes>... InComponentTypes)
+	{
+		BeginTrackingAndCachePreAnimatedValues(Linker, TSavePreAnimatedStateParams<ContributorTypes...>(), InComponentTypes...);
+	}
+
+	/**
+	 * Look at any entity with the specified component types, and set up new associations with storage indices for those entities, whilst also caching pre-animated values at the same time.
+	 */
+	template<typename TaskType, typename... ContributorTypes>
+	void BeginTrackingAndCachePreAnimatedValuesTask(UMovieSceneEntitySystemLinker* Linker, const TaskType& InParams, TComponentTypeID<ContributorTypes>... InComponentTypes)
+	{
+		auto VisitAllocation = [this, InParams](FEntityAllocationIteratorItem Item, TRead<FMovieSceneEntityID> EntityIDs, TRead<FRootInstanceHandle> RootInstanceHandles, TRead<ContributorTypes>... Inputs)
+		{
+			FPreAnimatedTrackerParams Params(Item);
+
+			const int32 Num = Params.Num;
+			const bool  bWantsRestore = Params.bWantsRestoreState;
+
+			if (!this->ParentExtension->IsCapturingGlobalState() && !bWantsRestore)
+			{
+				return;
+			}
+
+			FPreAnimatedEntityCaptureSource* EntityMetaData = this->ParentExtension->GetOrCreateEntityMetaData();
+			for (int32 Index = 0; Index < Num; ++Index)
+			{
+				if (!InParams.ShouldSavePreAnimatedState(Inputs[Index]...))
+				{
+					continue;
+				}
+
+				KeyType Key{ Inputs[Index]... };
+
+				FPreAnimatedStorageGroupHandle GroupHandle  = this->Traits.MakeGroup(Inputs[Index]...);
+				FPreAnimatedStorageIndex       StorageIndex = this->GetOrCreateStorageIndex(Key);
+
+				FPreAnimatedStateEntry Entry{ GroupHandle, FPreAnimatedStateCachedValueHandle{ this->StorageID, StorageIndex } };
+				EntityMetaData->BeginTrackingEntity(Entry, EntityIDs[Index], RootInstanceHandles[Index], bWantsRestore);
+
+				EPreAnimatedStorageRequirement StorageRequirement = this->ParentExtension->GetStorageRequirement(Entry);
+				if (!this->IsStorageRequirementSatisfied(StorageIndex, StorageRequirement))
+				{
+					this->AssignPreAnimatedValue(StorageIndex, StorageRequirement, this->Traits.CachePreAnimatedValue(Inputs[Index]...));
+				}
+
+			}
+		};
+
+		FBuiltInComponentTypes* BuiltInComponents = FBuiltInComponentTypes::Get();
+
+		FEntityTaskBuilder()
+		.ReadEntityIDs()
+		.Read(BuiltInComponents->RootInstanceHandle)
+		.ReadAllOf(InComponentTypes...)
+		.CombineFilter(InParams.AdditionalFilter)
+		.FilterAll({ BuiltInComponents->Tags.NeedsLink })
+		.Iterate_PerAllocation(&Linker->EntityManager, VisitAllocation);
+	}
+
+	void OnObjectReplaced(FPreAnimatedStorageIndex StorageIndex, const FObjectKey& OldObject, const FObjectKey& NewObject) override
+	{
+		KeyType ExistingKey = this->GetKey(StorageIndex);
+		this->Traits.ReplaceObject(ExistingKey, NewObject);
+		this->ReplaceKey(StorageIndex, ExistingKey);
+	}
+
 protected:
 
 	struct FCachedData
@@ -326,24 +728,36 @@ protected:
 		bool bPersistent : 1;
 	};
 
+	/** Map that associates a storage index with a key. This indirection allows common code to deal with indices without knowing concrete templated types. */
 	TMap<KeyType, FPreAnimatedStorageIndex> KeyToStorageIndex;
 
+	/** Sparse array of cached data that represents persistent, or both persistent _and_ transient storage (but never exclusively transient storage).
+	 * FPreAnimatedStorageIndex defines an index into this array (or into the rarely used TransientPreAnimatedStorage map) */
 	TSparseArray<FCachedData> PreAnimatedStorage;
 
-	/** Storage that holds values that need to be kept transiently (for evaluation) */
+	/**
+	 * Storage that holds values that need to be kept transiently (for evaluation).
+	 * This map is only used if a Keep State section previously captured a value (because it is evaluating in a 'Capture Global State' context)
+	 * and we end up animating the _same_ value using a Restore State section. The Restore State section has to re-capture its starting value
+	 * whilst also keeping the previously captured value alive to ensure that the section can restore to its starting value in addition to the
+	 * sequence itself restoring to the global starting value
+	 */
 	TSortedMap<FPreAnimatedStorageIndex, StorageType> TransientPreAnimatedStorage;
 
+	/** Pointer to our parent extension */
 	FPreAnimatedStateExtension* ParentExtension = nullptr;
 
+	/** Temporary restoration mask */
 	const IRestoreMask* RestoreMask = nullptr;
 
+	/** Our storage ID */
 	FPreAnimatedStorageID StorageID;
 
+	/** Traits instance - normally not required but some traits require state */
 	StorageTraits Traits;
 };
 
 
-
 } // namespace MovieScene
 } // namespace UE
 
diff --git a/Engine/Source/Runtime/MovieSceneTracks/Private/EntitySystem/Interrogation/MovieSceneInterrogatedPropertyInstantiator.cpp b/Engine/Source/Runtime/MovieSceneTracks/Private/EntitySystem/Interrogation/MovieSceneInterrogatedPropertyInstantiator.cpp
--- a/Engine/Source/Runtime/MovieSceneTracks/Private/EntitySystem/Interrogation/MovieSceneInterrogatedPropertyInstantiator.cpp
+++ b/Engine/Source/Runtime/MovieSceneTracks/Private/EntitySystem/Interrogation/MovieSceneInterrogatedPropertyInstantiator.cpp
@@ -291,6 +291,8 @@ void UMovieSceneInterrogatedPropertyInstantiatorSystem::DestroyOutput(UE::MovieS
 
 void UMovieSceneInterrogatedPropertyInstantiatorSystem::OnRun(FSystemTaskPrerequisites& InPrerequisites, FSystemSubsequentTasks& Subsequents)
 {
+	PropertyTracker.Initialize(this);
+
 	using namespace UE::MovieScene;
 
 	{
diff --git a/Engine/Source/Runtime/MovieSceneTracks/Private/PreAnimatedState/MovieScenePreAnimatedComponentTransformStorage.cpp b/Engine/Source/Runtime/MovieSceneTracks/Private/PreAnimatedState/MovieScenePreAnimatedComponentTransformStorage.cpp
--- a/Engine/Source/Runtime/MovieSceneTracks/Private/PreAnimatedState/MovieScenePreAnimatedComponentTransformStorage.cpp
+++ b/Engine/Source/Runtime/MovieSceneTracks/Private/PreAnimatedState/MovieScenePreAnimatedComponentTransformStorage.cpp
@@ -82,19 +82,19 @@ void FPreAnimatedComponentTransformStorage::CachePreAnimatedTransforms(const FCa
 			continue;
 		}
 
-		FAnimatedPropertyKey Key{ BoundObject, PropertyBinding.PropertyPath };
+		TTuple<FObjectKey, FName> Key{ BoundObject, PropertyBinding.PropertyPath };
 
-		FPreAnimatedStorageGroupHandle GroupHandle  = ObjectGroupManager->MakeGroupForKey(BoundObject);
-		FPreAnimatedStorageIndex       StorageIndex = Storage.GetOrCreateStorageIndex(Key);
+		FPreAnimatedStorageGroupHandle GroupHandle  = this->Traits.MakeGroup(BoundObject);
+		FPreAnimatedStorageIndex       StorageIndex = this->GetOrCreateStorageIndex(Key);
 
 		FPreAnimatedStateEntry Entry{ GroupHandle, FPreAnimatedStateCachedValueHandle{ StorageID, StorageIndex } };
 
 		ParentExtension->EnsureMetaData(Entry);
 
 		EPreAnimatedStorageRequirement StorageRequirement = ParentExtension->GetStorageRequirement(Entry);
-		if (!Storage.IsStorageRequirementSatisfied(StorageIndex, StorageRequirement))
+		if (!this->IsStorageRequirementSatisfied(StorageIndex, StorageRequirement))
 		{
-			FPreAnimatedProperty NewValue;
+			StorageType NewValue;
 
 			if (const uint16* Fast = Property->TryGet<uint16>())
 			{
@@ -116,12 +116,12 @@ void FPreAnimatedComponentTransformStorage::CachePreAnimatedTransforms(const FCa
 				FComponentTransformPropertyTraits::GetObjectPropertyValue(BoundObject, Bindings.Get(), NewValue.Data);
 			}
 
-			Storage.AssignPreAnimatedValue(StorageIndex, StorageRequirement, MoveTemp(NewValue));
+			this->AssignPreAnimatedValue(StorageIndex, StorageRequirement, MoveTemp(NewValue));
 		}
 
 		if (Params.bForcePersist)
 		{
-			Storage.ForciblyPersistStorage(StorageIndex);
+			this->ForciblyPersistStorage(StorageIndex);
 		}
 	}
 }
diff --git a/Engine/Source/Runtime/MovieSceneTracks/Private/Systems/MovieSceneComponentAttachmentSystem.cpp b/Engine/Source/Runtime/MovieSceneTracks/Private/Systems/MovieSceneComponentAttachmentSystem.cpp
--- a/Engine/Source/Runtime/MovieSceneTracks/Private/Systems/MovieSceneComponentAttachmentSystem.cpp
+++ b/Engine/Source/Runtime/MovieSceneTracks/Private/Systems/MovieSceneComponentAttachmentSystem.cpp
@@ -102,7 +102,7 @@ struct FAttachmentHandler
 	}
 };
 
-struct FComponentAttachmentPreAnimatedTraits
+struct FComponentAttachmentPreAnimatedTraits : FBoundObjectPreAnimatedStateTraits
 {
 	using KeyType     = FObjectKey;
 	using StorageType = FPreAnimAttachment;
@@ -127,8 +127,6 @@ struct FPreAnimatedComponentAttachmentStorage
 	: TPreAnimatedStateStorage_ObjectTraits<FComponentAttachmentPreAnimatedTraits>
 {
 	static TAutoRegisterPreAnimatedStorageID<FPreAnimatedComponentAttachmentStorage> StorageID;
-
-	FPreAnimatedStorageID GetStorageType() const override { return StorageID; }
 };
 
 TAutoRegisterPreAnimatedStorageID<FPreAnimatedComponentAttachmentStorage> FPreAnimatedComponentAttachmentStorage::StorageID;
@@ -192,17 +190,7 @@ void UMovieSceneComponentAttachmentSystem::OnLink()
 	Linker->SystemGraph.AddReference(this, AttachmentInvalidator);
 	Linker->SystemGraph.AddPrerequisite(AttachmentInvalidator, this);
 
-	Linker->Events.TagGarbage.AddUObject(this, &UMovieSceneComponentAttachmentSystem::TagGarbage);
-}
-
-void UMovieSceneComponentAttachmentSystem::TagGarbage(UMovieSceneEntitySystemLinker*)
-{
-	AttachmentTracker.CleanupGarbage();
-}
-
-void UMovieSceneComponentAttachmentSystem::AddReferencedObjects(UObject* InThis, FReferenceCollector& Collector)
-{
-	CastChecked<UMovieSceneComponentAttachmentSystem>(InThis)->AttachmentTracker.AddReferencedObjects(Collector);
+	AttachmentTracker.Initialize(this);
 }
 
 void UMovieSceneComponentAttachmentSystem::OnUnlink()
diff --git a/Engine/Source/Runtime/MovieSceneTracks/Private/Systems/MovieSceneComponentMobilitySystem.cpp b/Engine/Source/Runtime/MovieSceneTracks/Private/Systems/MovieSceneComponentMobilitySystem.cpp
--- a/Engine/Source/Runtime/MovieSceneTracks/Private/Systems/MovieSceneComponentMobilitySystem.cpp
+++ b/Engine/Source/Runtime/MovieSceneTracks/Private/Systems/MovieSceneComponentMobilitySystem.cpp
@@ -83,7 +83,7 @@ static void GetFlattenedHierarchy(USceneComponent* SceneComponent, TArray<UScene
 
 FPreAnimatedStateEntry FPreAnimatedComponentMobilityStorage::MakeEntry(USceneComponent* InSceneComponent)
 {
-	FPreAnimatedStorageGroupHandle GroupHandle  = ObjectGroupManager->MakeGroupForKey(InSceneComponent);
+	FPreAnimatedStorageGroupHandle GroupHandle  = this->Traits.ObjectGroupManager->MakeGroupForKey(InSceneComponent);
 	FPreAnimatedStorageIndex       StorageIndex = GetOrCreateStorageIndex(InSceneComponent);
 
 	return FPreAnimatedStateEntry{ GroupHandle, FPreAnimatedStateCachedValueHandle{ StorageID, StorageIndex } };
@@ -138,7 +138,7 @@ bool UMovieSceneComponentMobilitySystem::IsRelevantImpl(UMovieSceneEntitySystemL
 
 void UMovieSceneComponentMobilitySystem::OnLink()
 {
-	Linker->Events.TagGarbage.AddUObject(this, &UMovieSceneComponentMobilitySystem::TagGarbage);
+	MobilityTracker.Initialize(this);
 }
 
 void UMovieSceneComponentMobilitySystem::OnUnlink()
@@ -163,16 +163,6 @@ void UMovieSceneComponentMobilitySystem::OnRun(FSystemTaskPrerequisites& InPrere
 	MobilityTracker.ProcessInvalidatedOutputs(Linker, FMobilityCacheHandler(this));
 }
 
-void UMovieSceneComponentMobilitySystem::TagGarbage(UMovieSceneEntitySystemLinker*)
-{
-	MobilityTracker.CleanupGarbage();
-}
-
-void UMovieSceneComponentMobilitySystem::AddReferencedObjects(UObject* InThis, FReferenceCollector& Collector)
-{
-	CastChecked<UMovieSceneComponentMobilitySystem>(InThis)->MobilityTracker.AddReferencedObjects(Collector);
-}
-
 void UMovieSceneComponentMobilitySystem::SavePreAnimatedState(const FPreAnimationParameters& InParameters)
 {
 	using namespace UE::MovieScene;
diff --git a/Engine/Source/Runtime/MovieSceneTracks/Private/Systems/MovieSceneDataLayerSystem.cpp b/Engine/Source/Runtime/MovieSceneTracks/Private/Systems/MovieSceneDataLayerSystem.cpp
--- a/Engine/Source/Runtime/MovieSceneTracks/Private/Systems/MovieSceneDataLayerSystem.cpp
+++ b/Engine/Source/Runtime/MovieSceneTracks/Private/Systems/MovieSceneDataLayerSystem.cpp
@@ -42,7 +42,7 @@ enum class EDataLayerUpdateFlags : uint8
 ENUM_CLASS_FLAGS(EDataLayerUpdateFlags)
 
 /** Traits class governing how pre-animated state is (re)stored for data layers */
-struct FPreAnimatedDataLayerStorageTraits
+struct FPreAnimatedDataLayerStorageTraits : FBoundObjectPreAnimatedStateTraits
 {
 	using KeyType = TObjectKey<UDataLayerInstance>;
 	using StorageType = EDataLayerRuntimeState;
diff --git a/Engine/Source/Runtime/MovieSceneTracks/Private/TrackInstances/MovieSceneCVarTrackInstance.cpp b/Engine/Source/Runtime/MovieSceneTracks/Private/TrackInstances/MovieSceneCVarTrackInstance.cpp
--- a/Engine/Source/Runtime/MovieSceneTracks/Private/TrackInstances/MovieSceneCVarTrackInstance.cpp
+++ b/Engine/Source/Runtime/MovieSceneTracks/Private/TrackInstances/MovieSceneCVarTrackInstance.cpp
@@ -18,7 +18,7 @@ namespace UE
 namespace MovieScene
 {
 
-struct FPreAnimatedCVarTraits
+struct FPreAnimatedCVarTraits : FDefaultPreAnimatedStateTraits
 {
 	using KeyType     = FString;
 	using StorageType = FString;
diff --git a/Engine/Source/Runtime/MovieSceneTracks/Public/EntitySystem/Interrogation/MovieSceneInterrogatedPropertyInstantiator.h b/Engine/Source/Runtime/MovieSceneTracks/Public/EntitySystem/Interrogation/MovieSceneInterrogatedPropertyInstantiator.h
--- a/Engine/Source/Runtime/MovieSceneTracks/Public/EntitySystem/Interrogation/MovieSceneInterrogatedPropertyInstantiator.h
+++ b/Engine/Source/Runtime/MovieSceneTracks/Public/EntitySystem/Interrogation/MovieSceneInterrogatedPropertyInstantiator.h
@@ -117,7 +117,7 @@ private:
 
 	static constexpr uint16 INVALID_BLEND_CHANNEL = uint16(-1);
 
-	UE::MovieScene::TOverlappingEntityTracker<UE::MovieScene::FInterrogationKey, FPropertyInfo> PropertyTracker;
+	UE::MovieScene::TOverlappingEntityTracker<FPropertyInfo, UE::MovieScene::FInterrogationKey> PropertyTracker;
 	UE::MovieScene::FComponentMask CleanFastPathMask;
 	UE::MovieScene::FBuiltInComponentTypes* BuiltInComponents;
 	UE::MovieScene::FPropertyRecomposerImpl RecomposerImpl;
diff --git a/Engine/Source/Runtime/MovieSceneTracks/Public/Systems/MovieSceneComponentAttachmentSystem.h b/Engine/Source/Runtime/MovieSceneTracks/Public/Systems/MovieSceneComponentAttachmentSystem.h
--- a/Engine/Source/Runtime/MovieSceneTracks/Public/Systems/MovieSceneComponentAttachmentSystem.h
+++ b/Engine/Source/Runtime/MovieSceneTracks/Public/Systems/MovieSceneComponentAttachmentSystem.h
@@ -58,8 +58,6 @@ public:
 
 	void AddPendingDetach(USceneComponent* SceneComponent, const UE::MovieScene::FPreAnimAttachment& Attachment);
 
-	static void AddReferencedObjects(UObject* InThis, FReferenceCollector& Collector);
-
 private:
 
 	virtual void OnLink() override final;
@@ -70,9 +68,7 @@ private:
 	virtual void SavePreAnimatedState(const FPreAnimationParameters& InParameters) override;
 	virtual void RestorePreAnimatedState(const FPreAnimationParameters& InParameters) override;
 
-	void TagGarbage(UMovieSceneEntitySystemLinker*);
-
-	UE::MovieScene::TOverlappingEntityTracker_BoundObject<UE::MovieScene::FPreAnimAttachment> AttachmentTracker;
+	UE::MovieScene::TOverlappingEntityTracker<UE::MovieScene::FPreAnimAttachment, UObject*> AttachmentTracker;
 
 	TArray<TTuple<USceneComponent*, UE::MovieScene::FPreAnimAttachment>> PendingAttachmentsToRestore;
 };
diff --git a/Engine/Source/Runtime/MovieSceneTracks/Public/Systems/MovieSceneComponentMobilitySystem.h b/Engine/Source/Runtime/MovieSceneTracks/Public/Systems/MovieSceneComponentMobilitySystem.h
--- a/Engine/Source/Runtime/MovieSceneTracks/Public/Systems/MovieSceneComponentMobilitySystem.h
+++ b/Engine/Source/Runtime/MovieSceneTracks/Public/Systems/MovieSceneComponentMobilitySystem.h
@@ -23,8 +23,6 @@ public:
 
 	UMovieSceneComponentMobilitySystem(const FObjectInitializer& ObjInit);
 
-	static void AddReferencedObjects(UObject* InThis, FReferenceCollector& Collector);
-
 private:
 
 	virtual bool IsRelevantImpl(UMovieSceneEntitySystemLinker* InLinker) const override;
@@ -36,11 +34,9 @@ private:
 	virtual void SavePreAnimatedState(const FPreAnimationParameters& InParameters) override;
 	virtual void RestorePreAnimatedState(const FPreAnimationParameters& InParameters) override;
 
-	void TagGarbage(UMovieSceneEntitySystemLinker*);
-
 private:
 
-	UE::MovieScene::TOverlappingEntityTracker_BoundObject<EComponentMobility::Type> MobilityTracker;
+	UE::MovieScene::TOverlappingEntityTracker<EComponentMobility::Type, UObject*> MobilityTracker;
 
 	UE::MovieScene::FEntityComponentFilter Filter;
 
@@ -53,7 +49,7 @@ namespace UE
 namespace MovieScene
 {
 
-struct FPreAnimatedMobilityTraits
+struct FPreAnimatedMobilityTraits : FBoundObjectPreAnimatedStateTraits
 {
 	using KeyType     = FObjectKey;
 	using StorageType = EComponentMobility::Type;
