Sequencer: Fixed reference counting on overlapping entity tracker

diff --git a/Engine/Source/Runtime/MovieScene/Public/EntitySystem/MovieSceneOverlappingEntityTracker.h b/Engine/Source/Runtime/MovieScene/Public/EntitySystem/MovieSceneOverlappingEntityTracker.h
--- a/Engine/Source/Runtime/MovieScene/Public/EntitySystem/MovieSceneOverlappingEntityTracker.h
+++ b/Engine/Source/Runtime/MovieScene/Public/EntitySystem/MovieSceneOverlappingEntityTracker.h
@@ -465,8 +465,8 @@ protected:
 };
 
 
-template<typename OutputType, typename... KeyType>
-struct TOverlappingEntityTracker_NoGarbage : TOverlappingEntityTrackerImpl<OutputType, KeyType...>
+template<typename OutputType, typename... InputTypes>
+struct TOverlappingEntityTracker_NoGarbage : TOverlappingEntityTrackerImpl<OutputType, InputTypes...>
 {
 	void Initialize(UMovieSceneEntitySystem* OwningSystem)
 	{
@@ -474,11 +474,11 @@ struct TOverlappingEntityTracker_NoGarbage : TOverlappingEntityTrackerImpl<Outpu
 	}
 };
 
-template<typename OutputType, typename... KeyType>
-struct TOverlappingEntityTracker_WithGarbage : TOverlappingEntityTrackerImpl<OutputType, KeyType...>
+template<typename OutputType, typename... InputTypes>
+struct TOverlappingEntityTracker_WithGarbage : TOverlappingEntityTrackerImpl<OutputType, InputTypes...>
 {
-	using ThisType = TOverlappingEntityTracker_WithGarbage<OutputType, KeyType...>;
-	using Super = TOverlappingEntityTrackerImpl<OutputType, KeyType...>;
+	using ThisType = TOverlappingEntityTracker_WithGarbage<OutputType, InputTypes...>;
+	using Super = TOverlappingEntityTrackerImpl<OutputType, InputTypes...>;
 	using typename Super::FOutput;
 	using typename Super::KeyType;
 
@@ -541,17 +541,26 @@ struct TOverlappingEntityTracker_WithGarbage : TOverlappingEntityTrackerImpl<Out
 
 	void AddReferencedObjects(UMovieSceneEntitySystemLinker* Linker, FReferenceCollector& ReferenceCollector)
 	{
-		for (TPair<KeyType, uint16>& Pair : this->KeyToOutput)
+		constexpr bool bKeyCanBeGarbage = (THasAddReferencedObjectForComponent<InputTypes>::Value || ...);
+		constexpr bool bOutputCanBeGarbage = THasAddReferencedObjectForComponent<OutputType>::Value;
+
+		if constexpr (bKeyCanBeGarbage)
 		{
-			KeyType::GarbageTraits::AddReferencedObjects(ReferenceCollector, Pair.Key);
+			for (TPair<KeyType, uint16>& Pair : this->KeyToOutput)
+			{
+				KeyType::GarbageTraits::AddReferencedObjects(ReferenceCollector, Pair.Key);
+			}
 		}
 
-		if constexpr (THasAddReferencedObjectForComponent<KeyType>::Value || THasAddReferencedObjectForComponent<OutputType>::Value)
+		for (FOutput& Output : this->Outputs)
 		{
-			for (FOutput& Output : this->Outputs)
+			if constexpr (bKeyCanBeGarbage)
 			{
 				KeyType::GarbageTraits::AddReferencedObjects(ReferenceCollector, Output.Key);
-				KeyType::GarbageTraits::AddReferencedObjects(ReferenceCollector, Output.OutputData);
+			}
+			if constexpr (bOutputCanBeGarbage)
+			{
+				FGarbageTraits::AddReferencedObjects(ReferenceCollector, Output.OutputData);
 			}
 		}
 	}
