Sequencer: Added nullptr check for valid world

diff --git a/Engine/Source/Runtime/MovieScene/Private/MovieSceneSequencePlayer.cpp b/Engine/Source/Runtime/MovieScene/Private/MovieSceneSequencePlayer.cpp
--- a/Engine/Source/Runtime/MovieScene/Private/MovieSceneSequencePlayer.cpp
+++ b/Engine/Source/Runtime/MovieScene/Private/MovieSceneSequencePlayer.cpp
@@ -1225,7 +1225,7 @@ FString UMovieSceneSequencePlayer::GetSequenceName(bool bAddClientInfo) const
 		if (bAddClientInfo)
 		{
 			AActor* Actor = GetTypedOuter<AActor>();
-			if (Actor && Actor->GetWorld()->GetNetMode() == NM_Client)
+			if (Actor && Actor->GetWorld() && Actor->GetWorld()->GetNetMode() == NM_Client)
 			{
 				SequenceName += FString::Printf(TEXT(" (client %d)"), GPlayInEditorID - 1);
 			}
