Sequencer: Remove ensure in UnregisterTickClient

diff --git a/Engine/Source/Runtime/MovieScene/Private/MovieSceneSequenceTickManager.cpp b/Engine/Source/Runtime/MovieScene/Private/MovieSceneSequenceTickManager.cpp
--- a/Engine/Source/Runtime/MovieScene/Private/MovieSceneSequenceTickManager.cpp
+++ b/Engine/Source/Runtime/MovieScene/Private/MovieSceneSequenceTickManager.cpp
@@ -181,7 +181,7 @@ void UMovieSceneSequenceTickManager::UnregisterTickClient(TScriptInterface<IMovi
 	}
 
 	const int32 ClientIndex = Algo::IndexOfBy(TickableClients, ClientInterface, &FTickableClientData::Interface);
-	if (!ensure(TickableClients.IsValidIndex(ClientIndex)))
+	if (!TickableClients.IsValidIndex(ClientIndex))
 	{
 		return;
 	}
