Sequencer: Ensure that bound objects are correctly marked as producing entities

  - This fixes downstream dependencies such as the pre-animated-state system not correctly waiting until bound objects have been resolved before attempting to cache pre-animated state

diff --git a/Engine/Source/Runtime/MovieScene/Private/EntitySystem/MovieSceneBoundObjectInstantiator.cpp b/Engine/Source/Runtime/MovieScene/Private/EntitySystem/MovieSceneBoundObjectInstantiator.cpp
--- a/Engine/Source/Runtime/MovieScene/Private/EntitySystem/MovieSceneBoundObjectInstantiator.cpp
+++ b/Engine/Source/Runtime/MovieScene/Private/EntitySystem/MovieSceneBoundObjectInstantiator.cpp
@@ -23,6 +23,7 @@ UMovieSceneGenericBoundObjectInstantiator::UMovieSceneGenericBoundObjectInstanti
 	if (HasAnyFlags(RF_ClassDefaultObject))
 	{
 		DefineComponentProducer(GetClass(), Components->BoundObject);
+		DefineComponentProducer(GetClass(), Components->SymbolicTags.CreatesEntities);
 	}
 }
 
diff --git a/Engine/Source/Runtime/MovieScene/Private/EntitySystem/MovieSceneBoundSceneComponentInstantiator.cpp b/Engine/Source/Runtime/MovieScene/Private/EntitySystem/MovieSceneBoundSceneComponentInstantiator.cpp
--- a/Engine/Source/Runtime/MovieScene/Private/EntitySystem/MovieSceneBoundSceneComponentInstantiator.cpp
+++ b/Engine/Source/Runtime/MovieScene/Private/EntitySystem/MovieSceneBoundSceneComponentInstantiator.cpp
@@ -25,6 +25,7 @@ UMovieSceneBoundSceneComponentInstantiator::UMovieSceneBoundSceneComponentInstan
 	if (HasAnyFlags(RF_ClassDefaultObject))
 	{
 		DefineComponentProducer(GetClass(), Components->BoundObject);
+		DefineComponentProducer(GetClass(), Components->SymbolicTags.CreatesEntities);
 	}
 }
 
