Don't assert on exit if the runner's linker is pending kill.

This can happen if the linker is flagged for destruction before whatever owns the runner.

diff --git a/Engine/Source/Runtime/MovieScene/Private/EntitySystem/MovieSceneEntitySystemRunner.cpp b/Engine/Source/Runtime/MovieScene/Private/EntitySystem/MovieSceneEntitySystemRunner.cpp
--- a/Engine/Source/Runtime/MovieScene/Private/EntitySystem/MovieSceneEntitySystemRunner.cpp
+++ b/Engine/Source/Runtime/MovieScene/Private/EntitySystem/MovieSceneEntitySystemRunner.cpp
@@ -63,7 +63,7 @@ void FMovieSceneEntitySystemRunner::DetachFromLinker()
 {
 	if (ensureMsgf(!WeakLinker.IsExplicitlyNull(), TEXT("This runner is not attached to any linker")))
 	{
-		if (ensureMsgf(WeakLinker.IsValid(), TEXT("This runner is attached to an invalid linker!")))
+		if (!GExitPurge && ensureMsgf(WeakLinker.IsValid(), TEXT("This runner is attached to an invalid linker!")))
 		{
 			OnLinkerAbandon(WeakLinker.Get());
 		}
