Sequencer: Disable flushing on pre-garbage collection for now.

This was only added to prevent unwanted side-effects for budgeted evaluations which are currently not being used anywhere, and it causes issues with re-entrancy when a garbage collection is triggered from inside a Sequencer evaluation.

diff --git a/Engine/Source/Runtime/MovieScene/Private/EntitySystem/MovieSceneEntitySystemLinker.cpp b/Engine/Source/Runtime/MovieScene/Private/EntitySystem/MovieSceneEntitySystemLinker.cpp
--- a/Engine/Source/Runtime/MovieScene/Private/EntitySystem/MovieSceneEntitySystemLinker.cpp
+++ b/Engine/Source/Runtime/MovieScene/Private/EntitySystem/MovieSceneEntitySystemLinker.cpp
@@ -270,15 +270,22 @@ void UMovieSceneEntitySystemLinker::HandlePreGarbageCollection()
 {
 	using namespace UE::MovieScene;
 
+	// @todo: This is currently disabled because it is too indiscriminate
+	//         with regards to when garbage collection is run (ie, if it's run _inside_
+	//         the instantiation phase, we end up running everything again or if it's run without any
+	//         outstanding work, it performs an unnecessary flush.
+	//
+	//         For now nothing is using the budgeting which this code was written for, so we will revisit in future
+	
 	// If we have any active runners part-way through an evaluation,
 	// they must be flushed before we run a garbage collection
-	for (int32 Index = ActiveRunners.Num()-1; Index >= 0; --Index)
-	{
-		if (ActiveRunners[Index].Runner->IsAttachedToLinker())
-		{
-			ActiveRunners[Index].Runner->Flush();
-		}
-	}
+	//for (int32 Index = ActiveRunners.Num()-1; Index >= 0; --Index)
+	//{
+	//	if (ActiveRunners[Index].Runner->IsAttachedToLinker())
+	//	{
+	//		ActiveRunners[Index].Runner->Flush();
+	//	}
+	//}
 }
 
 void UMovieSceneEntitySystemLinker::HandlePostGarbageCollection()
