MovieScene: Fixed smoothed replicated playback causing desyncs during looped playback.

diff --git a/Engine/Source/Runtime/MovieScene/Private/MovieSceneSequencePlayer.cpp b/Engine/Source/Runtime/MovieScene/Private/MovieSceneSequencePlayer.cpp
--- a/Engine/Source/Runtime/MovieScene/Private/MovieSceneSequencePlayer.cpp
+++ b/Engine/Source/Runtime/MovieScene/Private/MovieSceneSequencePlayer.cpp
@@ -983,6 +983,13 @@ void UMovieSceneSequencePlayer::UpdateTimeCursorPosition_Internal(FFrameTime New
 		if (bLoopIndefinitely || CurrentNumLoops + NumTimesLooped <= PlaybackSettings.LoopCount.Value)
 		{
 			CurrentNumLoops += NumTimesLooped;
+			if (NumTimesLooped > 0)
+			{
+				// Reset server time samples when this player has looped. This ensures that
+				// smoothed playback (if enabled) does not result in a smoothed frame in the previous
+				// loop.
+				ServerTimeSamples.Reset();
+			}
 
 			// Finish evaluating any frames left in the current loop in case they have events attached
 			FFrameTime CurrentPosition = PlayPosition.GetCurrentPosition();
