Sequencer: pass owning section when evaluating a channel

Some new channel types (like perlin noise) need the section pointer to look up contextual data, like tick resolution.

diff --git a/Engine/Source/Editor/MovieSceneTools/Private/CurveKeyEditors/SequencerKeyEditor.h b/Engine/Source/Editor/MovieSceneTools/Private/CurveKeyEditors/SequencerKeyEditor.h
--- a/Engine/Source/Editor/MovieSceneTools/Private/CurveKeyEditors/SequencerKeyEditor.h
+++ b/Engine/Source/Editor/MovieSceneTools/Private/CurveKeyEditors/SequencerKeyEditor.h
@@ -76,7 +76,7 @@ struct TSequencerKeyEditor
 		{
 			const FFrameTime CurrentTime = UE::MovieScene::ClampToDiscreteRange(Sequencer->GetLocalTime().Time, OwningSection->GetRange());
 			//If we have no keys and no default, key with the external value if it exists
-			if (!EvaluateChannel(Channel, CurrentTime, Result))
+			if (!EvaluateChannel(OwningSection, Channel, CurrentTime, Result))
 			{
 				if (TOptional<ValueType> ExternalValue = GetExternalValue())
 				{
@@ -195,4 +195,4 @@ private:
 	TWeakPtr<ISequencer> WeakSequencer;
 	TWeakPtr<FTrackInstancePropertyBindings> WeakPropertyBindings;
 	TFunction<TOptional<ValueType>(UObject&, FTrackInstancePropertyBindings*)> OnGetExternalValue;
-};
\ No newline at end of file
+};
diff --git a/Engine/Source/Runtime/MovieScene/Public/Channels/MovieSceneChannelTraits.h b/Engine/Source/Runtime/MovieScene/Public/Channels/MovieSceneChannelTraits.h
--- a/Engine/Source/Runtime/MovieScene/Public/Channels/MovieSceneChannelTraits.h
+++ b/Engine/Source/Runtime/MovieScene/Public/Channels/MovieSceneChannelTraits.h
@@ -71,6 +71,21 @@ namespace MovieScene
 		return InChannel->Evaluate(InTime, OutValue);
 	}
 
+	/**
+	 * Called to evaluate a channel. Overload with specific channel types for custom behaviour.
+	 *
+	 * @param InSection	  The section that contains the channel
+	 * @param InChannel   The channel to evaluate
+	 * @param InTime      The time to evaluate at
+	 * @param OutValue    Value to receive the result
+	 * @return true if the channel was evaluated successfully, false otherwise
+	 */
+	template<typename ChannelType, typename ValueType>
+	bool EvaluateChannel(const UMovieSceneSection* InSection, const ChannelType* InChannel, FFrameTime InTime, ValueType& OutValue)
+	{
+		return EvaluateChannel(InChannel, InTime, OutValue);
+	}
+
 	/**
 	 * Called to assign a specific value in a channel.
 	 *
@@ -251,4 +266,4 @@ namespace MovieScene
 		}
 	}
 } // namespace MovieScene
-} // namespace UE
\ No newline at end of file
+} // namespace UE
