Sequencer: new channel proxy functions

- Function to add a channel with some default extended editor data. It's expected that one of the callers will eventually specify the correct extended editor data for the new channel, although we currently don't do any type-safety check to ensure that this extended editor data will be of the correct type. This whole thing is necessary for overriden channels which are added by the custom channel containers, but don't know about what they are.
- Function to get a channel from its SortOrder, instead of its index in the typed bucket.

diff --git a/Engine/Source/Runtime/MovieScene/Private/Channels/MovieSceneChannelProxy.cpp b/Engine/Source/Runtime/MovieScene/Private/Channels/MovieSceneChannelProxy.cpp
--- a/Engine/Source/Runtime/MovieScene/Private/Channels/MovieSceneChannelProxy.cpp
+++ b/Engine/Source/Runtime/MovieScene/Private/Channels/MovieSceneChannelProxy.cpp
@@ -41,4 +41,22 @@ FMovieSceneChannel* FMovieSceneChannelProxy::GetChannel(FName ChannelTypeName, i
 		return Channels.IsValidIndex(ChannelIndex) ? Channels[ChannelIndex] : nullptr;
 	}
 	return nullptr;
-}
\ No newline at end of file
+}
+
+#if WITH_EDITOR
+
+FMovieSceneChannel* FMovieSceneChannelProxy::GetChannelBySortOrder(FName ChannelTypeName, int32 ChannelSortOrder) const
+{
+	if (const FMovieSceneChannelEntry* Entry = FindEntry(ChannelTypeName))
+	{
+		TArrayView<FMovieSceneChannel* const> Channels = Entry->GetChannels();
+		const int32 ChannelIndex = Entry->GetChannelIndexBySortOrder(ChannelSortOrder);
+		if (ChannelIndex != INDEX_NONE)
+		{
+			return Channels[ChannelIndex];
+		}
+	}
+	return nullptr;
+}
+
+#endif // WITH_EDITOR
diff --git a/Engine/Source/Runtime/MovieScene/Public/Channels/MovieSceneChannelEditorDataEntry.h b/Engine/Source/Runtime/MovieScene/Public/Channels/MovieSceneChannelEditorDataEntry.h
--- a/Engine/Source/Runtime/MovieScene/Public/Channels/MovieSceneChannelEditorDataEntry.h
+++ b/Engine/Source/Runtime/MovieScene/Public/Channels/MovieSceneChannelEditorDataEntry.h
@@ -4,6 +4,7 @@
 
 #include "UObject/ObjectMacros.h"
 #include "Misc/InlineValue.h"
+#include "Algo/BinarySearch.h"
 #include "Algo/Sort.h"
 #include "Misc/Attribute.h"
 #include "Misc/Optional.h"
@@ -36,6 +37,20 @@ struct FMovieSceneChannelEditorDataEntry
 		return MetaDataArray;
 	}
 
+	/**
+	 * Returns the index of the first channel with the given sort order
+	 * Note that there can be multiple channels with the same sort order
+	 */
+	int32 GetChannelIndexBySortOrder(uint32 SortOrder) const
+	{
+		const int32 Index = Algo::BinarySearchBy(SortOrderedArray, SortOrder, &FChannelSortOrderIndex::SortOrder);
+		if (Index != INDEX_NONE)
+		{
+			return SortOrderedArray[Index].ChannelIndex;
+		}
+		return INDEX_NONE;
+	}
+
 	/**
 	 * Access the extended editor data for a specific channel
 	 */
@@ -58,6 +73,9 @@ protected:
 	{
 		static_assert(TIsSame<typename TMovieSceneChannelTraits<ChannelType>::ExtendedEditorDataType, void>::Value, "Must supply extended editor data according to the channel's traits.");
 
+		// Keep track of channels in sort order
+		InsertSortOrderEntry(FChannelSortOrderIndex { MetaData.SortOrder, MetaDataArray.Num() }, MetaData.Name);
+
 		// Add the editor meta-data
 		MetaDataArray.Add(MetaData);
 	}
@@ -68,6 +86,9 @@ protected:
 	template<typename ChannelType, typename ExtendedEditorDataType>
 	void AddMetaData(const FMovieSceneChannelMetaData& MetaData, ExtendedEditorDataType&& InExtendedEditorData)
 	{
+		// Keep track of channels in sort order
+		InsertSortOrderEntry(FChannelSortOrderIndex { MetaData.SortOrder, MetaDataArray.Num() }, MetaData.Name);
+
 		// Add the editor meta-data
 		MetaDataArray.Add(MetaData);
 
@@ -76,6 +97,17 @@ protected:
 		TypedImpl.Data.Add(Forward<ExtendedEditorDataType>(InExtendedEditorData));
 	}
 
+	/**
+	 * Set the extended channel-type specific editor data on an already added channel entry
+	 */
+	template<typename ChannelType, typename ExtendedEditorDataType>
+	void SetExtendedEditorData(int32 EntryIndex, ExtendedEditorDataType&& InExtendedEditorData)
+	{
+		// Set the extended channel-type specific editor data
+		auto& TypedImpl = static_cast<TMovieSceneExtendedEditorDataArray<ChannelType>&>(ExtendedEditorDataArray.GetValue());
+		TypedImpl.Data[EntryIndex] = Forward<ExtendedEditorDataType>(InExtendedEditorData);
+	}
+
 	/**
 	 * Access the extended editor data for channels stored in this entry
 	 */
@@ -90,6 +122,12 @@ protected:
 
 private:
 
+	struct FChannelSortOrderIndex
+	{
+		uint32 SortOrder;
+		int32 ChannelIndex;
+	};
+
 	/** Construct the extended editor data container for channel types that require it */
 	template<typename ChannelType, typename ExtendedEditorDataType>
 	void ConstructExtendedEditorDataArray(ExtendedEditorDataType*)
@@ -102,11 +140,33 @@ private:
 	{
 	}
 
+	void InsertSortOrderEntry(FChannelSortOrderIndex InEntry, FName InChannelName)
+	{
+		check(SortOrderedArray.Num() == MetaDataArray.Num());
+
+		// Get the first element where SortOrder >= InEntry.SortOrder.
+		int32 Index = Algo::LowerBoundBy(SortOrderedArray, InEntry.SortOrder, &FChannelSortOrderIndex::SortOrder);
+		// Keep going if there are more entries with the same sort order, but lesser names.
+		while (Index < SortOrderedArray.Num() - 1)
+		{
+			const FMovieSceneChannelMetaData& NextEntry(MetaDataArray[Index + 1]);
+			if (NextEntry.SortOrder != InEntry.SortOrder || NextEntry.Name.Compare(InChannelName) > 0)
+			{
+				break;
+			}
+			++Index;
+		}
+		SortOrderedArray.Insert(InEntry, Index);
+	}
+
 private:
 
 	/** Base editor data, one per channel */
 	TArray<FMovieSceneChannelMetaData, TInlineAllocator<1>> MetaDataArray;
 
+	/** Indirection between channel index and sort order */
+	TArray<FChannelSortOrderIndex> SortOrderedArray;
+
 	/**
 	 * We store the array behind an interface whose access is via void*
 	 * Typed access is only permitted using the original ChannelType templated methods
diff --git a/Engine/Source/Runtime/MovieScene/Public/Channels/MovieSceneChannelProxy.h b/Engine/Source/Runtime/MovieScene/Public/Channels/MovieSceneChannelProxy.h
--- a/Engine/Source/Runtime/MovieScene/Public/Channels/MovieSceneChannelProxy.h
+++ b/Engine/Source/Runtime/MovieScene/Public/Channels/MovieSceneChannelProxy.h
@@ -105,6 +105,27 @@ struct FMovieSceneChannelProxyData
 		Entries[ChannelTypeIndex].AddMetaData<ChannelType>(InMetaData);
 	}
 
+	/**
+	 * Add a new channel to the proxy. Channel's address is stored internally as a voic* and should exist as long as the channel proxy does
+	 *
+	 * @param InChannel          The channel to add to this proxy. Should live for as long as the proxy does. Any re-allocation should be accompanied with a re-creation of the proxy
+	 * @param InMetaData         The editor meta data to be associated with this channel
+	 */
+	template<typename ChannelType>
+	FMovieSceneChannelHandle AddWithDefaultEditorData(ChannelType& InChannel, const FMovieSceneChannelMetaData& InMetaData)
+	{
+		static_assert(!TIsSame<typename TMovieSceneChannelTraits<ChannelType>::ExtendedEditorDataType, void>::Value, "This method is for channels with typed editor data. You *must* call SetExtendedEditorData afterwards.");
+		static_assert(!TIsSame<ChannelType, FMovieSceneChannel>::Value, "Cannot add channels by their base FMovieSceneChannel type.");
+
+		// Add the channel
+		const int32 ChannelTypeIndex = AddInternal(InChannel);
+		// Add a default editor data at the same index, hopefully the caller will set it afterwards
+		Entries[ChannelTypeIndex].AddMetaData<ChannelType>(InMetaData, TMovieSceneChannelTraits<ChannelType>::ExtendedEditorDataType());
+		// Return index usable for SetExtendedEditorData
+		const FName ChannelTypeName = ChannelType::StaticStruct()->GetFName();
+		return FMovieSceneChannelHandle(nullptr, ChannelTypeName, Entries[ChannelTypeIndex].GetChannels().Num() - 1);
+	}
+
 	/**
 	 * Add a new channel to the proxy. Channel's address is stored internally as a voic* and should exist as long as the channel proxy does
 	 *
@@ -124,6 +145,20 @@ struct FMovieSceneChannelProxyData
 		Entries[ChannelTypeIndex].AddMetaData<ChannelType>(InMetaData, Forward<ExtendedEditorDataType>(InExtendedEditorData));
 	}
 
+	template<typename ChannelType, typename ExtendedEditorDataType>
+	void SetExtendedEditorData(FMovieSceneChannelHandle ChannelHandle, ExtendedEditorDataType&& InExtendedEditorData)
+	{
+		// Find the entry
+		// TODO: we rely on ChannelType's extended editor data to be the same as the overriden extended editor data!
+		const FName ChannelTypeName = ChannelHandle.GetChannelTypeName();
+		FMovieSceneChannelEntry* Entry = Entries.FindByPredicate([=](const FMovieSceneChannelEntry& CurEntry)
+				{ return CurEntry.ChannelTypeName == ChannelTypeName; });
+		if (ensure(Entry))
+		{
+			Entry->SetExtendedEditorData<ChannelType>(ChannelHandle.GetChannelIndex(), InExtendedEditorData);
+		}
+	}
+
 #else
 
 	/**
@@ -259,7 +294,6 @@ public:
 
 #if !WITH_EDITOR
 
-
 	/**
 	 * Construction via a single channel, and its editor data
 	 * Channel's address is stored internally as a voic* and should exist as long as this channel proxy does.
@@ -292,6 +326,21 @@ public:
 	template<typename ChannelType>
 	TArrayView<const FMovieSceneChannelMetaData> GetMetaData() const;
 
+	/**
+	 * Get the channel for the specified sort index of a particular type.
+	 *
+	 * @return A pointer to the channel, or nullptr if the index was invalid, or the type was not present
+	 */
+	template<typename ChannelType>
+	ChannelType* GetChannelBySortOrder(int32 ChannelSortOrder) const;
+
+	/**
+	 * Get the channel for the specified sort index of a particular type.
+	 *
+	 * @return A pointer to the channel, or nullptr if the index was invalid, or the type was not present
+	 */
+	FMovieSceneChannel* GetChannelBySortOrder(FName ChannelTypeName, int32 ChannelSortOrder) const;
+
 	/**
 	 * Access all the extended data for the templated channel type
 	 *
@@ -371,7 +420,6 @@ ChannelType* FMovieSceneChannelProxy::GetChannel(int32 ChannelIndex) const
 	return Channels.IsValidIndex(ChannelIndex) ? Channels[ChannelIndex] : nullptr;
 }
 
-
 #if !WITH_EDITOR
 
 
@@ -464,4 +512,18 @@ TArrayView<const FMovieSceneChannelMetaData> FMovieSceneChannelProxy::GetMetaDat
 	return TArrayView<const FMovieSceneChannelMetaData>();
 }
 
-#endif	// !WITH_EDITOR
\ No newline at end of file
+
+/**
+ * Get the channel for the specified sort index of a particular type.
+ *
+ * @return A pointer to the channel, or nullptr if the index was invalid, or the type was not present
+ */
+template<typename ChannelType>
+ChannelType* FMovieSceneChannelProxy::GetChannelBySortOrder(int32 ChannelSortOrder) const
+{
+	FName ChannelTypeName = ChannelType::StaticStruct()->GetFName();
+	return static_cast<ChannelType*>(GetChannelBySortOrder(ChannelTypeName, ChannelSortOrder));
+}
+
+
+#endif	// !WITH_EDITOR
