Sequencer: fix code for getting a component writer to a single entity

diff --git a/Engine/Source/Runtime/MovieScene/Public/EntitySystem/MovieSceneEntityManager.h b/Engine/Source/Runtime/MovieScene/Public/EntitySystem/MovieSceneEntityManager.h
--- a/Engine/Source/Runtime/MovieScene/Public/EntitySystem/MovieSceneEntityManager.h
+++ b/Engine/Source/Runtime/MovieScene/Public/EntitySystem/MovieSceneEntityManager.h
@@ -503,7 +503,7 @@ public:
 		{
 			if (Header.ComponentType == ComponentTypeID)
 			{
-				return TComponentLock<TWriteOptional<T>>(Header, FEntityAllocationWriteContext(*this), ComponentOffset);
+				return TComponentLock<TWriteOptional<T>>(&Header, FEntityAllocationWriteContext(*this), ComponentOffset);
 			}
 		}
 
@@ -520,7 +520,7 @@ public:
 	 * @param Value                The value to write
 	 */
 	template<typename T, typename ValueType>
-	void WriteComponentChecked(FMovieSceneEntityID Entity, TComponentTypeID<T> ComponentTypeID, ValueType&& Value) const
+	void WriteComponentChecked(FMovieSceneEntityID Entity, TComponentTypeID<T> ComponentTypeID, ValueType&& Value)
 	{
 		TComponentLock<TWriteOptional<T>> ComponentPtr = WriteComponent(Entity, ComponentTypeID);
 		check(ComponentPtr);
