Sequencer: Follow up to add support for the new property setter functions, ie UPROPERTY(Setter = MyPropertySetter) for specialized CallFunction overrides for bool and UObject properties

diff --git a/Engine/Source/Runtime/MovieScene/Private/MovieSceneCommonHelpers.cpp b/Engine/Source/Runtime/MovieScene/Private/MovieSceneCommonHelpers.cpp
--- a/Engine/Source/Runtime/MovieScene/Private/MovieSceneCommonHelpers.cpp
+++ b/Engine/Source/Runtime/MovieScene/Private/MovieSceneCommonHelpers.cpp
@@ -730,26 +730,29 @@ FTrackInstancePropertyBindings::FPropertyAddress FTrackInstancePropertyBindings:
 void FTrackInstancePropertyBindings::CallFunctionForEnum( UObject& InRuntimeObject, int64 PropertyValue )
 {
 	FPropertyAndFunction PropAndFunction = FindOrAdd(InRuntimeObject);
-	if (UFunction* SetterFunction = PropAndFunction.SetterFunction.Get())
+
+	FProperty* Property = GetProperty(InRuntimeObject);
+	if (Property && Property->HasSetter())
+	{
+		Property->CallSetter(&InRuntimeObject, &PropertyValue);
+	}
+	else if (UFunction* SetterFunction = PropAndFunction.SetterFunction.Get())
 	{
 		InvokeSetterFunction(&InRuntimeObject, SetterFunction, PropertyValue);
 	}
-	else if (FProperty* Property = PropAndFunction.PropertyAddress.GetProperty())
+	else if (Property && Property->IsA(FEnumProperty::StaticClass()))
 	{
-		if (Property->IsA(FEnumProperty::StaticClass()))
-		{
-			if (FEnumProperty* EnumProperty = CastFieldChecked<FEnumProperty>(Property))
-			{
-				FNumericProperty* UnderlyingProperty = EnumProperty->GetUnderlyingProperty();
-				void* ValueAddr = EnumProperty->ContainerPtrToValuePtr<void>(PropAndFunction.PropertyAddress.Address);
-				UnderlyingProperty->SetIntPropertyValue(ValueAddr, PropertyValue);
-			}
-		}
-		else
+		if (FEnumProperty* EnumProperty = CastFieldChecked<FEnumProperty>(Property))
 		{
-			UE_LOG(LogMovieScene, Error, TEXT("Mismatch in property evaluation. %s is not of type: %s"), *Property->GetName(), *FEnumProperty::StaticClass()->GetName());
+			FNumericProperty* UnderlyingProperty = EnumProperty->GetUnderlyingProperty();
+			void* ValueAddr = EnumProperty->ContainerPtrToValuePtr<void>(PropAndFunction.PropertyAddress.Address);
+			UnderlyingProperty->SetIntPropertyValue(ValueAddr, PropertyValue);
 		}
 	}
+	else
+	{
+		UE_LOG(LogMovieScene, Error, TEXT("Mismatch in property evaluation. %s is not of type: %s"), *Property->GetName(), *FEnumProperty::StaticClass()->GetName());
+	}
 
 	if (UFunction* NotifyFunction = PropAndFunction.NotifyFunction.Get())
 	{
@@ -818,25 +821,28 @@ int64 FTrackInstancePropertyBindings::GetCurrentValueForEnum(const UObject& Obje
 template<> void FTrackInstancePropertyBindings::CallFunction<bool>(UObject& InRuntimeObject, TCallTraits<bool>::ParamType PropertyValue)
 {
 	FPropertyAndFunction PropAndFunction = FindOrAdd(InRuntimeObject);
-	if (UFunction* SetterFunction = PropAndFunction.SetterFunction.Get())
+
+	FProperty* Property = GetProperty(InRuntimeObject);
+	if (Property && Property->HasSetter())
+	{
+		Property->CallSetter(&InRuntimeObject, &PropertyValue);
+	}
+	else if (UFunction* SetterFunction = PropAndFunction.SetterFunction.Get())
 	{
 		InvokeSetterFunction(&InRuntimeObject, SetterFunction, PropertyValue);
 	}
-	else if (FProperty* Property = PropAndFunction.PropertyAddress.GetProperty())
+	else if (Property && Property->IsA(FBoolProperty::StaticClass()))
 	{
-		if (Property->IsA(FBoolProperty::StaticClass()))
+		if (FBoolProperty* BoolProperty = CastFieldChecked<FBoolProperty>(Property))
 		{
-			if (FBoolProperty* BoolProperty = CastFieldChecked<FBoolProperty>(Property))
-			{
-				uint8* ValuePtr = BoolProperty->ContainerPtrToValuePtr<uint8>(PropAndFunction.PropertyAddress.Address);
-				BoolProperty->SetPropertyValue(ValuePtr, PropertyValue);
-			}
-		}
-		else
-		{
-			UE_LOG(LogMovieScene, Error, TEXT("Mismatch in property evaluation. %s is not of type: %s"), *Property->GetName(), *FBoolProperty::StaticClass()->GetName());
+			uint8* ValuePtr = BoolProperty->ContainerPtrToValuePtr<uint8>(PropAndFunction.PropertyAddress.Address);
+			BoolProperty->SetPropertyValue(ValuePtr, PropertyValue);
 		}
 	}
+	else
+	{
+		UE_LOG(LogMovieScene, Error, TEXT("Mismatch in property evaluation. %s is not of type: %s"), *Property->GetName(), *FBoolProperty::StaticClass()->GetName());
+	}
 
 	if (UFunction* NotifyFunction = PropAndFunction.NotifyFunction.Get())
 	{
@@ -888,14 +894,27 @@ template<> void FTrackInstancePropertyBindings::SetCurrentValue<bool>(UObject& O
 template<> void FTrackInstancePropertyBindings::CallFunction<UObject*>(UObject& InRuntimeObject, UObject* PropertyValue)
 {
 	FPropertyAndFunction PropAndFunction = FindOrAdd(InRuntimeObject);
-	if (UFunction* SetterFunction = PropAndFunction.SetterFunction.Get())
+
+	FProperty* Property = GetProperty(InRuntimeObject);
+	if (Property && Property->HasSetter())
+	{
+		Property->CallSetter(&InRuntimeObject, &PropertyValue);
+	}
+	else if (UFunction* SetterFunction = PropAndFunction.SetterFunction.Get())
 	{
 		InvokeSetterFunction(&InRuntimeObject, SetterFunction, PropertyValue);
 	}
-	else if (FObjectPropertyBase* ObjectProperty = CastField<FObjectPropertyBase>(PropAndFunction.PropertyAddress.GetProperty()))
+	else if (Property && Property->IsA(FObjectPropertyBase::StaticClass()))
 	{
-		uint8* ValuePtr = ObjectProperty->ContainerPtrToValuePtr<uint8>(PropAndFunction.PropertyAddress.Address);
-		ObjectProperty->SetObjectPropertyValue(ValuePtr, PropertyValue);
+		if (FObjectPropertyBase* ObjectProperty = CastFieldChecked<FObjectPropertyBase>(Property))
+		{
+			uint8* ValuePtr = ObjectProperty->ContainerPtrToValuePtr<uint8>(PropAndFunction.PropertyAddress.Address);
+			ObjectProperty->SetObjectPropertyValue(ValuePtr, PropertyValue);
+		}
+	}
+	else
+	{
+		UE_LOG(LogMovieScene, Error, TEXT("Mismatch in property evaluation. %s is not of type: %s"), *Property->GetName(), *FObjectPropertyBase::StaticClass()->GetName());
 	}
 
 	if (UFunction* NotifyFunction = PropAndFunction.NotifyFunction.Get())
