Sequencer: add functions to get all imported entities given an owner object

diff --git a/Engine/Source/Runtime/MovieScene/Private/EntitySystem/MovieSceneEntityLedger.cpp b/Engine/Source/Runtime/MovieScene/Private/EntitySystem/MovieSceneEntityLedger.cpp
--- a/Engine/Source/Runtime/MovieScene/Private/EntitySystem/MovieSceneEntityLedger.cpp
+++ b/Engine/Source/Runtime/MovieScene/Private/EntitySystem/MovieSceneEntityLedger.cpp
@@ -133,6 +133,17 @@ FMovieSceneEntityID FEntityLedger::FindImportedEntity(const FMovieSceneEvaluatio
 	return ImportedEntities.FindRef(EntityKey).EntityID;
 }
 
+void FEntityLedger::FindImportedEntities(TWeakObjectPtr<UObject> EntityOwner, TArray<FMovieSceneEntityID>& OutEntityIDs) const
+{
+	for (const TPair<FMovieSceneEvaluationFieldEntityKey, FImportedEntityData>& Pair : ImportedEntities)
+	{
+		if (Pair.Key.EntityOwner == EntityOwner)
+		{
+			OutEntityIDs.Add(Pair.Value.EntityID);
+		}
+	}
+}
+
 void FEntityLedger::ImportEntity(UMovieSceneEntitySystemLinker* Linker, const FEntityImportSequenceParams& ImportParams, const FMovieSceneEntityComponentField* EntityField, const FMovieSceneEvaluationFieldEntityQuery& Query)
 {
 	// We always add an entry even if no entity was imported by the provider to ensure that we do not repeatedly try and import the same entity every frame
@@ -239,4 +250,4 @@ void FEntityLedger::TagGarbage(UMovieSceneEntitySystemLinker* Linker)
 }
 
 } // namespace MovieScene
-} // namespace UE
\ No newline at end of file
+} // namespace UE
diff --git a/Engine/Source/Runtime/MovieScene/Private/EntitySystem/MovieSceneSequenceInstance.cpp b/Engine/Source/Runtime/MovieScene/Private/EntitySystem/MovieSceneSequenceInstance.cpp
--- a/Engine/Source/Runtime/MovieScene/Private/EntitySystem/MovieSceneSequenceInstance.cpp
+++ b/Engine/Source/Runtime/MovieScene/Private/EntitySystem/MovieSceneSequenceInstance.cpp
@@ -380,6 +380,11 @@ FMovieSceneEntityID FSequenceInstance::FindEntity(UObject* Owner, uint32 EntityI
 	return Ledger.FindImportedEntity(FMovieSceneEvaluationFieldEntityKey{ decltype(FMovieSceneEvaluationFieldEntityKey::EntityOwner)(Owner), EntityID });
 }
 
+void FSequenceInstance::FindEntities(UObject* Owner, TArray<FMovieSceneEntityID>& OutEntityIDs) const
+{
+	Ledger.FindImportedEntities(Owner, OutEntityIDs);
+}
+
 FSubSequencePath FSequenceInstance::GetSubSequencePath() const
 {
 	return FSubSequencePath(SequenceID, *GetPlayer());
diff --git a/Engine/Source/Runtime/MovieScene/Private/Evaluation/MovieSceneEvaluationTemplateInstance.cpp b/Engine/Source/Runtime/MovieScene/Private/Evaluation/MovieSceneEvaluationTemplateInstance.cpp
--- a/Engine/Source/Runtime/MovieScene/Private/Evaluation/MovieSceneEvaluationTemplateInstance.cpp
+++ b/Engine/Source/Runtime/MovieScene/Private/Evaluation/MovieSceneEvaluationTemplateInstance.cpp
@@ -317,6 +317,16 @@ UE::MovieScene::FMovieSceneEntityID FMovieSceneRootEvaluationTemplateInstance::F
 	return FMovieSceneEntityID::Invalid();
 }
 
+void FMovieSceneRootEvaluationTemplateInstance::FindEntitiesFromOwner(UObject* Owner, FMovieSceneSequenceID SequenceID, TArray<UE::MovieScene::FMovieSceneEntityID>& OutEntityIDs) const
+{
+	using namespace UE::MovieScene;
+
+	if (const FSequenceInstance* SequenceInstance = FindInstance(SequenceID))
+	{
+		SequenceInstance->FindEntities(Owner, OutEntityIDs);
+	}
+}
+
 UObject* FMovieSceneRootEvaluationTemplateInstance::GetOrCreateDirectorInstance(FMovieSceneSequenceIDRef SequenceID, IMovieScenePlayer& Player)
 {
 	UObject* ExistingDirectorInstance = DirectorInstances.FindRef(SequenceID);
diff --git a/Engine/Source/Runtime/MovieScene/Public/EntitySystem/MovieSceneEntityLedger.h b/Engine/Source/Runtime/MovieScene/Public/EntitySystem/MovieSceneEntityLedger.h
--- a/Engine/Source/Runtime/MovieScene/Public/EntitySystem/MovieSceneEntityLedger.h
+++ b/Engine/Source/Runtime/MovieScene/Public/EntitySystem/MovieSceneEntityLedger.h
@@ -72,6 +72,11 @@ public:
 	 */
 	FMovieSceneEntityID FindImportedEntity(const FMovieSceneEvaluationFieldEntityKey& EntityKey) const;
 
+	/**
+	 * Find imported entities
+	 */
+	void FindImportedEntities(TWeakObjectPtr<UObject> EntityOwner, TArray<FMovieSceneEntityID>& OutEntityIDs) const;
+
 	/**
 	 * Indicate that the specified field entity is currently being evaluated
 	 *
@@ -132,4 +137,4 @@ private:
 
 
 } // namespace MovieScene
-} // namespace UE
\ No newline at end of file
+} // namespace UE
diff --git a/Engine/Source/Runtime/MovieScene/Public/EntitySystem/MovieSceneSequenceInstance.h b/Engine/Source/Runtime/MovieScene/Public/EntitySystem/MovieSceneSequenceInstance.h
--- a/Engine/Source/Runtime/MovieScene/Public/EntitySystem/MovieSceneSequenceInstance.h
+++ b/Engine/Source/Runtime/MovieScene/Public/EntitySystem/MovieSceneSequenceInstance.h
@@ -202,6 +202,11 @@ public:
 	 */
 	FMovieSceneEntityID FindEntity(UObject* Owner, uint32 EntityID) const;
 
+	/**
+	 * Attempt to locate all entities given their owner
+	 */
+	void FindEntities(UObject* Owner, TArray<FMovieSceneEntityID>& OutEntityIDs) const;
+
 	/**
 	 * Retrieve the legacy evaluator for this sequence, if it is available (may return nullptr)
 	 */
diff --git a/Engine/Source/Runtime/MovieScene/Public/Evaluation/MovieSceneEvaluationTemplateInstance.h b/Engine/Source/Runtime/MovieScene/Public/Evaluation/MovieSceneEvaluationTemplateInstance.h
--- a/Engine/Source/Runtime/MovieScene/Public/Evaluation/MovieSceneEvaluationTemplateInstance.h
+++ b/Engine/Source/Runtime/MovieScene/Public/Evaluation/MovieSceneEvaluationTemplateInstance.h
@@ -4,6 +4,7 @@
 
 #include "CoreTypes.h"
 #include "Containers/Map.h"
+#include "EntitySystem/MovieSceneEntityIDs.h"
 #include "Templates/UniquePtr.h"
 #include "UObject/WeakObjectPtr.h"
 #include "UObject/StrongObjectPtr.h"
@@ -108,6 +109,7 @@ public:
 	MOVIESCENE_API const UE::MovieScene::FSequenceInstance* FindInstance(FMovieSceneSequenceID SequenceID) const;
 
 	MOVIESCENE_API UE::MovieScene::FMovieSceneEntityID FindEntityFromOwner(UObject* Owner, uint32 EntityID, FMovieSceneSequenceID SequenceID) const;
+	MOVIESCENE_API void FindEntitiesFromOwner(UObject* Owner, FMovieSceneSequenceID SequenceID, TArray<UE::MovieScene::FMovieSceneEntityID>& OutEntityIDs) const;
 
 	MOVIESCENE_API UMovieSceneSequence* GetSequence(FMovieSceneSequenceIDRef SequenceID) const;
 
