Sequencer: if an imported entity already specifies a result component, just duplicate it to the output/children entities.

This is because overriden channels are each modeled by a single entity. That entity's only way to know which result channel to write to is to specify that result channel directly from import.

diff --git a/Engine/Source/Runtime/MovieScene/Private/EntitySystem/BuiltInComponentTypes.cpp b/Engine/Source/Runtime/MovieScene/Private/EntitySystem/BuiltInComponentTypes.cpp
--- a/Engine/Source/Runtime/MovieScene/Private/EntitySystem/BuiltInComponentTypes.cpp
+++ b/Engine/Source/Runtime/MovieScene/Private/EntitySystem/BuiltInComponentTypes.cpp
@@ -193,6 +193,7 @@ FBuiltInComponentTypes::FBuiltInComponentTypes()
 	// Byte channel relationships
 	{
 		ComponentRegistry->Factories.DuplicateChildComponent(ByteChannel);
+		ComponentRegistry->Factories.DuplicateChildComponent(ByteResult);
 		ComponentRegistry->Factories.DefineMutuallyInclusiveComponent(ByteChannel, ByteResult);
 		ComponentRegistry->Factories.DefineMutuallyInclusiveComponent(ByteChannel, EvalTime);
 	}
@@ -200,6 +201,7 @@ FBuiltInComponentTypes::FBuiltInComponentTypes()
 	// Integer channel relationships
 	{
 		ComponentRegistry->Factories.DuplicateChildComponent(IntegerChannel);
+		ComponentRegistry->Factories.DuplicateChildComponent(IntegerResult);
 		ComponentRegistry->Factories.DefineMutuallyInclusiveComponent(IntegerChannel, IntegerResult);
 		ComponentRegistry->Factories.DefineMutuallyInclusiveComponent(IntegerChannel, EvalTime);
 
@@ -248,6 +250,7 @@ FBuiltInComponentTypes::FBuiltInComponentTypes()
 		for (int32 Index = 0; Index < UE_ARRAY_COUNT(DoubleChannel); ++Index)
 		{
 			ComponentRegistry->Factories.DuplicateChildComponent(DoubleChannel[Index]);
+			ComponentRegistry->Factories.DuplicateChildComponent(DoubleResult[Index]);
 			ComponentRegistry->Factories.DefineMutuallyInclusiveComponent(DoubleChannel[Index], DoubleResult[Index]);
 			ComponentRegistry->Factories.DefineMutuallyInclusiveComponent(DoubleChannel[Index], EvalTime);
 			ComponentRegistry->Factories.DefineMutuallyInclusiveComponent(DoubleChannel[Index], DoubleChannelFlags[Index]);
@@ -277,6 +280,7 @@ FBuiltInComponentTypes::FBuiltInComponentTypes()
 	{
 		// Weight channel components should be duplicated to children
 		ComponentRegistry->Factories.DuplicateChildComponent(WeightChannel);
+		ComponentRegistry->Factories.DuplicateChildComponent(WeightResult);
 
 		// Weight channel components need a time and result to evaluate
 		ComponentRegistry->Factories.DefineMutuallyInclusiveComponent(WeightChannel, EvalTime);
