Sequencer: Fix static analysis

diff --git a/Engine/Source/Runtime/MovieScene/Private/MovieSceneCommonHelpers.cpp b/Engine/Source/Runtime/MovieScene/Private/MovieSceneCommonHelpers.cpp
--- a/Engine/Source/Runtime/MovieScene/Private/MovieSceneCommonHelpers.cpp
+++ b/Engine/Source/Runtime/MovieScene/Private/MovieSceneCommonHelpers.cpp
@@ -680,7 +680,7 @@ FTrackInstancePropertyBindings::FPropertyAddress FTrackInstancePropertyBindings:
 				}
 			}
 		}
-		else
+		else if (PropertyAndIndex.Property)
 		{
 			UE_LOG(LogMovieScene, Error, TEXT("Mismatch in property evaluation. %s is not of type: %s"), *PropertyAndIndex.Property->GetName(), *FArrayProperty::StaticClass()->GetName());
 		}
@@ -749,7 +749,7 @@ void FTrackInstancePropertyBindings::CallFunctionForEnum( UObject& InRuntimeObje
 			UnderlyingProperty->SetIntPropertyValue(ValueAddr, PropertyValue);
 		}
 	}
-	else
+	else if (Property)
 	{
 		UE_LOG(LogMovieScene, Error, TEXT("Mismatch in property evaluation. %s is not of type: %s"), *Property->GetName(), *FEnumProperty::StaticClass()->GetName());
 	}
@@ -839,7 +839,7 @@ template<> void FTrackInstancePropertyBindings::CallFunction<bool>(UObject& InRu
 			BoolProperty->SetPropertyValue(ValuePtr, PropertyValue);
 		}
 	}
-	else
+	else if (Property)
 	{
 		UE_LOG(LogMovieScene, Error, TEXT("Mismatch in property evaluation. %s is not of type: %s"), *Property->GetName(), *FBoolProperty::StaticClass()->GetName());
 	}
@@ -912,7 +912,7 @@ template<> void FTrackInstancePropertyBindings::CallFunction<UObject*>(UObject&
 			ObjectProperty->SetObjectPropertyValue(ValuePtr, PropertyValue);
 		}
 	}
-	else
+	else if (Property)
 	{
 		UE_LOG(LogMovieScene, Error, TEXT("Mismatch in property evaluation. %s is not of type: %s"), *Property->GetName(), *FObjectPropertyBase::StaticClass()->GetName());
 	}
