Sequencer: Null pointer check for TickManager

diff --git a/Engine/Source/Runtime/MovieScene/Private/MovieSceneSequencePlayer.cpp b/Engine/Source/Runtime/MovieScene/Private/MovieSceneSequencePlayer.cpp
--- a/Engine/Source/Runtime/MovieScene/Private/MovieSceneSequencePlayer.cpp
+++ b/Engine/Source/Runtime/MovieScene/Private/MovieSceneSequencePlayer.cpp
@@ -1158,7 +1158,7 @@ void UMovieSceneSequencePlayer::UpdateMovieSceneInstance(FMovieSceneEvaluationRa
 	{
 		CurrentRunner->QueueUpdate(Context, RootTemplateInstance.GetRootInstanceHandle());
 	}
-	else
+	else if (TickManager)
 	{
 		// Look up the runner based on our tick interval
 		FMovieSceneEntitySystemRunner* Runner = TickManager->GetRunner(RegisteredTickInterval.GetValue());
