Sequencer: add finer control over what systems are included or not in a given linker

- Systems can belong to one of more categories, with the possibility of making custom categories
- Linker has allow/disallow lists for system categories or individual systems
- Put most core systems in the core categories
- When linker has auto-linking disabled, auto-unlinking is now also disabled

diff --git a/Engine/Source/Runtime/MovieScene/Private/EntitySystem/MovieSceneBlenderSystem.cpp b/Engine/Source/Runtime/MovieScene/Private/EntitySystem/MovieSceneBlenderSystem.cpp
--- a/Engine/Source/Runtime/MovieScene/Private/EntitySystem/MovieSceneBlenderSystem.cpp
+++ b/Engine/Source/Runtime/MovieScene/Private/EntitySystem/MovieSceneBlenderSystem.cpp
@@ -27,6 +27,8 @@ UMovieSceneBlenderSystem::UMovieSceneBlenderSystem(const FObjectInitializer& Obj
 {
 	using namespace UE::MovieScene;
 
+	SystemCategories = UE::MovieScene::EEntitySystemCategory::BlenderSystems;
+
 	if (HasAnyFlags(RF_ClassDefaultObject))
 	{
 		// You can only ever register blender systems, never unregister them.
diff --git a/Engine/Source/Runtime/MovieScene/Private/EntitySystem/MovieSceneBoundSceneComponentInstantiator.cpp b/Engine/Source/Runtime/MovieScene/Private/EntitySystem/MovieSceneBoundSceneComponentInstantiator.cpp
--- a/Engine/Source/Runtime/MovieScene/Private/EntitySystem/MovieSceneBoundSceneComponentInstantiator.cpp
+++ b/Engine/Source/Runtime/MovieScene/Private/EntitySystem/MovieSceneBoundSceneComponentInstantiator.cpp
@@ -67,12 +67,6 @@ void UMovieSceneBoundSceneComponentInstantiator::OnRun(FSystemTaskPrerequisites&
 						Add(InEntityIndex, SceneComponent);
 						Result = EResolveError::None;
 					}
-					else if (Object->Implements<UMovieSceneSceneComponentImpersonator>())
-					{
-						// Objects meant explicitly to be impersonators are also allowed.
-						Add(InEntityIndex, Object);
-						Result = EResolveError::None;
-					}
 				}
 				else
 				{
diff --git a/Engine/Source/Runtime/MovieScene/Private/EntitySystem/MovieSceneEntitySystem.cpp b/Engine/Source/Runtime/MovieScene/Private/EntitySystem/MovieSceneEntitySystem.cpp
--- a/Engine/Source/Runtime/MovieScene/Private/EntitySystem/MovieSceneEntitySystem.cpp
+++ b/Engine/Source/Runtime/MovieScene/Private/EntitySystem/MovieSceneEntitySystem.cpp
@@ -177,7 +177,7 @@ UMovieSceneEntitySystem::UMovieSceneEntitySystem(const FObjectInitializer& ObjIn
 	using namespace UE::MovieScene;
 
 	bSystemIsEnabled = true;
-	SystemExclusionContext = EEntitySystemContext::None;
+	SystemCategories = EEntitySystemCategory::Unspecified;
 
 	Phase = ESystemPhase::Evaluation;
 	GraphID = TNumericLimits<uint16>::Max();
@@ -251,7 +251,6 @@ void UMovieSceneEntitySystem::LinkRelevantSystems(UMovieSceneEntitySystemLinker*
 {
 	using namespace UE::MovieScene;
 
-	EEntitySystemContext LinkerContext = InLinker->GetSystemContext();
 	for (uint16 GraphID = 0; GraphID < GlobalDependencyGraph.NumGraphIDs(); ++GraphID)
 	{
 		if (InLinker->HasLinkedSystem(GraphID))
@@ -262,13 +261,71 @@ void UMovieSceneEntitySystem::LinkRelevantSystems(UMovieSceneEntitySystemLinker*
 		UClass* Class = GlobalDependencyGraph.ClassFromGraphID(GraphID);
 		UMovieSceneEntitySystem* SystemCDO = Class ? Cast<UMovieSceneEntitySystem>(Class->GetDefaultObject()) : nullptr;
 
-		if (SystemCDO && !EnumHasAnyFlags(SystemCDO->SystemExclusionContext, LinkerContext))
+		if (SystemCDO && InLinker->GetSystemFilter().CheckSystem(SystemCDO))
 		{
 			SystemCDO->ConditionalLinkSystem(InLinker);
 		}
 	}
 }
 
+void UMovieSceneEntitySystem::LinkCategorySystems(UMovieSceneEntitySystemLinker* InLinker, UE::MovieScene::EEntitySystemCategory InCategory)
+{
+	using namespace UE::MovieScene;
+
+	for (uint16 GraphID = 0; GraphID < GlobalDependencyGraph.NumGraphIDs(); ++GraphID)
+	{
+		if (InLinker->HasLinkedSystem(GraphID))
+		{
+			continue;
+		}
+
+		UClass* Class = GlobalDependencyGraph.ClassFromGraphID(GraphID);
+		if (Class == UMovieSceneEntitySystem::StaticClass())
+		{
+			// Ignore the base class
+			continue;
+		}
+
+		UMovieSceneEntitySystem* SystemCDO = Class ? Cast<UMovieSceneEntitySystem>(Class->GetDefaultObject()) : nullptr;
+		if (SystemCDO && EnumHasAnyFlags(SystemCDO->SystemCategories, InCategory))
+		{
+			InLinker->LinkSystemIfAllowed(Class);
+		}
+	}
+}
+
+void UMovieSceneEntitySystem::LinkAllSystems(UMovieSceneEntitySystemLinker* InLinker)
+{
+	using namespace UE::MovieScene;
+
+	for (uint16 GraphID = 0; GraphID < GlobalDependencyGraph.NumGraphIDs(); ++GraphID)
+	{
+		if (InLinker->HasLinkedSystem(GraphID))
+		{
+			continue;
+		}
+
+		UClass* Class = GlobalDependencyGraph.ClassFromGraphID(GraphID);
+		// Ignore the base class
+		if (Class && Class != UMovieSceneEntitySystem::StaticClass())
+		{
+			InLinker->LinkSystemIfAllowed(Class);
+		}
+	}
+}
+
+UE::MovieScene::EEntitySystemCategory UMovieSceneEntitySystem::RegisterCustomSystemCategory()
+{
+	using namespace UE::MovieScene;
+
+	static EEntitySystemCategory NextCustomCategory = EEntitySystemCategory::Custom;
+
+	EEntitySystemCategory Result = NextCustomCategory;
+	NextCustomCategory = (EEntitySystemCategory)((uint32)NextCustomCategory << 1);
+	check(NextCustomCategory != EEntitySystemCategory::Last);
+	return Result;
+}
+
 bool UMovieSceneEntitySystem::IsRelevant(UMovieSceneEntitySystemLinker* InLinker) const
 {
 	if (RelevantComponent && InLinker->EntityManager.ContainsComponent(RelevantComponent))
@@ -350,7 +407,7 @@ void UMovieSceneEntitySystem::Run(FSystemTaskPrerequisites& InPrerequisites, FSy
 	checkf(Linker != nullptr, TEXT("Attempting to evaluate a system that has been unlinked!"));
 
 	// We may have erroneously linked a system we should have done, but we must not run it in this case
-	if (EnumHasAnyFlags(SystemExclusionContext, Linker->GetSystemContext()))
+	if (!Linker->GetSystemFilter().CheckSystem(this))
 	{
 		return;
 	}
@@ -385,4 +442,4 @@ void UMovieSceneEntitySystem::Unlink()
 
 	Linker->SystemUnlinked(this);
 	Linker = nullptr;
-}
\ No newline at end of file
+}
diff --git a/Engine/Source/Runtime/MovieScene/Private/EntitySystem/MovieSceneEntitySystemLinker.cpp b/Engine/Source/Runtime/MovieScene/Private/EntitySystem/MovieSceneEntitySystemLinker.cpp
--- a/Engine/Source/Runtime/MovieScene/Private/EntitySystem/MovieSceneEntitySystemLinker.cpp
+++ b/Engine/Source/Runtime/MovieScene/Private/EntitySystem/MovieSceneEntitySystemLinker.cpp
@@ -48,6 +48,92 @@ EEntitySystemLinkerRole RegisterCustomEntitySystemLinkerRole()
 	return Result;
 }
 
+FSystemFilter::FSystemFilter()
+{
+	CategoriesAllowed = EEntitySystemCategory::All;
+	CategoriesDisallowed = EEntitySystemCategory::None;
+}
+
+bool FSystemFilter::CheckSystem(TSubclassOf<UMovieSceneEntitySystem> InClass) const
+{
+	const UMovieSceneEntitySystem* SystemCDO = InClass.GetDefaultObject();
+	if (SystemCDO)
+	{
+		return CheckSystem(SystemCDO);
+	}
+	// Allow any systems without a CDO.
+	return true;
+}
+
+bool FSystemFilter::CheckSystem(const UMovieSceneEntitySystem* InSystem) const
+{
+	// Specific allow/disallow rules for systems take precedence.
+	const uint16 SystemClassID = InSystem->GetGlobalDependencyGraphID();
+	if (SystemsDisallowed.IsValidIndex(SystemClassID) && SystemsDisallowed[SystemClassID])
+	{
+		return false;
+	}
+	if (SystemsAllowed.IsValidIndex(SystemClassID) && SystemsAllowed[SystemClassID])
+	{
+		return true;
+	}
+
+	// If any given category is disallowed, the entire thing is refused.
+	if (EnumHasAnyFlags(InSystem->GetCategories(), CategoriesDisallowed))
+	{
+		return false;
+	}
+	// If the given categories contain at least one allowed category, it's good.
+	if (EnumHasAnyFlags(InSystem->GetCategories(), CategoriesAllowed))
+	{
+		return true;
+	}
+	// Nothing explicitly disallowed, but nothing allowed either, so we don't want it.
+	return false;
+}
+
+void FSystemFilter::SetAllowedCategories(EEntitySystemCategory InCategory)
+{
+	CategoriesAllowed = InCategory;
+}
+
+void FSystemFilter::AllowCategory(EEntitySystemCategory InCategory)
+{
+	CategoriesAllowed |= InCategory;
+}
+
+void FSystemFilter::SetDisallowedCategories(EEntitySystemCategory InCategory)
+{
+	CategoriesDisallowed = InCategory;
+}
+
+void FSystemFilter::DisallowCategory(EEntitySystemCategory InCategory)
+{
+	CategoriesDisallowed |= InCategory;
+}
+
+void FSystemFilter::AllowSystem(TSubclassOf<UMovieSceneEntitySystem> InClass)
+{
+	const UMovieSceneEntitySystem* SystemCDO = InClass.GetDefaultObject();
+	if (ensure(SystemCDO))
+	{
+		const uint16 SystemClassID = SystemCDO->GetGlobalDependencyGraphID();
+		SystemsAllowed.PadToNum(SystemClassID + 1, false);
+		SystemsAllowed[SystemClassID] = true;
+	}
+}
+
+void FSystemFilter::DisallowSystem(TSubclassOf<UMovieSceneEntitySystem> InClass)
+{
+	const UMovieSceneEntitySystem* SystemCDO = InClass.GetDefaultObject();
+	if (ensure(SystemCDO))
+	{
+		const uint16 SystemClassID = SystemCDO->GetGlobalDependencyGraphID();
+		SystemsDisallowed.PadToNum(SystemClassID + 1, false);
+		SystemsDisallowed[SystemClassID] = true;
+	}
+}
+
 } // namespace MovieScene
 } // namespace UE
 
@@ -59,9 +145,9 @@ UMovieSceneEntitySystemLinker::UMovieSceneEntitySystemLinker(const FObjectInitia
 
 	Role = EEntitySystemLinkerRole::Unknown;
 	LastSystemLinkVersion = 0;
+	LastSystemUnlinkVersion = 0;
 	LastInstantiationVersion = 0;
 	AutoLinkMode = EAutoLinkRelevantSystems::Enabled;
-	SystemContext = EEntitySystemContext::Runtime;
 
 	if (!HasAnyFlags(RF_ClassDefaultObject))
 	{
@@ -413,6 +499,27 @@ UMovieSceneEntitySystem* UMovieSceneEntitySystemLinker::LinkSystem(TSubclassOf<U
 		return Existing;
 	}
 
+	return LinkSystemImpl(InClassType);
+}
+
+UMovieSceneEntitySystem* UMovieSceneEntitySystemLinker::LinkSystemIfAllowed(TSubclassOf<UMovieSceneEntitySystem> InClassType)
+{
+	UMovieSceneEntitySystem* Existing = FindSystem(InClassType);
+	if (Existing)
+	{
+		return Existing;
+	}
+
+	if (!SystemFilter.CheckSystem(InClassType))
+	{
+		return nullptr;
+	}
+
+	return LinkSystemImpl(InClassType);
+}
+
+UMovieSceneEntitySystem* UMovieSceneEntitySystemLinker::LinkSystemImpl(TSubclassOf<UMovieSceneEntitySystem> InClassType)
+{
 	// We always create systems with a fixed name (since there should only ever be one of that name)
 	// This means we can do our own recycling within the scope of this linker, to save on the cost of re-creating
 	// systems when the first instantiation phase kicks in after a period without any sequence playing.
@@ -436,7 +543,7 @@ UMovieSceneEntitySystem* UMovieSceneEntitySystemLinker::LinkSystem(TSubclassOf<U
 	}
 
 	// If a system implements a hard depdency on another (through direct use of LinkSystem<>), we can't break the client code by returning null, but we can still warn that it should have checked whether it can call LinkSystem first
-	ensureMsgf(!EnumHasAnyFlags(NewSystem->GetExclusionContext(), SystemContext), TEXT("Attempting to link a system that should have been excluded - this is probably an explicit call to Link a system that should have been excluded."));
+	ensureMsgf(SystemFilter.CheckSystem(NewSystem), TEXT("Attempting to link a system that should have been excluded - this is probably an explicit call to Link a system that should have been excluded."));
 
 	SystemGraph.AddSystem(NewSystem);
 	NewSystem->Link(this);
@@ -472,6 +579,16 @@ void UMovieSceneEntitySystemLinker::LinkRelevantSystems()
 	}
 }
 
+void UMovieSceneEntitySystemLinker::UnlinkIrrelevantSystems()
+{
+	if (EntityManager.HasStructureChangedSince(LastSystemUnlinkVersion))
+	{
+		SystemGraph.RemoveIrrelevantSystems(this);
+
+		LastSystemUnlinkVersion = EntityManager.GetSystemSerial();
+	}
+}
+
 void UMovieSceneEntitySystemLinker::AutoLinkRelevantSystems()
 {
 	if (AutoLinkMode == UE::MovieScene::EAutoLinkRelevantSystems::Enabled)
@@ -480,6 +597,14 @@ void UMovieSceneEntitySystemLinker::AutoLinkRelevantSystems()
 	}
 }
 
+void UMovieSceneEntitySystemLinker::AutoUnlinkIrrelevantSystems()
+{
+	if (AutoLinkMode == UE::MovieScene::EAutoLinkRelevantSystems::Enabled)
+	{
+		UnlinkIrrelevantSystems();
+	}
+}
+
 FMovieSceneEntitySystemEvaluationReentrancyWindow::FMovieSceneEntitySystemEvaluationReentrancyWindow(UMovieSceneEntitySystemLinker& InLinker)
 	: Linker(InLinker)
 {
diff --git a/Engine/Source/Runtime/MovieScene/Private/EntitySystem/MovieSceneEntitySystemRunner.cpp b/Engine/Source/Runtime/MovieScene/Private/EntitySystem/MovieSceneEntitySystemRunner.cpp
--- a/Engine/Source/Runtime/MovieScene/Private/EntitySystem/MovieSceneEntitySystemRunner.cpp
+++ b/Engine/Source/Runtime/MovieScene/Private/EntitySystem/MovieSceneEntitySystemRunner.cpp
@@ -552,7 +552,7 @@ bool FMovieSceneEntitySystemRunner::GameThread_PostInstantiation(UMovieSceneEnti
 	// Free anything that has been unlinked
 	EntityManager.FreeEntities(FEntityComponentFilter().All({ BuiltInComponentTypes->Tags.NeedsUnlink }));
 
-	Linker->SystemGraph.RemoveIrrelevantSystems(Linker);
+	Linker->AutoUnlinkIrrelevantSystems();
 
 	EntityManager.Compact();
 	return true;
diff --git a/Engine/Source/Runtime/MovieScene/Private/EntitySystem/MovieSceneEvalTimeSystem.cpp b/Engine/Source/Runtime/MovieScene/Private/EntitySystem/MovieSceneEvalTimeSystem.cpp
--- a/Engine/Source/Runtime/MovieScene/Private/EntitySystem/MovieSceneEvalTimeSystem.cpp
+++ b/Engine/Source/Runtime/MovieScene/Private/EntitySystem/MovieSceneEvalTimeSystem.cpp
@@ -29,11 +29,14 @@ struct FAssignEvalTimesTask
 UMovieSceneEvalTimeSystem::UMovieSceneEvalTimeSystem(const FObjectInitializer& ObjInit)
 	: Super(ObjInit)
 {
-	RelevantComponent = UE::MovieScene::FBuiltInComponentTypes::Get()->EvalTime;
+	using namespace UE::MovieScene;
+
+	RelevantComponent = FBuiltInComponentTypes::Get()->EvalTime;
+	SystemCategories = EEntitySystemCategory::Core;
 
 	if (HasAnyFlags(RF_ClassDefaultObject))
 	{
-		DefineComponentProducer(GetClass(), UE::MovieScene::FBuiltInComponentTypes::Get()->EvalTime);
+		DefineComponentProducer(GetClass(), FBuiltInComponentTypes::Get()->EvalTime);
 	}
 }
 
diff --git a/Engine/Source/Runtime/MovieScene/Public/EntitySystem/MovieSceneBoundSceneComponentInstantiator.h b/Engine/Source/Runtime/MovieScene/Public/EntitySystem/MovieSceneBoundSceneComponentInstantiator.h
--- a/Engine/Source/Runtime/MovieScene/Public/EntitySystem/MovieSceneBoundSceneComponentInstantiator.h
+++ b/Engine/Source/Runtime/MovieScene/Public/EntitySystem/MovieSceneBoundSceneComponentInstantiator.h
@@ -24,22 +24,3 @@ private:
 	virtual void OnRun(FSystemTaskPrerequisites& InPrerequisites, FSystemSubsequentTasks& Subsequents) override final;
 };
 
-/**
- * Empty interface that indicates that an object is able to "impersonate" a USceneComponent
- * for the purposes of Sequencer animation.
- *
- * This is useful when we override a sequence's object bindings with some other object that doesn't
- * have the same strongly-typed components.  Most component bindings will just be animated in a 
- * duck-typing kind of way (looking up properties by name), but scene components are "special",
- * and therefore require this interface to get around if needed.
- */
-UINTERFACE()
-class UMovieSceneSceneComponentImpersonator : public UInterface
-{
-	GENERATED_BODY()
-};
-
-class IMovieSceneSceneComponentImpersonator
-{
-	GENERATED_BODY()
-};
diff --git a/Engine/Source/Runtime/MovieScene/Public/EntitySystem/MovieSceneEntitySystem.h b/Engine/Source/Runtime/MovieScene/Public/EntitySystem/MovieSceneEntitySystem.h
--- a/Engine/Source/Runtime/MovieScene/Public/EntitySystem/MovieSceneEntitySystem.h
+++ b/Engine/Source/Runtime/MovieScene/Public/EntitySystem/MovieSceneEntitySystem.h
@@ -29,17 +29,29 @@ namespace MovieScene
 	struct FSystemSubsequentTasks;
 	struct FSystemTaskPrerequisites;
 
-	enum class EEntitySystemContext : uint8 
+	enum class EEntitySystemCategory : uint32
 	{
+		// No category
 		None = 0,
-
-		/** This system is relevant to runtime */
-		Runtime = 1 << 0,
-
-		/** This system is relevant to interrogation */
-		Interrogation = 1 << 1,
+		// Systems lacking any particular category
+		Unspecified = 1u << 0,
+		// Core systems, such as the time evaluation system
+		Core = 1u << 1,
+		// Systems that evaluate channel values
+		ChannelEvaluators = 1u << 2,
+		// Systems the blend values together
+		BlenderSystems = 1u << 3,
+		// Systems that set properties on objects
+		PropertySystems = 1u << 4,
+		// Start of custom categories
+		Custom = 1u << 5,
+
+		// Last entry, used as error condition
+		Last = 1u << 31,
+		// All categories
+		All = ~0u
 	};
-	ENUM_CLASS_FLAGS(EEntitySystemContext)
+	ENUM_CLASS_FLAGS(EEntitySystemCategory)
 
 } // namespace MovieScene
 } // namespace UE
@@ -93,12 +105,27 @@ public:
 	 */
 	static void LinkRelevantSystems(UMovieSceneEntitySystemLinker* InLinker);
 
+	/**
+	 * Link all systems in a given category
+	 */
+	static void LinkCategorySystems(UMovieSceneEntitySystemLinker* InLinker, UE::MovieScene::EEntitySystemCategory InCategory);
+
+	/**
+	 * Link all systems that pass the given linker's filter
+	 */
+	static void LinkAllSystems(UMovieSceneEntitySystemLinker* InLinker);
+
+	/**
+	 * Create a new system category
+	 */
+	static UE::MovieScene::EEntitySystemCategory RegisterCustomSystemCategory();
+
 public:
 
-	/** Returns linker contexts for which this system should not exist */
-	UE::MovieScene::EEntitySystemContext GetExclusionContext() const
+	/** Returns system categories */
+	UE::MovieScene::EEntitySystemCategory GetCategories() const
 	{
-		return SystemExclusionContext;
+		return SystemCategories;
 	}
 
 	/** Returns the phase(s) during which this system should be run */
@@ -198,7 +225,7 @@ protected:
 	uint16 GraphID;
 	uint16 GlobalDependencyGraphID;
 
-	UE::MovieScene::EEntitySystemContext SystemExclusionContext;
+	UE::MovieScene::EEntitySystemCategory SystemCategories;
 
 	/** When false, this system will not call its OnRun function, but will still be kept alive as long as IsRelevant is true */
 	bool bSystemIsEnabled;
diff --git a/Engine/Source/Runtime/MovieScene/Public/EntitySystem/MovieSceneEntitySystemLinker.h b/Engine/Source/Runtime/MovieScene/Public/EntitySystem/MovieSceneEntitySystemLinker.h
--- a/Engine/Source/Runtime/MovieScene/Public/EntitySystem/MovieSceneEntitySystemLinker.h
+++ b/Engine/Source/Runtime/MovieScene/Public/EntitySystem/MovieSceneEntitySystemLinker.h
@@ -26,8 +26,9 @@ namespace UE
 {
 namespace MovieScene
 {
+
 	struct FComponentRegistry;
-	enum class EEntitySystemContext : uint8;
+	enum class EEntitySystemCategory : uint32;
 
 	enum class EAutoLinkRelevantSystems : uint8
 	{
@@ -50,12 +51,54 @@ namespace MovieScene
 		Standalone,
 		/** This linker is owned by a Sequencer UI */
 		SequencerUI,
+		/** This linker is running interrogations */
+		Interrogation,
 		/** This value and any greater values are for other custom roles */
 		Custom
 	};
 
 	/** Register a new custom linker role */
 	MOVIESCENE_API EEntitySystemLinkerRole RegisterCustomEntitySystemLinkerRole();
+
+	/** Utility class for filtering systems */
+	struct MOVIESCENE_API FSystemFilter
+	{
+		/** Constructs a default filter that allows all systems */
+		FSystemFilter();
+
+		/** Checks whether the given system class passes all filters */
+		template<typename SystemClass>
+		bool CheckSystem() const
+		{
+			return CheckSystem(SystemClass::StaticClass());
+		}
+
+		/** Checks whether the given system class passes all filters */
+		bool CheckSystem(TSubclassOf<UMovieSceneEntitySystem> InClass) const;
+		/** Checks whether the given system passes all filters */
+		bool CheckSystem(const UMovieSceneEntitySystem* InSystem) const;
+
+		/** Sets system categories that are allowed */
+		void SetAllowedCategories(EEntitySystemCategory InCategory);
+		/** Add system categories to be allowed */
+		void AllowCategory(EEntitySystemCategory InCategory);
+		/** Sets system categories that are disallowed */
+		void SetDisallowedCategories(EEntitySystemCategory InCategory);
+		/** Add system categories to be disallowed */
+		void DisallowCategory(EEntitySystemCategory InCategory);
+
+		/** Specifically allow the given system type */
+		void AllowSystem(TSubclassOf<UMovieSceneEntitySystem> InClass);
+		/** Specifically disallow the given system type */
+		void DisallowSystem(TSubclassOf<UMovieSceneEntitySystem> InClass);
+
+	private:
+		UE::MovieScene::EEntitySystemCategory CategoriesAllowed;
+		UE::MovieScene::EEntitySystemCategory CategoriesDisallowed;
+		TBitArray<> SystemsAllowed;
+		TBitArray<> SystemsDisallowed;
+	};
+
 }
 }
 
@@ -90,33 +133,48 @@ public:
 
 	UE::MovieScene::FPreAnimatedStateExtension PreAnimatedState;
 
+	/** Constructs a new linker */
 	UMovieSceneEntitySystemLinker(const FObjectInitializer& ObjInit);
 
+	/** Gets the global component registry */
 	static FComponentRegistry* GetComponents();
 
+	/** Finds or creates a named linker */
 	static UMovieSceneEntitySystemLinker* FindOrCreateLinker(UObject* PreferredOuter, UE::MovieScene::EEntitySystemLinkerRole LinkerRole, const TCHAR* Name = TEXT("DefaultMovieSceneEntitySystemLinker"));
+	/** Creates a new linker */
 	static UMovieSceneEntitySystemLinker* CreateLinker(UObject* PreferredOuter, UE::MovieScene::EEntitySystemLinkerRole LinkerRole);
 
+	/** Gets this linker's instance registry */
 	FInstanceRegistry* GetInstanceRegistry()
 	{
 		check(InstanceRegistry.IsValid());
 		return InstanceRegistry.Get();
 	}
 
+	/** Gets this linker's instance registry */
 	const FInstanceRegistry* GetInstanceRegistry() const
 	{
 		check(InstanceRegistry.IsValid());
 		return InstanceRegistry.Get();
 	}
 
+	/** Tells the given instance to finish evaluating */
 	void FinishInstance(FInstanceHandle InstanceHandle);
 
+	/** Links a given type of system. Will assert if the system type isn't allowed on this linker */
 	template<typename SystemType>
 	SystemType* LinkSystem()
 	{
 		return CastChecked<SystemType>(LinkSystem(SystemType::StaticClass()));
 	}
 
+	/** Links a given type of system. Returns null if the system type isn't allowed on this linker */
+	template<typename SystemType>
+	SystemType* LinkSystemIfAllowed()
+	{
+		return Cast<SystemType>(LinkSystemIfAllowed(SystemType::StaticClass()));
+	}
+
 	template<typename SystemType>
 	SystemType* FindSystem() const
 	{
@@ -124,26 +182,18 @@ public:
 	}
 
 	UMovieSceneEntitySystem* LinkSystem(TSubclassOf<UMovieSceneEntitySystem> InClassType);
+	UMovieSceneEntitySystem* LinkSystemIfAllowed(TSubclassOf<UMovieSceneEntitySystem> InClassType);
+
 	UMovieSceneEntitySystem* FindSystem(TSubclassOf<UMovieSceneEntitySystem> Class) const;
 
 	static void AddReferencedObjects(UObject* InThis, FReferenceCollector& Collector);
 
 	/**
-	 * Retrieve this linker's context, specifying what kinds of systems should be allowed or disallowed
-	 */
-	UE::MovieScene::EEntitySystemContext GetSystemContext() const
-	{
-		return SystemContext;
-	}
-
-	/**
-	 * Set the system context for this linker allowing some systems to be excluded based on the context.
-	 *
-	 * @param InSystemContext    The new system context for this linker
+	 * Gets this linker's system filter.
 	 */
-	void SetSystemContext(UE::MovieScene::EEntitySystemContext InSystemContext)
+	UE::MovieScene::FSystemFilter& GetSystemFilter()
 	{
-		SystemContext = InSystemContext;
+		return SystemFilter;
 	}
 
 	/**
@@ -272,7 +322,9 @@ public:
 	bool HasLinkedSystem(const uint16 GlobalDependencyGraphID);
 
 	void LinkRelevantSystems();
+	void UnlinkIrrelevantSystems();
 	void AutoLinkRelevantSystems();
+	void AutoUnlinkIrrelevantSystems();
 
 	bool HasStructureChangedSinceLastRun() const;
 
@@ -286,6 +338,8 @@ public:
 
 private:
 
+	UMovieSceneEntitySystem* LinkSystemImpl(TSubclassOf<UMovieSceneEntitySystem> InClassType);
+
 	void HandlePreGarbageCollection();
 	void HandlePostGarbageCollection();
 
@@ -332,6 +386,7 @@ public:
 private:
 
 	uint64 LastSystemLinkVersion;
+	uint64 LastSystemUnlinkVersion;
 	uint64 LastInstantiationVersion;
 
 	TWeakPtr<bool> GlobalStateCaptureToken;
@@ -340,7 +395,7 @@ protected:
 
 	UE::MovieScene::EEntitySystemLinkerRole Role;
 	UE::MovieScene::EAutoLinkRelevantSystems AutoLinkMode;
-	UE::MovieScene::EEntitySystemContext SystemContext;
+	UE::MovieScene::FSystemFilter SystemFilter;
 };
 
 /**
diff --git a/Engine/Source/Runtime/MovieSceneTracks/Private/EntitySystem/Interrogation/MovieSceneInterrogatedPropertyInstantiator.cpp b/Engine/Source/Runtime/MovieSceneTracks/Private/EntitySystem/Interrogation/MovieSceneInterrogatedPropertyInstantiator.cpp
--- a/Engine/Source/Runtime/MovieSceneTracks/Private/EntitySystem/Interrogation/MovieSceneInterrogatedPropertyInstantiator.cpp
+++ b/Engine/Source/Runtime/MovieSceneTracks/Private/EntitySystem/Interrogation/MovieSceneInterrogatedPropertyInstantiator.cpp
@@ -1,9 +1,10 @@
 // Copyright Epic Games, Inc. All Rights Reserved.
 
 #include "EntitySystem/Interrogation/MovieSceneInterrogatedPropertyInstantiator.h"
+#include "EntitySystem/Interrogation/MovieSceneInterrogationLinker.h"
+#include "EntitySystem/MovieSceneBlenderSystem.h"
 #include "EntitySystem/MovieSceneEntityBuilder.h"
 #include "EntitySystem/MovieSceneEntitySystemLinker.h"
-#include "EntitySystem/MovieSceneBlenderSystem.h"
 #include "EntitySystem/MovieScenePropertyRegistry.h"
 #include "Systems/MovieScenePiecewiseDoubleBlenderSystem.h"
 
@@ -22,6 +23,8 @@ UMovieSceneInterrogatedPropertyInstantiatorSystem::UMovieSceneInterrogatedProper
 				this, &UMovieSceneInterrogatedPropertyInstantiatorSystem::FindPropertyFromSource);
 
 	RelevantComponent = BuiltInComponents->Interrogation.InputKey;
+	SystemCategories |= FSystemInterrogator::GetInterrogationCategory();
+
 	if (HasAnyFlags(RF_ClassDefaultObject))
 	{
 		DefineComponentProducer(GetClass(), BuiltInComponents->BlendChannelInput);
diff --git a/Engine/Source/Runtime/MovieSceneTracks/Private/EntitySystem/Interrogation/MovieSceneInterrogationLinker.cpp b/Engine/Source/Runtime/MovieSceneTracks/Private/EntitySystem/Interrogation/MovieSceneInterrogationLinker.cpp
--- a/Engine/Source/Runtime/MovieSceneTracks/Private/EntitySystem/Interrogation/MovieSceneInterrogationLinker.cpp
+++ b/Engine/Source/Runtime/MovieSceneTracks/Private/EntitySystem/Interrogation/MovieSceneInterrogationLinker.cpp
@@ -617,21 +617,33 @@ void FInterrogationChannels::QueryLocalSpaceTransforms(UMovieSceneEntitySystemLi
 	}
 }
 
+EEntitySystemCategory FSystemInterrogator::GetInterrogationCategory()
+{
+	static EEntitySystemCategory Interrogation = UMovieSceneEntitySystem::RegisterCustomSystemCategory();
+	return Interrogation;
+}
 
+EEntitySystemCategory FSystemInterrogator::GetExcludedFromInterrogationCategory()
+{
+	static EEntitySystemCategory ExcludedFromInterrogation = UMovieSceneEntitySystem::RegisterCustomSystemCategory();
+	return ExcludedFromInterrogation;
+}
 
 FSystemInterrogator::FSystemInterrogator()
 {
 	InitialValueCache = FInitialValueCache::GetGlobalInitialValues();
 
 	Linker = NewObject<UMovieSceneEntitySystemLinker>(GetTransientPackage());
-	Linker->SetSystemContext(EEntitySystemContext::Interrogation);
+	Linker->SetLinkerRole(EEntitySystemLinkerRole::Interrogation);
+	Linker->GetSystemFilter().DisallowCategory(FSystemInterrogator::GetExcludedFromInterrogationCategory());
 
 	Linker->AddExtension(IInterrogationExtension::GetExtensionID(), static_cast<IInterrogationExtension*>(this));
 	Linker->AddExtension(InitialValueCache.Get());
 }
 
 FSystemInterrogator::~FSystemInterrogator()
-{}
+{
+}
 
 void FSystemInterrogator::AddReferencedObjects(FReferenceCollector& Collector)
 {
diff --git a/Engine/Source/Runtime/MovieSceneTracks/Private/Systems/BoolChannelEvaluatorSystem.cpp b/Engine/Source/Runtime/MovieSceneTracks/Private/Systems/BoolChannelEvaluatorSystem.cpp
--- a/Engine/Source/Runtime/MovieSceneTracks/Private/Systems/BoolChannelEvaluatorSystem.cpp
+++ b/Engine/Source/Runtime/MovieSceneTracks/Private/Systems/BoolChannelEvaluatorSystem.cpp
@@ -15,6 +15,7 @@ UBoolChannelEvaluatorSystem::UBoolChannelEvaluatorSystem(const FObjectInitialize
 {
 	using namespace UE::MovieScene;
 
+	SystemCategories = EEntitySystemCategory::ChannelEvaluators;
 	RelevantComponent = FBuiltInComponentTypes::Get()->BoolChannel;
 
 	if (HasAnyFlags(RF_ClassDefaultObject))
diff --git a/Engine/Source/Runtime/MovieSceneTracks/Private/Systems/ByteChannelEvaluatorSystem.cpp b/Engine/Source/Runtime/MovieSceneTracks/Private/Systems/ByteChannelEvaluatorSystem.cpp
--- a/Engine/Source/Runtime/MovieSceneTracks/Private/Systems/ByteChannelEvaluatorSystem.cpp
+++ b/Engine/Source/Runtime/MovieSceneTracks/Private/Systems/ByteChannelEvaluatorSystem.cpp
@@ -38,6 +38,8 @@ UByteChannelEvaluatorSystem::UByteChannelEvaluatorSystem(const FObjectInitialize
 {
 	using namespace UE::MovieScene;
 
+	SystemCategories = EEntitySystemCategory::ChannelEvaluators;
+
 	const FBuiltInComponentTypes* BuiltInComponents = FBuiltInComponentTypes::Get();
 	RelevantComponent = BuiltInComponents->ByteChannel;
 
diff --git a/Engine/Source/Runtime/MovieSceneTracks/Private/Systems/DoubleChannelEvaluatorSystem.cpp b/Engine/Source/Runtime/MovieSceneTracks/Private/Systems/DoubleChannelEvaluatorSystem.cpp
--- a/Engine/Source/Runtime/MovieSceneTracks/Private/Systems/DoubleChannelEvaluatorSystem.cpp
+++ b/Engine/Source/Runtime/MovieSceneTracks/Private/Systems/DoubleChannelEvaluatorSystem.cpp
@@ -56,6 +56,8 @@ UDoubleChannelEvaluatorSystem::UDoubleChannelEvaluatorSystem(const FObjectInitia
 {
 	using namespace UE::MovieScene;
 
+	SystemCategories = EEntitySystemCategory::ChannelEvaluators;
+
 	if (HasAnyFlags(RF_ClassDefaultObject))
 	{
 		DefineImplicitPrerequisite(UMovieSceneEvalTimeSystem::StaticClass(), GetClass());
diff --git a/Engine/Source/Runtime/MovieSceneTracks/Private/Systems/DoublePerlinNoiseChannelEvaluatorSystem.cpp b/Engine/Source/Runtime/MovieSceneTracks/Private/Systems/DoublePerlinNoiseChannelEvaluatorSystem.cpp
--- a/Engine/Source/Runtime/MovieSceneTracks/Private/Systems/DoublePerlinNoiseChannelEvaluatorSystem.cpp
+++ b/Engine/Source/Runtime/MovieSceneTracks/Private/Systems/DoublePerlinNoiseChannelEvaluatorSystem.cpp
@@ -45,6 +45,8 @@ UDoublePerlinNoiseChannelEvaluatorSystem::UDoublePerlinNoiseChannelEvaluatorSyst
 {
 	using namespace UE::MovieScene;
 
+	SystemCategories = EEntitySystemCategory::ChannelEvaluators;
+
 	const FBuiltInComponentTypes* BuiltInComponents = FBuiltInComponentTypes::Get();
 	const FMovieSceneTracksComponentTypes* TrackComponents = FMovieSceneTracksComponentTypes::Get();
 
diff --git a/Engine/Source/Runtime/MovieSceneTracks/Private/Systems/FloatChannelEvaluatorSystem.cpp b/Engine/Source/Runtime/MovieSceneTracks/Private/Systems/FloatChannelEvaluatorSystem.cpp
--- a/Engine/Source/Runtime/MovieSceneTracks/Private/Systems/FloatChannelEvaluatorSystem.cpp
+++ b/Engine/Source/Runtime/MovieSceneTracks/Private/Systems/FloatChannelEvaluatorSystem.cpp
@@ -59,6 +59,8 @@ UFloatChannelEvaluatorSystem::UFloatChannelEvaluatorSystem(const FObjectInitiali
 {
 	using namespace UE::MovieScene;
 
+	SystemCategories = EEntitySystemCategory::ChannelEvaluators;
+
 	if (HasAnyFlags(RF_ClassDefaultObject))
 	{
 		DefineImplicitPrerequisite(UMovieSceneEvalTimeSystem::StaticClass(), GetClass());
diff --git a/Engine/Source/Runtime/MovieSceneTracks/Private/Systems/FloatPerlinNoiseChannelEvaluatorSystem.cpp b/Engine/Source/Runtime/MovieSceneTracks/Private/Systems/FloatPerlinNoiseChannelEvaluatorSystem.cpp
--- a/Engine/Source/Runtime/MovieSceneTracks/Private/Systems/FloatPerlinNoiseChannelEvaluatorSystem.cpp
+++ b/Engine/Source/Runtime/MovieSceneTracks/Private/Systems/FloatPerlinNoiseChannelEvaluatorSystem.cpp
@@ -46,6 +46,8 @@ UFloatPerlinNoiseChannelEvaluatorSystem::UFloatPerlinNoiseChannelEvaluatorSystem
 {
 	using namespace UE::MovieScene;
 
+	SystemCategories = EEntitySystemCategory::ChannelEvaluators;
+
 	const FBuiltInComponentTypes* BuiltInComponents = FBuiltInComponentTypes::Get();
 	const FMovieSceneTracksComponentTypes* TrackComponents = FMovieSceneTracksComponentTypes::Get();
 
diff --git a/Engine/Source/Runtime/MovieSceneTracks/Private/Systems/IntegerChannelEvaluatorSystem.cpp b/Engine/Source/Runtime/MovieSceneTracks/Private/Systems/IntegerChannelEvaluatorSystem.cpp
--- a/Engine/Source/Runtime/MovieSceneTracks/Private/Systems/IntegerChannelEvaluatorSystem.cpp
+++ b/Engine/Source/Runtime/MovieSceneTracks/Private/Systems/IntegerChannelEvaluatorSystem.cpp
@@ -38,6 +38,8 @@ UIntegerChannelEvaluatorSystem::UIntegerChannelEvaluatorSystem(const FObjectInit
 {
 	using namespace UE::MovieScene;
 
+	SystemCategories = EEntitySystemCategory::ChannelEvaluators;
+
 	const FBuiltInComponentTypes* BuiltInComponents = FBuiltInComponentTypes::Get();
 	RelevantComponent = BuiltInComponents->IntegerChannel;
 
diff --git a/Engine/Source/Runtime/MovieSceneTracks/Private/Systems/MovieScene3DTransformPropertySystem.cpp b/Engine/Source/Runtime/MovieSceneTracks/Private/Systems/MovieScene3DTransformPropertySystem.cpp
--- a/Engine/Source/Runtime/MovieSceneTracks/Private/Systems/MovieScene3DTransformPropertySystem.cpp
+++ b/Engine/Source/Runtime/MovieSceneTracks/Private/Systems/MovieScene3DTransformPropertySystem.cpp
@@ -11,8 +11,6 @@
 UMovieScene3DTransformPropertySystem::UMovieScene3DTransformPropertySystem(const FObjectInitializer& ObjInit)
 	: Super(ObjInit)
 {
-	SystemExclusionContext |= UE::MovieScene::EEntitySystemContext::Interrogation;
-
 	BindToProperty(UE::MovieScene::FMovieSceneTracksComponentTypes::Get()->Transform);
 
 	if (HasAnyFlags(RF_ClassDefaultObject))
diff --git a/Engine/Source/Runtime/MovieSceneTracks/Private/Systems/MovieSceneBaseValueEvaluatorSystem.cpp b/Engine/Source/Runtime/MovieSceneTracks/Private/Systems/MovieSceneBaseValueEvaluatorSystem.cpp
--- a/Engine/Source/Runtime/MovieSceneTracks/Private/Systems/MovieSceneBaseValueEvaluatorSystem.cpp
+++ b/Engine/Source/Runtime/MovieSceneTracks/Private/Systems/MovieSceneBaseValueEvaluatorSystem.cpp
@@ -49,6 +49,7 @@ UMovieSceneBaseValueEvaluatorSystem::UMovieSceneBaseValueEvaluatorSystem(const F
 
 	Phase = ESystemPhase::Instantiation;
 	RelevantComponent = FBuiltInComponentTypes::Get()->BaseValueEvalTime;
+	SystemCategories = EEntitySystemCategory::ChannelEvaluators;
 
 	if (HasAnyFlags(RF_ClassDefaultObject))
 	{
diff --git a/Engine/Source/Runtime/MovieSceneTracks/Private/Systems/MovieSceneBoolPropertySystem.cpp b/Engine/Source/Runtime/MovieSceneTracks/Private/Systems/MovieSceneBoolPropertySystem.cpp
--- a/Engine/Source/Runtime/MovieSceneTracks/Private/Systems/MovieSceneBoolPropertySystem.cpp
+++ b/Engine/Source/Runtime/MovieSceneTracks/Private/Systems/MovieSceneBoolPropertySystem.cpp
@@ -7,8 +7,6 @@
 UMovieSceneBoolPropertySystem::UMovieSceneBoolPropertySystem(const FObjectInitializer& ObjInit)
 	: Super(ObjInit)
 {
-	SystemExclusionContext |= UE::MovieScene::EEntitySystemContext::Interrogation;
-
 	BindToProperty(UE::MovieScene::FMovieSceneTracksComponentTypes::Get()->Bool);
 
 	if (HasAnyFlags(RF_ClassDefaultObject))
diff --git a/Engine/Source/Runtime/MovieSceneTracks/Private/Systems/MovieSceneBytePropertySystem.cpp b/Engine/Source/Runtime/MovieSceneTracks/Private/Systems/MovieSceneBytePropertySystem.cpp
--- a/Engine/Source/Runtime/MovieSceneTracks/Private/Systems/MovieSceneBytePropertySystem.cpp
+++ b/Engine/Source/Runtime/MovieSceneTracks/Private/Systems/MovieSceneBytePropertySystem.cpp
@@ -7,8 +7,6 @@
 UMovieSceneBytePropertySystem::UMovieSceneBytePropertySystem(const FObjectInitializer& ObjInit)
 	: Super(ObjInit)
 {
-	SystemExclusionContext |= UE::MovieScene::EEntitySystemContext::Interrogation;
-
 	BindToProperty(UE::MovieScene::FMovieSceneTracksComponentTypes::Get()->Byte);
 
 	if (HasAnyFlags(RF_ClassDefaultObject))
diff --git a/Engine/Source/Runtime/MovieSceneTracks/Private/Systems/MovieSceneColorPropertySystem.cpp b/Engine/Source/Runtime/MovieSceneTracks/Private/Systems/MovieSceneColorPropertySystem.cpp
--- a/Engine/Source/Runtime/MovieSceneTracks/Private/Systems/MovieSceneColorPropertySystem.cpp
+++ b/Engine/Source/Runtime/MovieSceneTracks/Private/Systems/MovieSceneColorPropertySystem.cpp
@@ -9,8 +9,6 @@
 UMovieSceneColorPropertySystem::UMovieSceneColorPropertySystem(const FObjectInitializer& ObjInit)
 	: Super(ObjInit)
 {
-	SystemExclusionContext |= UE::MovieScene::EEntitySystemContext::Interrogation;
-
 	BindToProperty(UE::MovieScene::FMovieSceneTracksComponentTypes::Get()->Color);
 
 	if (HasAnyFlags(RF_ClassDefaultObject))
diff --git a/Engine/Source/Runtime/MovieSceneTracks/Private/Systems/MovieSceneComponentAttachmentSystem.cpp b/Engine/Source/Runtime/MovieSceneTracks/Private/Systems/MovieSceneComponentAttachmentSystem.cpp
--- a/Engine/Source/Runtime/MovieSceneTracks/Private/Systems/MovieSceneComponentAttachmentSystem.cpp
+++ b/Engine/Source/Runtime/MovieSceneTracks/Private/Systems/MovieSceneComponentAttachmentSystem.cpp
@@ -7,6 +7,7 @@
 
 #include "EntitySystem/MovieSceneBoundObjectInstantiator.h"
 #include "EntitySystem/MovieSceneBoundSceneComponentInstantiator.h"
+#include "EntitySystem/Interrogation/MovieSceneInterrogationLinker.h"
 
 #include "PreAnimatedState/MovieScenePreAnimatedComponentTransformStorage.h"
 #include "Evaluation/PreAnimatedState/MovieScenePreAnimatedStorageID.inl"
@@ -163,7 +164,7 @@ UMovieSceneComponentAttachmentSystem::UMovieSceneComponentAttachmentSystem(const
 {
 	using namespace UE::MovieScene;
 
-	SystemExclusionContext |= EEntitySystemContext::Interrogation;
+	SystemCategories |= FSystemInterrogator::GetExcludedFromInterrogationCategory();
 
 	FMovieSceneTracksComponentTypes* TrackComponents = FMovieSceneTracksComponentTypes::Get();
 	RelevantComponent = TrackComponents->AttachParentBinding;
diff --git a/Engine/Source/Runtime/MovieSceneTracks/Private/Systems/MovieSceneComponentMobilitySystem.cpp b/Engine/Source/Runtime/MovieSceneTracks/Private/Systems/MovieSceneComponentMobilitySystem.cpp
--- a/Engine/Source/Runtime/MovieSceneTracks/Private/Systems/MovieSceneComponentMobilitySystem.cpp
+++ b/Engine/Source/Runtime/MovieSceneTracks/Private/Systems/MovieSceneComponentMobilitySystem.cpp
@@ -8,6 +8,7 @@
 #include "EntitySystem/BuiltInComponentTypes.h"
 #include "EntitySystem/MovieSceneEntitySystemTask.h"
 #include "EntitySystem/MovieSceneEntitySystemLinker.h"
+#include "EntitySystem/Interrogation/MovieSceneInterrogationLinker.h"
 
 #include "Evaluation/PreAnimatedState/MovieScenePreAnimatedObjectGroupManager.h"
 #include "Evaluation/PreAnimatedState/MovieScenePreAnimatedStorageID.inl"
@@ -114,7 +115,7 @@ UMovieSceneComponentMobilitySystem::UMovieSceneComponentMobilitySystem(const FOb
 {
 	using namespace UE::MovieScene;
 
-	SystemExclusionContext |= EEntitySystemContext::Interrogation;
+	SystemCategories |= FSystemInterrogator::GetExcludedFromInterrogationCategory();
 
 	FBuiltInComponentTypes*          BuiltInComponents = FBuiltInComponentTypes::Get();
 	FMovieSceneTracksComponentTypes* TrackComponents   = FMovieSceneTracksComponentTypes::Get();
diff --git a/Engine/Source/Runtime/MovieSceneTracks/Private/Systems/MovieSceneComponentTransformSystem.cpp b/Engine/Source/Runtime/MovieSceneTracks/Private/Systems/MovieSceneComponentTransformSystem.cpp
--- a/Engine/Source/Runtime/MovieSceneTracks/Private/Systems/MovieSceneComponentTransformSystem.cpp
+++ b/Engine/Source/Runtime/MovieSceneTracks/Private/Systems/MovieSceneComponentTransformSystem.cpp
@@ -6,16 +6,19 @@
 #include "Systems/MovieScenePropertyInstantiator.h"
 
 #include "EntitySystem/BuiltInComponentTypes.h"
+#include "EntitySystem/Interrogation/MovieSceneInterrogationLinker.h"
 
 #include "MovieSceneTracksComponentTypes.h"
 
 UMovieSceneComponentTransformSystem::UMovieSceneComponentTransformSystem(const FObjectInitializer& ObjInit)
 	: Super(ObjInit)
 {
+	using namespace UE::MovieScene;
+
 	// This system can be used for interrogation
-	SystemExclusionContext = UE::MovieScene::EEntitySystemContext::None;
+	SystemCategories &= ~FSystemInterrogator::GetExcludedFromInterrogationCategory();
 
-	BindToProperty(UE::MovieScene::FMovieSceneTracksComponentTypes::Get()->ComponentTransform);
+	BindToProperty(FMovieSceneTracksComponentTypes::Get()->ComponentTransform);
 
 	if (HasAnyFlags(RF_ClassDefaultObject))
 	{
diff --git a/Engine/Source/Runtime/MovieSceneTracks/Private/Systems/MovieSceneDoublePropertySystem.cpp b/Engine/Source/Runtime/MovieSceneTracks/Private/Systems/MovieSceneDoublePropertySystem.cpp
--- a/Engine/Source/Runtime/MovieSceneTracks/Private/Systems/MovieSceneDoublePropertySystem.cpp
+++ b/Engine/Source/Runtime/MovieSceneTracks/Private/Systems/MovieSceneDoublePropertySystem.cpp
@@ -10,8 +10,6 @@
 UMovieSceneDoublePropertySystem::UMovieSceneDoublePropertySystem(const FObjectInitializer& ObjInit)
 	: Super(ObjInit)
 {
-	SystemExclusionContext |= UE::MovieScene::EEntitySystemContext::Interrogation;
-
 	BindToProperty(UE::MovieScene::FMovieSceneTracksComponentTypes::Get()->Double);
 
 	if (HasAnyFlags(RF_ClassDefaultObject))
diff --git a/Engine/Source/Runtime/MovieSceneTracks/Private/Systems/MovieSceneEnumPropertySystem.cpp b/Engine/Source/Runtime/MovieSceneTracks/Private/Systems/MovieSceneEnumPropertySystem.cpp
--- a/Engine/Source/Runtime/MovieSceneTracks/Private/Systems/MovieSceneEnumPropertySystem.cpp
+++ b/Engine/Source/Runtime/MovieSceneTracks/Private/Systems/MovieSceneEnumPropertySystem.cpp
@@ -7,8 +7,6 @@
 UMovieSceneEnumPropertySystem::UMovieSceneEnumPropertySystem(const FObjectInitializer& ObjInit)
 	: Super(ObjInit)
 {
-	SystemExclusionContext |= UE::MovieScene::EEntitySystemContext::Interrogation;
-
 	BindToProperty(UE::MovieScene::FMovieSceneTracksComponentTypes::Get()->Enum);
 
 	if (HasAnyFlags(RF_ClassDefaultObject))
diff --git a/Engine/Source/Runtime/MovieSceneTracks/Private/Systems/MovieSceneEulerTransformPropertySystem.cpp b/Engine/Source/Runtime/MovieSceneTracks/Private/Systems/MovieSceneEulerTransformPropertySystem.cpp
--- a/Engine/Source/Runtime/MovieSceneTracks/Private/Systems/MovieSceneEulerTransformPropertySystem.cpp
+++ b/Engine/Source/Runtime/MovieSceneTracks/Private/Systems/MovieSceneEulerTransformPropertySystem.cpp
@@ -11,8 +11,6 @@
 UMovieSceneEulerTransformPropertySystem::UMovieSceneEulerTransformPropertySystem(const FObjectInitializer& ObjInit)
 	: Super(ObjInit)
 {
-	SystemExclusionContext |= UE::MovieScene::EEntitySystemContext::Interrogation;
-
 	BindToProperty(UE::MovieScene::FMovieSceneTracksComponentTypes::Get()->EulerTransform);
 
 	if (HasAnyFlags(RF_ClassDefaultObject))
diff --git a/Engine/Source/Runtime/MovieSceneTracks/Private/Systems/MovieSceneFloatPropertySystem.cpp b/Engine/Source/Runtime/MovieSceneTracks/Private/Systems/MovieSceneFloatPropertySystem.cpp
--- a/Engine/Source/Runtime/MovieSceneTracks/Private/Systems/MovieSceneFloatPropertySystem.cpp
+++ b/Engine/Source/Runtime/MovieSceneTracks/Private/Systems/MovieSceneFloatPropertySystem.cpp
@@ -10,8 +10,6 @@
 UMovieSceneFloatPropertySystem::UMovieSceneFloatPropertySystem(const FObjectInitializer& ObjInit)
 	: Super(ObjInit)
 {
-	SystemExclusionContext |= UE::MovieScene::EEntitySystemContext::Interrogation;
-
 	BindToProperty(UE::MovieScene::FMovieSceneTracksComponentTypes::Get()->Float);
 
 	if (HasAnyFlags(RF_ClassDefaultObject))
@@ -24,4 +22,4 @@ UMovieSceneFloatPropertySystem::UMovieSceneFloatPropertySystem(const FObjectInit
 void UMovieSceneFloatPropertySystem::OnRun(FSystemTaskPrerequisites& InPrerequisites, FSystemSubsequentTasks& Subsequents)
 {
 	Super::OnRun(InPrerequisites, Subsequents);
-}
\ No newline at end of file
+}
diff --git a/Engine/Source/Runtime/MovieSceneTracks/Private/Systems/MovieSceneHierarchicalBiasSystem.cpp b/Engine/Source/Runtime/MovieSceneTracks/Private/Systems/MovieSceneHierarchicalBiasSystem.cpp
--- a/Engine/Source/Runtime/MovieSceneTracks/Private/Systems/MovieSceneHierarchicalBiasSystem.cpp
+++ b/Engine/Source/Runtime/MovieSceneTracks/Private/Systems/MovieSceneHierarchicalBiasSystem.cpp
@@ -109,6 +109,8 @@ UMovieSceneHierarchicalBiasSystem::UMovieSceneHierarchicalBiasSystem(const FObje
 {
 	using namespace UE::MovieScene;
 
+	SystemCategories = EEntitySystemCategory::Core;
+
 	if (HasAnyFlags(RF_ClassDefaultObject))
 	{
 		DefineComponentConsumer(GetClass(), FBuiltInComponentTypes::Get()->BlendChannelInput);
diff --git a/Engine/Source/Runtime/MovieSceneTracks/Private/Systems/MovieSceneInitialValueSystem.cpp b/Engine/Source/Runtime/MovieSceneTracks/Private/Systems/MovieSceneInitialValueSystem.cpp
--- a/Engine/Source/Runtime/MovieSceneTracks/Private/Systems/MovieSceneInitialValueSystem.cpp
+++ b/Engine/Source/Runtime/MovieSceneTracks/Private/Systems/MovieSceneInitialValueSystem.cpp
@@ -77,6 +77,8 @@ struct FInitialValueMutation : IMovieSceneEntityMutation
 UMovieSceneInitialValueSystem::UMovieSceneInitialValueSystem(const FObjectInitializer& ObjInit)
 	: Super(ObjInit)
 {
+	SystemCategories = UE::MovieScene::EEntitySystemCategory::Core;
+
 	if (HasAnyFlags(RF_ClassDefaultObject))
 	{
 		DefineImplicitPrerequisite(UMovieScenePropertyInstantiatorSystem::StaticClass(), StaticClass());
@@ -169,4 +171,4 @@ void UMovieSceneInitialValueSystem::OnRun(FSystemTaskPrerequisites& InPrerequisi
 			}
 		}
 	}
-}
\ No newline at end of file
+}
diff --git a/Engine/Source/Runtime/MovieSceneTracks/Private/Systems/MovieSceneIntegerPropertySystem.cpp b/Engine/Source/Runtime/MovieSceneTracks/Private/Systems/MovieSceneIntegerPropertySystem.cpp
--- a/Engine/Source/Runtime/MovieSceneTracks/Private/Systems/MovieSceneIntegerPropertySystem.cpp
+++ b/Engine/Source/Runtime/MovieSceneTracks/Private/Systems/MovieSceneIntegerPropertySystem.cpp
@@ -8,8 +8,6 @@
 UMovieSceneIntegerPropertySystem::UMovieSceneIntegerPropertySystem(const FObjectInitializer& ObjInit)
 	: Super(ObjInit)
 {
-	SystemExclusionContext |= UE::MovieScene::EEntitySystemContext::Interrogation;
-
 	BindToProperty(UE::MovieScene::FMovieSceneTracksComponentTypes::Get()->Integer);
 
 	if (HasAnyFlags(RF_ClassDefaultObject))
diff --git a/Engine/Source/Runtime/MovieSceneTracks/Private/Systems/MovieScenePropertyInstantiator.cpp b/Engine/Source/Runtime/MovieSceneTracks/Private/Systems/MovieScenePropertyInstantiator.cpp
--- a/Engine/Source/Runtime/MovieSceneTracks/Private/Systems/MovieScenePropertyInstantiator.cpp
+++ b/Engine/Source/Runtime/MovieSceneTracks/Private/Systems/MovieScenePropertyInstantiator.cpp
@@ -8,6 +8,7 @@
 #include "EntitySystem/MovieSceneBlenderSystem.h"
 #include "EntitySystem/MovieScenePropertyRegistry.h"
 #include "Systems/MovieScenePiecewiseDoubleBlenderSystem.h"
+#include "EntitySystem/Interrogation/MovieSceneInterrogationLinker.h"
 
 #include "Algo/AllOf.h"
 #include "Algo/IndexOf.h"
@@ -27,7 +28,7 @@ UMovieScenePropertyInstantiatorSystem::UMovieScenePropertyInstantiatorSystem(con
 	RecomposerImpl.OnGetPropertyInfo = FOnGetPropertyRecomposerPropertyInfo::CreateUObject(
 				this, &UMovieScenePropertyInstantiatorSystem::FindPropertyFromSource);
 
-	SystemExclusionContext = EEntitySystemContext::Interrogation;
+	SystemCategories = FSystemInterrogator::GetExcludedFromInterrogationCategory();
 	RelevantComponent = BuiltInComponents->PropertyBinding;
 	if (HasAnyFlags(RF_ClassDefaultObject))
 	{
diff --git a/Engine/Source/Runtime/MovieSceneTracks/Private/Systems/MovieScenePropertySystem.cpp b/Engine/Source/Runtime/MovieSceneTracks/Private/Systems/MovieScenePropertySystem.cpp
--- a/Engine/Source/Runtime/MovieSceneTracks/Private/Systems/MovieScenePropertySystem.cpp
+++ b/Engine/Source/Runtime/MovieSceneTracks/Private/Systems/MovieScenePropertySystem.cpp
@@ -4,12 +4,15 @@
 #include "Systems/MovieScenePropertyInstantiator.h"
 
 #include "EntitySystem/MovieSceneEntitySystemLinker.h"
+#include "EntitySystem/Interrogation/MovieSceneInterrogationLinker.h"
 #include "Evaluation/PreAnimatedState/MovieScenePreAnimatedEntityCaptureSource.h"
 
 UMovieScenePropertySystem::UMovieScenePropertySystem(const FObjectInitializer& ObjInit)
 	: Super(ObjInit)
 {
-	SystemExclusionContext |= UE::MovieScene::EEntitySystemContext::Interrogation;
+	using namespace UE::MovieScene;
+
+	SystemCategories = EEntitySystemCategory::PropertySystems | FSystemInterrogator::GetExcludedFromInterrogationCategory();
 }
 
 void UMovieScenePropertySystem::OnLink()
@@ -17,9 +20,9 @@ void UMovieScenePropertySystem::OnLink()
 	using namespace UE::MovieScene;
 
 	// Never apply properties during evaluation. This code is necessary if derived types do support interrogation.
-	if (!EnumHasAnyFlags(Linker->GetSystemContext(), EEntitySystemContext::Interrogation))
+	InstantiatorSystem = Linker->LinkSystemIfAllowed<UMovieScenePropertyInstantiatorSystem>();
+	if (InstantiatorSystem)
 	{
-		InstantiatorSystem = Linker->LinkSystem<UMovieScenePropertyInstantiatorSystem>();
 		Linker->SystemGraph.AddReference(this, InstantiatorSystem);
 	}
 }
@@ -29,7 +32,7 @@ void UMovieScenePropertySystem::OnRun(FSystemTaskPrerequisites& InPrerequisites,
 	using namespace UE::MovieScene;
 
 	// Never apply properties during evaluation. This code is necessary if derived types do support interrogation.
-	if (EnumHasAnyFlags(Linker->GetSystemContext(), EEntitySystemContext::Interrogation))
+	if (!InstantiatorSystem)
 	{
 		return;
 	}
@@ -126,4 +129,4 @@ void UMovieScenePropertySystem::SavePreAnimatedState(const FPreAnimationParamete
 void UMovieScenePropertySystem::RestorePreAnimatedState(const FPreAnimationParameters& InParameters)
 {
 
-}
\ No newline at end of file
+}
diff --git a/Engine/Source/Runtime/MovieSceneTracks/Private/Systems/MovieSceneQuaternionInterpolationRotationSystem.cpp b/Engine/Source/Runtime/MovieSceneTracks/Private/Systems/MovieSceneQuaternionInterpolationRotationSystem.cpp
--- a/Engine/Source/Runtime/MovieSceneTracks/Private/Systems/MovieSceneQuaternionInterpolationRotationSystem.cpp
+++ b/Engine/Source/Runtime/MovieSceneTracks/Private/Systems/MovieSceneQuaternionInterpolationRotationSystem.cpp
@@ -185,6 +185,8 @@ UMovieSceneQuaternionInterpolationRotationSystem::UMovieSceneQuaternionInterpola
 {
 	using namespace UE::MovieScene;
 
+	SystemCategories = EEntitySystemCategory::ChannelEvaluators;
+
 	if (HasAnyFlags(RF_ClassDefaultObject))
 	{
 		DefineImplicitPrerequisite(UMovieSceneEvalTimeSystem::StaticClass(), GetClass());
diff --git a/Engine/Source/Runtime/MovieSceneTracks/Private/Systems/MovieSceneVectorPropertySystem.cpp b/Engine/Source/Runtime/MovieSceneTracks/Private/Systems/MovieSceneVectorPropertySystem.cpp
--- a/Engine/Source/Runtime/MovieSceneTracks/Private/Systems/MovieSceneVectorPropertySystem.cpp
+++ b/Engine/Source/Runtime/MovieSceneTracks/Private/Systems/MovieSceneVectorPropertySystem.cpp
@@ -13,8 +13,6 @@
 UMovieSceneFloatVectorPropertySystem::UMovieSceneFloatVectorPropertySystem(const FObjectInitializer& ObjInit)
 	: Super(ObjInit)
 {
-	SystemExclusionContext |= UE::MovieScene::EEntitySystemContext::Interrogation;
-
 	BindToProperty(UE::MovieScene::FMovieSceneTracksComponentTypes::Get()->FloatVector);
 
 	if (HasAnyFlags(RF_ClassDefaultObject))
@@ -33,8 +31,6 @@ void UMovieSceneFloatVectorPropertySystem::OnRun(FSystemTaskPrerequisites& InPre
 UMovieSceneDoubleVectorPropertySystem::UMovieSceneDoubleVectorPropertySystem(const FObjectInitializer& ObjInit)
 	: Super(ObjInit)
 {
-	SystemExclusionContext |= UE::MovieScene::EEntitySystemContext::Interrogation;
-
 	BindToProperty(UE::MovieScene::FMovieSceneTracksComponentTypes::Get()->DoubleVector);
 
 	if (HasAnyFlags(RF_ClassDefaultObject))
diff --git a/Engine/Source/Runtime/MovieSceneTracks/Private/Systems/WeightAndEasingEvaluatorSystem.cpp b/Engine/Source/Runtime/MovieSceneTracks/Private/Systems/WeightAndEasingEvaluatorSystem.cpp
--- a/Engine/Source/Runtime/MovieSceneTracks/Private/Systems/WeightAndEasingEvaluatorSystem.cpp
+++ b/Engine/Source/Runtime/MovieSceneTracks/Private/Systems/WeightAndEasingEvaluatorSystem.cpp
@@ -26,6 +26,7 @@ UMovieSceneHierarchicalEasingInstantiatorSystem::UMovieSceneHierarchicalEasingIn
 
 	FBuiltInComponentTypes* BuiltInComponents = FBuiltInComponentTypes::Get();
 	RelevantComponent = BuiltInComponents->HierarchicalEasingProvider;
+	SystemCategories = EEntitySystemCategory::Core;
 }
 
 void UMovieSceneHierarchicalEasingInstantiatorSystem::OnRun(FSystemTaskPrerequisites& InPrerequisites, FSystemSubsequentTasks& Subsequents)
@@ -271,6 +272,8 @@ struct FPropagateHierarchicalEasings
 UWeightAndEasingEvaluatorSystem::UWeightAndEasingEvaluatorSystem(const FObjectInitializer& ObjInit)
 	: Super(ObjInit)
 {
+	SystemCategories = UE::MovieScene::EEntitySystemCategory::ChannelEvaluators;
+
 	if (HasAnyFlags(RF_ClassDefaultObject))
 	{
 		DefineImplicitPrerequisite(UMovieSceneEvalTimeSystem::StaticClass(), GetClass());
diff --git a/Engine/Source/Runtime/MovieSceneTracks/Public/EntitySystem/Interrogation/MovieSceneInterrogationLinker.h b/Engine/Source/Runtime/MovieSceneTracks/Public/EntitySystem/Interrogation/MovieSceneInterrogationLinker.h
--- a/Engine/Source/Runtime/MovieSceneTracks/Public/EntitySystem/Interrogation/MovieSceneInterrogationLinker.h
+++ b/Engine/Source/Runtime/MovieSceneTracks/Public/EntitySystem/Interrogation/MovieSceneInterrogationLinker.h
@@ -28,6 +28,7 @@ namespace MovieScene
 
 struct FInitialValueCache;
 struct FSystemInterrogatorEntityTracker;
+enum class EEntitySystemCategory : uint32;
 
 struct FInterrogationParams
 {
@@ -44,7 +45,7 @@ struct FInterrogationParams
 /**
  * A class specialized for interrogating Sequencer entity data without applying any state to objects.
  * Currently only tracks within the same time-base are supported.
- * Will only link systems that are not excluded from EEntitySystemContext::Interrogation.
+ * Will not link systems that are in the custom "ExcludedFromInterrogation" category.
  * Multiple different outputs can be interrogated simultaneously by Importing tracks onto separate channels
  * allocated through AllocateChannel.
  * 
@@ -274,6 +275,18 @@ public:
 	FSystemInterrogator();
 	~FSystemInterrogator();
 
+
+	/**
+	 * Gets the custom system category for interrogation-specific systems
+	 */
+	static EEntitySystemCategory GetInterrogationCategory();
+
+
+	/**
+	 * Gets the custom system category for systems who should be excluded from interrogation linkers
+	 */
+	static EEntitySystemCategory GetExcludedFromInterrogationCategory();
+
 	
 	/**
 	 * Allocate a new interrogation channel that relates to a specific object
diff --git a/Engine/Source/Runtime/UMG/Private/Animation/MovieScene2DTransformPropertySystem.cpp b/Engine/Source/Runtime/UMG/Private/Animation/MovieScene2DTransformPropertySystem.cpp
--- a/Engine/Source/Runtime/UMG/Private/Animation/MovieScene2DTransformPropertySystem.cpp
+++ b/Engine/Source/Runtime/UMG/Private/Animation/MovieScene2DTransformPropertySystem.cpp
@@ -2,6 +2,7 @@
 
 #include "Animation/MovieScene2DTransformPropertySystem.h"
 
+#include "EntitySystem/MovieSceneEntitySystemLinker.h"
 #include "Systems/MovieScenePropertyInstantiator.h"
 #include "Systems/FloatChannelEvaluatorSystem.h"
 #include "Systems/MovieScenePiecewiseDoubleBlenderSystem.h"
@@ -12,8 +13,6 @@
 UMovieScene2DTransformPropertySystem::UMovieScene2DTransformPropertySystem(const FObjectInitializer& ObjInit)
 	: Super(ObjInit)
 {
-	SystemExclusionContext |= UE::MovieScene::EEntitySystemContext::Interrogation;
-
 	BindToProperty(UE::MovieScene::FMovieSceneUMGComponentTypes::Get()->WidgetTransform);
 
 	if (HasAnyFlags(RF_ClassDefaultObject))
diff --git a/Engine/Source/Runtime/UMG/Private/Animation/MovieSceneMarginPropertySystem.cpp b/Engine/Source/Runtime/UMG/Private/Animation/MovieSceneMarginPropertySystem.cpp
--- a/Engine/Source/Runtime/UMG/Private/Animation/MovieSceneMarginPropertySystem.cpp
+++ b/Engine/Source/Runtime/UMG/Private/Animation/MovieSceneMarginPropertySystem.cpp
@@ -9,8 +9,6 @@
 UMovieSceneMarginPropertySystem::UMovieSceneMarginPropertySystem(const FObjectInitializer& ObjInit)
 	: Super(ObjInit)
 {
-	SystemExclusionContext |= UE::MovieScene::EEntitySystemContext::Interrogation;
-
 	BindToProperty(UE::MovieScene::FMovieSceneUMGComponentTypes::Get()->Margin);
 
 	if (HasAnyFlags(RF_ClassDefaultObject))
