Sequencer: default-construct components created by a mutation

diff --git a/Engine/Source/Runtime/MovieScene/Private/EntitySystem/MovieSceneEntityManager.cpp b/Engine/Source/Runtime/MovieScene/Private/EntitySystem/MovieSceneEntityManager.cpp
--- a/Engine/Source/Runtime/MovieScene/Private/EntitySystem/MovieSceneEntityManager.cpp
+++ b/Engine/Source/Runtime/MovieScene/Private/EntitySystem/MovieSceneEntityManager.cpp
@@ -1639,9 +1639,28 @@ int32 FEntityManager::MutateAll(const FEntityComponentFilter& Filter, const IMov
 		FEntityAllocation* NewAllocation = MigrateAllocation(AllocationIndex, Pair.Value);
 		TotalNumMutations += NewAllocation->Num();
 
+		FComponentMask OldAllocationType = EntityAllocationMasks[Pair.Key];
 		EntityAllocationMasks[Pair.Key] = Pair.Value;
 		EntityAllocations[Pair.Key] = NewAllocation;
 
+		// Default construct all the new components in the allocation, and then allow the mutation to further initialize this data if needed
+		for (FComponentMaskIterator Component = Pair.Value.Iterate(); Component; ++Component)
+		{
+			FComponentTypeID ComponentTypeID = FComponentTypeID::FromBitIndex(Component.GetIndex());
+			if (OldAllocationType.Contains(ComponentTypeID))
+			{
+				continue;
+			}
+
+			const FComponentHeader& ComponentHeader = NewAllocation->GetComponentHeaderChecked(ComponentTypeID);
+			if (!ComponentHeader.IsTag())
+			{
+				const FComponentTypeInfo& ComponentTypeInfo = ComponentRegistry->GetComponentTypeChecked(ComponentTypeID);
+
+				void* Components = ComponentHeader.GetValuePtr(0);
+				ComponentTypeInfo.ConstructItems(Components, NewAllocation->Num());
+			}
+		}
 		Mutation.InitializeAllocation(NewAllocation, Pair.Value);
 
 		DestroyAllocation(SourceAllocation);
