PR #9550: Sequencer Bezier interpolation to be done in double (Contributed by [REDACTED_PERSON])

diff --git a/Engine/Source/Runtime/MovieScene/Private/Channels/MovieSceneCurveChannelImpl.cpp b/Engine/Source/Runtime/MovieScene/Private/Channels/MovieSceneCurveChannelImpl.cpp
--- a/Engine/Source/Runtime/MovieScene/Private/Channels/MovieSceneCurveChannelImpl.cpp
+++ b/Engine/Source/Runtime/MovieScene/Private/Channels/MovieSceneCurveChannelImpl.cpp
@@ -40,13 +40,13 @@ EvalForTwoKeys(
 
 	double DecimalRate = DisplayRate.AsDecimal();
 
-	float Diff = (float)(Key2Time - Key1Time).Value;
+	double Diff = (double)(Key2Time - Key1Time).Value;
 	Diff /= DecimalRate;
 	const int CheckBothLinear = GSequencerLinearCubicInterpolation;
 
 	if (Diff > 0 && Key1.InterpMode != RCIM_Constant)
 	{
-		const float Alpha = ((float)(InTime - Key1Time).Value / DecimalRate) / Diff;
+		const double Alpha = ((double)(InTime - Key1Time).Value / DecimalRate) / Diff;
 		const CurveValueType P0 = Key1.Value;
 		const CurveValueType P3 = Key2.Value;
 
@@ -56,10 +56,10 @@ EvalForTwoKeys(
 		}
 		else
 		{
-			float LeaveTangent = Key1.Tangent.LeaveTangent * DecimalRate;
-			float ArriveTangent = Key2.Tangent.ArriveTangent * DecimalRate;
+			double LeaveTangent = Key1.Tangent.LeaveTangent * DecimalRate;
+			double ArriveTangent = Key2.Tangent.ArriveTangent * DecimalRate;
 
-			const float OneThird = 1.0f / 3.0f;
+			const double OneThird = 1.0 / 3.0;
 			const CurveValueType P1 = P0 + (LeaveTangent * Diff*OneThird);
 			const CurveValueType P2 = P3 - (ArriveTangent * Diff*OneThird);
 
@@ -387,15 +387,15 @@ bool TMovieSceneCurveChannelImpl<ChannelType>::Evaluate(const ChannelType* InCha
 		{
 		case RCIM_Cubic:
 		{
-			const float OneThird = 1.0f / 3.0f;
+			const double OneThird = 1.0 / 3.0;
 			if ((Key1.Tangent.TangentWeightMode == RCTWM_WeightedNone || Key1.Tangent.TangentWeightMode == RCTWM_WeightedArrive)
 				&& (Key2.Tangent.TangentWeightMode == RCTWM_WeightedNone || Key2.Tangent.TangentWeightMode == RCTWM_WeightedLeave))
 			{
 				const int32 Diff = InChannel->Times[Index2].Value - InChannel->Times[Index1].Value;
-				const float P0 = Key1.Value;
-				const float P1 = P0 + (Key1.Tangent.LeaveTangent * Diff * OneThird);
-				const float P3 = Key2.Value;
-				const float P2 = P3 - (Key2.Tangent.ArriveTangent * Diff * OneThird);
+				const double P0 = Key1.Value;
+				const double P1 = P0 + (Key1.Tangent.LeaveTangent * Diff * OneThird);
+				const double P3 = Key2.Value;
+				const double P2 = P3 - (Key2.Tangent.ArriveTangent * Diff * OneThird);
 
 				OutValue = Params.ValueOffset + UE::Curves::BezierInterp(P0, P1, P2, P3, Interp);
 				break;
@@ -407,49 +407,49 @@ bool TMovieSceneCurveChannelImpl<ChannelType>::Evaluate(const ChannelType* InCha
 
 				const double Time1 = InChannel->TickResolution.AsSeconds(InChannel->Times[Index1].Value);
 				const double Time2 = InChannel->TickResolution.AsSeconds(InChannel->Times[Index2].Value);
-				const float X = Time2 - Time1;
-				float CosAngle, SinAngle;
-				float Angle = FMath::Atan(Key1.Tangent.LeaveTangent * ToSeconds);
+				const double X = Time2 - Time1;
+				double CosAngle, SinAngle;
+				double Angle = FMath::Atan(Key1.Tangent.LeaveTangent * ToSeconds);
 				FMath::SinCos(&SinAngle, &CosAngle, Angle);
-				float LeaveWeight;
+				double LeaveWeight;
 				if (Key1.Tangent.TangentWeightMode == RCTWM_WeightedNone || Key1.Tangent.TangentWeightMode == RCTWM_WeightedArrive)
 				{
-					const float LeaveTangentNormalized = Key1.Tangent.LeaveTangent / (TimeInterval);
-					const float Y = LeaveTangentNormalized * X;
+					const double LeaveTangentNormalized = Key1.Tangent.LeaveTangent / (TimeInterval);
+					const double Y = LeaveTangentNormalized * X;
 					LeaveWeight = FMath::Sqrt(X*X + Y * Y) * OneThird;
 				}
 				else
 				{
 					LeaveWeight = Key1.Tangent.LeaveTangentWeight;
 				}
-				const float Key1TanX = CosAngle * LeaveWeight + Time1;
-				const float Key1TanY = SinAngle * LeaveWeight + Key1.Value;
+				const double Key1TanX = CosAngle * LeaveWeight + Time1;
+				const double Key1TanY = SinAngle * LeaveWeight + Key1.Value;
 
 				Angle = FMath::Atan(Key2.Tangent.ArriveTangent * ToSeconds);
 				FMath::SinCos(&SinAngle, &CosAngle, Angle);
-				float ArriveWeight;
+				double ArriveWeight;
 				if (Key2.Tangent.TangentWeightMode == RCTWM_WeightedNone || Key2.Tangent.TangentWeightMode == RCTWM_WeightedLeave)
 				{
-					const float ArriveTangentNormalized = Key2.Tangent.ArriveTangent / (TimeInterval);
-					const float Y = ArriveTangentNormalized * X;
+					const double ArriveTangentNormalized = Key2.Tangent.ArriveTangent / (TimeInterval);
+					const double Y = ArriveTangentNormalized * X;
 					ArriveWeight = FMath::Sqrt(X*X + Y * Y) * OneThird;
 				}
 				else
 				{
 					ArriveWeight =  Key2.Tangent.ArriveTangentWeight;
 				}
-				const float Key2TanX = -CosAngle * ArriveWeight + Time2;
-				const float Key2TanY = -SinAngle * ArriveWeight + Key2.Value;
+				const double Key2TanX = -CosAngle * ArriveWeight + Time2;
+				const double Key2TanY = -SinAngle * ArriveWeight + Key2.Value;
 
 				//Normalize the Time Range
-				const float RangeX = Time2 - Time1;
+				const double RangeX = Time2 - Time1;
 
-				const float Dx1 = Key1TanX - Time1;
-				const float Dx2 = Key2TanX - Time1;
+				const double Dx1 = Key1TanX - Time1;
+				const double Dx2 = Key2TanX - Time1;
 
 				// Normalize values
-				const float NormalizedX1 = Dx1 / RangeX;
-				const float NormalizedX2 = Dx2 / RangeX;
+				const double NormalizedX1 = Dx1 / RangeX;
+				const double NormalizedX2 = Dx2 / RangeX;
 				
 				double Coeff[4];
 				double Results[3];
@@ -489,10 +489,10 @@ bool TMovieSceneCurveChannelImpl<ChannelType>::Evaluate(const ChannelType* InCha
 
 				}
 				//now use NewInterp and adjusted tangents plugged into the Y (Value) part of the graph.
-				const float P0 = Key1.Value;
-				const float P1 = Key1TanY;
-				const float P3 = Key2.Value;
-				const float P2 = Key2TanY;
+				const double P0 = Key1.Value;
+				const double P1 = Key1TanY;
+				const double P3 = Key2.Value;
+				const double P2 = Key2TanY;
 
 				OutValue = Params.ValueOffset + UE::Curves::BezierInterp(P0, P1, P2, P3,  NewInterp);
 			}
@@ -528,8 +528,8 @@ void TMovieSceneCurveChannelImpl<ChannelType>::AutoSetTangents(ChannelType* InCh
 		{
 			FirstValue.Tangent.TangentWeightMode = RCTWM_WeightedNone;
 			ChannelValueType& NextKey = InChannel->Values[1];
-			const float NextTimeDiff = FMath::Max<double>(KINDA_SMALL_NUMBER, InChannel->Times[1].Value - InChannel->Times[0].Value);
-			const float NewTangent = (NextKey.Value - FirstValue.Value) / NextTimeDiff;
+			const double NextTimeDiff = FMath::Max<double>(KINDA_SMALL_NUMBER, InChannel->Times[1].Value - InChannel->Times[0].Value);
+			const double NewTangent = (NextKey.Value - FirstValue.Value) / NextTimeDiff;
 			FirstValue.Tangent.LeaveTangent = NewTangent;
 		}
 		else if (FirstValue.InterpMode == RCIM_Cubic && FirstValue.TangentMode == RCTM_Auto)
@@ -546,8 +546,8 @@ void TMovieSceneCurveChannelImpl<ChannelType>::AutoSetTangents(ChannelType* InCh
 			LastValue.Tangent.TangentWeightMode = RCTWM_WeightedNone;
 			int32 Index = InChannel->Values.Num() - 1;
 			ChannelValueType& PrevKey = InChannel->Values[Index-1];
-			const float PrevTimeDiff = FMath::Max<double>(KINDA_SMALL_NUMBER, InChannel->Times[Index].Value - InChannel->Times[Index - 1].Value);
-			const float NewTangent = (LastValue.Value - PrevKey.Value) / PrevTimeDiff;
+			const double PrevTimeDiff = FMath::Max<double>(KINDA_SMALL_NUMBER, InChannel->Times[Index].Value - InChannel->Times[Index - 1].Value);
+			const double NewTangent = (LastValue.Value - PrevKey.Value) / PrevTimeDiff;
 			LastValue.Tangent.ArriveTangent = NewTangent;
 		}
 		else if (LastValue.InterpMode == RCIM_Cubic && LastValue.TangentMode == RCTM_Auto)
@@ -565,7 +565,7 @@ void TMovieSceneCurveChannelImpl<ChannelType>::AutoSetTangents(ChannelType* InCh
 		if (ThisKey.InterpMode == RCIM_Cubic && ThisKey.TangentMode == RCTM_Auto)
 		{
 			ChannelValueType& NextKey = InChannel->Values[Index+1];
-			CurveValueType NewTangent = 0.f;
+			CurveValueType NewTangent = 0.0;
 			const double PrevToNextTimeDiff = FMath::Max<double>(KINDA_SMALL_NUMBER, InChannel->Times[Index + 1].Value - InChannel->Times[Index - 1].Value);
 
 			if (!UseNewAutoTangent)
@@ -582,26 +582,26 @@ void TMovieSceneCurveChannelImpl<ChannelType>::AutoSetTangents(ChannelType* InCh
 					AutoCalcTangent(PrevKey.Value, ThisKey.Value, NextKey.Value, Tension, NewTangent);
 					NewTangent /= PrevToNextTimeDiff;
 					//if within 0 to 15% or 85% to 100% range we gradually weight tangent to zero
-					const float AverageToZeroRange = 0.85f;
-					const float ValDiff = FMath::Abs<float>(NextKey.Value - PrevKey.Value);
-					const float OurDiff = FMath::Abs<float>(ThisKey.Value - PrevKey.Value);
+					const double AverageToZeroRange = 0.85;
+					const double ValDiff = FMath::Abs<double>(NextKey.Value - PrevKey.Value);
+					const double OurDiff = FMath::Abs<double>(ThisKey.Value - PrevKey.Value);
 					//ValDiff won't be zero ever due to previous check
-					float PercDiff = OurDiff / ValDiff;
+					double PercDiff = OurDiff / ValDiff;
 					if (PercDiff > AverageToZeroRange)
 					{
-						PercDiff = (PercDiff - AverageToZeroRange) / (1.0f - AverageToZeroRange);
-						NewTangent = NewTangent * (1.0f - PercDiff);
+						PercDiff = (PercDiff - AverageToZeroRange) / (1.0 - AverageToZeroRange);
+						NewTangent = NewTangent * (1.0 - PercDiff);
 					}
-					else if (PercDiff < (1.0f - AverageToZeroRange))
+					else if (PercDiff < (1.0 - AverageToZeroRange))
 					{
-						PercDiff = PercDiff  / (1.0f - AverageToZeroRange);
+						PercDiff = PercDiff  / (1.0 - AverageToZeroRange);
 						NewTangent = NewTangent * PercDiff;
 					}
 				}
 			}
 
 			// In 'auto' mode, arrive and leave tangents are always the same
-			ThisKey.Tangent.LeaveTangent = ThisKey.Tangent.ArriveTangent = (float)NewTangent;
+			ThisKey.Tangent.LeaveTangent = ThisKey.Tangent.ArriveTangent = (double)NewTangent;
 			ThisKey.Tangent.TangentWeightMode = RCTWM_WeightedNone;
 		}
 		else if (ThisKey.InterpMode == RCIM_Linear)
@@ -610,7 +610,7 @@ void TMovieSceneCurveChannelImpl<ChannelType>::AutoSetTangents(ChannelType* InCh
 			ChannelValueType& NextKey = InChannel->Values[Index + 1];
 
 			const double PrevTimeDiff = FMath::Max<double>(KINDA_SMALL_NUMBER, InChannel->Times[Index].Value - InChannel->Times[Index - 1].Value);
-			float NewTangent  = (ThisKey.Value - PrevKey.Value) / PrevTimeDiff;
+			double NewTangent  = (ThisKey.Value - PrevKey.Value) / PrevTimeDiff;
 			ThisKey.Tangent.ArriveTangent = NewTangent;
 			
 			const double NextTimeDiff = FMath::Max<double>(KINDA_SMALL_NUMBER, InChannel->Times[Index + 1].Value - InChannel->Times[Index].Value);
@@ -807,7 +807,7 @@ void TMovieSceneCurveChannelImpl<ChannelType>::RefineCurvePoints(const ChannelTy
 
 				EvalTime = FMath::Lerp(Lower.Get<0>(), Upper.Get<0>(), InterpTimes[InterpIndex]);
 
-				CurveValueType Value = 0.f;
+				CurveValueType Value = 0.0;
 				Evaluate(InChannel, EvalTime * InTickResolution, Value);
 
 				const CurveValueType LinearValue = FMath::Lerp(Lower.Get<1>(), Upper.Get<1>(), InterpTimes[InterpIndex]);
@@ -834,7 +834,7 @@ bool TMovieSceneCurveChannelImpl<ChannelType>::ValueExistsAtTime(const ChannelTy
 {
 	const FFrameTime FrameTime(InFrameNumber);
 
-	CurveValueType ExistingValue = 0.f;
+	CurveValueType ExistingValue = 0.0;
 	return Channel->Evaluate(FrameTime, ExistingValue) && FMath::IsNearlyEqual(ExistingValue, Value, (CurveValueType)KINDA_SMALL_NUMBER);
 }
 
