Sequencer: Don't register the same tick client multiple times

A player could be registered multiple times if "SetSequence" was called multiple times on it. This resulted in that player getting ticked more than once per frame, making it appear like its play rate was higher than expected

diff --git a/Engine/Source/Runtime/MovieScene/Private/MovieSceneSequenceTickManager.cpp b/Engine/Source/Runtime/MovieScene/Private/MovieSceneSequenceTickManager.cpp
--- a/Engine/Source/Runtime/MovieScene/Private/MovieSceneSequenceTickManager.cpp
+++ b/Engine/Source/Runtime/MovieScene/Private/MovieSceneSequenceTickManager.cpp
@@ -114,6 +114,9 @@ void UMovieSceneSequenceTickManager::RegisterTickClient(const FMovieSceneSequenc
 		*InterfaceObject->GetName(), ObjectWorld ? *ObjectWorld->GetName() : TEXT("nullptr"), ThisWorld ? *ThisWorld->GetName() : TEXT("nullptr"));
 #endif
 
+	// Remove the client from any existing groups if we already know about it
+	UnregisterTickClientImpl(InTickInterface.GetInterface());
+
 	const int32 DesiredTickIntervalMs = ResolvedTickInterval.RoundTickIntervalMs();
 
 	// Find a linker group with the specified tick interval
@@ -180,14 +183,19 @@ void UMovieSceneSequenceTickManager::UnregisterTickClient(TScriptInterface<IMovi
 		return;
 	}
 
+	UnregisterTickClientImpl(ClientInterface);
+}
+
+void UMovieSceneSequenceTickManager::UnregisterTickClientImpl(IMovieSceneSequenceTickManagerClient* InClientInterface)
+{
 	// Remove any latent actions tied to the given client.
 	UObject* Object = InTickInterface.GetObject();
 	if (Object)
 	{
 		ClearLatentActions(Object);
 	}
-
-	const int32 ClientIndex = Algo::IndexOfBy(TickableClients, ClientInterface, &FTickableClientData::Interface);
+	
+	const int32 ClientIndex = Algo::IndexOfBy(TickableClients, InClientInterface, &FTickableClientData::Interface);
 	if (!TickableClients.IsValidIndex(ClientIndex))
 	{
 		return;
diff --git a/Engine/Source/Runtime/MovieScene/Public/MovieSceneSequenceTickManager.h b/Engine/Source/Runtime/MovieScene/Public/MovieSceneSequenceTickManager.h
--- a/Engine/Source/Runtime/MovieScene/Public/MovieSceneSequenceTickManager.h
+++ b/Engine/Source/Runtime/MovieScene/Public/MovieSceneSequenceTickManager.h
@@ -212,6 +212,11 @@ private:
 	 */
 	void ProcessPendingOperations(TArrayView<const FPendingOperation> InOperations);
 
+	/**
+	 * Unregister tick client without performing basic checks
+	 */
+	void UnregisterTickClientImpl(IMovieSceneSequenceTickManagerClient* InClientInterface);
+
 private:
 
 	/** Flat array of tickable client data with indices into LinkerGroups */
