Sequencer: Re-implemented budgeting changes that were lost in CL 21275076

 - Mostly this includes improvements to the way we handle re-entrancy inside FMovieSceneEntitySystemRunner::GameThread_PostEvaluationPhase by caching any state that could change on the stack and restoring it later, firming up the asyncronous update mechanisms to include options for queuing final updates that optionally finish and destroy instances, and adding callbacks that are called when a request has been fulfilled.

 - FMovieSceneRootEvaluationTemplateInstance no longer contains an implicit synchronous runner and must be supplied a single runner that is externally flushed. This avoids the confusing state where re-entrancy could previously cause updates to be flushed on different runners simultaneously.

diff --git a/Engine/Plugins/Animation/ControlRig/Source/ControlRigEditor/Private/ControlRigSequencerEditorLibrary.cpp b/Engine/Plugins/Animation/ControlRig/Source/ControlRigEditor/Private/ControlRigSequencerEditorLibrary.cpp
--- a/Engine/Plugins/Animation/ControlRig/Source/ControlRigEditor/Private/ControlRigSequencerEditorLibrary.cpp
+++ b/Engine/Plugins/Animation/ControlRig/Source/ControlRigEditor/Private/ControlRigSequencerEditorLibrary.cpp
@@ -872,7 +872,7 @@ static bool LocalGetControlRigControlValues(IMovieScenePlayer* Player, UMovieSce
 			GlobalTime = GlobalTime * RootToLocalTransform.InverseLinearOnly();
 			FMovieSceneContext Context = FMovieSceneContext(FMovieSceneEvaluationRange(GlobalTime, TickResolution), Player->GetPlaybackStatus()).SetHasJumped(true);
 
-			Player->GetEvaluationTemplate().Evaluate(Context, *Player);
+			Player->GetEvaluationTemplate().EvaluateSynchronousBlocking(Context, *Player);
 			ControlRig->Evaluate_AnyThread();
 			OutValues[Index] = ControlRig->GetControlValue(ControlName);
 		}
diff --git a/Engine/Plugins/Animation/ControlRig/Source/ControlRigEditor/Private/Sequencer/ControlRigParameterTrackEditor.cpp b/Engine/Plugins/Animation/ControlRig/Source/ControlRigEditor/Private/Sequencer/ControlRigParameterTrackEditor.cpp
--- a/Engine/Plugins/Animation/ControlRig/Source/ControlRigEditor/Private/Sequencer/ControlRigParameterTrackEditor.cpp
+++ b/Engine/Plugins/Animation/ControlRig/Source/ControlRigEditor/Private/Sequencer/ControlRigParameterTrackEditor.cpp
@@ -4110,7 +4110,7 @@ bool FControlRigParameterTrackEditor::CollapseAllLayers(TSharedPtr<ISequencer>&
 
 			FMovieSceneContext Context = FMovieSceneContext(FMovieSceneEvaluationRange(GlobalTime, TickResolution), SequencerPtr->GetPlaybackStatus()).SetHasJumped(true);
 
-			SequencerPtr->GetEvaluationTemplate().Evaluate(Context, *SequencerPtr);
+			SequencerPtr->GetEvaluationTemplate().EvaluateSynchronousBlocking(Context, *SequencerPtr);
 			ControlRig->Evaluate_AnyThread();
 			for (TPair<FName, TArray<FTransform>>& TrailControlTransform : ControlLocalTransforms)
 			{
diff --git a/Engine/Plugins/Animation/ControlRig/Source/ControlRigEditor/Private/Sequencer/ControlRigSpaceChannelEditors.cpp b/Engine/Plugins/Animation/ControlRig/Source/ControlRigEditor/Private/Sequencer/ControlRigSpaceChannelEditors.cpp
--- a/Engine/Plugins/Animation/ControlRig/Source/ControlRigEditor/Private/Sequencer/ControlRigSpaceChannelEditors.cpp
+++ b/Engine/Plugins/Animation/ControlRig/Source/ControlRigEditor/Private/Sequencer/ControlRigSpaceChannelEditors.cpp
@@ -430,7 +430,7 @@ FKeyHandle FControlRigSpaceChannelHelpers::SequencerKeyControlRigSpaceChannel(UC
 			GlobalTime = GlobalTime * RootToLocalTransform.InverseLinearOnly();
 			
 			FMovieSceneContext SceneContext = FMovieSceneContext(FMovieSceneEvaluationRange(GlobalTime, TickResolution), Sequencer->GetPlaybackStatus()).SetHasJumped(true);
-			Sequencer->GetEvaluationTemplate().Evaluate(SceneContext, *Sequencer);
+			Sequencer->GetEvaluationTemplate().EvaluateSynchronousBlocking(SceneContext, *Sequencer);
 
 			Context.LocalTime = TickResolution.AsSeconds(FFrameTime(Time - 1));
 			ControlRig->SetControlGlobalTransform(ControlKey.Name, ControlWorldTransforms[0], true, Context);
@@ -454,7 +454,7 @@ FKeyHandle FControlRigSpaceChannelHelpers::SequencerKeyControlRigSpaceChannel(UC
 			GlobalTime = GlobalTime * RootToLocalTransform.InverseLinearOnly();
 
 			FMovieSceneContext SceneContext = FMovieSceneContext(FMovieSceneEvaluationRange(GlobalTime, TickResolution), Sequencer->GetPlaybackStatus()).SetHasJumped(true);
-			Sequencer->GetEvaluationTemplate().Evaluate(SceneContext, *Sequencer);
+			Sequencer->GetEvaluationTemplate().EvaluateSynchronousBlocking(SceneContext, *Sequencer);
 
 			ControlRig->Evaluate_AnyThread();
 			Context.LocalTime = TickResolution.AsSeconds(FFrameTime(Frame));
@@ -657,7 +657,7 @@ void  FControlRigSpaceChannelHelpers::SequencerSpaceChannelKeyDeleted(UControlRi
 			GlobalTime = GlobalTime * RootToLocalTransform.InverseLinearOnly();
 
 			FMovieSceneContext SceneContext = FMovieSceneContext(FMovieSceneEvaluationRange(GlobalTime, TickResolution), Sequencer->GetPlaybackStatus()).SetHasJumped(true);
-			Sequencer->GetEvaluationTemplate().Evaluate(SceneContext, *Sequencer);
+			Sequencer->GetEvaluationTemplate().EvaluateSynchronousBlocking(SceneContext, *Sequencer);
 			//make sure to set rig hierarchy correct since key is not deleted yet
 			switch (PreviousValue.SpaceType)
 			{
@@ -935,7 +935,7 @@ void FControlRigSpaceChannelHelpers::SequencerBakeControlInSpace(UControlRig* Co
 				GlobalTime = GlobalTime * RootToLocalTransform.InverseLinearOnly();
 
 				FMovieSceneContext SceneContext = FMovieSceneContext(FMovieSceneEvaluationRange(GlobalTime, TickResolution), Sequencer->GetPlaybackStatus()).SetHasJumped(true);
-				Sequencer->GetEvaluationTemplate().Evaluate(SceneContext, *Sequencer);
+				Sequencer->GetEvaluationTemplate().EvaluateSynchronousBlocking(SceneContext, *Sequencer);
 
 				//evaluate control rig
 				ControlRig->Evaluate_AnyThread();
@@ -965,7 +965,7 @@ void FControlRigSpaceChannelHelpers::SequencerBakeControlInSpace(UControlRig* Co
 				GlobalTime = GlobalTime * RootToLocalTransform.InverseLinearOnly();
 
 				FMovieSceneContext SceneContext = FMovieSceneContext(FMovieSceneEvaluationRange(GlobalTime, TickResolution), Sequencer->GetPlaybackStatus()).SetHasJumped(true);
-				Sequencer->GetEvaluationTemplate().Evaluate(SceneContext, *Sequencer);
+				Sequencer->GetEvaluationTemplate().EvaluateSynchronousBlocking(SceneContext, *Sequencer);
 		
 				//evaluate control rig
 				ControlRig->Evaluate_AnyThread();
diff --git a/Engine/Plugins/Animation/ControlRig/Source/ControlRigEditor/Private/Tools/ControlRigSnapper.cpp b/Engine/Plugins/Animation/ControlRig/Source/ControlRigEditor/Private/Tools/ControlRigSnapper.cpp
--- a/Engine/Plugins/Animation/ControlRig/Source/ControlRigEditor/Private/Tools/ControlRigSnapper.cpp
+++ b/Engine/Plugins/Animation/ControlRig/Source/ControlRigEditor/Private/Tools/ControlRigSnapper.cpp
@@ -133,7 +133,7 @@ static bool LocalGetControlRigControlTransforms(IMovieScenePlayer* Player, const
 				GlobalTime = GlobalTime * RootToLocalTransform.InverseLinearOnly(); //player evals in root time so need to go back to it.
 
 				FMovieSceneContext Context = FMovieSceneContext(FMovieSceneEvaluationRange(GlobalTime, TickResolution), Player->GetPlaybackStatus()).SetHasJumped(true);
-				Player->GetEvaluationTemplate().Evaluate(Context, *Player);
+				Player->GetEvaluationTemplate().EvaluateSynchronousBlocking(Context, *Player);
 			}
 			ControlRig->Evaluate_AnyThread();
 			OutTransforms[Index] = ControlRig->GetControlGlobalTransform(ControlName) * ParentTransforms[Index];
diff --git a/Engine/Plugins/Editor/SequencerAnimTools/Source/SequencerAnimTools/Private/Sequencer/MovieSceneTransformTrail.cpp b/Engine/Plugins/Editor/SequencerAnimTools/Source/SequencerAnimTools/Private/Sequencer/MovieSceneTransformTrail.cpp
--- a/Engine/Plugins/Editor/SequencerAnimTools/Source/SequencerAnimTools/Private/Sequencer/MovieSceneTransformTrail.cpp
+++ b/Engine/Plugins/Editor/SequencerAnimTools/Source/SequencerAnimTools/Private/Sequencer/MovieSceneTransformTrail.cpp
@@ -768,7 +768,7 @@ bool FMovieSceneControlRigTransformTrail::ApplyDelta(const FVector& Pos, const F
 				NewTransform.SetLocation(NewTransform.GetLocation() + Pos);
 				KeyInfo->Transform = NewTransform;
 				NewTransform = NewTransform.GetRelativeTransform(KeyInfo->ParentTransform);
-				GetSequencer()->GetEvaluationTemplate().Evaluate(MovieSceneContext, *Player);
+				GetSequencer()->GetEvaluationTemplate().EvaluateSynchronousBlocking(MovieSceneContext, *Player);
 				ControlRig->Evaluate_AnyThread();
 				ControlRig->SetControlGlobalTransform(ControlName, NewTransform, true, Context, false);
 			}
@@ -823,7 +823,7 @@ bool FMovieSceneControlRigTransformTrail::EndTracking()
 
 		FMovieSceneContext MovieSceneContext = FMovieSceneContext(FMovieSceneEvaluationRange(StartTime, TickResolution), Player->GetPlaybackStatus()).SetHasJumped(true);
 		
-		Player->GetEvaluationTemplate().Evaluate(MovieSceneContext, *Player);
+		Player->GetEvaluationTemplate().EvaluateSynchronousBlocking(MovieSceneContext, *Player);
 		ControlRig->Evaluate_AnyThread();
 		return true;
 	}
@@ -863,7 +863,7 @@ bool FMovieSceneControlRigTransformTrail::HandleAltClick(FEditorViewportClient*
 	Context.LocalTime = TickResolution.AsSeconds(GlobalTime);
 	Context.KeyMask = (uint32)EControlRigContextChannelToKey::Translation;
 	FMovieSceneContext MovieSceneContext = FMovieSceneContext(FMovieSceneEvaluationRange(GlobalTime, TickResolution), Player->GetPlaybackStatus()).SetHasJumped(true);
-	GetSequencer()->GetEvaluationTemplate().Evaluate(MovieSceneContext, *Player);
+	GetSequencer()->GetEvaluationTemplate().EvaluateSynchronousBlocking(MovieSceneContext, *Player);
 	ControlRig->Evaluate_AnyThread();
 	FTransform NewTransform(ControlRig->GetControlGlobalTransform(ControlName));
 	ControlRig->SetControlGlobalTransform(ControlName, NewTransform, true, Context, false);
@@ -873,7 +873,7 @@ bool FMovieSceneControlRigTransformTrail::HandleAltClick(FEditorViewportClient*
 	StartTime = StartTime * RootToLocalTransform.InverseLinearOnly(); //player evals in root time so need to go back to it.
 
 	MovieSceneContext = FMovieSceneContext(FMovieSceneEvaluationRange(StartTime, TickResolution), Player->GetPlaybackStatus()).SetHasJumped(true);
-	Player->GetEvaluationTemplate().Evaluate(MovieSceneContext, *Player);
+	Player->GetEvaluationTemplate().EvaluateSynchronousBlocking(MovieSceneContext, *Player);
 	ControlRig->Evaluate_AnyThread();
 
 	//create new keys
diff --git a/Engine/Plugins/Editor/SequencerAnimTools/Source/SequencerAnimTools/Private/Sequencer/SequencerTrailHierarchy.cpp b/Engine/Plugins/Editor/SequencerAnimTools/Source/SequencerAnimTools/Private/Sequencer/SequencerTrailHierarchy.cpp
--- a/Engine/Plugins/Editor/SequencerAnimTools/Source/SequencerAnimTools/Private/Sequencer/SequencerTrailHierarchy.cpp
+++ b/Engine/Plugins/Editor/SequencerAnimTools/Source/SequencerAnimTools/Private/Sequencer/SequencerTrailHierarchy.cpp
@@ -266,7 +266,7 @@ void FSequencerTrailHierarchy::UpdateControlRig(const TArray<FFrameNumber> &Fram
 
 					FMovieSceneContext Context = FMovieSceneContext(FMovieSceneEvaluationRange(GlobalTime, TickResolution), Player->GetPlaybackStatus()).SetHasJumped(true);
 
-					Player->GetEvaluationTemplate().Evaluate(Context, *Player);
+					Player->GetEvaluationTemplate().EvaluateSynchronousBlocking(Context, *Player);
 					ControlRig->Evaluate_AnyThread();
 					for (FTrailControlTransforms& TrailControlTransform : TrailControlTransforms)
 					{
@@ -293,7 +293,7 @@ void FSequencerTrailHierarchy::UpdateControlRig(const TArray<FFrameNumber> &Fram
 			FFrameTime StartTime = Sequencer->GetLocalTime().Time;
 			StartTime = StartTime * RootToLocalTransform.InverseLinearOnly(); //player evals in root time so need to go back to it.
 			FMovieSceneContext Context = FMovieSceneContext(FMovieSceneEvaluationRange(StartTime, TickResolution), Player->GetPlaybackStatus()).SetHasJumped(true);
-			Player->GetEvaluationTemplate().Evaluate(Context, *Player);
+			Player->GetEvaluationTemplate().EvaluateSynchronousBlocking(Context, *Player);
 			ControlRig->Evaluate_AnyThread();
 		}
 	}
diff --git a/Engine/Plugins/Experimental/GameFeatures/Source/GameFeatures/Private/GameFeaturePluginStateMachine.cpp b/Engine/Plugins/Experimental/GameFeatures/Source/GameFeatures/Private/GameFeaturePluginStateMachine.cpp
--- a/Engine/Plugins/Experimental/GameFeatures/Source/GameFeatures/Private/GameFeaturePluginStateMachine.cpp
+++ b/Engine/Plugins/Experimental/GameFeatures/Source/GameFeatures/Private/GameFeaturePluginStateMachine.cpp
@@ -87,7 +87,7 @@ namespace UE::GameFeatures
 					FReferenceChainSearch(AssetObj, EReferenceChainSearchMode::Shortest | EReferenceChainSearchMode::PrintResults);
 				}
 
-				ensureAlwaysMsgf(false, TEXT("GFP %s failed to unload asset %s!"), *PluginName, *AssetData.GetFullName());
+				//ensureAlwaysMsgf(false, TEXT("GFP %s failed to unload asset %s!"), *PluginName, *AssetData.GetFullName());
 			}
 		};
 
diff --git a/Engine/Plugins/Importers/USDImporter/Source/USDUtilities/Private/USDSkeletalDataConversion.cpp b/Engine/Plugins/Importers/USDImporter/Source/USDUtilities/Private/USDSkeletalDataConversion.cpp
--- a/Engine/Plugins/Importers/USDImporter/Source/USDUtilities/Private/USDSkeletalDataConversion.cpp
+++ b/Engine/Plugins/Importers/USDImporter/Source/USDUtilities/Private/USDSkeletalDataConversion.cpp
@@ -3540,7 +3540,7 @@ bool UnrealToUsd::ConvertControlRigSection(
 			InPlayer->GetPlaybackStatus()
 		).SetHasJumped( true );
 
-		InPlayer->GetEvaluationTemplate().Evaluate( Context, *InPlayer );
+		InPlayer->GetEvaluationTemplate().EvaluateSynchronousBlocking( Context, *InPlayer );
 		ControlRig->Evaluate_AnyThread();
 
 		URigHierarchy* Hierarchy = ControlRig->GetHierarchy();
diff --git a/Engine/Plugins/MovieScene/ReplayTracks/Source/ReplayTracks/Private/Systems/MovieSceneReplaySystem.cpp b/Engine/Plugins/MovieScene/ReplayTracks/Source/ReplayTracks/Private/Systems/MovieSceneReplaySystem.cpp
--- a/Engine/Plugins/MovieScene/ReplayTracks/Source/ReplayTracks/Private/Systems/MovieSceneReplaySystem.cpp
+++ b/Engine/Plugins/MovieScene/ReplayTracks/Source/ReplayTracks/Private/Systems/MovieSceneReplaySystem.cpp
@@ -445,7 +445,7 @@ void UMovieSceneReplaySystem::OnPostLoadMap(UWorld* World, IMovieScenePlayer* La
 				World->GetTimerManager().ClearTimer(ReEvaluateHandle);
 
 				FMovieSceneRootEvaluationTemplateInstance& RootEvalTemplate = LastPlayer->GetEvaluationTemplate();
-				RootEvalTemplate.Evaluate(LastContext, *LastPlayer);
+				RootEvalTemplate.EvaluateSynchronousBlocking(LastContext, *LastPlayer);
 			}
 		},
 		0.1f,
diff --git a/Engine/Plugins/MovieScene/TemplateSequence/Source/TemplateSequence/Private/CameraAnimationSequencePlayer.cpp b/Engine/Plugins/MovieScene/TemplateSequence/Source/TemplateSequence/Private/CameraAnimationSequencePlayer.cpp
--- a/Engine/Plugins/MovieScene/TemplateSequence/Source/TemplateSequence/Private/CameraAnimationSequencePlayer.cpp
+++ b/Engine/Plugins/MovieScene/TemplateSequence/Source/TemplateSequence/Private/CameraAnimationSequencePlayer.cpp
@@ -325,7 +325,11 @@ void UCameraAnimationSequencePlayer::Initialize(UMovieSceneSequence* InSequence)
 
 	PlayPosition.Reset(StartFrame);
 
-	RootTemplateInstance.Initialize(*Sequence, *this, nullptr);
+	UCameraAnimationSequenceSubsystem* Subsystem = UCameraAnimationSequenceSubsystem::GetCameraAnimationSequenceSubsystem(GetWorld());
+	ensureMsgf(Subsystem, TEXT("Unable to locate a valid camera animation sub-system. Camera anim sequences will not play."));
+
+	TSharedPtr<FMovieSceneEntitySystemRunner> Runner = Subsystem ? Subsystem->GetRunner() : nullptr;
+	RootTemplateInstance.Initialize(*Sequence, *this, nullptr, Runner);
 }
 
 void UCameraAnimationSequencePlayer::Play(bool bLoop, bool bRandomStartTime)
@@ -384,9 +388,16 @@ void UCameraAnimationSequencePlayer::Update(FFrameTime NewPosition)
 		}
 	}
 
-	const FMovieSceneEvaluationRange Range = PlayPosition.PlayTo(NewPosition);
-	const FMovieSceneContext Context(Range, Status);
-	RootTemplateInstance.Evaluate(Context, *this);
+	if (TSharedPtr<FMovieSceneEntitySystemRunner> Runner = RootTemplateInstance.GetRunner())
+	{
+		const FMovieSceneEvaluationRange Range = PlayPosition.PlayTo(NewPosition);
+		const FMovieSceneContext Context(Range, Status);
+
+		Runner->QueueUpdate(Context, RootTemplateInstance.GetRootInstanceHandle());
+
+		// @todo: we should really only call flush _once_ for all camera anim sequences
+		Runner->Flush();
+	}
 
 	if (bShouldStop)
 	{
@@ -405,9 +416,13 @@ void UCameraAnimationSequencePlayer::Stop()
 
 	PlayPosition.Reset(StartFrame);
 
-	if (RootTemplateInstance.IsValid())
+	if (TSharedPtr<FMovieSceneEntitySystemRunner> Runner = RootTemplateInstance.GetRunner())
 	{
-		RootTemplateInstance.Finish(*this);
+		if (Runner->QueueFinalUpdate(RootTemplateInstance.GetRootInstanceHandle(), UE::MovieScene::ERunnerUpdateFlags::Finish))
+		{
+			// @todo: we should really only call flush _once_ for all camera anim sequences
+			Runner->Flush();
+		}
 	}
 }
 
diff --git a/Engine/Plugins/MovieScene/TemplateSequence/Source/TemplateSequence/Private/CameraAnimationSequenceSubsystem.cpp b/Engine/Plugins/MovieScene/TemplateSequence/Source/TemplateSequence/Private/CameraAnimationSequenceSubsystem.cpp
--- a/Engine/Plugins/MovieScene/TemplateSequence/Source/TemplateSequence/Private/CameraAnimationSequenceSubsystem.cpp
+++ b/Engine/Plugins/MovieScene/TemplateSequence/Source/TemplateSequence/Private/CameraAnimationSequenceSubsystem.cpp
@@ -7,6 +7,7 @@
 #include "EntitySystem/MovieSceneEntityManager.h"
 #include "EntitySystem/MovieSceneEntityMutations.h"
 #include "EntitySystem/MovieSceneEntitySystemLinker.h"
+#include "EntitySystem/MovieSceneEntitySystemRunner.h"
 #include "EntitySystem/MovieSceneEntitySystemTask.h"
 #include "EntitySystem/MovieSceneMasterInstantiatorSystem.h"
 #include "EntitySystem/MovieSceneBoundSceneComponentInstantiator.h"
@@ -212,13 +213,19 @@ UCameraAnimationSequenceSubsystem::~UCameraAnimationSequenceSubsystem()
 {
 }
 
+void UCameraAnimationSequenceSubsystem::Initialize(FSubsystemCollectionBase& Collection)
+{
+	Runner = MakeShared<FMovieSceneEntitySystemRunner>();
+}
+
 void UCameraAnimationSequenceSubsystem::Deinitialize()
 {
 	// We check if the runner still has a valid pointer on the linker because the linker could
 	// have been GC'ed just now, which would make DetachFromLinker complain.
-	if (Runner.GetLinker())
+	if (Runner->GetLinker())
 	{
-		Runner.DetachFromLinker();
+		Runner->DetachFromLinker();
+		Runner = nullptr;
 	}
 	Linker = nullptr;
 
@@ -230,10 +237,15 @@ UMovieSceneEntitySystemLinker* UCameraAnimationSequenceSubsystem::GetLinker(bool
 	if (!Linker && bAutoCreate)
 	{
 		Linker = CreateLinker(this, TEXT("CameraAnimationSequenceSubsystemLinker"));
-		Runner.AttachToLinker(Linker);
+		Runner->AttachToLinker(Linker);
 	}
 	return Linker;
 }
 
+TSharedPtr<FMovieSceneEntitySystemRunner> UCameraAnimationSequenceSubsystem::GetRunner() const
+{
+	return Runner;
+}
+
 #undef LOCTEXT_NAMESPACE
 
diff --git a/Engine/Plugins/MovieScene/TemplateSequence/Source/TemplateSequence/Public/CameraAnimationSequenceSubsystem.h b/Engine/Plugins/MovieScene/TemplateSequence/Source/TemplateSequence/Public/CameraAnimationSequenceSubsystem.h
--- a/Engine/Plugins/MovieScene/TemplateSequence/Source/TemplateSequence/Public/CameraAnimationSequenceSubsystem.h
+++ b/Engine/Plugins/MovieScene/TemplateSequence/Source/TemplateSequence/Public/CameraAnimationSequenceSubsystem.h
@@ -4,11 +4,13 @@
 
 #include "CoreMinimal.h"
 #include "EntitySystem/MovieSceneEntityInstantiatorSystem.h"
-#include "EntitySystem/MovieSceneEntitySystemRunner.h"
 #include "Subsystems/WorldSubsystem.h"
 #include "UObject/ObjectMacros.h"
 #include "CameraAnimationSequenceSubsystem.generated.h"
 
+class UMovieSceneEntitySystemLinker;
+class FMovieSceneEntitySystemRunner;
+
 UCLASS()
 class UCameraAnimationSpawnableSystem : public UMovieSceneEntitySystem
 {
@@ -76,8 +78,12 @@ public:
 	/** Gets the Sequencer linker owned by this sybsystem */
 	UMovieSceneEntitySystemLinker* GetLinker(bool bAutoCreate = true);
 
+	/** Get's this subsystem's camera anim runner */
+	TSharedPtr<FMovieSceneEntitySystemRunner> GetRunner() const;
+
 private:
 	// Begin USubsystem
+	virtual void Initialize(FSubsystemCollectionBase& Collection) override;
 	virtual void Deinitialize() override;
 	// End USubsystem
 	
@@ -87,6 +93,6 @@ private:
 	TObjectPtr<UMovieSceneEntitySystemLinker> Linker;
 
 	/** The global Sequencer runner that runs all the shakes and camera animations */
-	FMovieSceneEntitySystemRunner Runner;
+	TSharedPtr<FMovieSceneEntitySystemRunner> Runner;
 };
 
diff --git a/Engine/Source/Editor/MovieSceneTools/Private/MovieSceneConstraintChannelHelper.cpp b/Engine/Source/Editor/MovieSceneTools/Private/MovieSceneConstraintChannelHelper.cpp
--- a/Engine/Source/Editor/MovieSceneTools/Private/MovieSceneConstraintChannelHelper.cpp
+++ b/Engine/Source/Editor/MovieSceneTools/Private/MovieSceneConstraintChannelHelper.cpp
@@ -401,7 +401,7 @@ void FCompensationEvaluator::ComputeCompensation(UWorld* InWorld, const TSharedP
 				BakeHelper->PreEvaluation(MovieScene, InFrame);
 			}
 		}
-		InSequencer->GetEvaluationTemplate().Evaluate(Context0, *InSequencer);
+		InSequencer->GetEvaluationTemplate().EvaluateSynchronousBlocking(Context0, *InSequencer);
 
 		for (const UTickableTransformConstraint* InConstraint : Constraints)
 		{
diff --git a/Engine/Source/Editor/MovieSceneTools/Private/MovieSceneToolHelpers.cpp b/Engine/Source/Editor/MovieSceneTools/Private/MovieSceneToolHelpers.cpp
--- a/Engine/Source/Editor/MovieSceneTools/Private/MovieSceneToolHelpers.cpp
+++ b/Engine/Source/Editor/MovieSceneTools/Private/MovieSceneToolHelpers.cpp
@@ -4257,9 +4257,9 @@ static void GetSequencerActorWorldTransforms(IMovieScenePlayer* Player, UMovieSc
 					FMovieSceneContext Context = FMovieSceneContext(FMovieSceneEvaluationRange(GlobalTime, TickResolution), Player->GetPlaybackStatus()).SetHasJumped(true);
 					if (Index == 0) // similar with baking first time in we need to evaluate twice (think due to double buffering that happens with skel mesh components).
 					{
-						Player->GetEvaluationTemplate().Evaluate(Context, *Player);
+						Player->GetEvaluationTemplate().EvaluateSynchronousBlocking(Context, *Player);
 					}
-					Player->GetEvaluationTemplate().Evaluate(Context, *Player);
+					Player->GetEvaluationTemplate().EvaluateSynchronousBlocking(Context, *Player);
 
 
 					const FConstraintsManagerController& Controller = FConstraintsManagerController::Get(Actor->GetWorld());
diff --git a/Engine/Source/Editor/Sequencer/Private/SSequencer.cpp b/Engine/Source/Editor/Sequencer/Private/SSequencer.cpp
--- a/Engine/Source/Editor/Sequencer/Private/SSequencer.cpp
+++ b/Engine/Source/Editor/Sequencer/Private/SSequencer.cpp
@@ -1929,14 +1929,12 @@ void SSequencer::FillAdvancedMenu(FMenuBuilder& MenuBuilder)
 
 				UMovieSceneSequence* RootSequence = Template.GetRootSequence();
 
-				// Destroy the old template
-				Template.Finish(*SequencerPin);
 				// Set the new emulation mode
 				Template.SetEmulatedNetworkMask(InMode, *SequencerPin);
 				// Since sequencer owns its own compiled data manager, it's ok to override the mask here and reset everything
 				Template.GetCompiledDataManager()->SetEmulatedNetworkMask(InMode);
 				// Reinitialize the template again
-				Template.Initialize(*RootSequence, *SequencerPin, Template.GetCompiledDataManager());
+				Template.Initialize(*RootSequence, *SequencerPin, Template.GetCompiledDataManager(), Template.GetRunner());
 			}
 		};
 		auto IsNetworkModeChecked = [WeakSequencer = SequencerPtr](EMovieSceneServerClientMask InMode)
diff --git a/Engine/Source/Editor/Sequencer/Private/Sequencer.cpp b/Engine/Source/Editor/Sequencer/Private/Sequencer.cpp
--- a/Engine/Source/Editor/Sequencer/Private/Sequencer.cpp
+++ b/Engine/Source/Editor/Sequencer/Private/Sequencer.cpp
@@ -397,7 +397,9 @@ void FSequencer::InitSequencer(const FSequencerInitParams& InitParams, const TSh
 
 	ActiveTemplateIDs.Add(MovieSceneSequenceID::Root);
 	ActiveTemplateStates.Add(true);
-	RootTemplateInstance.Initialize(*InitParams.RootSequence, *this, CompiledDataManager);
+
+	Runner = MakeShared<FMovieSceneEntitySystemRunner>();
+	RootTemplateInstance.Initialize(*InitParams.RootSequence, *this, CompiledDataManager, Runner);
 
 	RootTemplateInstance.EnableGlobalPreAnimatedStateCapture();
 
@@ -582,11 +584,10 @@ FSequencer::FSequencer()
 
 FSequencer::~FSequencer()
 {
-	RootTemplateInstance.Finish(*this);
-
-	if (RootTemplateInstance.GetEntitySystemRunner().IsAttachedToLinker())
+	if (Runner->IsAttachedToLinker())
 	{
-		RootTemplateInstance.GetEntitySystemRunner().Flush();
+		Runner->QueueFinalUpdate(RootTemplateInstance.GetRootInstanceHandle(), UE::MovieScene::ERunnerUpdateFlags::Finish);
+		Runner->Flush();
 	}
 
 	if (GEditor)
@@ -626,11 +627,10 @@ void FSequencer::Close()
 		OldMaxTickRate.Reset();
 	}
 
-	RootTemplateInstance.Finish(*this);
-
-	if (RootTemplateInstance.GetEntitySystemRunner().IsAttachedToLinker())
+	if (Runner->IsAttachedToLinker())
 	{
-		RootTemplateInstance.GetEntitySystemRunner().Flush();
+		Runner->QueueFinalUpdate(RootTemplateInstance.GetRootInstanceHandle(), UE::MovieScene::ERunnerUpdateFlags::Finish);
+		Runner->Flush();
 	}
 
 	RestorePreAnimatedState();
@@ -974,11 +974,10 @@ void FSequencer::ResetToNewRootSequence(UMovieSceneSequence& NewSequence)
 		}
 	}
 
-	RootTemplateInstance.Finish(*this);
-
-	if (RootTemplateInstance.GetEntitySystemRunner().IsAttachedToLinker())
+	if (Runner->IsAttachedToLinker())
 	{
-		RootTemplateInstance.GetEntitySystemRunner().Flush();
+		Runner->QueueFinalUpdate(RootTemplateInstance.GetRootInstanceHandle(), UE::MovieScene::ERunnerUpdateFlags::Finish);
+		Runner->Flush();
 	}
 
 	ActiveTemplateIDs.Reset();
@@ -986,7 +985,7 @@ void FSequencer::ResetToNewRootSequence(UMovieSceneSequence& NewSequence)
 	ActiveTemplateStates.Reset();
 	ActiveTemplateStates.Add(true);
 
-	RootTemplateInstance.Initialize(NewSequence, *this, CompiledDataManager);
+	RootTemplateInstance.Initialize(NewSequence, *this, CompiledDataManager, Runner);
 	RootTemplateInstance.EnableGlobalPreAnimatedStateCapture();
 
 	RootToLocalTransform = FMovieSceneSequenceTransform();
@@ -3126,20 +3125,10 @@ void FSequencer::EvaluateInternal(FMovieSceneEvaluationRange InRange, bool bHasJ
 	FMovieSceneContext Context = FMovieSceneContext(InRange, PlaybackState).SetIsSilent(SilentModeCount != 0);
 	Context.SetHasJumped(bHasJumped);
 
-	FMovieSceneSequenceID RootOverride = MovieSceneSequenceID::Root;
-	if (Settings->ShouldEvaluateSubSequencesInIsolation())
-	{
-		RootOverride = ActiveTemplateIDs.Top();
-	}
-
-	RootTemplateInstance.Evaluate(Context, *this);
+	
+	RootTemplateInstance.EvaluateSynchronousBlocking(Context, *this);
 	SuppressAutoEvalSignature.Reset();
 
-	if (RootTemplateInstance.GetEntitySystemRunner().IsAttachedToLinker())
-	{
-		RootTemplateInstance.GetEntitySystemRunner().Flush();
-	}
-
 	if (Settings->ShouldRerunConstructionScripts())
 	{
 		RerunConstructionScripts();
diff --git a/Engine/Source/Editor/Sequencer/Private/Sequencer.h b/Engine/Source/Editor/Sequencer/Private/Sequencer.h
--- a/Engine/Source/Editor/Sequencer/Private/Sequencer.h
+++ b/Engine/Source/Editor/Sequencer/Private/Sequencer.h
@@ -1398,6 +1398,7 @@ private:
 	TRange<FFrameNumber> SubSequenceRange;
 
 	UMovieSceneCompiledDataManager* CompiledDataManager;
+	TSharedPtr<FMovieSceneEntitySystemRunner> Runner;
 
 	TMap<FName, TFunction<void()>> CleanupFunctions;
 
diff --git a/Engine/Source/Editor/UnrealEd/Private/Fbx/FbxMainExport.cpp b/Engine/Source/Editor/UnrealEd/Private/Fbx/FbxMainExport.cpp
--- a/Engine/Source/Editor/UnrealEd/Private/Fbx/FbxMainExport.cpp
+++ b/Engine/Source/Editor/UnrealEd/Private/Fbx/FbxMainExport.cpp
@@ -1712,7 +1712,7 @@ void FLevelSequenceAnimTrackAdapter::UpdateAnimation( int32 LocalFrame )
 
 	FMovieSceneContext Context = FMovieSceneContext(FMovieSceneEvaluationRange(GlobalTime, TickResolution), MovieScenePlayer->GetPlaybackStatus()).SetHasJumped(true);
 
-	MovieScenePlayer->GetEvaluationTemplate().Evaluate( Context, *MovieScenePlayer );
+	MovieScenePlayer->GetEvaluationTemplate().EvaluateSynchronousBlocking( Context, *MovieScenePlayer );
 }
 
 float FLevelSequenceAnimTrackAdapter::GetFrameRate() const
diff --git a/Engine/Source/Runtime/MovieScene/Private/EntitySystem/MovieSceneEntitySystemLinker.cpp b/Engine/Source/Runtime/MovieScene/Private/EntitySystem/MovieSceneEntitySystemLinker.cpp
--- a/Engine/Source/Runtime/MovieScene/Private/EntitySystem/MovieSceneEntitySystemLinker.cpp
+++ b/Engine/Source/Runtime/MovieScene/Private/EntitySystem/MovieSceneEntitySystemLinker.cpp
@@ -153,6 +153,7 @@ UMovieSceneEntitySystemLinker::UMovieSceneEntitySystemLinker(const FObjectInitia
 	{
 		PreAnimatedState.Initialize(this);
 
+		FCoreUObjectDelegates::GetPreGarbageCollectDelegate().AddUObject(this, &UMovieSceneEntitySystemLinker::HandlePreGarbageCollection);
 		FCoreUObjectDelegates::GetPostGarbageCollect().AddUObject(this, &UMovieSceneEntitySystemLinker::HandlePostGarbageCollection);
 
 		EntityManager.SetDebugName(GetName() + TEXT("[Entity Manager]"));
@@ -315,11 +316,13 @@ bool UMovieSceneEntitySystemLinker::HasStructureChangedSinceLastRun() const
 
 bool UMovieSceneEntitySystemLinker::StartEvaluation(FMovieSceneEntitySystemRunner& InRunner)
 {
-	if (ActiveRunners.Num() == 0 || ActiveRunners.Last().bIsReentrancyAllowed)
+
+	if (ActiveRunners.Num() == 0 || ActiveRunnerReentrancyFlags[ActiveRunners.Num()-1])
 	{
 		// Default to re-entrancy being forbidden. The runner will allow re-entrancy at specific spots
 		// in the evaluation loop, via a "re-entrancy window".
-		ActiveRunners.Emplace(FActiveRunnerInfo{ &InRunner, false });
+		ActiveRunners.Emplace(&InRunner);
+		ActiveRunnerReentrancyFlags.Add(false);
 		return true;
 	}
 		
@@ -331,7 +334,7 @@ FMovieSceneEntitySystemRunner* UMovieSceneEntitySystemLinker::GetActiveRunner()
 {
 	if (ActiveRunners.Num() > 0)
 	{
-		return ActiveRunners.Last().Runner;
+		return ActiveRunners.Last();
 	}
 	return nullptr;
 }
@@ -345,13 +348,21 @@ void UMovieSceneEntitySystemLinker::PostInstantation(FMovieSceneEntitySystemRunn
 
 void UMovieSceneEntitySystemLinker::EndEvaluation(FMovieSceneEntitySystemRunner& InRunner)
 {
-	if (ensureMsgf((ActiveRunners.Num() > 0 && ActiveRunners.Last().Runner == &InRunner),
+	if (ensureMsgf((ActiveRunners.Num() > 0 && ActiveRunners.Last() == &InRunner),
 				TEXT("Trying end the evaluation of a runner that's not the latest one to run.")))
 	{
+		const int32 LastIndex = ActiveRunners.Num()-1;
+		ensureAlways(ActiveRunnerReentrancyFlags[LastIndex] == false);
+
 		ActiveRunners.Pop();
+		ActiveRunnerReentrancyFlags.RemoveAt(LastIndex);
 	}
 }
 
+void UMovieSceneEntitySystemLinker::HandlePreGarbageCollection()
+{
+}
+
 void UMovieSceneEntitySystemLinker::HandlePostGarbageCollection()
 {
 	using namespace UE::MovieScene;
@@ -387,7 +398,7 @@ void UMovieSceneEntitySystemLinker::CleanGarbage()
 	// the next time a runner gets flushed
 	LastInstantiationVersion = 0;
 
-	// Allow any other system to tag garbage
+	// Allow any other system to cleanup garbage
 	Events.CleanTaggedGarbage.Broadcast(this);
 
 	auto RouteCleanTaggedGarbage = [](UMovieSceneEntitySystem* System){ System->CleanTaggedGarbage(); };
@@ -413,11 +424,22 @@ void UMovieSceneEntitySystemLinker::CleanGarbage()
 		EntityManager.AddComponent(EntityID, BuiltInComponents->Tags.HasUnresolvedBinding);
 	}
 
+	// Remove NeedsLink off any NeedsUnLink entities - these tags are incompatible. Any such entities should get cleaned up
+	// through a binding to Events.CleanTaggedGarbage rather than through the usual NeedsUnlink methods
+	FRemoveSingleMutation RemoveNeedsLink(BuiltInComponents->Tags.NeedsLink);
+	EntityManager.MutateAll(FEntityComponentFilter().All({ BuiltInComponents->Tags.NeedsLink, NeedsUnlink }), RemoveNeedsLink);
+
 	// Free the entities
 	TSet<FMovieSceneEntityID> FreedEntities;
 	EntityManager.FreeEntities(FEntityComponentFilter().All({ NeedsUnlink }), &FreedEntities);
 
 	InstanceRegistry->CleanupLinkerEntities(FreedEntities);
+
+	// If we have any runners part-way through an evaluation, we need to reset them so that they re-evaluate from the start
+	for (int32 Index = ActiveRunners.Num()-1; Index >= 0; --Index)
+	{
+		ActiveRunners[Index]->ResetFlushState();
+	}
 }
 
 void UMovieSceneEntitySystemLinker::OnObjectsReplaced(const TMap<UObject*, UObject*>& ReplacementMap)
@@ -563,7 +585,6 @@ void UMovieSceneEntitySystemLinker::AutoLinkRelevantSystems()
 		LinkRelevantSystems();
 	}
 }
-
 void UMovieSceneEntitySystemLinker::AutoUnlinkIrrelevantSystems()
 {
 	if (AutoLinkMode == UE::MovieScene::EAutoLinkRelevantSystems::Enabled)
@@ -571,18 +592,3 @@ void UMovieSceneEntitySystemLinker::AutoUnlinkIrrelevantSystems()
 		UnlinkIrrelevantSystems();
 	}
 }
-
-FMovieSceneEntitySystemEvaluationReentrancyWindow::FMovieSceneEntitySystemEvaluationReentrancyWindow(UMovieSceneEntitySystemLinker& InLinker)
-	: Linker(InLinker)
-{
-	CurrentLevel = Linker.ActiveRunners.Num() - 1;
-	Linker.ActiveRunners[CurrentLevel].bIsReentrancyAllowed = true;
-}
-
-FMovieSceneEntitySystemEvaluationReentrancyWindow::~FMovieSceneEntitySystemEvaluationReentrancyWindow()
-{
-	if (ensure(Linker.ActiveRunners.IsValidIndex(CurrentLevel)))
-	{
-		Linker.ActiveRunners[CurrentLevel].bIsReentrancyAllowed = false;
-	}
-}
diff --git a/Engine/Source/Runtime/MovieScene/Private/EntitySystem/MovieSceneEntitySystemRunner.cpp b/Engine/Source/Runtime/MovieScene/Private/EntitySystem/MovieSceneEntitySystemRunner.cpp
--- a/Engine/Source/Runtime/MovieScene/Private/EntitySystem/MovieSceneEntitySystemRunner.cpp
+++ b/Engine/Source/Runtime/MovieScene/Private/EntitySystem/MovieSceneEntitySystemRunner.cpp
@@ -7,7 +7,7 @@
 #include "Evaluation/PreAnimatedState/MovieScenePreAnimatedCaptureSource.h"
 #include "ProfilingDebugging/CountersTrace.h"
 
-DECLARE_CYCLE_STAT(TEXT("ECS System Cost"), 			MovieSceneEval_TotalGTCost, 				STATGROUP_MovieSceneEval);
+DECLARE_CYCLE_STAT(TEXT("Runner Flush"), 				MovieSceneEval_RunnerFlush, 				STATGROUP_MovieSceneEval);
 
 DECLARE_CYCLE_STAT(TEXT("Spawn Phase"),                 MovieSceneEval_SpawnPhase,              	STATGROUP_MovieSceneECS);
 DECLARE_CYCLE_STAT(TEXT("Instantiation Phase"), 		MovieSceneEval_InstantiationPhase, 			STATGROUP_MovieSceneECS);
@@ -18,13 +18,55 @@ DECLARE_CYCLE_STAT(TEXT("Evaluation Phase"), 			MovieSceneEval_EvaluationPhase,
 DECLARE_CYCLE_STAT(TEXT("Finalization Phase"),          MovieSceneEval_FinalizationPhase,       	STATGROUP_MovieSceneECS);
 DECLARE_CYCLE_STAT(TEXT("Post Evaluation Phase"),       MovieSceneEval_PostEvaluationPhase,     	STATGROUP_MovieSceneECS);
 
-TRACE_DECLARE_INT_COUNTER(MovieSceneEntitySystemFlushes, TEXT("MovieScene/ECSFlushes"));
-TRACE_DECLARE_INT_COUNTER(MovieSceneEntitySystemEvaluations, TEXT("MovieScene/ECSEvaluations"));
+
+/**
+ * Structure for making it possible to make re-entrant evaluation on a linker.
+ */
+struct FMovieSceneEntitySystemEvaluationReentrancyWindow
+{
+	FMovieSceneEntitySystemRunner* Runner;
+	UMovieSceneEntitySystemLinker* Linker;
+	int32 RunnerIndex;
+	UE::MovieScene::ERunnerFlushState CachedFlushState;
+
+	FMovieSceneEntitySystemEvaluationReentrancyWindow(FMovieSceneEntitySystemRunner* InRunner, UMovieSceneEntitySystemLinker* InLinker)
+		: Runner(InRunner)
+		, Linker(InLinker)
+		, RunnerIndex(Linker->ActiveRunners.Num() - 1)
+		, CachedFlushState(Runner->FlushState)
+	{
+		check(RunnerIndex >= 0);
+		checkf(Linker->ActiveRunnerReentrancyFlags[RunnerIndex] == false, TEXT("Nested FMovieSceneEntitySystemEvaluationReentrancyWindows are not supported for the same active runner"));
+
+		Linker->ActiveRunnerReentrancyFlags[RunnerIndex] = true;
+
+		// Clear the current flush state to prevent inner scopes from flushing
+		//      our pending states and set the flag that allows us to be re-entrant
+		Runner->FlushState = UE::MovieScene::ERunnerFlushState::None;
+	}
+
+	~FMovieSceneEntitySystemEvaluationReentrancyWindow()
+	{
+		// If a re-entrant call left some evaluation still outstanding, we have to flush that before we can continue
+		if (Runner->FlushState != UE::MovieScene::ERunnerFlushState::None)
+		{
+			Runner->FlushOutstanding();
+		}
+
+		Runner->FlushState = CachedFlushState;
+		if (ensure(Linker->ActiveRunnerReentrancyFlags.IsValidIndex(RunnerIndex)))
+		{
+			Linker->ActiveRunnerReentrancyFlags[RunnerIndex] = false;
+		}
+	}
+};
+
 
 FMovieSceneEntitySystemRunner::FMovieSceneEntitySystemRunner()
-	: CompletionTask(nullptr)
-	, GameThread(ENamedThreads::GameThread_Local)
+	: GameThread(ENamedThreads::GameThread_Local)
 	, CurrentPhase(UE::MovieScene::ESystemPhase::None)
+	, FlushState(UE::MovieScene::ERunnerFlushState::None)
+	, bRequireFullFlush(false)
 {
 }
 
@@ -93,6 +135,16 @@ UE::MovieScene::FInstanceRegistry* FMovieSceneEntitySystemRunner::GetInstanceReg
 	return nullptr;
 }
 
+bool FMovieSceneEntitySystemRunner::IsCurrentlyEvaluating() const
+{
+	return FlushState != UE::MovieScene::ERunnerFlushState::None;
+}
+
+int32 FMovieSceneEntitySystemRunner::GetQueuedUpdateCount() const
+{
+	return UpdateQueue.Num() + DissectedUpdates.Num();
+}
+
 bool FMovieSceneEntitySystemRunner::HasQueuedUpdates() const
 {
 	if (UpdateQueue.Num() != 0 || DissectedUpdates.Num() != 0)
@@ -108,49 +160,120 @@ bool FMovieSceneEntitySystemRunner::HasQueuedUpdates() const
 
 bool FMovieSceneEntitySystemRunner::HasQueuedUpdates(FInstanceHandle InInstanceHandle) const
 {
-	return Algo::FindBy(UpdateQueue, InInstanceHandle, &FMovieSceneUpdateRequest::InstanceHandle) != nullptr ||
+	using namespace UE::MovieScene;
+
+	return Algo::FindBy(UpdateQueue, InInstanceHandle, [](const FUpdateParamsAndContext& In){ return In.Params.InstanceHandle; }) != nullptr ||
 		Algo::FindBy(DissectedUpdates, InInstanceHandle, &FDissectedUpdate::InstanceHandle) != nullptr;
 }
 
-void FMovieSceneEntitySystemRunner::QueueUpdate(const FMovieSceneContext& Context, FInstanceHandle InInstanceHandle)
+void FMovieSceneEntitySystemRunner::QueueUpdate(const FMovieSceneContext& Context, FInstanceHandle Instance, UE::MovieScene::ERunnerUpdateFlags Flags)
 {
-	UpdateQueue.Add(FMovieSceneUpdateRequest{ Context, InInstanceHandle });
+	if (EnumHasAnyFlags(Flags, UE::MovieScene::ERunnerUpdateFlags::Flush))
+	{
+		bRequireFullFlush = true;
+	}
+	UpdateQueue.Add(FUpdateParamsAndContext{ FSimpleDelegate(), Context, { Instance, Flags } });
 }
 
-void FMovieSceneEntitySystemRunner::Update(const FMovieSceneContext& Context, FInstanceHandle Instance)
+void FMovieSceneEntitySystemRunner::QueueUpdate(const FMovieSceneContext& Context, FInstanceHandle Instance, FSimpleDelegate&& OnFlushedDelegate, UE::MovieScene::ERunnerUpdateFlags Flags)
 {
-	if (UpdateQueue.Num() > 0)
+	if (EnumHasAnyFlags(Flags, UE::MovieScene::ERunnerUpdateFlags::Flush))
 	{
-		UE_LOG(LogMovieSceneECS, Warning, TEXT("Updates are already queued! This will run those updates as well, which might not be what's intended."));
+		bRequireFullFlush = true;
 	}
 
-	// Queue our one update and flush immediately.
-	QueueUpdate(Context, Instance);
-	Flush();
+	UpdateQueue.Add(FUpdateParamsAndContext{ OnFlushedDelegate, Context, { Instance, Flags } });
 }
 
-void FMovieSceneEntitySystemRunner::Flush()
+bool FMovieSceneEntitySystemRunner::QueueFinalUpdate(FInstanceHandle InInstanceHandle, UE::MovieScene::ERunnerUpdateFlags Flags)
 {
 	using namespace UE::MovieScene;
 
+	// It is not well defined to call this function with no flags - that would signify the function should not have been called at all
+	check(Flags != ERunnerUpdateFlags::None);
+
 	UMovieSceneEntitySystemLinker* Linker = GetLinker();
+	if (!Linker)
+	{
+		return false;
+	}
 
-	// Check that we are attached to a linker that allows starting a new evaluation.
-	if (!ensureMsgf(Linker, TEXT("Runner isn't attached to a valid linker")))
+	FInstanceRegistry* InstanceRegistry = GetInstanceRegistry();
+	if (!InstanceRegistry->IsHandleValid(InInstanceHandle))
 	{
-		return;
+		return false;
 	}
 
-	if (!Linker->StartEvaluation(*this))
+	FSequenceInstance& Instance = InstanceRegistry->MutateInstance(InInstanceHandle);
+
+	if (EnumHasAnyFlags(Flags, UE::MovieScene::ERunnerUpdateFlags::Flush))
 	{
-		return;
+		bRequireFullFlush = true;
 	}
 
-	SCOPE_CYCLE_COUNTER(MovieSceneEval_TotalGTCost);
-	TRACE_COUNTER_INCREMENT(MovieSceneEntitySystemFlushes);
+	ERunnerFlushState UnsafeDestroyMask = ERunnerFlushState::Everything & ~(ERunnerFlushState::PostEvaluation | ERunnerFlushState::End);
 
-	// We need to run the system from the game thread so we know we can fire events and callbacks from here.
-	check(IsInGameThread());
+	// If this instance has not current updates and we're not in the middle of an update loop
+	// we can finish and destroy this instance right away. The effects of the instance being torn down
+	// (ie entities being unlinked) will be seen once the runner is flushed
+	const bool bSafeToDestroyNow = !EnumHasAnyFlags(FlushState, UnsafeDestroyMask);
+	if (bSafeToDestroyNow && !HasQueuedUpdates(InInstanceHandle))
+	{
+		if (EnumHasAnyFlags(Flags, ERunnerUpdateFlags::Finish))
+		{
+			Instance.Finish(Linker);
+		}
+		Instance.PostEvaluation(Linker);
+
+		// PostEvaluation could have prompted a bunch of other logic that may have destroyed our instance handle so we have to check it for validity again
+		if (EnumHasAnyFlags(Flags, ERunnerUpdateFlags::Destroy) && InstanceRegistry->IsHandleValid(InInstanceHandle))
+		{
+			InstanceRegistry->DestroyInstance(InInstanceHandle);
+		}
+
+		return false;
+	}
+	else
+	{
+		FUpdateParamsAndContext Update;
+		Update.Params.InstanceHandle = InInstanceHandle;
+		Update.Params.UpdateFlags = Flags;
+		UpdateQueue.Add(Update);
+
+		return true;
+	}
+}
+
+void FMovieSceneEntitySystemRunner::AbandonAndDestroyInstance(FInstanceHandle Instance)
+{
+	if (QueueFinalUpdate(Instance, UE::MovieScene::ERunnerUpdateFlags::Finish | UE::MovieScene::ERunnerUpdateFlags::Destroy))
+	{
+		Flush();
+	}
+}
+
+bool FMovieSceneEntitySystemRunner::StartEvaluation(UMovieSceneEntitySystemLinker* Linker)
+{
+	using namespace UE::MovieScene;
+
+	// We specifically only check whether the entity manager has changed since the last instantiation once
+	// to ensure that we are not vulnerable to infinite loops where components are added/removed in post-evaluation
+	const bool bStructureHadChanged = Linker->HasStructureChangedSinceLastRun();
+
+	if (!bStructureHadChanged && UpdateQueue.Num() == 0 && DissectedUpdates.Num() == 0)
+	{
+		// If nothing has changed, and we have no pending updates, we don't run any evaluation
+		FlushState = ERunnerFlushState::None;
+		return false;
+	}
+
+	if (!Linker->StartEvaluation(*this))
+	{
+		return false;
+	}
+
+	// Don't run start again until explicitly told to
+	FlushState &= ~ERunnerFlushState::Start;
 
 	// Our entity manager cannot be locked down for us to continue. Something must have left it locked if 
 	// this check fails
@@ -160,33 +283,137 @@ void FMovieSceneEntitySystemRunner::Flush()
 	EntityManager.SetDispatchThread(ENamedThreads::GameThread_Local);
 	EntityManager.SetGatherThread(ENamedThreads::GameThread_Local);
 
-	// We specifically only check whether the entity manager has changed since the last instantation once
-	// to ensure that we are not vulnerable to infinite loops where components are added/removed in post-evaluation
-	bool bStructureHadChanged = Linker->HasStructureChangedSinceLastRun();
+	return true;
+}
+
+void FMovieSceneEntitySystemRunner::EndEvaluation(UMovieSceneEntitySystemLinker* Linker)
+{
+	using namespace UE::MovieScene;
+
+	Linker->EndEvaluation(*this);
+	FlushState &= ~ERunnerFlushState::End;
+}
+
+bool FMovieSceneEntitySystemRunner::FlushNext(UMovieSceneEntitySystemLinker* Linker)
+{
+	using namespace UE::MovieScene;
+
+	// The evaluation loop is set up as a crude state machine which allows us to break out of the
+	// loop and pick up where we left off next frame. Each bit within the FlushState bitmask defines
+	// a particular task that needs to execute. In this way we can loop the evaluation while there
+	// are updates pending by simply setting the necessary flags for each task. For convenience, the
+	// required tasks to perform another iteration are defined in ERunnerFlushState::LoopEval
+
+	// Guard the value of the CurrentPhase so it always returns back to ::None
+	TGuardValue<ESystemPhase> CurrentPhaseGuard(CurrentPhase, ESystemPhase::None);
 
-	// Start flushing the update queue... keep flushing as long as we have work to do.
-	while (UpdateQueue.Num() > 0 ||
-		DissectedUpdates.Num() > 0 ||
-		bStructureHadChanged)
+	// Step 1: Initialize all external flags and broadcast 'begin eval' events.
+	//         NOTE: This is only ever called once, regardless of how many iterations we perform
+	if (EnumHasAnyFlags(FlushState, ERunnerFlushState::Start))
 	{
-		DoFlushUpdateQueueOnce();
+		return StartEvaluation(Linker);
+	}
 
-		bStructureHadChanged = false;
+	// Step 2: Update sequence instances and import entities
+	//         This is the entry-point for for each iteration
+	if (EnumHasAnyFlags(FlushState, ERunnerFlushState::Import))
+	{
+		return GameThread_UpdateSequenceInstances(Linker);
 	}
 
-	Linker->EndEvaluation(*this);
+	// Step 3: Conditionally run the spawn phase of the system graph
+	//
+	if (EnumHasAnyFlags(FlushState, ERunnerFlushState::Spawn))
+	{
+		return GameThread_SpawnPhase(Linker);
+	}
+
+	// Step 4: Conditionally run the instantiation phase of the system graph
+	//
+	if (EnumHasAnyFlags(FlushState, ERunnerFlushState::Instantiation))
+	{
+		return GameThread_InstantiationPhase(Linker);
+	}
+
+	// Step 5: Run the evaluation phase of the system graph
+	//
+	if (EnumHasAnyFlags(FlushState, ERunnerFlushState::Evaluation))
+	{
+		return GameThread_EvaluationPhase(Linker);
+	}
+
+	// Step 6: Run the finalization phase of the system graph, including legacy templates
+	//
+	if (EnumHasAnyFlags(FlushState, ERunnerFlushState::Finalization))
+	{
+		GameThread_EvaluationFinalizationPhase(Linker);
+		return true;
+	}
+
+	// Step 7: Trigger events
+	//
+	if (EnumHasAnyFlags(FlushState, ERunnerFlushState::EventTriggers))
+	{
+		GameThread_EventTriggerPhase(Linker);
+		return true;
+	}
+
+	// Step 7: Call PostEvaluation on all current sequence instances
+	//
+	if (EnumHasAnyFlags(FlushState, ERunnerFlushState::PostEvaluation))
+	{
+		GameThread_PostEvaluationPhase(Linker);
+		return true;
+	}
+
+	// Step 7: Perform any clean up and broadcast 'end eval' events
+	//         NOTE: Only ever called once regardless of how many iterations we perform
+	if (EnumHasAnyFlags(FlushState, ERunnerFlushState::End))
+	{
+		EndEvaluation(Linker);
+		return true;
+	}
+
+	return false;
 }
 
-void FMovieSceneEntitySystemRunner::DoFlushUpdateQueueOnce()
+void FMovieSceneEntitySystemRunner::Flush(double BudgetMs)
 {
 	using namespace UE::MovieScene;
 
-	TRACE_COUNTER_INCREMENT(MovieSceneEntitySystemEvaluations);
-	TRACE_CPUPROFILER_EVENT_SCOPE(FMovieSceneEntitySystemRunner::DoFlushUpdateQueueOnce);
+	// If we're not currently evaluating, start by flushing everything
+	if (FlushState == ERunnerFlushState::None)
+	{
+		FlushState = ERunnerFlushState::Everything;
+	}
 
-	// Setup the completion task that we can wait on.
-	CompletionTask = TGraphTask<FNullGraphTask>::CreateTask(nullptr, ENamedThreads::GameThread)
-		.ConstructAndHold(TStatId(), ENamedThreads::GameThread_Local);
+	FlushOutstanding(BudgetMs);
+}
+
+void FMovieSceneEntitySystemRunner::FlushOutstanding(double BudgetMs, UE::MovieScene::ERunnerFlushState TargetState)
+{
+	using namespace UE::MovieScene;
+
+	// Simple case is we have nothing outstanding, so just early return
+	if (FlushState == ERunnerFlushState::None)
+	{
+		return;
+	}
+
+	UMovieSceneEntitySystemLinker* Linker = GetLinker();
+
+	// Check that we are attached to a linker that allows starting a new evaluation.
+	if (!ensureMsgf(Linker, TEXT("Runner isn't attached to a valid linker")))
+	{
+		return;
+	}
+
+	SCOPE_CYCLE_COUNTER(MovieSceneEval_RunnerFlush);
+
+	ensureMsgf(CurrentPhase == ESystemPhase::None, TEXT("Cannot flush this runner while it is already inside a system phase"));
+
+	// We need to run the system from the game thread so we know we can fire events and callbacks from here.
+	check(IsInGameThread());
 
 	// Set the debug vizualizer's entity manager pointer, so all debugging happening here will show
 	// relevant information. We need to set it here instead of higher up because we could have, say,
@@ -197,80 +424,143 @@ void FMovieSceneEntitySystemRunner::DoFlushUpdateQueueOnce()
 	// again.
 	TGuardValue<FEntityManager*> DebugVizGuard(GEntityManagerForDebuggingVisualizers, GetEntityManager());
 
-	// Also reset the capture source scope so that each group of sequences tied to a given linker starts
-	// with a clean slate.
-	TGuardValue<FScopedPreAnimatedCaptureSource*> CaptureSourceGuard(FScopedPreAnimatedCaptureSource::GetCaptureSourcePtr(), nullptr);
+	// For the purposes of this function, None means that everything must be flushed
+	if (TargetState == ERunnerFlushState::None)
+	{
+		TargetState = ERunnerFlushState::Everything;
+	}
 
-	// Entry point to the whole ECS loop... this will either unroll in the current thread's call stack
-	// if there's not much to do, or it will start queuing up tasks on the task graph.
-	// We immediately wait for the completion task to be executed.
+	const double BudgetSeconds = BudgetMs / 1000.f;
+	if (!bRequireFullFlush && BudgetSeconds > 0.0)
+	{
+		double StartTime = FPlatformTime::Seconds();
+		do
+		{
+			if (!FlushNext(Linker))
+			{
+				break;
+			}
+		}
+		while (FPlatformTime::Seconds() - StartTime < BudgetSeconds && EnumHasAnyFlags(FlushState, TargetState));
+	}
+	else
 	{
-		GameThread_ProcessQueue();
+		while (EnumHasAnyFlags(FlushState, TargetState))
+		{
+			if (!FlushNext(Linker))
+			{
+				break;
+			}
+		}
 	}
 
-	FTaskGraphInterface::Get().WaitUntilTaskCompletes(CompletionTask->GetCompletionEvent(), ENamedThreads::GameThread_Local);
+	bRequireFullFlush = false;
+}
+
+void FMovieSceneEntitySystemRunner::ResetFlushState()
+{
+	using namespace UE::MovieScene;
 
-	CompletionTask = nullptr;
+	constexpr ERunnerFlushState StatesThatTriggerReset = ERunnerFlushState::Instantiation | ERunnerFlushState::Evaluation | ERunnerFlushState::Finalization | ERunnerFlushState::EventTriggers;
+	constexpr ERunnerFlushState StatesThatDontTriggerReset = ERunnerFlushState::Import | ERunnerFlushState::Spawn;
 
-	// Now run the post-evaluation logic, which contains stuff we don't want to run from inside a task
-	// graph call.
-	GameThread_EvaluationFinalizationPhase();
-	GameThread_PostEvaluationPhase();
+	// We only need to reset the flush state if we are far enough through an evaluation
+	if (EnumHasAnyFlags(FlushState, StatesThatTriggerReset) && !EnumHasAnyFlags(FlushState, StatesThatDontTriggerReset))
+	{
+		// When resetting - we don't need to (or want to) re-import anything, we just want to re-run the current
+		// frame of updates 
+		FlushState = UE::MovieScene::ERunnerFlushState::LoopEval & ~UE::MovieScene::ERunnerFlushState::Import;
+	}
 }
 
-void FMovieSceneEntitySystemRunner::GameThread_ProcessQueue()
+bool FMovieSceneEntitySystemRunner::GameThread_UpdateSequenceInstances(UMovieSceneEntitySystemLinker* Linker)
 {
 	using namespace UE::MovieScene;
 
-	UMovieSceneEntitySystemLinker* Linker = GetLinker();
-	check(Linker);
+	TRACE_CPUPROFILER_EVENT_SCOPE(FMovieSceneEntitySystemRunner::GameThread_UpdateSequenceInstances);
+
+	// Also reset the capture source scope so that each group of sequences tied to a given linker starts
+	// with a clean slate.
+	TGuardValue<FScopedPreAnimatedCaptureSource*> CaptureSourceGuard(FScopedPreAnimatedCaptureSource::GetCaptureSourcePtr(), nullptr);
+
+	// Allow the next step to run
+	FlushState &= ~ERunnerFlushState::Import;
 
 	FInstanceRegistry* InstanceRegistry = GetInstanceRegistry();
 
 	if (DissectedUpdates.Num() == 0)
 	{
 		TArray<TRange<FFrameTime>> Dissections;
-		TSet<FInstanceHandle, DefaultKeyFuncs<FInstanceHandle>, TInlineSetAllocator<256>> Updates;
+		TBitArray<> UpdatedSequenceInstances;
 
-		TArray<FMovieSceneUpdateRequest> TempUpdateQueue;
+		TArray<FUpdateParamsAndContext> TempUpdateQueue;
 		Swap(TempUpdateQueue, UpdateQueue);
+
+		DissectedUpdates.Reserve(TempUpdateQueue.Num());
+
 		for (int32 UpdateIndex = 0; UpdateIndex < TempUpdateQueue.Num(); ++UpdateIndex)
 		{
-			FMovieSceneUpdateRequest Request = TempUpdateQueue[UpdateIndex];
-			if (!InstanceRegistry->IsHandleValid(Request.InstanceHandle))
+			FUpdateParamsAndContext& Request = TempUpdateQueue[UpdateIndex];
+			const int32 InstanceID = Request.Params.InstanceHandle.InstanceID;
+
+			if (!InstanceRegistry->IsHandleValid(Request.Params.InstanceHandle))
 			{
 				continue;
 			}
 			// Already have an update for this 
-			else if (Updates.Contains(Request.InstanceHandle))
+			else if (UpdatedSequenceInstances.IsValidIndex(InstanceID) && UpdatedSequenceInstances[InstanceID] == true)
 			{
 				UpdateQueue.Add(Request);
 				continue;
 			}
 
-			Updates.Add(Request.InstanceHandle);
+			// Set the bit to indicate this sequence is being updated
+			UpdatedSequenceInstances.PadToNum(InstanceID + 1, false);
+			UpdatedSequenceInstances[InstanceID] = true;
+
+			// If the request is to destroy the instance, we must assume it is no longer valid and should
+			// not be updated
+			if (EnumHasAnyFlags(Request.Params.UpdateFlags, ERunnerUpdateFlags::Destroy))
+			{
+				DissectedUpdates.Add(FDissectedUpdate{ MoveTemp(Request.OnFlushed), Request.Context, Request.Params.InstanceHandle, MAX_int32, Request.Params.UpdateFlags });
+				MarkForUpdate(Request.Params.InstanceHandle, Request.Params.UpdateFlags);
+				continue;
+			}
 
 			// Give the instance an opportunity to dissect the range into distinct evaluations
-			FSequenceInstance& Instance = InstanceRegistry->MutateInstance(Request.InstanceHandle);
+			FSequenceInstance& Instance = InstanceRegistry->MutateInstance(Request.Params.InstanceHandle);
 			Instance.DissectContext(Linker, Request.Context, Dissections);
 
 			if (Dissections.Num() != 0)
 			{
 				for (int32 Index = 0; Index < Dissections.Num() - 1; ++Index)
 				{
+					// Never finish or destroy sequence instances until the _last_ dissected update
+					ERunnerUpdateFlags Flags = Request.Params.UpdateFlags & ~ERunnerUpdateFlags::FinalDissectionMask;
+
 					FDissectedUpdate Dissection{
+						FSimpleDelegate(), // Never trigger the OnFlushed delegate until the final dissection
 						FMovieSceneContext(FMovieSceneEvaluationRange(Dissections[Index], Request.Context.GetFrameRate(), Request.Context.GetDirection()), Request.Context.GetStatus()),
-						Request.InstanceHandle,
-						Index
+						Request.Params.InstanceHandle,
+						Index,
+						Flags
 					};
 					DissectedUpdates.Add(Dissection);
+
+					// Only mark the first dissection for update this round
+					if (Index == 0)
+					{
+						MarkForUpdate(Request.Params.InstanceHandle, Flags);
+					}
 				}
 
 				// Add the last one with MAX_int32 so it gets evaluated with all the others in this flush
 				FDissectedUpdate Dissection{
+					MoveTemp(Request.OnFlushed),
 					FMovieSceneContext(FMovieSceneEvaluationRange(Dissections.Last(), Request.Context.GetFrameRate(), Request.Context.GetDirection()), Request.Context.GetStatus()),
-					Request.InstanceHandle,
-					MAX_int32
+					Request.Params.InstanceHandle,
+					MAX_int32,
+					Request.Params.UpdateFlags
 				};
 				DissectedUpdates.Add(Dissection);
 
@@ -278,10 +568,9 @@ void FMovieSceneEntitySystemRunner::GameThread_ProcessQueue()
 			}
 			else
 			{
-				DissectedUpdates.Add(FDissectedUpdate{ Request.Context, Request.InstanceHandle, MAX_int32 });
+				DissectedUpdates.Add(FDissectedUpdate{ MoveTemp(Request.OnFlushed), Request.Context, Request.Params.InstanceHandle, MAX_int32, Request.Params.UpdateFlags });
+				MarkForUpdate(Request.Params.InstanceHandle, Request.Params.UpdateFlags);
 			}
-
-			MarkForUpdate(Request.InstanceHandle);
 		}
 	}
 	else
@@ -290,48 +579,36 @@ void FMovieSceneEntitySystemRunner::GameThread_ProcessQueue()
 		const int32 PredicateOrder = DissectedUpdates[0].Order;
 		for (int32 Index = 0; Index < DissectedUpdates.Num() && DissectedUpdates[Index].Order == PredicateOrder; ++Index)
 		{
-			const FDissectedUpdate& Update = DissectedUpdates[Index];
-			MarkForUpdate(Update.InstanceHandle);
+			FDissectedUpdate& Update = DissectedUpdates[Index];
+			MarkForUpdate(Update.InstanceHandle, Update.UpdateFlags);
 		}
 	}
 
 	// If we have no instances marked for update, we are running an evaluation probably because some
-	// structural changes have occured in the entity manager (out of date instantiation serial number
+	// structural changes have occurred in the entity manager (out of date instantiation serial number
 	// in the linker). So we mark everything for update, so that PreEvaluation/PostEvaluation callbacks
 	// and legacy templates are correctly executed.
 	if (CurrentInstances.Num() == 0)
 	{
 		for (const FSequenceInstance& Instance : InstanceRegistry->GetSparseInstances())
 		{
-			MarkForUpdate(Instance.GetInstanceHandle());
+			MarkForUpdate(Instance.GetInstanceHandle(), ERunnerUpdateFlags::None);
 		}
 	}
 
 	// Let sequence instances do any pre-evaluation work.
-	for (FInstanceHandle UpdatedInstanceHandle : CurrentInstances)
+	for (const FQueuedUpdateParams& UpdatedInstance : CurrentInstances)
 	{
-		InstanceRegistry->MutateInstance(UpdatedInstanceHandle).PreEvaluation(Linker);
+		if (!EnumHasAnyFlags(UpdatedInstance.UpdateFlags, ERunnerUpdateFlags::Destroy))
+		{
+			InstanceRegistry->MutateInstance(UpdatedInstance.InstanceHandle).PreEvaluation(Linker);
+		}
 	}
 
-	// Process updates
-	GameThread_SpawnPhase();
-}
-
-void FMovieSceneEntitySystemRunner::GameThread_SpawnPhase()
-{
-	using namespace UE::MovieScene;
-
-	check(GameThread == ENamedThreads::GameThread || GameThread == ENamedThreads::GameThread_Local);
-
-	UMovieSceneEntitySystemLinker* Linker = GetLinker();
-	check(Linker);
-
+	// NOTE: IncrementSystemSerial must be called before any instance updates are made
+	//       to ensure that up-to-date versions are used inside FEntityManager::OnStructureChanged
 	Linker->EntityManager.IncrementSystemSerial();
 
-	CurrentPhase = ESystemPhase::Spawn;
-
-	FInstanceRegistry* InstanceRegistry = GetInstanceRegistry();
-
 	// Update all systems
 	if (DissectedUpdates.Num() != 0)
 	{
@@ -339,17 +616,70 @@ void FMovieSceneEntitySystemRunner::GameThread_SpawnPhase()
 		int32 Index = 0;
 		for (; Index < DissectedUpdates.Num() && DissectedUpdates[Index].Order == PredicateOrder; ++Index)
 		{
-			FDissectedUpdate Update = DissectedUpdates[Index];
+			FDissectedUpdate& Update = DissectedUpdates[Index];
+
+			// Always forward the OnFlushed delegate to be called at the end of the frame, even if the instance is no longer valid
+			if (Update.OnFlushed.IsBound())
+			{
+				OnFlushedDelegates.Add(MoveTemp(Update.OnFlushed));
+			}
+
+			if (EnumHasAnyFlags(Update.UpdateFlags, ERunnerUpdateFlags::Destroy))
+			{
+				continue;
+			}
+
 			if (ensure(InstanceRegistry->IsHandleValid(Update.InstanceHandle)))
 			{
 				FSequenceInstance& Instance = InstanceRegistry->MutateInstance(Update.InstanceHandle);
 
-				Instance.Update(Linker, Update.Context);
+				if (EnumHasAnyFlags(Update.UpdateFlags, ERunnerUpdateFlags::Finish))
+				{
+					// Context is irrelevant for Finishing sequences
+					Instance.Finish(Linker);
+				}
+				else
+				{
+					Instance.Update(Linker, Update.Context);
+				}
 			}
 		}
 		DissectedUpdates.RemoveAt(0, Index);
 	}
 
+#if DO_GUARD_SLOW
+	// Check that there are no duplicates in the CurrentInstances array
+	{
+		TBitArray<> VisitedBits;
+		VisitedBits.Reserve(InstanceRegistry->GetSparseInstances().GetMaxIndex());
+		for (const FQueuedUpdateParams& Update : CurrentInstances)
+		{
+			const int32 Index = Update.InstanceHandle.InstanceID;
+
+			checkf(!VisitedBits.IsValidIndex(Index) || VisitedBits[Index] == false, TEXT("Multiple updates exist for the same sequence instance. This indicates a bookkeeping error with the CurrentInstances array."));
+
+			VisitedBits.PadToNum(Index + 1, false);
+			VisitedBits[Index] = true;
+		}
+	}
+#endif
+
+	return true;
+}
+
+bool FMovieSceneEntitySystemRunner::GameThread_SpawnPhase(UMovieSceneEntitySystemLinker* Linker)
+{
+	using namespace UE::MovieScene;
+
+	check(GameThread == ENamedThreads::GameThread || GameThread == ENamedThreads::GameThread_Local);
+
+	// Allow the next step to run
+	FlushState &= ~ERunnerFlushState::Spawn;
+
+	CurrentPhase = ESystemPhase::Spawn;
+
+	FInstanceRegistry* InstanceRegistry = GetInstanceRegistry();
+
 	const bool bInstantiationDirty = Linker->HasStructureChangedSinceLastRun() || InstanceRegistry->HasInvalidatedBindings();
 
 	FGraphEventArray AllTasks;
@@ -357,24 +687,34 @@ void FMovieSceneEntitySystemRunner::GameThread_SpawnPhase()
 	Linker->AutoLinkRelevantSystems();
 
 	// --------------------------------------------------------------------------------------------------------------------------------------------
-	// Step 1: Run the spawn phase if there were any changes to the current entity instantiations
-	if (bInstantiationDirty)
+	// Run the spawn phase if there were any changes to the current entity instantiations, and wait on the result
 	{
 		SCOPE_CYCLE_COUNTER(MovieSceneEval_SpawnPhase);
-
-		// The spawn phase can queue events to trigger from the event tracks.
-		bCanQueueEventTriggers = true;
+		if (bInstantiationDirty)
 		{
-			Linker->SystemGraph.ExecutePhase(ESystemPhase::Spawn, Linker, AllTasks);
+			// The spawn phase can queue events to trigger from the event tracks.
+			bCanQueueEventTriggers = true;
+			{
+				Linker->SystemGraph.ExecutePhase(ESystemPhase::Spawn, Linker, AllTasks);
+			}
+			bCanQueueEventTriggers = false;
+
+			// We don't open a re-entrancy window, however, because there's no way we can recursively evaluate things at this point... too many
+			// things are in an intermediate state. So events triggered as PreSpawn/PostSpawn can't be wired to something that starts a sequence.
+			if (EventTriggers.IsBound())
+			{
+				EventTriggers.Broadcast();
+				EventTriggers.Clear();
+			}
 		}
-		bCanQueueEventTriggers = false;
 
-		// We don't open a re-entrancy window, however, because there's no way we can recursively evaluate things at this point... too many
-		// things are in an intermediate state. So events triggered as PreSpawn/PostSpawn can't be wired to something that starts a sequence.
-		if (EventTriggers.IsBound())
+		// If there were any tasks created, we need to wait on them before proceeding. This is rare for the spawn phase.
+		if (AllTasks.Num() != 0)
 		{
-			EventTriggers.Broadcast();
-			EventTriggers.Clear();
+			FGraphEventRef SpawnEvent = TGraphTask<FNullGraphTask>::CreateTask(&AllTasks, ENamedThreads::GameThread)
+			.ConstructAndDispatchWhenReady(TStatId(), GameThread);
+
+			FTaskGraphInterface::Get().WaitUntilTaskCompletes(SpawnEvent, ENamedThreads::GameThread_Local);
 		}
 	}
 
@@ -384,160 +724,117 @@ void FMovieSceneEntitySystemRunner::GameThread_SpawnPhase()
 	// --------------------------------------------------------------------------------------------------------------------------------------------
 	// Only run the instantiation phase if there is anything to instantiate. This must come after the spawn phase because new instantiations may
 	// be created during the spawn phase
-	auto NextStep = [this, bInstantiationDirty, Linker]
-	{
-		const bool bAnyPending = Linker->EntityManager.ContainsAnyComponent(FBuiltInComponentTypes::Get()->RequiresInstantiationMask) ||  Linker->GetInstanceRegistry()->HasInvalidatedBindings();
-		if (bInstantiationDirty || bAnyPending)
-		{
-			this->GameThread_InstantiationPhase();
-		}
-		else
-		{
-			// Go straight to evaluation
-			this->GameThread_EvaluationPhase();
-		}
-	};
-
-	if (AllTasks.Num() != 0)
-	{
-		TGraphTask<TFunctionGraphTaskImpl<void(), ESubsequentsMode::TrackSubsequents>>::CreateTask(&AllTasks, ENamedThreads::GameThread)
-			.ConstructAndDispatchWhenReady(MoveTemp(NextStep), TStatId(), GameThread);
-	}
-	else
+	const bool bAnyPending = Linker->EntityManager.ContainsAnyComponent(FBuiltInComponentTypes::Get()->RequiresInstantiationMask) || InstanceRegistry->HasInvalidatedBindings();
+	if (bInstantiationDirty == false && bAnyPending == false)
 	{
-		NextStep();
+		FlushState &= ~ERunnerFlushState::Instantiation;
 	}
+
+	return true;
 }
 
-void FMovieSceneEntitySystemRunner::GameThread_InstantiationPhase()
+bool FMovieSceneEntitySystemRunner::GameThread_InstantiationPhase(UMovieSceneEntitySystemLinker* Linker)
 {
 	using namespace UE::MovieScene;
 
-	check(GameThread == ENamedThreads::GameThread || GameThread == ENamedThreads::GameThread_Local);
+	SCOPE_CYCLE_COUNTER(MovieSceneEval_InstantiationPhase);
 
-	UMovieSceneEntitySystemLinker* Linker = GetLinker();
-	check(Linker);
+	check(GameThread == ENamedThreads::GameThread || GameThread == ENamedThreads::GameThread_Local);
 
+	// Allow the next step to run
+	FlushState &= ~ERunnerFlushState::Instantiation;
 	CurrentPhase = ESystemPhase::Instantiation;
 
 	FGraphEventArray AllTasks;
-	{
-		SCOPE_CYCLE_COUNTER(MovieSceneEval_InstantiationPhase);
-
-		Linker->SystemGraph.ExecutePhase(ESystemPhase::Instantiation, Linker, AllTasks);
-	}
+	Linker->SystemGraph.ExecutePhase(ESystemPhase::Instantiation, Linker, AllTasks);
 
+	// If there were any tasks created, we need to wait on them before proceeding. This is rare for the instantiation phase.
 	if (AllTasks.Num() != 0)
 	{
-		TGraphTask<TFunctionGraphTaskImpl<void(), ESubsequentsMode::TrackSubsequents>>::CreateTask(&AllTasks, ENamedThreads::GameThread)
-			.ConstructAndDispatchWhenReady(
-				[this]
-				{
-					this->GameThread_PostInstantiation();
-				}
-				, TStatId(), GameThread);
-	}
-	else
-	{
-		this->GameThread_PostInstantiation();
+		FGraphEventRef InstantiationEvent = TGraphTask<FNullGraphTask>::CreateTask(&AllTasks, ENamedThreads::GameThread)
+		.ConstructAndDispatchWhenReady(TStatId(), GameThread);
+
+		FTaskGraphInterface::Get().WaitUntilTaskCompletes(InstantiationEvent, ENamedThreads::GameThread_Local);
 	}
+
+	return GameThread_PostInstantiation(Linker);
 }
 
-void FMovieSceneEntitySystemRunner::GameThread_PostInstantiation()
+bool FMovieSceneEntitySystemRunner::GameThread_PostInstantiation(UMovieSceneEntitySystemLinker* Linker)
 {
 	using namespace UE::MovieScene;
 
-	{
-		SCOPE_CYCLE_COUNTER(MovieSceneEval_PostInstantiation);
-
-		check(GameThread == ENamedThreads::GameThread || GameThread == ENamedThreads::GameThread_Local);
+	SCOPE_CYCLE_COUNTER(MovieSceneEval_PostInstantiation);
 
-		UMovieSceneEntitySystemLinker* Linker = GetLinker();
-		check(Linker);
+	check(GameThread == ENamedThreads::GameThread || GameThread == ENamedThreads::GameThread_Local);
 
-		Linker->PostInstantation(*this);
+	Linker->PostInstantation(*this);
 
-		FEntityManager& EntityManager = Linker->EntityManager;
-		FBuiltInComponentTypes* BuiltInComponentTypes = FBuiltInComponentTypes::Get();
+	FEntityManager& EntityManager = Linker->EntityManager;
+	FBuiltInComponentTypes* BuiltInComponentTypes = FBuiltInComponentTypes::Get();
 
-		// Nothing needs linking, caching or restoring any more
-		FRemoveMultipleMutation Mutation;
-		Mutation.MaskToRemove = BuiltInComponentTypes->RequiresInstantiationMask;
-		// Inverse the filter because MaskToRemove is applied as a binary AND against all matching allocations.
-		// Therefore, any set bits in RequiresInstantiationMask will be removed from components
-		Mutation.MaskToRemove.BitwiseNOT();
+	// Nothing needs linking, caching or restoring any more
+	FRemoveMultipleMutation Mutation;
+	Mutation.MaskToRemove = BuiltInComponentTypes->RequiresInstantiationMask;
+	// Inverse the filter because MaskToRemove is applied as a binary AND against all matching allocations.
+	// Therefore, any set bits in RequiresInstantiationMask will be removed from components
+	Mutation.MaskToRemove.BitwiseNOT();
 
-		FEntityComponentFilter Filter = FEntityComponentFilter().Any(BuiltInComponentTypes->RequiresInstantiationMask);
-		EntityManager.MutateAll(Filter, Mutation);
+	FEntityComponentFilter Filter = FEntityComponentFilter().Any(BuiltInComponentTypes->RequiresInstantiationMask);
+	EntityManager.MutateAll(Filter, Mutation);
 
-		// Free anything that has been unlinked
-		EntityManager.FreeEntities(FEntityComponentFilter().All({ BuiltInComponentTypes->Tags.NeedsUnlink }));
+	// Free anything that has been unlinked
+	EntityManager.FreeEntities(FEntityComponentFilter().All({ BuiltInComponentTypes->Tags.NeedsUnlink }));
 
-		Linker->AutoUnlinkIrrelevantSystems();
+	Linker->AutoUnlinkIrrelevantSystems();
 
-		EntityManager.Compact();
-	}
-
-	GameThread_EvaluationPhase();
+	EntityManager.Compact();
+	return true;
 }
 
-void FMovieSceneEntitySystemRunner::GameThread_EvaluationPhase()
+bool FMovieSceneEntitySystemRunner::GameThread_EvaluationPhase(UMovieSceneEntitySystemLinker* Linker)
 {
 	using namespace UE::MovieScene;
 
-	UMovieSceneEntitySystemLinker* Linker = GetLinker();
-	check(Linker);
+	SCOPE_CYCLE_COUNTER(MovieSceneEval_EvaluationPhase);
 
+	FlushState &= ~ERunnerFlushState::Evaluation;
 	CurrentPhase = ESystemPhase::Evaluation;
 
 	FGraphEventArray AllTasks;
-	{
-		SCOPE_CYCLE_COUNTER(MovieSceneEval_EvaluationPhase);
-
-		// --------------------------------------------------------------------------------------------------------------------------------------------
-		// Step 2: Run the evaluation phase. The entity manager is locked down for this phase, meaning no changes to entity-component structure is allowed
-		//         This vastly simplifies the concurrent handling of entity component allocations
-		Linker->EntityManager.LockDown();
-
-		checkf(!Linker->EntityManager.ContainsComponent(FBuiltInComponentTypes::Get()->Tags.NeedsUnlink), TEXT("Stale entities remain in the entity manager during evaluation - these should have been destroyed during the instantiation phase. Did it run?"));
+	// --------------------------------------------------------------------------------------------------------------------------------------------
+	// Step 2: Run the evaluation phase. The entity manager is locked down for this phase, meaning no changes to entity-component structure is allowed
+	//         This vastly simplifies the concurrent handling of entity component allocations
+	Linker->EntityManager.LockDown();
 
-		Linker->SystemGraph.ExecutePhase(ESystemPhase::Evaluation, Linker, AllTasks);
-	}
+	checkf(!Linker->EntityManager.ContainsComponent(FBuiltInComponentTypes::Get()->Tags.NeedsUnlink), TEXT("Stale entities remain in the entity manager during evaluation - these should have been destroyed during the instantiation phase. Did it run?"));
 
-	auto Finish = [this]
-	{
-		// We are now done with the current update batch. Let's unlock the completion task to unblock
-		// the main thread, which is waiting on it inside Flush().
-		check(this->CompletionTask != nullptr);
-		this->CompletionTask->Unlock();
-	};
+	Linker->SystemGraph.ExecutePhase(ESystemPhase::Evaluation, Linker, AllTasks);
 
 	if (AllTasks.Num() != 0)
 	{
-		TGraphTask<TFunctionGraphTaskImpl<void(), ESubsequentsMode::TrackSubsequents>>::CreateTask(&AllTasks, ENamedThreads::GameThread)
-			.ConstructAndDispatchWhenReady(MoveTemp(Finish), TStatId(), GameThread);
-	}
-	else
-	{
-		Finish();
+		FGraphEventRef EvaluationEvent = TGraphTask<FNullGraphTask>::CreateTask(&AllTasks, ENamedThreads::GameThread)
+		.ConstructAndDispatchWhenReady(TStatId(), GameThread);
+
+		FTaskGraphInterface::Get().WaitUntilTaskCompletes(EvaluationEvent, ENamedThreads::GameThread_Local);
 	}
+
+	Linker->EntityManager.ReleaseLockDown();
+
+	return true;
 }
 
-void FMovieSceneEntitySystemRunner::GameThread_EvaluationFinalizationPhase()
+void FMovieSceneEntitySystemRunner::GameThread_EvaluationFinalizationPhase(UMovieSceneEntitySystemLinker* Linker)
 {
 	using namespace UE::MovieScene;
 
 	check(GameThread == ENamedThreads::GameThread || GameThread == ENamedThreads::GameThread_Local);
 
-	UMovieSceneEntitySystemLinker* Linker = GetLinker();
-	check(Linker);
-
-	Linker->EntityManager.ReleaseLockDown();
-
 	CurrentPhase = ESystemPhase::Finalization;
+	FlushState &= ~ERunnerFlushState::Finalization;
 
 	// Post-eval events can be queued during the finalization phase so let's open that up.
-	// The events are actually executed a bit later, in GameThread_PostEvaluationPhase.
+	// The events are actually executed a bit later, in GameThread_EventTriggerPhase.
 	bCanQueueEventTriggers = true;
 	{
 		SCOPE_CYCLE_COUNTER(MovieSceneEval_FinalizationPhase);
@@ -545,12 +842,12 @@ void FMovieSceneEntitySystemRunner::GameThread_EvaluationFinalizationPhase()
 		FInstanceRegistry* InstanceRegistry = GetInstanceRegistry();
 
 		// Iterate on a copy of our current instances, since LegacyEvaluator->Evaluate() could change the instance handle, which would affect PostEvaluationPhase
-		TArray<FInstanceHandle> CurrentInstancesCopy(CurrentInstances);
-		for (FInstanceHandle UpdatedInstanceHandle : CurrentInstancesCopy)
+		TArray<FQueuedUpdateParams> CurrentInstancesCopy(CurrentInstances);
+		for (const FQueuedUpdateParams& UpdateParams : CurrentInstancesCopy)
 		{
-			if (InstanceRegistry->IsHandleValid(UpdatedInstanceHandle))
+			if (InstanceRegistry->IsHandleValid(UpdateParams.InstanceHandle))
 			{
-				FSequenceInstance& Instance = InstanceRegistry->MutateInstance(UpdatedInstanceHandle);
+				FSequenceInstance& Instance = InstanceRegistry->MutateInstance(UpdateParams.InstanceHandle);
 				if (Instance.IsRootSequence())
 				{
 					Instance.RunLegacyTrackTemplates();
@@ -564,86 +861,136 @@ void FMovieSceneEntitySystemRunner::GameThread_EvaluationFinalizationPhase()
 	}
 	bCanQueueEventTriggers = false;
 
-	CurrentPhase = ESystemPhase::None;
+	if (!EventTriggers.IsBound())
+	{
+		// Skip event triggers if there are none
+		FlushState &= ~ERunnerFlushState::EventTriggers;
+	}
 }
 
-void FMovieSceneEntitySystemRunner::GameThread_PostEvaluationPhase()
+void FMovieSceneEntitySystemRunner::GameThread_EventTriggerPhase(UMovieSceneEntitySystemLinker* Linker)
 {
 	using namespace UE::MovieScene;
 
-	SCOPE_CYCLE_COUNTER(MovieSceneEval_PostEvaluationPhase);
-
-	UMovieSceneEntitySystemLinker* Linker = GetLinker();
-	check(Linker);
+	FlushState &= ~ERunnerFlushState::EventTriggers;
 
 	// Execute any queued events from the evaluation finalization phase.
 	if (EventTriggers.IsBound())
 	{
-		// Let's allow re-entrant evaluation at this point.
-		FMovieSceneEntitySystemEvaluationReentrancyWindow Window(*Linker);
+		// Event triggers are allowed to be re-entrant
+		FMovieSceneEntitySystemEvaluationReentrancyWindow Window(this, Linker);
+
+		// Trigger from a temporary delegate to ensure that re-entrant evaluations do not re-trigger events
+		FMovieSceneEntitySystemEventTriggers TmpEventTriggers;
 
-		EventTriggers.Broadcast();
+		Swap(TmpEventTriggers, EventTriggers);
 		EventTriggers.Clear();
+
+		TmpEventTriggers.Broadcast();
 	}
+}
+
+void FMovieSceneEntitySystemRunner::GameThread_PostEvaluationPhase(UMovieSceneEntitySystemLinker* Linker)
+{
+	using namespace UE::MovieScene;
+
+	SCOPE_CYCLE_COUNTER(MovieSceneEval_PostEvaluationPhase);
+
+	FlushState &= ~ERunnerFlushState::PostEvaluation;
 
 	// Now run the post-evaluation logic so that we can safely handle broadcast events (like OnFinished)
 	// that trigger some new evaluations (such as connecting it to another sequence's Play in Blueprint).
 	//
 	// If we are the global linker (and not a "private" linker, as is the case with "blocking" sequences),
 	// we may find ourselves in a re-entrant call, which means we need to save our state here and restore
-	// it afterwards. We also:
-	//   - iterate on a copy of our current instances, since a re-entrant call would modify that array.
-	//   - store our pending update result in a separate variable before assigning it to ourselves, for
-	//     the same reasons as above.
-	//
-	FInstanceRegistry* InstanceRegistry = GetInstanceRegistry();
-	TArray<FInstanceHandle> CurrentInstancesCopy(CurrentInstances);
-	CurrentInstances.Empty();
+	// it afterwards. We also iterate on a copy of our current instances, since a re-entrant call would
+	// modify that array.
+
+	// Temporarily cache the pending updates and dissected updates so we can preserve re-entrant ordering
+	//    by appending the arrays afterwards. This ensures that any updates queued during a re-entrant call
+	//    get evaluated _before_ the updates that are specified in the outer scope, even if the evaluation was
+	//    budgeted and didn't fully flush, whilst simultaneously correctly only flushing inner updates if a
+	//    full flush is performed.
+	TArray<FQueuedUpdateParams> TmpCurrentInstances;
+	TArray<FUpdateParamsAndContext> TmpUpdateQueue;
+	TArray<FDissectedUpdate> TmpDissectedUpdates;
+	TArray<FInstanceHandle> TmpDestroyInstances;
+	TArray<FSimpleDelegate> TmpOnFlushedDelegates;
+
+	Swap(UpdateQueue, TmpUpdateQueue);
+	Swap(DissectedUpdates, TmpDissectedUpdates);
+	Swap(CurrentInstances, TmpCurrentInstances);
+	Swap(OnFlushedDelegates, TmpOnFlushedDelegates);
+
+	FInstanceRegistry* InstanceRegistry = Linker->GetInstanceRegistry();
 
 	{
-		FMovieSceneEntitySystemEvaluationReentrancyWindow Window(*Linker);
+		FMovieSceneEntitySystemEvaluationReentrancyWindow Window(this, Linker);
 
-		for (FInstanceHandle UpdatedInstanceHandle : CurrentInstancesCopy)
+		for (const FQueuedUpdateParams& UpdateParams : TmpCurrentInstances)
 		{
 			// We must check for validity here because the cache handles may have become invalid
 			// during this iteration (since there is a re-entrancy window open)
-			if (InstanceRegistry->IsHandleValid(UpdatedInstanceHandle))
+			if (InstanceRegistry->IsHandleValid(UpdateParams.InstanceHandle))
 			{
-				InstanceRegistry->MutateInstance(UpdatedInstanceHandle).PostEvaluation(Linker);
+				FSequenceInstance& Instance = InstanceRegistry->MutateInstance(UpdateParams.InstanceHandle);
+				Instance.PostEvaluation(Linker);
+
+				if (EnumHasAnyFlags(UpdateParams.UpdateFlags, ERunnerUpdateFlags::Destroy))
+				{
+					TmpDestroyInstances.Add(UpdateParams.InstanceHandle);
+				}
 			}
 		}
-	}
-}
 
-void FMovieSceneEntitySystemRunner::FinishInstance(FInstanceHandle InInstanceHandle)
-{
-	using namespace UE::MovieScene;
+		for (const FSimpleDelegate& OnFlushed : TmpOnFlushedDelegates)
+		{
+			OnFlushed.Execute();
+		}
+	}
 
-	UMovieSceneEntitySystemLinker* Linker = GetLinker();
-	check(Linker);
+	// Destroy any instances that need destroying
+	for (FInstanceHandle Handle : TmpDestroyInstances)
+	{
+		if (InstanceRegistry->IsHandleValid(Handle))
+		{
+			InstanceRegistry->DestroyInstance(Handle);
+		}
+	}
 
-	// If we've already got queued updates for this instance we need to flush the linker first so that those updates are reflected correctly
-	FInstanceRegistry* InstanceRegistry = GetInstanceRegistry();
-	if (InstanceRegistry->GetInstance(InInstanceHandle).HasEverUpdated() && HasQueuedUpdates(InInstanceHandle))
+	// Append our cached updates if there are outstanding updates from a re-entrant call, otherwise
+	// just swap them back
+	if (UpdateQueue.Num() > 0 && TmpUpdateQueue.Num() > 0)
 	{
-		Flush();
+		UpdateQueue.Append(TmpUpdateQueue);
+	}
+	else
+	{
+		Swap(TmpUpdateQueue, UpdateQueue);
 	}
 
-	InstanceRegistry->MutateInstance(InInstanceHandle).Finish(Linker);
-	if (Linker->HasStructureChangedSinceLastRun())
+	// Same with the dissected updates
+	if (DissectedUpdates.Num() > 0 && TmpDissectedUpdates.Num() > 0)
 	{
-		MarkForUpdate(InInstanceHandle);
-		Flush();
+		DissectedUpdates.Append(TmpDissectedUpdates);
 	}
 	else
 	{
-		InstanceRegistry->MutateInstance(InInstanceHandle).PostEvaluation(Linker);
+		Swap(TmpDissectedUpdates, DissectedUpdates);
+	}
+
+	// If we have any pending updates, we need to run another evaluation.
+	// This will cause us to effectively loop back to ERunnerFlushState::Import
+	// which will pick up the next updates
+	if (UpdateQueue.Num() > 0 || DissectedUpdates.Num() > 0)
+	{
+		FlushState = ERunnerFlushState::LoopEval;
 	}
 }
 
-void FMovieSceneEntitySystemRunner::MarkForUpdate(FInstanceHandle InInstanceHandle)
+void FMovieSceneEntitySystemRunner::MarkForUpdate(FInstanceHandle InInstanceHandle, UE::MovieScene::ERunnerUpdateFlags UpdateFlags)
 {
-	CurrentInstances.AddUnique(InInstanceHandle);
+	CurrentInstances.Add(FQueuedUpdateParams{ InInstanceHandle, UpdateFlags });
 }
 
 void FMovieSceneEntitySystemRunner::OnLinkerAbandon(UMovieSceneEntitySystemLinker* InLinker)
diff --git a/Engine/Source/Runtime/MovieScene/Private/EntitySystem/MovieSceneSequenceUpdaters.cpp b/Engine/Source/Runtime/MovieScene/Private/EntitySystem/MovieSceneSequenceUpdaters.cpp
--- a/Engine/Source/Runtime/MovieScene/Private/EntitySystem/MovieSceneSequenceUpdaters.cpp
+++ b/Engine/Source/Runtime/MovieScene/Private/EntitySystem/MovieSceneSequenceUpdaters.cpp
@@ -627,7 +627,7 @@ void FSequenceUpdater_Hierarchical::Update(UMovieSceneEntitySystemLinker* Linker
 			continue;
 		}
 
-		Runner->MarkForUpdate(SubData.Handle);
+		Runner->MarkForUpdate(SubData.Handle, ERunnerUpdateFlags::None);
 
 		if (!ActiveSequences.Contains(InstanceIt.Key()))
 		{
diff --git a/Engine/Source/Runtime/MovieScene/Private/Evaluation/MovieSceneEvaluationTemplateInstance.cpp b/Engine/Source/Runtime/MovieScene/Private/Evaluation/MovieSceneEvaluationTemplateInstance.cpp
--- a/Engine/Source/Runtime/MovieScene/Private/Evaluation/MovieSceneEvaluationTemplateInstance.cpp
+++ b/Engine/Source/Runtime/MovieScene/Private/Evaluation/MovieSceneEvaluationTemplateInstance.cpp
@@ -32,7 +32,22 @@ void FMovieSceneRootEvaluationTemplateInstance::BeginDestroy()
 {
 	if (EntitySystemLinker && IsValidChecked(EntitySystemLinker) && !EntitySystemLinker->IsUnreachable() && !EntitySystemLinker->HasAnyFlags(RF_BeginDestroyed))
 	{
-		EntitySystemLinker->GetInstanceRegistry()->DestroyInstance(RootInstanceHandle);
+		if (TSharedPtr<FMovieSceneEntitySystemRunner> Runner = WeakRunner.Pin())
+		{
+			// There's no way to know whether this instance is currently being included in the evaluation so if it is currently evaluating we
+			// just need to flush any current work. FlushOutstanding does nothing if it is not currently in the middle of an evaluation
+			Runner->FlushOutstanding();
+
+			const bool bRequiresFlush = Runner->QueueFinalUpdate(RootInstanceHandle, UE::MovieScene::ERunnerUpdateFlags::Destroy);
+			if (bRequiresFlush)
+			{
+			  Runner->Flush();
+			}
+		}
+		else
+		{
+			EntitySystemLinker->GetInstanceRegistry()->DestroyInstance(RootInstanceHandle);
+		}
 	}
 
 	CompiledDataManager = nullptr;
@@ -56,15 +71,13 @@ UMovieSceneEntitySystemLinker* FMovieSceneRootEvaluationTemplateInstance::Constr
 	return UMovieSceneEntitySystemLinker::FindOrCreateLinker(PlaybackContext, UE::MovieScene::EEntitySystemLinkerRole::Standalone, TEXT("DefaultEntitySystemLinker"));
 }
 
-void FMovieSceneRootEvaluationTemplateInstance::Initialize(UMovieSceneSequence& InRootSequence, IMovieScenePlayer& Player, UMovieSceneCompiledDataManager* InCompiledDataManager)
+void FMovieSceneRootEvaluationTemplateInstance::Initialize(UMovieSceneSequence& InRootSequence, IMovieScenePlayer& Player, UMovieSceneCompiledDataManager* InCompiledDataManager, TWeakPtr<FMovieSceneEntitySystemRunner> InWeakRunner)
 {
 	bool bReinitialize = (
 			// Initialize if we weren't initialized before and this is our first sequence.
 			WeakRootSequence.Get() == nullptr ||
 			// Initialize if we lost our linker.
-			EntitySystemLinker == nullptr ||
-			// Initialize if our linker was reset and forced our runner to detach.
-			!EntitySystemRunner.IsAttachedToLinker());
+			EntitySystemLinker == nullptr);
 
 	const UMovieSceneCompiledDataManager* PreviousCompiledDataManager = CompiledDataManager;
 	if (InCompiledDataManager)
@@ -100,12 +113,12 @@ void FMovieSceneRootEvaluationTemplateInstance::Initialize(UMovieSceneSequence&
 	}
 	bReinitialize |= (PreviousCompiledDataManager != CompiledDataManager);
 
+	TSharedPtr<FMovieSceneEntitySystemRunner> PreviousRunner = WeakRunner.Pin();
+
 	if (UMovieSceneSequence* ExistingSequence = WeakRootSequence.Get())
 	{
 		if (ExistingSequence != &InRootSequence)
 		{
-			Finish(Player);
-
 			bReinitialize = true;
 		}
 	}
@@ -113,52 +126,65 @@ void FMovieSceneRootEvaluationTemplateInstance::Initialize(UMovieSceneSequence&
 	CompiledDataID = CompiledDataManager->GetDataID(&InRootSequence);
 	WeakRootSequence = &InRootSequence;
 	RootID = MovieSceneSequenceID::Root;
+	WeakRunner = InWeakRunner;
 
 	if (bReinitialize)
 	{
-		if (RootInstanceHandle.IsValid() && EntitySystemLinker)
+		if (RootInstanceHandle.IsValid())
 		{
-			EntitySystemLinker->GetInstanceRegistry()->DestroyInstance(RootInstanceHandle);
+			if (PreviousRunner)
+			{
+				PreviousRunner->AbandonAndDestroyInstance(RootInstanceHandle);
+			}
+			else if (EntitySystemLinker)
+			{
+				EntitySystemLinker->GetInstanceRegistry()->DestroyInstance(RootInstanceHandle);
+			}
+			else
+			{
+				ensureMsgf(false, TEXT("Unable to destroy previously allocated sequence instance - this could indicate a memory leak."));
+			}
 		}
 
 		Player.State.PersistentEntityData.Reset();
 		Player.State.PersistentSharedData.Reset();
 
-		if (EntitySystemRunner.IsAttachedToLinker())
-		{
-			EntitySystemRunner.DetachFromLinker();
-		}
 		EntitySystemLinker = ConstructEntityLinker(Player);
-		EntitySystemRunner.AttachToLinker(EntitySystemLinker);
 		RootInstanceHandle = EntitySystemLinker->GetInstanceRegistry()->AllocateRootInstance(&Player);
 
+		TSharedPtr<FMovieSceneEntitySystemRunner> Runner = WeakRunner.Pin();
+		if (ensure(Runner) && EntitySystemLinker)
+		{
+			if (!Runner->IsAttachedToLinker())
+			{
+				Runner->AttachToLinker(EntitySystemLinker);
+			}
+		}
+
 		Player.PreAnimatedState.Initialize(EntitySystemLinker, RootInstanceHandle);
 	}
 }
 
 void FMovieSceneRootEvaluationTemplateInstance::Evaluate(FMovieSceneContext Context, IMovieScenePlayer& Player)
 {
-	SCOPE_CYCLE_COUNTER(MovieSceneEval_EntireEvaluationCost);
-
-	check(EntitySystemLinker);
-
-	if (EntitySystemRunner.IsAttachedToLinker())
-	{
-		EntitySystemRunner.Update(Context, RootInstanceHandle);
-	}
+	EvaluateSynchronousBlocking(Context, Player);
 }
 
-void FMovieSceneRootEvaluationTemplateInstance::Finish(IMovieScenePlayer& Player)
+void FMovieSceneRootEvaluationTemplateInstance::EvaluateSynchronousBlocking(FMovieSceneContext Context, IMovieScenePlayer& Player)
 {
-	if (EntitySystemRunner.IsAttachedToLinker())
+	if (TSharedPtr<FMovieSceneEntitySystemRunner> Runner = WeakRunner.Pin())
 	{
-		EntitySystemRunner.FinishInstance(RootInstanceHandle);
+		Runner->QueueUpdate(Context, RootInstanceHandle);
+		Runner->Flush();
 	}
+}
 
-	OnFinish();
+void FMovieSceneRootEvaluationTemplateInstance::OnFinished()
+{
+	ResetDirectorInstances();
 }
 
-void FMovieSceneRootEvaluationTemplateInstance::OnFinish()
+void FMovieSceneRootEvaluationTemplateInstance::ResetDirectorInstances()
 {
 	DirectorInstances.Reset();
 }
@@ -197,11 +223,6 @@ UMovieSceneEntitySystemLinker* FMovieSceneRootEvaluationTemplateInstance::GetEnt
 	return EntitySystemLinker;
 }
 
-FMovieSceneEntitySystemRunner& FMovieSceneRootEvaluationTemplateInstance::GetEntitySystemRunner()
-{
-	return EntitySystemRunner;
-}
-
 bool FMovieSceneRootEvaluationTemplateInstance::HasEverUpdated() const
 {
 	if (EntitySystemLinker)
@@ -364,11 +385,16 @@ void FMovieSceneRootEvaluationTemplateInstance::PlaybackContextChanged(IMovieSce
 
 	const bool bGlobalCapture = Player.PreAnimatedState.IsCapturingGlobalPreAnimatedState();
 
-	if (EntitySystemLinker && IsValidChecked(EntitySystemLinker) && !EntitySystemLinker->IsUnreachable() && !EntitySystemLinker->HasAnyFlags(RF_BeginDestroyed))
+	if (RootInstanceHandle.IsValid() && EntitySystemLinker && IsValidChecked(EntitySystemLinker) && !EntitySystemLinker->IsUnreachable() && !EntitySystemLinker->HasAnyFlags(RF_BeginDestroyed))
 	{
 		EntitySystemLinker->CleanupInvalidBoundObjects();
 
-		Finish(Player);
+		if (TSharedPtr<FMovieSceneEntitySystemRunner> Runner = WeakRunner.Pin())
+		{
+			Runner->QueueFinalUpdate(RootInstanceHandle, ERunnerUpdateFlags::Finish);
+			Runner->Flush();
+		}
+
 		if (bGlobalCapture)
 		{
 			Player.RestorePreAnimatedState();
@@ -376,12 +402,15 @@ void FMovieSceneRootEvaluationTemplateInstance::PlaybackContextChanged(IMovieSce
 		EntitySystemLinker->GetInstanceRegistry()->DestroyInstance(RootInstanceHandle);
 	}
 
-	if (EntitySystemRunner.IsAttachedToLinker())
+	EntitySystemLinker = ConstructEntityLinker(Player);
+	if (TSharedPtr<FMovieSceneEntitySystemRunner> Runner = WeakRunner.Pin())
 	{
-		EntitySystemRunner.DetachFromLinker();
+		if (Runner->IsAttachedToLinker())
+		{
+			Runner->DetachFromLinker();
+			Runner->AttachToLinker(EntitySystemLinker);
+		}
 	}
-	EntitySystemLinker = ConstructEntityLinker(Player);
-	EntitySystemRunner.AttachToLinker(EntitySystemLinker);
 
 	RootInstanceHandle = EntitySystemLinker->GetInstanceRegistry()->AllocateRootInstance(&Player);
 	DirectorInstances.Reset();
diff --git a/Engine/Source/Runtime/MovieScene/Private/MovieSceneSequencePlayer.cpp b/Engine/Source/Runtime/MovieScene/Private/MovieSceneSequencePlayer.cpp
--- a/Engine/Source/Runtime/MovieScene/Private/MovieSceneSequencePlayer.cpp
+++ b/Engine/Source/Runtime/MovieScene/Private/MovieSceneSequencePlayer.cpp
@@ -254,15 +254,6 @@ void UMovieSceneSequencePlayer::PlayInternal()
 			}
 		}
 
-		// Start playing
-		// @todo Sequencer playback: Should we recreate the instance every time?
-		// We must not recreate the instance since it holds stateful information (such as which objects it has spawned). Recreating the instance would break any 
-		// @todo: Is this still the case now that eval state is stored (correctly) in the player?
-		if (!RootTemplateInstance.IsValid())
-		{
-			RootTemplateInstance.Initialize(*Sequence, *this, nullptr);
-		}
-
 		// Update now
 		if (PlaybackSettings.bRestoreState)
 		{
@@ -333,43 +324,40 @@ void UMovieSceneSequencePlayer::Pause()
 		PauseOnFrame.Reset();
 		LastTickGameTimeSeconds.Reset();
 
-		// Evaluate the sequence at its current time, with a status of 'stopped' to ensure that animated state pauses correctly. (ie. audio sounds should stop/pause)
+		auto FinishPause = [this]
 		{
-			FMovieSceneEvaluationRange CurrentTimeRange = PlayPosition.GetCurrentPositionAsRange();
-			const FMovieSceneContext Context(CurrentTimeRange, EMovieScenePlayerStatus::Stopped);
-			RootTemplateInstance.Evaluate(Context, *this);
-		}
+			this->RunLatentActions();
+			this->UpdateNetworkSyncProperties();
 
-		RunLatentActions();
+			UMovieSceneSequence* MovieSceneSequence = this->RootTemplateInstance.GetSequence(MovieSceneSequenceID::Root);
+			if (MovieSceneSequence)
+			{
+				UE_LOG(LogMovieScene, Verbose, TEXT("Pause - MovieSceneSequence: %s"), *MovieSceneSequence->GetName());
+			}
 
-		UpdateNetworkSyncProperties();
+			if (this->OnPause.IsBound())
+			{
+				this->OnPause.Broadcast();
+			}
+		};
 
-		UMovieSceneSequence* MovieSceneSequence = RootTemplateInstance.GetSequence(MovieSceneSequenceID::Root);
-		if (MovieSceneSequence)
+		// Evaluate the sequence at its current time, with a status of 'stopped' to ensure that animated state pauses correctly. (ie. audio sounds should stop/pause)
+		if (TSharedPtr<FMovieSceneEntitySystemRunner> Runner = RootTemplateInstance.GetRunner())
 		{
-			UE_LOG(LogMovieScene, Verbose, TEXT("Pause - MovieSceneSequence: %s"), *MovieSceneSequence->GetName());
-		}
+			FMovieSceneEvaluationRange CurrentTimeRange = PlayPosition.GetCurrentPositionAsRange();
+			const FMovieSceneContext Context(CurrentTimeRange, EMovieScenePlayerStatus::Stopped);
 
-		if (OnPause.IsBound())
+			Runner->QueueUpdate(Context, RootTemplateInstance.GetRootInstanceHandle(), FSimpleDelegate::CreateWeakLambda(this, FinishPause));
+		}
+		else
 		{
-			OnPause.Broadcast();
+			FinishPause();
 		}
 	}
 }
 
 void UMovieSceneSequencePlayer::Scrub()
 {
-	// @todo Sequencer playback: Should we recreate the instance every time?
-	// We must not recreate the instance since it holds stateful information (such as which objects it has spawned). Recreating the instance would break any 
-	// @todo: Is this still the case now that eval state is stored (correctly) in the player?
-	if (ensureAsRuntimeWarning(Sequence != nullptr))
-	{
-		if (!RootTemplateInstance.IsValid())
-		{
-			RootTemplateInstance.Initialize(*Sequence, *this, nullptr);
-		}
-	}
-
 	Status = EMovieScenePlayerStatus::Scrubbing;
 	TimeController->StopPlaying(GetCurrentTime());
 
@@ -417,41 +405,53 @@ void UMovieSceneSequencePlayer::StopInternal(FFrameTime TimeToResetTo)
 			RestorePreAnimatedState();
 		}
 
-		if (RootTemplateInstance.IsValid())
+		// Lambda that is invoked when the request to finish this sequence has been fulfilled
+		auto OnFlushed = [this, TimeToResetTo]
 		{
-			RootTemplateInstance.Finish(*this);
-		}
+			if (this->OldMaxTickRate.IsSet())
+			{
+				GEngine->SetMaxFPS(OldMaxTickRate.GetValue());
+				this->OldMaxTickRate.Reset();
+			}
 
-		if (OldMaxTickRate.IsSet())
-		{
-			GEngine->SetMaxFPS(OldMaxTickRate.GetValue());
-			OldMaxTickRate.Reset();
-		}
+			if (this->HasAuthority())
+			{
+				// Explicitly handle Stop() events through an RPC call
+				this->RPC_OnStopEvent(TimeToResetTo);
+			}
+			this->UpdateNetworkSyncProperties();
 
-		if (HasAuthority())
-		{
-			// Explicitly handle Stop() events through an RPC call
-			RPC_OnStopEvent(TimeToResetTo);
-		}
-		UpdateNetworkSyncProperties();
+			this->OnStopped();
 
-		OnStopped();
+			if (this->RootTemplateInstance.IsValid())
+			{
+				UMovieSceneSequence* MovieSceneSequence = RootTemplateInstance.GetSequence(MovieSceneSequenceID::Root);
+				if (MovieSceneSequence)
+				{
+					UE_LOG(LogMovieScene, Verbose, TEXT("Stop - MovieSceneSequence: %s"), *MovieSceneSequence->GetName());
+				}
+			}
 
-		if (RootTemplateInstance.IsValid())
-		{
-			UMovieSceneSequence* MovieSceneSequence = RootTemplateInstance.GetSequence(MovieSceneSequenceID::Root);
-			if (MovieSceneSequence)
+			if (this->OnStop.IsBound())
 			{
-				UE_LOG(LogMovieScene, Verbose, TEXT("Stop - MovieSceneSequence: %s"), *MovieSceneSequence->GetName());
+				this->OnStop.Broadcast();
 			}
-		}
 
-		if (OnStop.IsBound())
+			this->RunLatentActions();
+		};
+
+		// Add the post eval callback that will be invoked when QueueFinalUpdate has been flushed
+		PostEvaluationCallbacks.Add(FOnEvaluationCallback::CreateWeakLambda(this, OnFlushed));
+
+		TSharedPtr<FMovieSceneEntitySystemRunner> Runner = RootTemplateInstance.GetRunner();
+		if (Runner)
 		{
-			OnStop.Broadcast();
+			// Finish but do not destroy
+			if (!Runner->QueueFinalUpdate(RootTemplateInstance.GetRootInstanceHandle(), UE::MovieScene::ERunnerUpdateFlags::Finish))
+			{
+				Runner->Flush();
+			}
 		}
-
-		RunLatentActions();
 	}
 	else if (RootTemplateInstance.IsValid() && RootTemplateInstance.HasEverUpdated())
 	{
@@ -459,7 +459,16 @@ void UMovieSceneSequencePlayer::StopInternal(FFrameTime TimeToResetTo)
 		{
 			RestorePreAnimatedState();
 		}
-		RootTemplateInstance.Finish(*this);
+
+		TSharedPtr<FMovieSceneEntitySystemRunner> Runner = RootTemplateInstance.GetRunner();
+		if (Runner)
+		{
+			// Finish but do not destroy
+			if (!Runner->QueueFinalUpdate(RootTemplateInstance.GetRootInstanceHandle(), UE::MovieScene::ERunnerUpdateFlags::Finish))
+			{
+				Runner->Flush();
+			}
+		}
 	}
 }
 
@@ -850,7 +859,15 @@ void UMovieSceneSequencePlayer::Initialize(UMovieSceneSequence* InSequence, cons
 	RegisteredTickInterval = TickInterval;
 	TickManager->RegisterTickClient(RegisteredTickInterval.GetValue(), this);
 
-	RootTemplateInstance.Initialize(*Sequence, *this, nullptr);
+	TSharedPtr<FMovieSceneEntitySystemRunner> RunnerToUse = TickManager->GetRunner(RegisteredTickInterval.GetValue());
+	if (EnumHasAnyFlags(Sequence->GetFlags(), EMovieSceneSequenceFlags::BlockingEvaluation))
+	{
+		SynchronousRunner = MakeShared<FMovieSceneEntitySystemRunner>();
+		RunnerToUse = SynchronousRunner;
+	}
+
+	check(RunnerToUse);
+	RootTemplateInstance.Initialize(*Sequence, *this, nullptr, RunnerToUse);
 
 	LatentActionManager.ClearLatentActions();
 
@@ -1149,22 +1166,13 @@ void UMovieSceneSequencePlayer::UpdateMovieSceneInstance(FMovieSceneEvaluationRa
 	FMovieSceneContext Context(InRange, PlayerStatus);
 	Context.SetHasJumped(Args.bHasJumped);
 
-	if (!Args.bIsAsync)
-	{
-		// Evaluate the sequence synchronously.
-		RootTemplateInstance.Evaluate(Context, *this);
-	}
-	else if (CurrentRunner)
-	{
-		CurrentRunner->QueueUpdate(Context, RootTemplateInstance.GetRootInstanceHandle());
-	}
-	else if (TickManager)
+	TSharedPtr<FMovieSceneEntitySystemRunner> Runner = RootTemplateInstance.GetRunner();
+	if (Runner)
 	{
-		// Look up the runner based on our tick interval
-		FMovieSceneEntitySystemRunner* Runner = TickManager->GetRunner(RegisteredTickInterval.GetValue());
-		if (ensure(Runner))
+		Runner->QueueUpdate(Context, RootTemplateInstance.GetRootInstanceHandle());
+		if (Runner == SynchronousRunner || !Args.bIsAsync)
 		{
-			Runner->QueueUpdate(Context, RootTemplateInstance.GetRootInstanceHandle());
+			Runner->Flush();
 		}
 	}
 }
@@ -1742,9 +1750,10 @@ void UMovieSceneSequencePlayer::RunLatentActions()
 	}
 	else
 	{
+		check(SynchronousRunner);
 		LatentActionManager.RunLatentActions([this]
 		{
-			this->RootTemplateInstance.GetEntitySystemRunner().Flush();
+			this->SynchronousRunner->Flush();
 		});
 	}
 }
diff --git a/Engine/Source/Runtime/MovieScene/Private/MovieSceneSequenceTickInterval.cpp b/Engine/Source/Runtime/MovieScene/Private/MovieSceneSequenceTickInterval.cpp
--- a/Engine/Source/Runtime/MovieScene/Private/MovieSceneSequenceTickInterval.cpp
+++ b/Engine/Source/Runtime/MovieScene/Private/MovieSceneSequenceTickInterval.cpp
@@ -44,10 +44,10 @@ int32 FMovieSceneSequenceTickInterval::RoundTickIntervalMs() const
 	// TickInterval is in seconds, so we multiply by 1000/<rounding> to round it as ms in integer space
 	if (bAllowRounding)
 	{
-		return FMath::RoundToInt(TickIntervalSeconds * (1000.f/GMovieSceneTickIntervalResolutionMs)) * GMovieSceneTickIntervalResolutionMs;
+		return FMath::CeilToInt(TickIntervalSeconds * (1000.f/GMovieSceneTickIntervalResolutionMs)) * GMovieSceneTickIntervalResolutionMs;
 	}
 
-	return FMath::RoundToInt(TickIntervalSeconds * 1000.f);
+	return FMath::CeilToInt(TickIntervalSeconds * 1000.f);
 }
 
 FMovieSceneSequenceTickInterval FMovieSceneSequenceTickInterval::GetInheritedInterval(UObject* ContextObject)
diff --git a/Engine/Source/Runtime/MovieScene/Private/MovieSceneSequenceTickManager.cpp b/Engine/Source/Runtime/MovieScene/Private/MovieSceneSequenceTickManager.cpp
--- a/Engine/Source/Runtime/MovieScene/Private/MovieSceneSequenceTickManager.cpp
+++ b/Engine/Source/Runtime/MovieScene/Private/MovieSceneSequenceTickManager.cpp
@@ -5,9 +5,15 @@
 #include "EntitySystem/MovieSceneEntitySystemLinker.h"
 #include "ProfilingDebugging/CountersTrace.h"
 
+#include "Logging/MessageLog.h"
+#include "Logging/TokenizedMessage.h"
+#include "Misc/UObjectToken.h"
 #include "Algo/IndexOf.h"
 
+#define LOCTEXT_NAMESPACE "MovieSceneSequenceTickManager"
+
 DECLARE_CYCLE_STAT(TEXT("Sequence Tick Manager"), MovieSceneEval_SequenceTickManager, STATGROUP_MovieSceneEval);
+DECLARE_CYCLE_STAT(TEXT("Tick Clients"), MovieSceneEval_TickClients, STATGROUP_MovieSceneEval);
 
 namespace UE::MovieScene
 {
@@ -61,14 +67,14 @@ UMovieSceneEntitySystemLinker* UMovieSceneSequenceTickManager::GetLinker(const F
 	return nullptr;
 }
 
-FMovieSceneEntitySystemRunner* UMovieSceneSequenceTickManager::GetRunner(const FMovieSceneSequenceTickInterval& TickInterval)
+TSharedPtr<FMovieSceneEntitySystemRunner> UMovieSceneSequenceTickManager::GetRunner(const FMovieSceneSequenceTickInterval& TickInterval)
 {
 	const int32 RoundedTickIntervalMs = TickInterval.RoundTickIntervalMs();
 	for (FLinkerGroup& Group : LinkerGroups)
 	{
 		if (Group.RoundedTickIntervalMs == RoundedTickIntervalMs)
 		{
-			return &Group.Runner;
+			return Group.Runner;
 		}
 	}
 	return nullptr;
@@ -130,6 +136,7 @@ void UMovieSceneSequenceTickManager::RegisterTickClient(const FMovieSceneSequenc
 		}
 	}
 
+	const float DesiredBudgetMs = ResolvedTickInterval.EvaluationBudgetMicroseconds / 1000.f;
 	if (!LinkerGroups.IsValidIndex(LinkerIndex))
 	{
 		LinkerIndex = LinkerGroups.Emplace();
@@ -147,13 +154,53 @@ void UMovieSceneSequenceTickManager::RegisterTickClient(const FMovieSceneSequenc
 
 		FLinkerGroup& NewGroup = LinkerGroups[LinkerIndex];
 		NewGroup.Linker = Linker;
-		NewGroup.Runner.AttachToLinker(Linker);
+		NewGroup.Runner = MakeShared<FMovieSceneEntitySystemRunner>();
+		NewGroup.Runner->AttachToLinker(Linker);
 		NewGroup.RoundedTickIntervalMs = DesiredTickIntervalMs;
+		NewGroup.FrameBudgetMs = DesiredBudgetMs;
 		NewGroup.NumClients = 1;
 	}
 	else
 	{
-		++LinkerGroups[LinkerIndex].NumClients;
+		FLinkerGroup& ExistingGroup = LinkerGroups[LinkerIndex];
+		++ExistingGroup.NumClients;
+
+		// If we're attempting to play sequences at the same tick interval with different budgets,
+		// we have to reconcile those. We emit a warning for mismatches of budgeted vs. not-budgeted
+		// as it will actually change behavior (non-budgeted will always eval in one frame, budgeted
+		// will sometimes eval over multiple frames).
+		if (ExistingGroup.FrameBudgetMs != DesiredBudgetMs)
+		{
+			if (ExistingGroup.FrameBudgetMs == 0.f)
+			{
+				ExistingGroup.FrameBudgetMs = DesiredBudgetMs;
+
+#if !UE_BUILD_SHIPPING && !UE_BUILD_TEST
+				TSharedRef<FTokenizedMessage> Message = FMessageLog("PIE")
+				.Warning(LOCTEXT("BudgetMismatch", "Attempting to play multiple sequences with the same tick interval but different budgets. Lowest budget will be used, but may result in undesired multi-frame effects."))
+				->AddToken(FTextToken::Create(LOCTEXT("ClientList", "Initiating tick client ")))
+				->AddToken(FUObjectToken::Create(InterfaceObject))
+				->AddToken(FTextToken::Create(FText::Format(LOCTEXT("ExistingClients", " with budget of {0}μs into group with an interval of {1}ms and budget of {2}μs. Existing clients: "),
+					FMath::RoundToInt(DesiredBudgetMs/1000.f), ExistingGroup.RoundedTickIntervalMs, FMath::RoundToInt(ExistingGroup.FrameBudgetMs/1000.f))));
+				for (const FTickableClientData& Client : TickableClients)
+				{
+					if (Client.LinkerIndex == LinkerIndex)
+					{
+						UObject* Object = Client.WeakObject.Get();
+						while (Object && !(Object->IsA<AActor>() || Object->IsA<USceneComponent>()))
+						{
+							Object = Object->GetOuter();
+						}
+						Message->AddToken(FUObjectToken::Create(Object));
+					}
+				}
+#endif
+			}
+			else if (DesiredBudgetMs != 0.f)
+			{
+				ExistingGroup.FrameBudgetMs = FMath::Min(ExistingGroup.FrameBudgetMs, DesiredBudgetMs);
+			}
+		}
 	}
 
 	checkSlow(LinkerIndex >= 0 && LinkerIndex < uint16(-1));
@@ -272,6 +319,8 @@ void UMovieSceneSequenceTickManager::TickSequenceActors(float DeltaSeconds)
 	// Sparse array where an allocated entry maps to the FLinkerGroup within LinkerGroups
 	TSparseArray<FUpdatedDeltaTimes, InlineSparseAllocator> UpdatedDeltaTimes;
 
+	TBitArray<> OutstandingLinkers;
+
 	// -----------------------------------------------------------------------------
 	// Step 1: Check the tick intervals of all our linker groups, and mark any whose
 	//         interval has passed for update
@@ -285,6 +334,15 @@ void UMovieSceneSequenceTickManager::TickSequenceActors(float DeltaSeconds)
 			float UnpausedDeltaTime = DeltaSeconds;
 			float DeltaTime = DeltaSeconds;
 
+			// If we're currently evaluating this linker group,
+			// skip ticking it until we're completely finished
+			if (Group.Runner->IsCurrentlyEvaluating())
+			{
+				OutstandingLinkers.PadToNum(Index + 1, false);
+				OutstandingLinkers[Index] = true;
+				continue;
+			}
+
 			if (Group.LastUnpausedTimeSeconds >= 0.f)
 			{
 				UnpausedDeltaTime = CurrentUnpausedTimeSeconds - Group.LastUnpausedTimeSeconds;
@@ -308,7 +366,7 @@ void UMovieSceneSequenceTickManager::TickSequenceActors(float DeltaSeconds)
 	}
 
 	// Skip any work if there are no updates scheduled
-	if (UpdatedDeltaTimes.Num() == 0)
+	if (UpdatedDeltaTimes.Num() == 0 && OutstandingLinkers.Num() == 0)
 	{
 		return;
 	}
@@ -318,28 +376,46 @@ void UMovieSceneSequenceTickManager::TickSequenceActors(float DeltaSeconds)
 		// Guard against any mutation while the entries are being ticked - anything added this tick cycle will have to wait until next tick
 		TGuardValue<TArray<FPendingOperation>*> Guard(PendingActorOperations, &CurrentPendingActorOperations);
 
-		// Step 2: Tick clients for any that need ticking
+		// Step 2: Tick unbudgeted clients for any with updated delta times
 		// 
-		for (const FTickableClientData& Client : TickableClients)
+		if (UpdatedDeltaTimes.Num() != 0)
 		{
-			const bool bCanTick = (Client.bTickWhenPaused || !bIsPaused) && Client.WeakObject.Get() != nullptr;
-			if (bCanTick && UpdatedDeltaTimes.IsValidIndex(Client.LinkerIndex))
+			SCOPE_CYCLE_COUNTER(MovieSceneEval_TickClients);
+			for (const FTickableClientData& Client : TickableClients)
 			{
-				const FUpdatedDeltaTimes Delta = UpdatedDeltaTimes[Client.LinkerIndex];
-				const float DeltaTime = Client.bTickWhenPaused ? Delta.UnpausedDeltaTime : Delta.DeltaTime;
+				const bool bCanTick = (Client.bTickWhenPaused || !bIsPaused) && Client.WeakObject.Get() != nullptr;
+				if (bCanTick && UpdatedDeltaTimes.IsValidIndex(Client.LinkerIndex))
+				{
+					const FUpdatedDeltaTimes Delta = UpdatedDeltaTimes[Client.LinkerIndex];
+					const float DeltaTime = Client.bTickWhenPaused ? Delta.UnpausedDeltaTime : Delta.DeltaTime;
 
-				FMovieSceneEntitySystemRunner* Runner = &LinkerGroups[Client.LinkerIndex].Runner;
-				Client.Interface->TickFromSequenceTickManager(DeltaTime, Runner);
+					TSharedPtr<FMovieSceneEntitySystemRunner> Runner = LinkerGroups[Client.LinkerIndex].Runner;
+					Client.Interface->TickFromSequenceTickManager(DeltaTime, Runner.Get());
+				}
 			}
 		}
 
-		// Step 3: Update and flush linkers as needed
+		// Step 3: Flush any budgeted evaluations that are still pending
+		// 
+		for (TConstSetBitIterator<> LinkerIndex(OutstandingLinkers); LinkerIndex; ++LinkerIndex)
+		{
+			FLinkerGroup& Group = LinkerGroups[LinkerIndex.GetIndex()];
+
+			check(!UpdatedDeltaTimes.IsValidIndex(LinkerIndex.GetIndex()));
+
+			Group.Runner->Flush(Group.FrameBudgetMs);
+		}
+
+		// Step 4: Update and flush linkers as needed
 		//
 		for (int32 Index = 0; Index < UpdatedDeltaTimes.GetMaxIndex(); ++Index)
 		{
 			if (UpdatedDeltaTimes.IsAllocated(Index))
 			{
-				LinkerGroups[Index].Runner.Flush();
+				FLinkerGroup& Group = LinkerGroups[Index];
+				// Hitting this check would indicate that the loop above that processes OutstandingLinkers either failed, or some other partial flush happened between then and now.
+				ensureMsgf(!Group.Runner->IsCurrentlyEvaluating(), TEXT("Linker is part-way thorugh a flush when a new flush is being instigated. This is undefined behavior."));
+				Group.Runner->Flush(Group.FrameBudgetMs);
 			}
 		}
 	}
@@ -396,7 +472,7 @@ void UMovieSceneSequenceTickManager::FlushRunners()
 		TGuardValue<TArray<FPendingOperation>*> Guard(PendingActorOperations, &CurrentPendingActorOperations);
 		for (FLinkerGroup& LinkerGroup : LinkerGroups)
 		{
-			LinkerGroup.Runner.Flush();
+			LinkerGroup.Runner->Flush();
 		}
 	}
 
@@ -486,3 +562,4 @@ void FMovieSceneLatentActionManager::RunLatentActions(TFunctionRef<void()> Flush
 	}
 }
 
+#undef LOCTEXT_NAMESPACE
\ No newline at end of file
diff --git a/Engine/Source/Runtime/MovieScene/Private/Tests/MovieScenePreAnimatedStateTests.cpp b/Engine/Source/Runtime/MovieScene/Private/Tests/MovieScenePreAnimatedStateTests.cpp
--- a/Engine/Source/Runtime/MovieScene/Private/Tests/MovieScenePreAnimatedStateTests.cpp
+++ b/Engine/Source/Runtime/MovieScene/Private/Tests/MovieScenePreAnimatedStateTests.cpp
@@ -8,6 +8,7 @@
 #include "Evaluation/MovieScenePreAnimatedState.h"
 #include "Evaluation/MovieSceneEvaluationTemplateInstance.h"
 #include "EntitySystem/MovieSceneEntitySystemLinker.h"
+#include "EntitySystem/MovieSceneEntitySystemRunner.h"
 #include "Evaluation/PreAnimatedState/MovieScenePreAnimatedStateExtension.h"
 #include "MovieSceneTestObjects.h"
 #include "UObject/Package.h"
@@ -382,8 +383,10 @@ bool FMovieScenePreAnimatedStateTrackTypesTest::RunTest(const FString& Parameter
 
 	UMovieSceneCompiledDataManager* CompiledDataManager = UMovieSceneCompiledDataManager::GetPrecompiledData();
 
+	TSharedPtr<FMovieSceneEntitySystemRunner> Runner = MakeShared<FMovieSceneEntitySystemRunner>();
+
 	FMovieSceneCompiledDataID DataID = CompiledDataManager->Compile(Sequence);
-	TestPlayer.Template.Initialize(*Sequence, TestPlayer, CompiledDataManager);
+	TestPlayer.Template.Initialize(*Sequence, TestPlayer, CompiledDataManager, Runner);
 	TestPlayer.Template.EnableGlobalPreAnimatedStateCapture();
 
 	// Test the keep state section
@@ -392,13 +395,13 @@ bool FMovieScenePreAnimatedStateTrackTypesTest::RunTest(const FString& Parameter
 		for (int32 i = 0; i < NumEvaluations; ++i)
 		{
 			FMovieSceneEvaluationRange EvaluatedRange(TRange<FFrameTime>(i*(SectionLength/NumEvaluations), (i+1)*(SectionLength/NumEvaluations)), TickResolution, EPlayDirection::Forwards);
-			TestPlayer.Template.Evaluate(EvaluatedRange, TestPlayer);
+			TestPlayer.Template.EvaluateSynchronousBlocking(EvaluatedRange, TestPlayer);
 
 			Assert(this, TestValue1, StartValue + i, TEXT("Keep-State EvaluationHook did not Begin or Update correctly."));
 		}
 
 		FMovieSceneEvaluationRange EvaluatedRange(TRange<FFrameTime>(SectionLength, SectionLength+100), TickResolution, EPlayDirection::Forwards);
-		TestPlayer.Template.Evaluate(EvaluatedRange, TestPlayer);
+		TestPlayer.Template.EvaluateSynchronousBlocking(EvaluatedRange, TestPlayer);
 		Assert(this, TestValue1, EndValue, TEXT("Keep-State EvaluationHook did not End correctly."));
 
 		TestPlayer.RestorePreAnimatedState();
@@ -411,13 +414,13 @@ bool FMovieScenePreAnimatedStateTrackTypesTest::RunTest(const FString& Parameter
 		for (int32 i = 0; i < NumEvaluations; ++i)
 		{
 			FMovieSceneEvaluationRange EvaluatedRange(TRange<FFrameTime>(2000 + i*(SectionLength/NumEvaluations), 2000 + (i+1)*(SectionLength/NumEvaluations)), TickResolution, EPlayDirection::Forwards);
-			TestPlayer.Template.Evaluate(EvaluatedRange, TestPlayer);
+			TestPlayer.Template.EvaluateSynchronousBlocking(EvaluatedRange, TestPlayer);
 
 			Assert(this, TestValue1, StartValue + i, TEXT("Restore-State EvaluationHook did not Begin or Update correctly."));
 		}
 
 		FMovieSceneEvaluationRange FinalRange(TRange<FFrameTime>(2000 + SectionLength, 2000 + SectionLength + 100), TickResolution, EPlayDirection::Forwards);
-		TestPlayer.Template.Evaluate(FinalRange, TestPlayer);
+		TestPlayer.Template.EvaluateSynchronousBlocking(FinalRange, TestPlayer);
 		Assert(this, TestValue1, TestMagicNumber, TEXT("Restore-State EvaluationHook did not End correctly."));
 
 		TestPlayer.RestorePreAnimatedState();
@@ -462,8 +465,10 @@ bool FMovieScenePreAnimatedStateContextChangedTest::RunTest(const FString& Param
 
 	UMovieSceneCompiledDataManager* CompiledDataManager = UMovieSceneCompiledDataManager::GetPrecompiledData();
 
+	TSharedPtr<FMovieSceneEntitySystemRunner> Runner = MakeShared<FMovieSceneEntitySystemRunner>();
+
 	FMovieSceneCompiledDataID DataID = CompiledDataManager->Compile(Sequence);
-	TestPlayer.Template.Initialize(*Sequence, TestPlayer, CompiledDataManager);
+	TestPlayer.Template.Initialize(*Sequence, TestPlayer, CompiledDataManager, Runner);
 	TestPlayer.Template.EnableGlobalPreAnimatedStateCapture();
 
 	// ---------------------------------------------
@@ -473,7 +478,7 @@ bool FMovieScenePreAnimatedStateContextChangedTest::RunTest(const FString& Param
 		for (int32 i = 0; i < NumEvaluations; ++i)
 		{
 			FMovieSceneEvaluationRange EvaluatedRange(TRange<FFrameTime>(i*(SectionLength/NumEvaluations), (i+1)*(SectionLength/NumEvaluations)), TickResolution, EPlayDirection::Forwards);
-			TestPlayer.Template.Evaluate(EvaluatedRange, TestPlayer);
+			TestPlayer.Template.EvaluateSynchronousBlocking(EvaluatedRange, TestPlayer);
 
 			Assert(this, TestValue1, StartValue + i, TEXT("In-Editor 1: EvaluationHook did not Begin or Update correctly."));
 		}
@@ -490,7 +495,7 @@ bool FMovieScenePreAnimatedStateContextChangedTest::RunTest(const FString& Param
 		for (int32 i = 0; i < NumEvaluations; ++i)
 		{
 			FMovieSceneEvaluationRange EvaluatedRange(TRange<FFrameTime>(i*(SectionLength/NumEvaluations), (i+1)*(SectionLength/NumEvaluations)), TickResolution, EPlayDirection::Forwards);
-			TestPlayer.Template.Evaluate(EvaluatedRange, TestPlayer);
+			TestPlayer.Template.EvaluateSynchronousBlocking(EvaluatedRange, TestPlayer);
 
 			Assert(this, TestValue1, StartValue + i, TEXT("PIE: EvaluationHook did not Begin or Update correctly."));
 		}
diff --git a/Engine/Source/Runtime/MovieScene/Public/EntitySystem/MovieSceneEntitySystemLinker.h b/Engine/Source/Runtime/MovieScene/Public/EntitySystem/MovieSceneEntitySystemLinker.h
--- a/Engine/Source/Runtime/MovieScene/Public/EntitySystem/MovieSceneEntitySystemLinker.h
+++ b/Engine/Source/Runtime/MovieScene/Public/EntitySystem/MovieSceneEntitySystemLinker.h
@@ -28,6 +28,7 @@ namespace MovieScene
 {
 
 	struct FComponentRegistry;
+	enum class ERunnerFlushState;
 	enum class EEntitySystemCategory : uint32;
 
 	enum class EAutoLinkRelevantSystems : uint8
@@ -156,10 +157,6 @@ public:
 		return InstanceRegistry.Get();
 	}
 
-	/** Tells the given instance to finish evaluating */
-	void FinishInstance(FInstanceHandle InstanceHandle);
-
-	/** Links a given type of system. Will assert if the system type isn't allowed on this linker */
 	template<typename SystemType>
 	SystemType* LinkSystem()
 	{
@@ -338,6 +335,7 @@ private:
 
 	UMovieSceneEntitySystem* LinkSystemImpl(TSubclassOf<UMovieSceneEntitySystem> InClassType);
 
+	void HandlePreGarbageCollection();
 	void HandlePostGarbageCollection();
 
 	void TagInvalidBoundObjects();
@@ -357,12 +355,8 @@ private:
 	TSparseArray<UMovieSceneEntitySystem*> EntitySystemsByGlobalGraphID;
 	TMap<UClass*, UMovieSceneEntitySystem*> EntitySystemsRecyclingPool;
 
-	struct FActiveRunnerInfo
-	{
-		FMovieSceneEntitySystemRunner* Runner;
-		bool bIsReentrancyAllowed;
-	};
-	TArray<FActiveRunnerInfo> ActiveRunners;
+	TArray<FMovieSceneEntitySystemRunner*> ActiveRunners;
+	TBitArray<> ActiveRunnerReentrancyFlags;
 
 	TSparseArray<void*> ExtensionsByID;
 
@@ -394,15 +388,3 @@ protected:
 	UE::MovieScene::EAutoLinkRelevantSystems AutoLinkMode;
 	UE::MovieScene::FSystemFilter SystemFilter;
 };
-
-/**
- * Structure for making it possible to make re-entrant evaluation on a linker.
- */
-struct FMovieSceneEntitySystemEvaluationReentrancyWindow
-{
-	UMovieSceneEntitySystemLinker& Linker;
-	int32 CurrentLevel;
-
-	FMovieSceneEntitySystemEvaluationReentrancyWindow(UMovieSceneEntitySystemLinker& InLinker);
-	~FMovieSceneEntitySystemEvaluationReentrancyWindow();
-};
diff --git a/Engine/Source/Runtime/MovieScene/Public/EntitySystem/MovieSceneEntitySystemRunner.h b/Engine/Source/Runtime/MovieScene/Public/EntitySystem/MovieSceneEntitySystemRunner.h
--- a/Engine/Source/Runtime/MovieScene/Public/EntitySystem/MovieSceneEntitySystemRunner.h
+++ b/Engine/Source/Runtime/MovieScene/Public/EntitySystem/MovieSceneEntitySystemRunner.h
@@ -2,6 +2,9 @@
 
 #pragma once
 
+#include "MovieSceneEntityIDs.h"
+#include "MovieSceneSequenceID.h"
+#include "Evaluation/MovieScenePlayback.h"
 #include "Async/TaskGraphInterfaces.h"
 #include "Containers/Array.h"
 #include "Delegates/Delegate.h"
@@ -16,16 +19,51 @@
 
 class UMovieSceneEntitySystemLinker;
 
-namespace UE {
-	namespace MovieScene {
-		class FEntityManager;
-		struct FInstanceRegistry;
-	}
-}
+namespace UE::MovieScene
+{
+
+class FEntityManager;
+struct FInstanceRegistry;
+
+/** Bit-mask enumeration that defines tasks that need running */
+enum class ERunnerFlushState
+{
+	None           = 0,				// Signifies no evaluation is currently underway
+	Start          = 1 << 0,		// Sets up initial evaluation flags for external players and listeners
+	Import         = 1 << 1,		// Update sequence instances and import entities into the entity manager
+	Spawn          = 1 << 2,		// Perform the Spawn phase in the Entity System Graph
+	Instantiation  = 1 << 3,		// Perform the Instantiation phase in the Entity System Graph
+	Evaluation     = 1 << 4,		// Perform the Evaluation phase in the Entity System Graph
+	Finalization   = 1 << 5,		// Perform the Finalization phase in the Entity System Graph and trigger any external events
+	EventTriggers  = 1 << 6,		// (re-entrant) Triggers any bound event triggers - skipped by Finalization if the delegate is not bound
+	PostEvaluation = 1 << 7,		// (re-entrant) Call post evaluation callbacks on sequence instances
+	End            = 1 << 8,		// Counterpart for Start - resets external players and listeners
+
+	// Signifies that, during the Finalization task, there were still outstanding tasks and we need to perform another iteration
+	LoopEval       = Import | Spawn | Instantiation | Evaluation | Finalization | EventTriggers | PostEvaluation | End,
+
+	// Initial flush state
+	Everything     = Start | LoopEval,
+};
+ENUM_CLASS_FLAGS(ERunnerFlushState)
+
+/** Bit-field enumeration that defines flags for an asynchronous update request */
+enum class ERunnerUpdateFlags
+{
+	None           = 0,			// No special behavior - just update the instance along with everything else when the runner is next flushed
+	Flush          = 1<<0,		// Indicate that the runner should be completely flushed next update, ignoring any budgets
+	Finish         = 1<<1,		// Finish the sequence instance and restore state (if enabled) after this update
+	Destroy        = 1<<2,		// Destroy the sequence instance once this request has been fulfilled
+
+	FinalDissectionMask = Finish|Destroy, // A mask that includes flags that should only ever be included on final (or non-)dissected updates, not intermediate updates
+};
+ENUM_CLASS_FLAGS(ERunnerUpdateFlags)
+
+} // namespace UE::MovieScene
 
 DECLARE_MULTICAST_DELEGATE(FMovieSceneEntitySystemEventTriggers);
 
-class MOVIESCENE_API FMovieSceneEntitySystemRunner
+class MOVIESCENE_API FMovieSceneEntitySystemRunner : public TSharedFromThis<FMovieSceneEntitySystemRunner>
 {
 public:
 
@@ -46,21 +84,50 @@ public:
 	/** Detaches this runner from a linker */
 	void DetachFromLinker();
 
+	int32 GetQueuedUpdateCount() const;
 	/** Returns whether this runner has any outstanding updates. */
 	bool HasQueuedUpdates() const;
 	/** Returns whether the given instance is queued for any updates. */
 	bool HasQueuedUpdates(FInstanceHandle Instance) const;
 	/** Queue the given instance for an update with the given context. */
-	void QueueUpdate(const FMovieSceneContext& Context, FInstanceHandle Instance);
-
-	/** Updates the given instance with the given context immediately. This doesn't affect the current update queue. */
-	void Update(const FMovieSceneContext& Context, FInstanceHandle Instance);
-
-	/** Flushes the update queue, blocking the current thread until all queued updates have run. */
-	void Flush();
-
-	/** Finish updating an instance. */
-	void FinishInstance(FInstanceHandle InInstanceHandle);
+	void QueueUpdate(const FMovieSceneContext& Context, FInstanceHandle Instance, UE::MovieScene::ERunnerUpdateFlags UpdateFlags = UE::MovieScene::ERunnerUpdateFlags::None);
+	void QueueUpdate(const FMovieSceneContext& Context, FInstanceHandle Instance, FSimpleDelegate&& InOnFlushedDelegate, UE::MovieScene::ERunnerUpdateFlags UpdateFlags = UE::MovieScene::ERunnerUpdateFlags::None);
+
+	/**
+	 * Queue a final update for the specified instance, optionally finishing and destroying it
+	 * @return true if the update requires a flush, or false if it was finished and/or destroyed immediately without requiring a flush
+	 */
+	bool QueueFinalUpdate(FInstanceHandle Instance, UE::MovieScene::ERunnerUpdateFlags Flags);
+
+	/**
+	 * Abandon the specified instance handle and destroy it immediately. May flush this runner if necessary
+	 */
+	void AbandonAndDestroyInstance(FInstanceHandle Instance);
+
+	/**
+	 * Flushes the update queue and applies any outstanding evaluation logic
+	 * 
+	 * @param BudgetMs      A budget (in milliseconds) to use for evaluation. Evaluation will cease prematurely once this budget is spent
+	 *                      and will process the outstanding work on the next call to Flush. A value of 0.0 signifies no budget - the queue
+	 *                      will be fully processed without leaving any outstanding work
+	 */
+	void Flush(double BudgetMs = 0.f);
+
+	/**
+	 * Flushes any outstanding update tasks in the current evaluation scope with a given budget. Only performs work if this runner is part-way through an evaluation
+	 * 
+	 * @param BudgetMs      A budget (in milliseconds) to use for evaluation. Evaluation will cease prematurely once this budget is spent
+	 *                      and will process the outstanding work on the next call to Flush. A value of 0.0 signifies no budget - the queue
+	 *                      will be fully processed without leaving any outstanding work
+	 * @param TargetState   The desired state to reach. The runner will stop flushing as soon as all the states specified in TargetState have been flushed.
+	 */
+	void FlushOutstanding(double BudgetMs = 0.f, UE::MovieScene::ERunnerFlushState TargetState = UE::MovieScene::ERunnerFlushState::None);
+
+	/**
+	 * Called in the event that the structure of the entity manager has been unexpectedly changed while this runner is active.
+	 * This allows the runner to be reset so it runs from the start of its evaluation loop next time it is evaluated, rather than half-way through
+	 */
+	void ResetFlushState();
 
 	/** Access this runner's currently executing phase */
 	UE::MovieScene::ESystemPhase GetCurrentPhase() const
@@ -68,6 +135,11 @@ public:
 		return CurrentPhase;
 	}
 
+	/**
+	 * Check whether this runner is currently inside an active evaluation loop
+	 */
+	bool IsCurrentlyEvaluating() const;
+
 public:
 
 	UMovieSceneEntitySystemLinker* GetLinker() const;
@@ -77,8 +149,8 @@ public:
 public:
 	
 	// Internal API
-	
-	void MarkForUpdate(FInstanceHandle InInstanceHandle);
+
+	void MarkForUpdate(FInstanceHandle InInstanceHandle, UE::MovieScene::ERunnerUpdateFlags UpdateFlags);
 	FMovieSceneEntitySystemEventTriggers& GetQueuedEventTriggers();
 
 private:
@@ -87,51 +159,99 @@ private:
 
 private:
 
-	// Update loop
-	
-	void DoFlushUpdateQueueOnce();
-	
-	void GameThread_ProcessQueue();
-	void GameThread_SpawnPhase();
-	void GameThread_InstantiationPhase();
-	void GameThread_PostInstantiation();
-	void GameThread_EvaluationPhase();
-	void GameThread_EvaluationFinalizationPhase();
-	void GameThread_PostEvaluationPhase();
+	/**
+	 * Flush the next item in our update loop based off the contents of FlushState
+	 *
+	 * @param  Linker   The linker we are arrached to
+	 * @return True if the loop is allowed to continue, or false if we should not flush any more
+	 */
+	bool FlushNext(UMovieSceneEntitySystemLinker* Linker);
+
+	/**
+	 * Set up initial state before any evaluation runs. Only called once regardless of the number of pending updates we have to process
+	 * Primarily used for setting up external 'is evaluating' flags for re-entrancy and async checks.
+	 */
+	bool StartEvaluation(UMovieSceneEntitySystemLinker* Linker);
+
+	/** Update sequence instances based on currently queued update requests, or outstanding dissected updates */
+	bool GameThread_UpdateSequenceInstances(UMovieSceneEntitySystemLinker* Linker);
+	/** Execute the spawn phase of the entity system graph, if there is anything to do */
+	bool GameThread_SpawnPhase(UMovieSceneEntitySystemLinker* Linker);
+	/** Execute the instantiation phase of the entity system graph, if there is anything to do */
+	bool GameThread_InstantiationPhase(UMovieSceneEntitySystemLinker* Linker);
+	/** Called immediately after instantiation to execute cleanup and bookkeeping tasks. Skipped if instantiation is skipped.*/
+	bool GameThread_PostInstantiation(UMovieSceneEntitySystemLinker* Linker);
+	/** Main entity-system evaluation phase. Blocks this thread until completion. */
+	bool GameThread_EvaluationPhase(UMovieSceneEntitySystemLinker* Linker);
+	/** Finalization phase for triggering external events and other behavior. */
+	void GameThread_EvaluationFinalizationPhase(UMovieSceneEntitySystemLinker* Linker);
+	/** Post-evaluation phase for triggering events. */
+	void GameThread_EventTriggerPhase(UMovieSceneEntitySystemLinker* Linker);
+	/** Post-evaluation phase for clean up. */
+	void GameThread_PostEvaluationPhase(UMovieSceneEntitySystemLinker* Linker);
+
+	/**
+	 * Counterpart for StartEvaluation.
+	 * Called only when our UpdateQueue and DissectedUpdates have been fully processed and there is nothing left to do.
+	 */
+	void EndEvaluation(UMovieSceneEntitySystemLinker* Linker);
 
 private:
 
-	struct FMovieSceneUpdateRequest
-	{
-		FMovieSceneContext Context;
-		FInstanceHandle InstanceHandle;
-	};
+	friend struct FMovieSceneEntitySystemEvaluationReentrancyWindow;
 
 	struct FDissectedUpdate
 	{
+		FSimpleDelegate OnFlushed;
 		FMovieSceneContext Context;
 
 		FInstanceHandle InstanceHandle;
 		int32 Order;
+		UE::MovieScene::ERunnerUpdateFlags UpdateFlags = UE::MovieScene::ERunnerUpdateFlags::None;
 	};
 
 	/** Owner linker */
 	TWeakObjectPtr<UMovieSceneEntitySystemLinker> WeakLinker;
 
+	struct FQueuedUpdateParams
+	{
+		FInstanceHandle InstanceHandle;
+		UE::MovieScene::ERunnerUpdateFlags UpdateFlags = UE::MovieScene::ERunnerUpdateFlags::None;
+	};
+	struct FUpdateParamsAndContext
+	{
+		FSimpleDelegate OnFlushed;
+		FMovieSceneContext Context;
+		FQueuedUpdateParams Params;
+	};
 	/** Queue of sequence instances to be updated */
-	TArray<FMovieSceneUpdateRequest> UpdateQueue;
+	TArray<FUpdateParamsAndContext> UpdateQueue;
 
 	/** When an update is running, the list of actual instances being updated */
-	TArray<FInstanceHandle> CurrentInstances;
+	TArray<FQueuedUpdateParams> CurrentInstances;
 	/** When an update is running, the list of sub-contexts for the requested update */
 	TArray<FDissectedUpdate> DissectedUpdates;
-
-	TGraphTask<FNullGraphTask>* CompletionTask;
-
-	bool bCanQueueEventTriggers;
+	TArray<FSimpleDelegate> OnFlushedDelegates;
 	FMovieSceneEntitySystemEventTriggers EventTriggers;
 
+	/** The number of times this runner has been re-entrant */
+	uint32 ReentrancyCount;
+
 	ENamedThreads::Type GameThread;
 
 	UE::MovieScene::ESystemPhase CurrentPhase;
-};
+
+	/**
+	 * Defines a bitmask of outstanding tasks that need running.
+	 * When a task has completed its bit becomes unset.
+	 * Evaluation can trigger tasks to run again by setting previously cleared bits in this mask.
+	 * A state of ERunnerFlushState::Everything means an evaluation has not yet started.
+	 * A state of ERunnerFlushState::None means an evaluation has just finished.
+	 */
+	UE::MovieScene::ERunnerFlushState FlushState;
+
+	bool bCanQueueEventTriggers;
+
+	/** True if any update has been queued with the ERunnerUpdateFlags::Flush flag since our last flush */
+	bool bRequireFullFlush;
+};
\ No newline at end of file
diff --git a/Engine/Source/Runtime/MovieScene/Public/Evaluation/MovieSceneEvaluationTemplateInstance.h b/Engine/Source/Runtime/MovieScene/Public/Evaluation/MovieSceneEvaluationTemplateInstance.h
--- a/Engine/Source/Runtime/MovieScene/Public/Evaluation/MovieSceneEvaluationTemplateInstance.h
+++ b/Engine/Source/Runtime/MovieScene/Public/Evaluation/MovieSceneEvaluationTemplateInstance.h
@@ -9,10 +9,10 @@
 #include "UObject/WeakObjectPtr.h"
 #include "UObject/StrongObjectPtr.h"
 #include "UObject/WeakObjectPtrTemplates.h"
+#include "MovieSceneFwd.h"
 #include "MovieSceneSequenceID.h"
 #include "Evaluation/MovieSceneRootOverridePath.h"
 #include "Compilation/MovieSceneCompiledDataID.h"
-#include "EntitySystem/MovieSceneEntitySystemRunner.h"
 #include "EntitySystem/MovieSceneSequenceInstanceHandle.h"
 #include "MovieSceneEvaluationTemplateInstance.generated.h"
 
@@ -22,6 +22,7 @@ struct FMovieSceneBlendingAccumulator;
 
 class IMovieScenePlayer;
 class UMovieSceneSequence;
+class FMovieSceneEntitySystemRunner;
 class UMovieSceneEntitySystemLinker;
 class UMovieSceneCompiledDataManager;
 
@@ -47,7 +48,7 @@ public:
 	 * @param Player					The player responsible for playback
 	 * @param TemplateStore				Template store responsible for supplying templates for a given sequence
 	 */
-	MOVIESCENE_API void Initialize(UMovieSceneSequence& RootSequence, IMovieScenePlayer& Player, UMovieSceneCompiledDataManager* InCompiledDataManager);
+	MOVIESCENE_API void Initialize(UMovieSceneSequence& RootSequence, IMovieScenePlayer& Player, UMovieSceneCompiledDataManager* InCompiledDataManager, TWeakPtr<FMovieSceneEntitySystemRunner> InWeakRunner);
 
 	/**
 	 * Evaluate this sequence in a synchronous way.
@@ -55,19 +56,17 @@ public:
 	 * @param Context				Evaluation context containing the time (or range) to evaluate
 	 * @param Player				The player responsible for playback
 	 */
+	MOVIESCENE_API void EvaluateSynchronousBlocking(FMovieSceneContext Context, IMovieScenePlayer& Player);
+
+	UE_DEPRECATED(5.1, "Use EvaluateSynchronousBlocking instead.")
 	MOVIESCENE_API void Evaluate(FMovieSceneContext Context, IMovieScenePlayer& Player);
 
 	/**
-	 * Indicate that we're not going to evaluate this instance again, and that we should tear down any current state
-	 *
-	 * @param Player				The player responsible for playback
+	 * Called when this instance has completely finished playing
 	 */
-	MOVIESCENE_API void Finish(IMovieScenePlayer& Player);
+	void OnFinished();
 
-	/**
-	 * Called when this template has finished evaluating to perform cleanup
-	 */
-	MOVIESCENE_API void OnFinish();
+	MOVIESCENE_API void ResetDirectorInstances();
 
 	bool IsValid() const
 	{
@@ -94,12 +93,15 @@ public:
 		return CompiledDataManager;
 	}
 
+	TSharedPtr<FMovieSceneEntitySystemRunner> GetRunner() const
+	{
+		return WeakRunner.Pin();
+	}
+
 	MOVIESCENE_API bool HasEverUpdated() const;
 
 	MOVIESCENE_API UMovieSceneEntitySystemLinker* GetEntitySystemLinker() const;
 
-	MOVIESCENE_API FMovieSceneEntitySystemRunner& GetEntitySystemRunner();
-
 	MOVIESCENE_API const FMovieSceneSequenceHierarchy* GetHierarchy() const;
 
 	MOVIESCENE_API void GetSequenceParentage(const UE::MovieScene::FInstanceHandle InstanceHandle, TArray<UE::MovieScene::FInstanceHandle>& OutParentHandles) const;
@@ -142,13 +144,13 @@ private:
 	UPROPERTY()
 	TObjectPtr<UMovieSceneCompiledDataManager> CompiledDataManager;
 
+	TWeakPtr<FMovieSceneEntitySystemRunner> WeakRunner;
+
 	UE::MovieScene::FRootInstanceHandle RootInstanceHandle;
 
 	UPROPERTY()
 	TObjectPtr<UMovieSceneEntitySystemLinker> EntitySystemLinker;
 
-	FMovieSceneEntitySystemRunner EntitySystemRunner;
-
 	/** Map of director instances by sequence ID. Kept alive by this map assuming this struct is reference collected */
 	UPROPERTY()
 	TMap<FMovieSceneSequenceID, TObjectPtr<UObject>> DirectorInstances;
diff --git a/Engine/Source/Runtime/MovieScene/Public/MovieSceneSequencePlayer.h b/Engine/Source/Runtime/MovieScene/Public/MovieSceneSequencePlayer.h
--- a/Engine/Source/Runtime/MovieScene/Public/MovieSceneSequencePlayer.h
+++ b/Engine/Source/Runtime/MovieScene/Public/MovieSceneSequencePlayer.h
@@ -766,6 +766,9 @@ protected:
 	/** (Optional) Externally supplied time controller */
 	TSharedPtr<FMovieSceneTimeController> TimeController;
 
+	/** (Optional) Synchronous runner to use when no tick manager is in use */
+	TSharedPtr<FMovieSceneEntitySystemRunner> SynchronousRunner;
+
 private:
 
 	/** The event that will be broadcast every time the sequence is updated */
diff --git a/Engine/Source/Runtime/MovieScene/Public/MovieSceneSequenceTickInterval.h b/Engine/Source/Runtime/MovieScene/Public/MovieSceneSequenceTickInterval.h
--- a/Engine/Source/Runtime/MovieScene/Public/MovieSceneSequenceTickInterval.h
+++ b/Engine/Source/Runtime/MovieScene/Public/MovieSceneSequenceTickInterval.h
@@ -37,6 +37,10 @@ public:
 	UPROPERTY(EditAnywhere, Category="Playback", meta=(DisplayName="Tick Interval", Units=s))
 	float TickIntervalSeconds = 0.f;
 
+	/** Defines an approximate budget for evaluation of this sequence (and any other sequences with the same tick interval) */
+	UPROPERTY(EditAnywhere, Category="Playback", meta=(DisplayName="Evaluation Budget", ForceUnits=us))
+	float EvaluationBudgetMicroseconds = 0.f;
+
 	/** When true, the sequence will continue to tick and progress even when the world is paused */
 	UPROPERTY(EditAnywhere, Category="Playback")
 	bool bTickWhenPaused = false;
@@ -64,6 +68,7 @@ public:
 	friend bool operator==(const FMovieSceneSequenceTickInterval& A, const FMovieSceneSequenceTickInterval& B)
 	{
 		return A.TickIntervalSeconds == B.TickIntervalSeconds
+			&& A.EvaluationBudgetMicroseconds == B.EvaluationBudgetMicroseconds
 			&& A.bTickWhenPaused == B.bTickWhenPaused
 			&& A.bAllowRounding == B.bAllowRounding;
 	}
diff --git a/Engine/Source/Runtime/MovieScene/Public/MovieSceneSequenceTickManager.h b/Engine/Source/Runtime/MovieScene/Public/MovieSceneSequenceTickManager.h
--- a/Engine/Source/Runtime/MovieScene/Public/MovieSceneSequenceTickManager.h
+++ b/Engine/Source/Runtime/MovieScene/Public/MovieSceneSequenceTickManager.h
@@ -103,7 +103,7 @@ public:
 	 * Retrieve the runner associated with the specified (resolved) tick interval.
 	 * @return The runner that flushes evaluations for the specified tick interval, or nullptr if there is none
 	 */
-	FMovieSceneEntitySystemRunner* GetRunner(const FMovieSceneSequenceTickInterval& TickInterval);
+	TSharedPtr<FMovieSceneEntitySystemRunner> GetRunner(const FMovieSceneSequenceTickInterval& TickInterval);
 
 	/**
 	 * Add an action that will be executed once all the pending evaluations in this tick manager have been flushed.
@@ -163,16 +163,22 @@ private:
 		/** The linker that owns all the entity data, systems and instances */
 		UMovieSceneEntitySystemLinker* Linker;
 		/** Runner responsible for evaluating each phase of the pipeline */
-		FMovieSceneEntitySystemRunner Runner;
+		TSharedPtr<FMovieSceneEntitySystemRunner> Runner;
 
 		/** The tick interval of this group, rounded to the nearest Sequencer.TickIntervalGroupingResolutionMs */
 		int32 RoundedTickIntervalMs = 0;
 
+		/** The frame budget for all linkers within this linker group in milliseconds */
+		float FrameBudgetMs = 0.f;
+
 		/** The value of UWorld::GetUnpausedTimeSeconds last time this group was evaluated */
 		float LastUnpausedTimeSeconds = -1;
 		/** The value of UWorld::GetTimeSeconds last time this group was evaluated */
 		float LastTimeSeconds = -1;
 
+		/** The number of times this group has been starved of frames for budgeted evaluation */
+		uint16 BudgetedStarvationCount = 0;
+
 		/** Total number of clients in this group */
 		uint16 NumClients = 0;
 	};
diff --git a/Engine/Source/Runtime/UMG/Private/Animation/UMGSequencePlayer.cpp b/Engine/Source/Runtime/UMG/Private/Animation/UMGSequencePlayer.cpp
--- a/Engine/Source/Runtime/UMG/Private/Animation/UMGSequencePlayer.cpp
+++ b/Engine/Source/Runtime/UMG/Private/Animation/UMGSequencePlayer.cpp
@@ -33,8 +33,10 @@ UUMGSequencePlayer::UUMGSequencePlayer(const FObjectInitializer& ObjectInitializ
 	Animation = nullptr;
 	bIsEvaluating = false;
 	bIsStopping = false;
+	bIsBeginningPlay = false;
 	bCompleteOnPostEvaluation = false;
 	UserTag = NAME_None;
+	BlockedDeltaTimeCompensation = 0.f;
 }
 
 void UUMGSequencePlayer::InitSequencePlayer(UWidgetAnimation& InAnimation, UUserWidget& InUserWidget)
@@ -71,11 +73,18 @@ UMovieSceneEntitySystemLinker* UUMGSequencePlayer::ConstructEntitySystemLinker()
 
 void UUMGSequencePlayer::Tick(float DeltaTime)
 {
-	if (bIsStopping)
+	if (bIsStopping || bIsBeginningPlay || bIsEvaluating)
 	{
+		if (bIsEvaluating)
+		{
+			BlockedDeltaTimeCompensation += DeltaTime;
+		}
 		return;
 	}
 
+	DeltaTime += BlockedDeltaTimeCompensation;
+	BlockedDeltaTimeCompensation = 0.f;
+
 	if ( PlayerStatus == EMovieScenePlayerStatus::Playing )
 	{
 		FFrameTime DeltaFrameTime = (bIsPlayingForward ? DeltaTime * PlaybackSpeed : -DeltaTime * PlaybackSpeed) * AnimationResolution;
@@ -177,56 +186,52 @@ void UUMGSequencePlayer::UpdateInternal(FFrameTime LastTimePosition, FFrameTime
 				PlayerStatus);
 		Context.SetHasJumped(bHasJumped);
 
-		UMovieSceneSequence* MovieSceneSequence = RootTemplateInstance.GetSequence(MovieSceneSequenceID::Root);
-		const bool bIsSequenceBlocking = EnumHasAnyFlags(MovieSceneSequence->GetFlags(), EMovieSceneSequenceFlags::BlockingEvaluation);
-
-		const bool bIsRunningWithTickManager = CVarUserWidgetUseParallelAnimation.GetValueOnGameThread();
-
-		if (bIsRunningWithTickManager)
+		TSharedPtr<FMovieSceneEntitySystemRunner> Runner = RootTemplateInstance.GetRunner();
+		if (Runner)
 		{
-			UUserWidget* Widget = UserWidget.Get();
-			UUMGSequenceTickManager* TickManager = Widget ? ToRawPtr(Widget->AnimationTickManager) : nullptr;
+			UE::MovieScene::ERunnerUpdateFlags UpdateFlags = UE::MovieScene::ERunnerUpdateFlags::None;
+			FSimpleDelegate OnComplete;
 
-			if (!TickManager)
+			if (bCompleteOnPostEvaluation)
 			{
-				return;
-			}
+				bIsStopping = true;
+				bCompleteOnPostEvaluation = false;
 
-			if (!bIsSequenceBlocking)
-			{
-				// Queue an evaluation of this player's widget animation, to be evaluated later by our
-				// global tick manager as part of a glorious multi-threaded job fest.
-				FMovieSceneEntitySystemRunner& Runner = TickManager->GetRunner();
-				Runner.QueueUpdate(Context, RootTemplateInstance.GetRootInstanceHandle());
-				// WARNING: the evalution hasn't run yet so don't run any stateful code after this point
-				// unless you know it's OK to do so. Most likely, you want to run stateful code in the
-				// PostEvaluation method, or queue up a latent action.
+				OnComplete = FSimpleDelegate::CreateUObject(this, &UUMGSequencePlayer::HandleLatentStop);
+				UpdateFlags = UE::MovieScene::ERunnerUpdateFlags::Flush;
 			}
-			else
+
+			Runner->QueueUpdate(Context, RootTemplateInstance.GetRootInstanceHandle(), MoveTemp(OnComplete), UpdateFlags);
+
+			if (Runner == SynchronousRunner || !UE::UMG::GAsyncAnimationControlFlow)
 			{
-				// Synchronous evaluation. Sucks for performance.
-				RootTemplateInstance.Evaluate(Context, *this);
-				// The latent actions will be run by the tick manager.
+				Runner->Flush();
+				ApplyLatentActions();
 			}
 		}
-		else
-		{
-			// Synchronous evaluation. Sucks for performance.
-			RootTemplateInstance.Evaluate(Context, *this);
-			ApplyLatentActions();
-		}
 	}
 }
 
 void UUMGSequencePlayer::PlayInternal(double StartAtTime, double EndAtTime, int32 InNumLoopsToPlay, EUMGSequencePlayMode::Type InPlayMode, float InPlaybackSpeed, bool bInRestoreState)
 {
-	RootTemplateInstance.Initialize(*Animation, *this, nullptr);
+	UUserWidget* Widget = UserWidget.Get();
+	UUMGSequenceTickManager* TickManager = Widget ? ToRawPtr(Widget->AnimationTickManager) : nullptr;
+
+	TSharedPtr<FMovieSceneEntitySystemRunner> RunnerToUse = TickManager ? TickManager->GetRunner() : nullptr;
+	if (EnumHasAnyFlags(Animation->GetFlags(), EMovieSceneSequenceFlags::BlockingEvaluation))
+	{
+		SynchronousRunner = MakeShared<FMovieSceneEntitySystemRunner>();
+		RunnerToUse = SynchronousRunner;
+	}
+
+	RootTemplateInstance.Initialize(*Animation, *this, nullptr, RunnerToUse);
 
 	if (bInRestoreState)
 	{
 		RootTemplateInstance.EnableGlobalPreAnimatedStateCapture();
 	}
 
+	bIsBeginningPlay = true;
 	bRestoreState = bInRestoreState;
 	PlaybackSpeed = FMath::Abs(InPlaybackSpeed);
 	PlayMode = InPlayMode;
@@ -262,19 +267,24 @@ void UUMGSequencePlayer::PlayInternal(double StartAtTime, double EndAtTime, int3
 
 	PlayerStatus = EMovieScenePlayerStatus::Playing;
 
-	UUserWidget* Widget = UserWidget.Get();
-	UUMGSequenceTickManager* TickManager = Widget ? ToRawPtr(Widget->AnimationTickManager) : nullptr;
-
 	// Playback assumes the start frame has already been evaulated, so we also want to evaluate any events on the start frame here.
-	if (TickManager && RootTemplateInstance.IsValid())
+	if (RunnerToUse)
 	{
 		const FMovieSceneContext Context(FMovieSceneEvaluationRange(AbsolutePlaybackStart + TimeCursorPosition, AbsolutePlaybackStart + TimeCursorPosition, AnimationResolution), PlayerStatus);
 
+		auto OnBegunPlay = [this]
+		{
+			this->bIsBeginningPlay = false;
+		};
+
 		// We queue an update instead of immediately flushing the entire linker so that we don't incur a cascade of flushes on frames when multiple animations are played
 		// In rare cases where the linker must be flushed immediately PreTick, the queue should be manually flushed 
+		RunnerToUse->QueueUpdate(Context, RootTemplateInstance.GetRootInstanceHandle(), FSimpleDelegate::CreateWeakLambda(this, OnBegunPlay), UE::MovieScene::ERunnerUpdateFlags::Flush);
 
-		FMovieSceneEntitySystemRunner& Runner = TickManager->GetRunner();
-		Runner.QueueUpdate(Context, RootTemplateInstance.GetRootInstanceHandle());
+		if (RunnerToUse == SynchronousRunner || !UE::UMG::GAsyncAnimationControlFlow)
+		{
+			RunnerToUse->Flush();
+		}
 	}
 }
 
@@ -314,16 +324,15 @@ void UUMGSequencePlayer::Pause()
 	PlayerStatus = EMovieScenePlayerStatus::Stopped;
 
 	// Evaluate the sequence at its current time, with a status of 'stopped' to ensure that animated state pauses correctly. (ie. audio sounds should stop/pause)
-	const FMovieSceneContext Context(FMovieSceneEvaluationRange(AbsolutePlaybackStart + TimeCursorPosition, AbsolutePlaybackStart + TimeCursorPosition, AnimationResolution), PlayerStatus);
-	if (RootTemplateInstance.HasEverUpdated())
+	TSharedPtr<FMovieSceneEntitySystemRunner> RunnerToUpdate = RootTemplateInstance.HasEverUpdated() ? RootTemplateInstance.GetRunner() : nullptr;
+	if (RunnerToUpdate)
 	{
-		UUserWidget* Widget = UserWidget.Get();
-		UUMGSequenceTickManager* TickManager = Widget ? ToRawPtr(Widget->AnimationTickManager) : nullptr;
+		const FMovieSceneContext Context(FMovieSceneEvaluationRange(AbsolutePlaybackStart + TimeCursorPosition, AbsolutePlaybackStart + TimeCursorPosition, AnimationResolution), PlayerStatus);
+		RunnerToUpdate->QueueUpdate(Context, RootTemplateInstance.GetRootInstanceHandle(), UE::MovieScene::ERunnerUpdateFlags::Flush);
 
-		if (TickManager)
+		if (RunnerToUpdate == SynchronousRunner || !UE::UMG::GAsyncAnimationControlFlow)
 		{
-			FMovieSceneEntitySystemRunner& Runner = TickManager->GetRunner();
-			Runner.QueueUpdate(Context, RootTemplateInstance.GetRootInstanceHandle());
+			RunnerToUpdate->Flush();
 		}
 	}
 }
@@ -338,12 +347,19 @@ void UUMGSequencePlayer::Reverse()
 
 void UUMGSequencePlayer::Stop()
 {
+	using namespace UE::MovieScene;
+
 	if (NeedsQueueLatentAction())
 	{
 		QueueLatentAction(FMovieSceneSequenceLatentActionDelegate::CreateUObject(this, &UUMGSequencePlayer::Stop));
 		return;
 	}
 
+	if (PlayerStatus == EMovieScenePlayerStatus::Stopped)
+	{
+		return;
+	}
+
 	PlayerStatus = EMovieScenePlayerStatus::Stopped;
 
 	UMovieSceneSequence* MovieSceneSequence = RootTemplateInstance.GetSequence(MovieSceneSequenceID::Root);
@@ -375,33 +391,35 @@ void UUMGSequencePlayer::Stop()
 	// Prevent any other updates
 	bIsStopping = true;
 
-	if (bIsAsync)
+	if (TSharedPtr<FMovieSceneEntitySystemRunner> Runner = RootTemplateInstance.GetRunner())
 	{
+		// It would be better if we did not have to actively re-evaluate the animation at the start time when Stop is called
+		// so we don't have to queue an additional update, but this has to be included for legacy reasons.
+		Runner->QueueUpdate(Context, RootTemplateInstance.GetRootInstanceHandle(), FSimpleDelegate::CreateUObject(this, &UUMGSequencePlayer::HandleLatentStop), UE::MovieScene::ERunnerUpdateFlags::Flush);
 
-		TickManager->AddLatentAction(FMovieSceneSequenceLatentActionDelegate::CreateUObject(this, &UUMGSequencePlayer::HandleLatentStop));
-		TickManager->GetRunner().QueueUpdate(Context, RootTemplateInstance.GetRootInstanceHandle());
-	}
-	else
-	{
-		RootTemplateInstance.Evaluate(Context, *this);
-		HandleLatentStop();
+		if (Runner == SynchronousRunner || !UE::UMG::GAsyncAnimationControlFlow)
+		{
+			Runner->Flush();
+		}
 	}
 }
 
 void UUMGSequencePlayer::HandleLatentStop()
 {
-	// Linker may not necessarily be valid here if we have already torn down the sequence
-	if (UMovieSceneEntitySystemLinker* Linker = RootTemplateInstance.GetEntitySystemLinker())
+	if (TSharedPtr<FMovieSceneEntitySystemRunner> Runner = RootTemplateInstance.GetRunner())
 	{
-		UE::MovieScene::FSequenceInstance* RootInstance = RootTemplateInstance.FindInstance(MovieSceneSequenceID::Root);
-		// Finish the root instance
-		if (RootInstance)
-		{
-			RootInstance->Finish(Linker);
-		}
+		Runner->QueueFinalUpdate(RootTemplateInstance.GetRootInstanceHandle(), UE::MovieScene::ERunnerUpdateFlags::Finish);
+
+		// Even if our request to Finish the instance was queued, we can wait until the next flush for those effects to be seen
+		// This will most likely happen immediately anyway since the runner will keep looping until its queue is empty,
+		// And we are already inside an active evaluation
 	}
 
-	RootTemplateInstance.OnFinish();
+	PlayerStatus = EMovieScenePlayerStatus::Stopped;
+
+	// This isn't really necessary because UMG animations cannot have director instances
+	// but it is here for completeness
+	RootTemplateInstance.ResetDirectorInstances();
 
 	if (bRestoreState)
 	{
@@ -469,14 +487,6 @@ void UUMGSequencePlayer::PreEvaluation(const FMovieSceneContext& Context)
 void UUMGSequencePlayer::PostEvaluation(const FMovieSceneContext& Context)
 {
 	bIsEvaluating = false;
-
-	if (bCompleteOnPostEvaluation)
-	{
-		bCompleteOnPostEvaluation = false;
-
-		PlayerStatus = EMovieScenePlayerStatus::Stopped;
-		HandleLatentStop();
-	}
 }
 
 bool UUMGSequencePlayer::NeedsQueueLatentAction() const
@@ -533,6 +543,12 @@ void UUMGSequencePlayer::RemoveEvaluationData()
 
 	if (SequenceInstance && Linker)
 	{
+		TSharedPtr<FMovieSceneEntitySystemRunner> Runner = RootTemplateInstance.GetRunner();
+		if (Runner && Runner->IsCurrentlyEvaluating())
+		{
+			Runner->FlushOutstanding();
+		}
+
 		SequenceInstance->Ledger.UnlinkEverything(Linker);
 		SequenceInstance->InvalidateCachedData(Linker);
 	}
diff --git a/Engine/Source/Runtime/UMG/Private/Animation/UMGSequenceTickManager.cpp b/Engine/Source/Runtime/UMG/Private/Animation/UMGSequenceTickManager.cpp
--- a/Engine/Source/Runtime/UMG/Private/Animation/UMGSequenceTickManager.cpp
+++ b/Engine/Source/Runtime/UMG/Private/Animation/UMGSequenceTickManager.cpp
@@ -11,19 +11,31 @@
 
 DECLARE_CYCLE_STAT(TEXT("Flush End of Frame Animations"), MovieSceneEval_FlushEndOfFrameAnimations, STATGROUP_MovieSceneEval);
 
-static TAutoConsoleVariable<int32> CVarUMGMaxAnimationLatentActions(
-	TEXT("Widget.MaxAnimationLatentActions"),
-	100,
-	TEXT("Defines the maximum number of latent actions that can be run in one frame."),
-	ECVF_Default
-);
-int32 GFlushUMGAnimationsAtEndOfFrame = 1;
-static FAutoConsoleVariableRef CVarUMGAnimationsAtEndOfFrame(
-	TEXT("UMG.FlushAnimationsAtEndOfFrame"),
-	GFlushUMGAnimationsAtEndOfFrame,
-	TEXT("Whether to automatically flush any outstanding animations at the end of the frame, or just wait until next frame."),
-	ECVF_Default
-);
+
+namespace UE::UMG
+{
+
+	static TAutoConsoleVariable<int32> CVarUMGMaxAnimationLatentActions(
+		TEXT("Widget.MaxAnimationLatentActions"),
+		100,
+		TEXT("Defines the maximum number of latent actions that can be run in one frame."),
+		ECVF_Default
+	);
+	int32 GFlushUMGAnimationsAtEndOfFrame = 1;
+	static FAutoConsoleVariableRef CVarUMGAnimationsAtEndOfFrame(
+		TEXT("UMG.FlushAnimationsAtEndOfFrame"),
+		GFlushUMGAnimationsAtEndOfFrame,
+		TEXT("Whether to automatically flush any outstanding animations at the end of the frame, or just wait until next frame."),
+		ECVF_Default
+	);
+
+	float GAnimationBudgetMs = 0.0f;
+	FAutoConsoleVariableRef CVarAnimationBudgetMs(
+		TEXT("UMG.AnimationBudgetMs"),
+		GAnimationBudgetMs,
+		TEXT("(Default: 0.0) EXPERIMENTAL: A per-frame animation budget to use for evaluation of all UMG animations this frame.")
+	);
+} // namespace UE::UMG
 
 extern TAutoConsoleVariable<bool> CVarUserWidgetUseParallelAnimation;
 
@@ -31,6 +43,10 @@ UUMGSequenceTickManager::UUMGSequenceTickManager(const FObjectInitializer& Init)
 	: Super(Init)
 	, bIsTicking(false)
 {
+	if (!HasAnyFlags(RF_ClassDefaultObject))
+	{
+		Runner = MakeShared<FMovieSceneEntitySystemRunner>();
+	}
 }
 
 void UUMGSequenceTickManager::AddWidget(UUserWidget* InWidget)
@@ -100,7 +116,7 @@ void UUMGSequenceTickManager::TickWidgetAnimations(float DeltaSeconds)
 	//
 	// The WidgetData have one frame delay (they are updated at the end of the frame).
 	// This may delay the animation update by one frame.
-
+	const bool bIsCurrentlyEvaluating = Runner->IsCurrentlyEvaluating();
 	{
 		SCOPE_CYCLE_UOBJECT(ContextScope, this);
 
@@ -109,6 +125,7 @@ void UUMGSequenceTickManager::TickWidgetAnimations(float DeltaSeconds)
 		{
 			UUserWidget* UserWidget = WidgetIter.Key().Get();
 			FSequenceTickManagerWidgetData& WidgetData = WidgetIter.Value();
+
 			WidgetData.bActionsAndAnimationTicked = false;
 
 			if (!UserWidget)
@@ -117,10 +134,14 @@ void UUMGSequenceTickManager::TickWidgetAnimations(float DeltaSeconds)
 			}
 			else if (!UserWidget->IsConstructed())
 			{
-				UserWidget->TearDownAnimations();
-				UserWidget->AnimationTickManager = nullptr;
+				if (!bIsCurrentlyEvaluating)
+				{
+					// Wait until the runner has finished evaluating before we tear stuff down
+					UserWidget->TearDownAnimations();
+					UserWidget->AnimationTickManager = nullptr;
 
-				WidgetIter.RemoveCurrent();
+					WidgetIter.RemoveCurrent();
+				}
 			}
 			else if (!WidgetData.bIsTicking)
 			{
@@ -131,12 +152,15 @@ void UUMGSequenceTickManager::TickWidgetAnimations(float DeltaSeconds)
 				// @todo: There is a chance that relative animations hitting this code path will resume with
 				// different relative bases due to the way the ecs data is destroyed and re-created.
 				// In order to fix this we would have to annex that data instead of destroying it.
-				if (WidgetData.bLastKnownTickState)
+				if (!bIsCurrentlyEvaluating)
 				{
-					UserWidget->DisableAnimations();
-				}
+					if (WidgetData.bLastKnownTickState)
+					{
+						UserWidget->DisableAnimations();
+					}
 
-				WidgetData.bLastKnownTickState = false;
+					WidgetData.bLastKnownTickState = false;
+				}
 			}
 			else
 			{
@@ -162,38 +186,36 @@ void UUMGSequenceTickManager::TickWidgetAnimations(float DeltaSeconds)
 
 	ForceFlush();
 
-	for (auto WidgetIter = WeakUserWidgetData.CreateIterator(); WidgetIter; ++WidgetIter)
+	if (!Runner->IsCurrentlyEvaluating())
 	{
-		UUserWidget* UserWidget = WidgetIter.Key().Get();
-		ensureMsgf(UserWidget, TEXT("Widget became null during animation tick!"));
-
-		if (UserWidget)
+		for (auto WidgetIter = WeakUserWidgetData.CreateIterator(); WidgetIter; ++WidgetIter)
 		{
-			if (WidgetIter.Value().bActionsAndAnimationTicked)
+			UUserWidget* UserWidget = WidgetIter.Key().Get();
+			ensureMsgf(UserWidget, TEXT("Widget became null during animation tick!"));
+
+			if (UserWidget)
 			{
-				UserWidget->PostTickActionsAndAnimation(DeltaSeconds);
+				// If this widget no longer has any animations playing, it doesn't need to be ticked any more
+				if (UserWidget->ActiveSequencePlayers.Num() == 0)
+				{
+					UserWidget->UpdateCanTick();
+					UserWidget->AnimationTickManager = nullptr;
+					WidgetIter.RemoveCurrent();
+				}
 			}
-
-			// If this widget no longer has any animations playing, it doesn't need to be ticked any more
-			if (UserWidget->ActiveSequencePlayers.Num() == 0)
+			else
 			{
-				UserWidget->UpdateCanTick();
-				UserWidget->AnimationTickManager = nullptr;
 				WidgetIter.RemoveCurrent();
 			}
 		}
-		else
-		{
-			WidgetIter.RemoveCurrent();
-		}
 	}
 }
 
 void UUMGSequenceTickManager::ForceFlush()
 {
-	if (Runner.IsAttachedToLinker())
+	if (Runner->IsAttachedToLinker())
 	{
-		Runner.Flush();
+		Runner->Flush(UE::UMG::GAnimationBudgetMs);
 		RunLatentActions();
 	}
 }
@@ -206,11 +228,12 @@ void UUMGSequenceTickManager::HandleSlatePostTick(float DeltaSeconds)
 		return;
 	}
 
-	if (GFlushUMGAnimationsAtEndOfFrame && Runner.IsAttachedToLinker() && Runner.HasQueuedUpdates())
+	// Only tick widgets at the end of the frame if our runner has completely finished, and we still have updates
+	if (UE::UMG::GFlushUMGAnimationsAtEndOfFrame && Runner->IsAttachedToLinker() && Runner->HasQueuedUpdates() && !Runner->IsCurrentlyEvaluating())
 	{
 		SCOPE_CYCLE_COUNTER(MovieSceneEval_FlushEndOfFrameAnimations);
 
-		Runner.Flush();
+		Runner->Flush();
 		RunLatentActions();
 	}
 }
@@ -227,10 +250,24 @@ void UUMGSequenceTickManager::ClearLatentActions(UObject* Object)
 
 void UUMGSequenceTickManager::RunLatentActions()
 {
-	LatentActionManager.RunLatentActions([this]
+	if (!this->Runner->IsCurrentlyEvaluating())
 	{
-		this->Runner.Flush();
-	});
+		int32 UpdateCount   = this->Runner->GetQueuedUpdateCount();
+		uint64 SystemSerial = this->Linker->EntityManager.GetSystemSerial();
+
+		LatentActionManager.RunLatentActions([this, &SystemSerial, &UpdateCount]
+		{
+			int32  NewUpdateCount  = this->Runner->GetQueuedUpdateCount();
+			uint64 NewSystemSerial = this->Linker->EntityManager.GetSystemSerial();
+			if (NewUpdateCount != UpdateCount || NewSystemSerial != SystemSerial)
+			{
+				UpdateCount = NewUpdateCount;
+				SystemSerial = NewSystemSerial;
+
+				this->Runner->Flush();
+			}
+		});
+	}
 }
 
 UUMGSequenceTickManager* UUMGSequenceTickManager::Get(UObject* PlaybackContext)
@@ -257,7 +294,7 @@ UUMGSequenceTickManager* UUMGSequenceTickManager::Get(UObject* PlaybackContext)
 
 		TickManager->Linker = UMovieSceneEntitySystemLinker::FindOrCreateLinker(Owner, UE::MovieScene::EEntitySystemLinkerRole::UMG, TEXT("UMGAnimationEntitySystemLinker"));
 		check(TickManager->Linker);
-		TickManager->Runner.AttachToLinker(TickManager->Linker);
+		TickManager->Runner->AttachToLinker(TickManager->Linker);
 
 		FSlateApplication& SlateApp = FSlateApplication::Get();
 		FDelegateHandle PreTickHandle = SlateApp.OnPreTick().AddUObject(TickManager, &UUMGSequenceTickManager::TickWidgetAnimations);
diff --git a/Engine/Source/Runtime/UMG/Private/UserWidget.cpp b/Engine/Source/Runtime/UMG/Private/UserWidget.cpp
--- a/Engine/Source/Runtime/UMG/Private/UserWidget.cpp
+++ b/Engine/Source/Runtime/UMG/Private/UserWidget.cpp
@@ -709,6 +709,7 @@ void UUserWidget::OnAnimationFinishedPlaying(UUMGSequencePlayer& Player)
 
 	if ( Player.GetPlaybackStatus() == EMovieScenePlayerStatus::Stopped )
 	{
+		ensureAlways(!StoppedSequencePlayers.Contains(&Player));
 		StoppedSequencePlayers.Add(&Player);
 
 		if (AnimationTickManager)
@@ -1056,9 +1057,15 @@ bool UUserWidget::GetIsVisible() const
 
 void UUserWidget::SetVisibility(ESlateVisibility InVisibility)
 {
+	ESlateVisibility OldVisibility = GetVisibility();
+
 	Super::SetVisibility(InVisibility);
-	OnNativeVisibilityChanged.Broadcast(InVisibility);
-	OnVisibilityChanged.Broadcast(InVisibility);
+
+	if (OldVisibility != GetVisibility())
+	{
+		OnNativeVisibilityChanged.Broadcast(InVisibility);
+		OnVisibilityChanged.Broadcast(InVisibility);
+	}
 }
 
 bool UUserWidget::IsInViewport() const
@@ -1530,7 +1537,6 @@ void UUserWidget::TickActionsAndAnimation(float InDeltaTime)
 
 void UUserWidget::PostTickActionsAndAnimation(float InDeltaTime)
 {
-	ClearStoppedSequencePlayers();
 }
 
 void UUserWidget::FlushAnimations()
diff --git a/Engine/Source/Runtime/UMG/Public/Animation/UMGSequencePlayer.h b/Engine/Source/Runtime/UMG/Public/Animation/UMGSequencePlayer.h
--- a/Engine/Source/Runtime/UMG/Public/Animation/UMGSequencePlayer.h
+++ b/Engine/Source/Runtime/UMG/Public/Animation/UMGSequencePlayer.h
@@ -113,6 +113,8 @@ private:
 	UPROPERTY()
 	FMovieSceneRootEvaluationTemplateInstance RootTemplateInstance;
 
+	TSharedPtr<FMovieSceneEntitySystemRunner> SynchronousRunner;
+
 	/** The resolution at which all FFrameNumbers are stored */
 	FFrameRate AnimationResolution;
 
@@ -128,6 +130,9 @@ private:
 	/** Time at which to end the animation after looping */
 	FFrameTime EndTime;
 
+	/** A compensation for delta-time we were not able to evaluate with due to ongoing budgeted evaluation */
+	float BlockedDeltaTimeCompensation;
+
 	/** Status of the player (e.g play, stopped) */
 	EMovieScenePlayerStatus::Type PlayerStatus;
 
@@ -165,6 +170,9 @@ private:
 	/** Set to true when this player is in the process of being stopped */
 	bool bIsStopping : 1;
 
+	/** Set to true when this player has just started playing, but hasn't finished evaluating yet */
+	bool bIsBeginningPlay : 1;
+
 	/** Set to true if we need to run the finishing logic in post-evaluation */
 	bool bCompleteOnPostEvaluation : 1;
 
diff --git a/Engine/Source/Runtime/UMG/Public/Animation/UMGSequenceTickManager.h b/Engine/Source/Runtime/UMG/Public/Animation/UMGSequenceTickManager.h
--- a/Engine/Source/Runtime/UMG/Public/Animation/UMGSequenceTickManager.h
+++ b/Engine/Source/Runtime/UMG/Public/Animation/UMGSequenceTickManager.h
@@ -40,7 +40,7 @@ public:
 	UUMGSequenceTickManager(const FObjectInitializer& Init);
 
 	UMovieSceneEntitySystemLinker* GetLinker() { return Linker; }
-	FMovieSceneEntitySystemRunner& GetRunner() { return Runner; }
+	TSharedPtr<FMovieSceneEntitySystemRunner> GetRunner() { return Runner; }
 
 	void AddLatentAction(FMovieSceneSequenceLatentActionDelegate Delegate);
 	void ClearLatentActions(UObject* Object);
@@ -69,10 +69,11 @@ private:
 	UPROPERTY(transient)
 	TObjectPtr<UMovieSceneEntitySystemLinker> Linker;
 
-	FMovieSceneEntitySystemRunner Runner;
+	TSharedPtr<FMovieSceneEntitySystemRunner> Runner;
 
-	bool bIsTicking;
 	FDelegateHandle SlateApplicationPreTickHandle, SlateApplicationPostTickHandle;
 
 	FMovieSceneLatentActionManager LatentActionManager;
+
+	bool bIsTicking;
 };
