Sequencer: Fixed migrating entities sometimes resulting in component values getting lost

FEntityManager::MigrateEntity is responsible for moving the components of an entity between allocations with mismatching combinations of components (ie, if an entity containing ABCDE has B and D removed, it migrates to an allocation containing ACE). The algorithm for this iterates all the destination component types and relocates component data from the old allocation to the new one, for any components that exist in both. For efficiency, the hole that is left in the old allocation is filled with the last entity in that allocation, the same way TArray::RemoveAtSwap works. For example, if entity ID 3 is migrated, we should see the following result:

Before migration:
[ 0, 1, 2, 3, 4, 5 ]

After migration:
[ 0, 1, 2, 5, 4  ]

Unfortunately there were a few flaws in the way this was implemented. Firstly, any components that existed in the old allocation that do not exist in the new allocation were not correctly relocated when performing the RemoveAtSwap. This meant that (in the example above) entity 5 would get entity 3's values for those components. Secondly there were some errors in the way we relocated component values. Relocation _must_ be performed on uninitialized memory to be correct for RAII, and the memory that is left behind _must not_ be destructed. These constraints were not adhered to because AddEntityToAllocation would always default construct its components, and we would not destruct the previous values for components that were not migrated.

I've tried to clarify this code by encapsulating the RemoveAtSwap logic that destructs a component and relocates the last element in a lambda, and this is now called for all source headers that do not get migrated to the new allocation. For any components that _do_ get migrated, we simply relocate the component and perform the swap manually - no destruction is necessary in this case because both values are kept alive.

diff --git a/Engine/Source/Runtime/MovieScene/Private/EntitySystem/MovieSceneEntityManager.cpp b/Engine/Source/Runtime/MovieScene/Private/EntitySystem/MovieSceneEntityManager.cpp
--- a/Engine/Source/Runtime/MovieScene/Private/EntitySystem/MovieSceneEntityManager.cpp
+++ b/Engine/Source/Runtime/MovieScene/Private/EntitySystem/MovieSceneEntityManager.cpp
@@ -1999,14 +1999,12 @@ int32 FEntityManager::MigrateEntity(int32 DestAllocationIndex, int32 SourceAlloc
 	{
 		Dst = GrowAllocation(DestAllocationIndex);
 	}
-	
-	FComponentHeader* SrcComponentHeader = Src->ComponentHeaders;
-	FComponentHeader* SrcLastComponentHeader = Src->ComponentHeaders + Src->GetNumComponentTypes() - 1;
 
 	FMovieSceneEntityID EntityID = Src->GetEntityIDs()[SourceEntityIndex];
 
-	// Make space for the new entity in the destination
-	const int32 DestEntityIndex = AddEntityToAllocation(DestAllocationIndex, EntityID);
+	// Make space for the new entity in the destination without initializing the memory
+	// This is important because we either default construct or relocate construct into this new allocation ourselves
+	const int32 DestEntityIndex = AddEntityToAllocation(DestAllocationIndex, EntityID, EMemoryType::DefaultConstructed);
 
 	const int32 SrcOffset = SourceEntityIndex;
 	const int32 DstOffset = DestEntityIndex;
@@ -2019,38 +2017,86 @@ int32 FEntityManager::MigrateEntity(int32 DestAllocationIndex, int32 SourceAlloc
 	Src->PostModifyStructureExcludingHeaders(WriteContext);
 	Dst->PostModifyStructureExcludingHeaders(WriteContext);
 
-	// Iterate all destination component types
+	const int32 NumSrcHeaders = Src->GetNumComponentTypes();
+	int32 SrcHeaderIndex = 0;
+
+	// This function takes a component header and address from the original allocation and removes it from the allocation
+	// by destructing the component, and relocating the last element in the component array into its place, similar to TArray::RemoveAtSwap
+	// This allows minimal shuffling of component data when moving entities between allocations.
+	auto RemoveAtSwapComponent = [this, LastEntityIndex, WriteContext](const FComponentHeader& SrcComponentHeader, int32 RemoveAtIndex)
+	{
+		SrcComponentHeader.PostWriteComponents(WriteContext);
+		if (!SrcComponentHeader.IsTag())
+		{
+			const FComponentTypeInfo& ComponentTypeInfo = this->ComponentRegistry->GetComponentTypeChecked(SrcComponentHeader.ComponentType);
+
+			void* RemovedValueAddress = SrcComponentHeader.GetValuePtr(RemoveAtIndex);
+
+			// Destroy the component at the address
+			ComponentTypeInfo.DestructItems(RemovedValueAddress, 1);
+
+			// Swap the last entity's components in this allocation with the migrated one
+			if (RemoveAtIndex != LastEntityIndex)
+			{
+				void* LastItem = SrcComponentHeader.GetValuePtr(LastEntityIndex);
+				ComponentTypeInfo.RelocateConstructItems(RemovedValueAddress, LastItem, 1);
+			}
+		}
+	};
+
+	// Iterate all destination component types and either default construct, or relocate the existing component
 	for (int32 DstHeaderOffset = 0; DstHeaderOffset < Dst->GetNumComponentTypes(); ++DstHeaderOffset)
 	{
 		FComponentHeader* DstComponentHeader = &Dst->ComponentHeaders[DstHeaderOffset];
-		DstComponentHeader->PostWriteComponents(WriteContext);
-
-		// Try to find a matching source component type
-		while (SrcComponentHeader != SrcLastComponentHeader && SrcComponentHeader->ComponentType.BitIndex() < DstComponentHeader->ComponentType.BitIndex())
+		if (DstComponentHeader->IsTag())
 		{
-			SrcComponentHeader->PostWriteComponents(WriteContext);
-			++SrcComponentHeader;
+			continue;
 		}
 
-		// Copy the component value to the new allocation
-		if (DstComponentHeader->ComponentType == SrcComponentHeader->ComponentType && !SrcComponentHeader->IsTag())
+		// Try to locate a matching source component header for this component type by walking through the source headers
+		// Component headers are sorted by component type, so if we encounter any with a type ID less than the current destination type
+		// it must not exist in the new allocation, and so should be destroyed from this one
+		while (SrcHeaderIndex < NumSrcHeaders && Src->ComponentHeaders[SrcHeaderIndex].ComponentType.BitIndex() < DstComponentHeader->ComponentType.BitIndex())
 		{
-			const FComponentTypeInfo& ComponentTypeInfo = ComponentRegistry->GetComponentTypeChecked(SrcComponentHeader->ComponentType);
+			// Destination does not have this component type so it needs destroying
+			RemoveAtSwapComponent(Src->ComponentHeaders[SrcHeaderIndex], SrcOffset);
+			++SrcHeaderIndex;
+		}
 
-			void* DstValue = DstComponentHeader->GetValuePtr(DstOffset);
-			void* SrcValue = SrcComponentHeader->GetValuePtr(SrcOffset);
+		const FComponentTypeInfo& ComponentTypeInfo = ComponentRegistry->GetComponentTypeChecked(DstComponentHeader->ComponentType);
+		void* DstValue = DstComponentHeader->GetValuePtr(DstOffset);
 
-			ComponentTypeInfo.RelocateConstructItems(DstValue, SrcValue, 1);
+		// Relocate the component value to the new allocation if it is the same type as our current source header
+		if (SrcHeaderIndex < NumSrcHeaders && Src->ComponentHeaders[SrcHeaderIndex].ComponentType.BitIndex() == DstComponentHeader->ComponentType.BitIndex())
+		{
+			void* MigratedValueAddress = Src->ComponentHeaders[SrcHeaderIndex].GetValuePtr(SrcOffset);
+			ComponentTypeInfo.RelocateConstructItems(DstValue, MigratedValueAddress, 1);
 
-			// If we need to swap a tail component with the one we just moved out, do that now
+			// We do not use RemoveAtSwapComponent here because the SrcValue is now already considered destructed
+			// Manually swap the last entity's components in this allocation with the migrated one
 			if (LastEntityIndex != SrcOffset)
 			{
-				void* SwapSource = SrcComponentHeader->GetValuePtr(LastEntityIndex);
-				ComponentTypeInfo.RelocateConstructItems(SrcValue, SwapSource, 1);
+				void* LastItem = Src->ComponentHeaders[SrcHeaderIndex].GetValuePtr(LastEntityIndex);
+				ComponentTypeInfo.RelocateConstructItems(MigratedValueAddress, LastItem, 1);
 			}
+
+			// This source header is now dealt with so skip over it
+			++SrcHeaderIndex;
+		}
+		// or default construct it if it's a new component that didn't exist before
+		else
+		{
+			ComponentTypeInfo.ConstructItems(DstValue, 1);
 		}
 	}
 
+	// Process any remaining components in the source allocation that were not in the destination
+	// by destructing the components and potentially relocating the RemoveAtSwap candidate
+	for ( ; SrcHeaderIndex < NumSrcHeaders; ++SrcHeaderIndex)
+	{
+		RemoveAtSwapComponent(Src->ComponentHeaders[SrcHeaderIndex], SrcOffset);
+	}
+
 	// When removing we just swap the tail element with the element to remove, and fix up the indices
 	if (LastEntityIndex != SrcOffset)
 	{
@@ -2206,7 +2252,7 @@ FEntityAllocation* FEntityManager::GrowAllocation(int32 AllocationIndex, int32 M
 	return NewAllocation;
 }
 
-int32 FEntityManager::AddEntityToAllocation(int32 AllocationIndex, FMovieSceneEntityID ID)
+int32 FEntityManager::AddEntityToAllocation(int32 AllocationIndex, FMovieSceneEntityID ID, EMemoryType MemoryType)
 {
 	CheckCanChangeStructure();
 
@@ -2228,7 +2274,7 @@ int32 FEntityManager::AddEntityToAllocation(int32 AllocationIndex, FMovieSceneEn
 	{
 		Header.PostWriteComponents(WriteContext);
 
-		if (!Header.IsTag())
+		if (!Header.IsTag() && MemoryType == EMemoryType::DefaultConstructed)
 		{
 			const FComponentTypeInfo& ComponentTypeInfo = ComponentRegistry->GetComponentTypeChecked(Header.ComponentType);
 			void* Value = Header.GetValuePtr(ActualEntityOffset);
diff --git a/Engine/Source/Runtime/MovieScene/Public/EntitySystem/MovieSceneEntityManager.h b/Engine/Source/Runtime/MovieScene/Public/EntitySystem/MovieSceneEntityManager.h
--- a/Engine/Source/Runtime/MovieScene/Public/EntitySystem/MovieSceneEntityManager.h
+++ b/Engine/Source/Runtime/MovieScene/Public/EntitySystem/MovieSceneEntityManager.h
@@ -861,6 +861,11 @@ public:
 private:
 	friend struct FEntityInitializer; 
 
+	enum class EMemoryType
+	{
+		Uninitialized,
+		DefaultConstructed
+	};
 	void OnStructureChanged();
 
 	FEntityAllocation* CreateEntityAllocation(const FComponentMask& EntityComponentMask, uint16 InitialCapacity, uint16 MaxCapacity, FEntityAllocation* MigrateComponentDataFrom = nullptr);
@@ -871,7 +876,7 @@ private:
 
 	void CopyComponents(int32 DestAllocationIndex, int32 DestEntityIndex, int32 SourceAllocationIndex, int32 SourceEntityIndex, const FComponentMask* OptionalMask = nullptr);
 
-	int32 AddEntityToAllocation(int32 AllocationIndex, FMovieSceneEntityID ID);
+	int32 AddEntityToAllocation(int32 AllocationIndex, FMovieSceneEntityID ID, EMemoryType MemoryType = EMemoryType::DefaultConstructed);
 	void RemoveEntityFromAllocation(int32 AllocationIndex, int32 SourceEntryIndexWithinAllocation);
 
 	FEntityAllocation* MigrateAllocation(int32 AllocationIndex, const FComponentMask& NewComponentMask);
diff --git a/Engine/Source/Runtime/MovieSceneTracks/Private/Systems/MovieSceneMaterialParameterSystem.cpp b/Engine/Source/Runtime/MovieSceneTracks/Private/Systems/MovieSceneMaterialParameterSystem.cpp
--- a/Engine/Source/Runtime/MovieSceneTracks/Private/Systems/MovieSceneMaterialParameterSystem.cpp
+++ b/Engine/Source/Runtime/MovieSceneTracks/Private/Systems/MovieSceneMaterialParameterSystem.cpp
@@ -137,6 +137,11 @@ struct FOverlappingMaterialParameterHandler
 				{
 					Linker->EntityManager.AddComponent(Input, BuiltInComponents->BlendChannelInput, Output->BlendChannelID);
 				}
+				else
+				{
+					// If the bound material changed, we might have been re-assigned a different blend channel so make sure it's up to date
+					Linker->EntityManager.WriteComponentChecked(Input, BuiltInComponents->BlendChannelInput, Output->BlendChannelID);
+				}
 			}
 		}
 		else
@@ -157,6 +162,7 @@ struct FOverlappingMaterialParameterHandler
 			{
 				Linker->EntityManager.RemoveComponent(Input, BuiltInComponents->BlendChannelInput);
 			}
+			Output->BlendChannelID = FMovieSceneBlendChannelID();
 		}
 
 		Output->NumContributors = NumContributors;
@@ -170,6 +176,12 @@ struct FOverlappingMaterialParameterHandler
 
 			Linker->EntityManager.AddComponent(Output->OutputEntityID, BuiltInComponents->Tags.NeedsUnlink);
 			Output->OutputEntityID = FMovieSceneEntityID();
+
+			if (System->DoubleBlenderSystem)
+			{
+				System->DoubleBlenderSystem->ReleaseBlendChannel(Output->BlendChannelID);
+			}
+			Output->BlendChannelID = FMovieSceneBlendChannelID();
 		}
 	}
 
