Sequencer: Rework tearing down sequences so that we do not require any re-evaluation

diff --git a/Engine/Source/Runtime/MovieScene/Private/EntitySystem/MovieSceneEntitySystemLinker.cpp b/Engine/Source/Runtime/MovieScene/Private/EntitySystem/MovieSceneEntitySystemLinker.cpp
--- a/Engine/Source/Runtime/MovieScene/Private/EntitySystem/MovieSceneEntitySystemLinker.cpp
+++ b/Engine/Source/Runtime/MovieScene/Private/EntitySystem/MovieSceneEntitySystemLinker.cpp
@@ -435,6 +435,12 @@ void UMovieSceneEntitySystemLinker::CleanGarbage()
 
 	InstanceRegistry->CleanupLinkerEntities(FreedEntities);
 
+	// If we have any runners part-way through an evaluation, we need to reset them so that they re-evaluate from the start
+	ResetActiveRunners();
+}
+
+void UMovieSceneEntitySystemLinker::ResetActiveRunners()
+{
 	// If we have any runners part-way through an evaluation, we need to reset them so that they re-evaluate from the start
 	for (int32 Index = ActiveRunners.Num()-1; Index >= 0; --Index)
 	{
diff --git a/Engine/Source/Runtime/MovieScene/Private/EntitySystem/MovieSceneEntitySystemRunner.cpp b/Engine/Source/Runtime/MovieScene/Private/EntitySystem/MovieSceneEntitySystemRunner.cpp
--- a/Engine/Source/Runtime/MovieScene/Private/EntitySystem/MovieSceneEntitySystemRunner.cpp
+++ b/Engine/Source/Runtime/MovieScene/Private/EntitySystem/MovieSceneEntitySystemRunner.cpp
@@ -473,6 +473,34 @@ void FMovieSceneEntitySystemRunner::ResetFlushState()
 	}
 }
 
+void FMovieSceneEntitySystemRunner::DiscardQueuedUpdates(FInstanceHandle Instance)
+{
+	using namespace UE::MovieScene;
+
+	for (int32 Index = UpdateQueue.Num()-1; Index >= 0; --Index)
+	{
+		if (UpdateQueue[Index].Params.InstanceHandle == Instance)
+		{
+			UpdateQueue.RemoveAt(Index, 1, false);
+		}
+	}
+	for (int32 Index = DissectedUpdates.Num()-1; Index >= 0; --Index)
+	{
+		if (DissectedUpdates[Index].InstanceHandle == Instance)
+		{
+			DissectedUpdates.RemoveAt(Index, 1, false);
+		}
+	}
+
+	for (int32 Index = CurrentInstances.Num()-1; Index >= 0; --Index)
+	{
+		if (CurrentInstances[Index].InstanceHandle == Instance)
+		{
+			CurrentInstances.RemoveAt(Index, 1, false);
+		}
+	}
+}
+
 bool FMovieSceneEntitySystemRunner::GameThread_UpdateSequenceInstances(UMovieSceneEntitySystemLinker* Linker)
 {
 	using namespace UE::MovieScene;
diff --git a/Engine/Source/Runtime/MovieScene/Private/Evaluation/MovieSceneEvaluationTemplateInstance.cpp b/Engine/Source/Runtime/MovieScene/Private/Evaluation/MovieSceneEvaluationTemplateInstance.cpp
--- a/Engine/Source/Runtime/MovieScene/Private/Evaluation/MovieSceneEvaluationTemplateInstance.cpp
+++ b/Engine/Source/Runtime/MovieScene/Private/Evaluation/MovieSceneEvaluationTemplateInstance.cpp
@@ -34,20 +34,11 @@ void FMovieSceneRootEvaluationTemplateInstance::BeginDestroy()
 	{
 		if (TSharedPtr<FMovieSceneEntitySystemRunner> Runner = WeakRunner.Pin())
 		{
-			// There's no way to know whether this instance is currently being included in the evaluation so if it is currently evaluating we
-			// just need to flush any current work. FlushOutstanding does nothing if it is not currently in the middle of an evaluation
-			Runner->FlushOutstanding();
-
-			const bool bRequiresFlush = Runner->QueueFinalUpdate(RootInstanceHandle, UE::MovieScene::ERunnerUpdateFlags::Destroy);
-			if (bRequiresFlush)
-			{
-			  Runner->Flush();
-			}
-		}
-		else
-		{
-			EntitySystemLinker->GetInstanceRegistry()->DestroyInstance(RootInstanceHandle);
+			Runner->DiscardQueuedUpdates(RootInstanceHandle);
 		}
+
+		EntitySystemLinker->GetInstanceRegistry()->DestroyInstance(RootInstanceHandle);
+		EntitySystemLinker->ResetActiveRunners();
 	}
 
 	CompiledDataManager = nullptr;
diff --git a/Engine/Source/Runtime/MovieScene/Public/EntitySystem/MovieSceneEntitySystemLinker.h b/Engine/Source/Runtime/MovieScene/Public/EntitySystem/MovieSceneEntitySystemLinker.h
--- a/Engine/Source/Runtime/MovieScene/Public/EntitySystem/MovieSceneEntitySystemLinker.h
+++ b/Engine/Source/Runtime/MovieScene/Public/EntitySystem/MovieSceneEntitySystemLinker.h
@@ -331,6 +331,8 @@ public:
 	void PostInstantation(FMovieSceneEntitySystemRunner& InRunner);
 	void EndEvaluation(FMovieSceneEntitySystemRunner& InRunner);
 
+	void ResetActiveRunners();
+
 private:
 
 	UMovieSceneEntitySystem* LinkSystemImpl(TSubclassOf<UMovieSceneEntitySystem> InClassType);
diff --git a/Engine/Source/Runtime/MovieScene/Public/EntitySystem/MovieSceneEntitySystemRunner.h b/Engine/Source/Runtime/MovieScene/Public/EntitySystem/MovieSceneEntitySystemRunner.h
--- a/Engine/Source/Runtime/MovieScene/Public/EntitySystem/MovieSceneEntitySystemRunner.h
+++ b/Engine/Source/Runtime/MovieScene/Public/EntitySystem/MovieSceneEntitySystemRunner.h
@@ -129,6 +129,11 @@ public:
 	 */
 	void ResetFlushState();
 
+	/**
+	 * Discard any queued updates for the specified sequence instance.
+	 */
+	void DiscardQueuedUpdates(FInstanceHandle Instance);
+
 	/** Access this runner's currently executing phase */
 	UE::MovieScene::ESystemPhase GetCurrentPhase() const
 	{
