Sequencer: Clear out instance handle after destroying the instance on BeginDestroy

diff --git a/Engine/Source/Runtime/MovieScene/Private/Evaluation/MovieSceneEvaluationTemplateInstance.cpp b/Engine/Source/Runtime/MovieScene/Private/Evaluation/MovieSceneEvaluationTemplateInstance.cpp
--- a/Engine/Source/Runtime/MovieScene/Private/Evaluation/MovieSceneEvaluationTemplateInstance.cpp
+++ b/Engine/Source/Runtime/MovieScene/Private/Evaluation/MovieSceneEvaluationTemplateInstance.cpp
@@ -41,6 +41,7 @@ void FMovieSceneRootEvaluationTemplateInstance::BeginDestroy()
 		EntitySystemLinker->ResetActiveRunners();
 	}
 
+	RootInstanceHandle = UE::MovieScene::FRootInstanceHandle();
 	CompiledDataManager = nullptr;
 	EntitySystemLinker = nullptr;
 }
