Sequencer: Fixed a crash when calling OnStop that nulled out the current Sequence

Also corrected flushing the runner when not necessary

diff --git a/Engine/Source/Runtime/MovieScene/Private/MovieSceneSequencePlayer.cpp b/Engine/Source/Runtime/MovieScene/Private/MovieSceneSequencePlayer.cpp
--- a/Engine/Source/Runtime/MovieScene/Private/MovieSceneSequencePlayer.cpp
+++ b/Engine/Source/Runtime/MovieScene/Private/MovieSceneSequencePlayer.cpp
@@ -447,7 +447,7 @@ void UMovieSceneSequencePlayer::StopInternal(FFrameTime TimeToResetTo)
 		if (Runner)
 		{
 			// Finish but do not destroy
-			if (!Runner->QueueFinalUpdate(RootTemplateInstance.GetRootInstanceHandle(), UE::MovieScene::ERunnerUpdateFlags::Finish))
+			if (Runner->QueueFinalUpdate(RootTemplateInstance.GetRootInstanceHandle(), UE::MovieScene::ERunnerUpdateFlags::Finish))
 			{
 				Runner->Flush();
 			}
@@ -464,7 +464,7 @@ void UMovieSceneSequencePlayer::StopInternal(FFrameTime TimeToResetTo)
 		if (Runner)
 		{
 			// Finish but do not destroy
-			if (!Runner->QueueFinalUpdate(RootTemplateInstance.GetRootInstanceHandle(), UE::MovieScene::ERunnerUpdateFlags::Finish))
+			if (Runner->QueueFinalUpdate(RootTemplateInstance.GetRootInstanceHandle(), UE::MovieScene::ERunnerUpdateFlags::Finish))
 			{
 				Runner->Flush();
 			}
@@ -1744,6 +1744,11 @@ void UMovieSceneSequencePlayer::QueueLatentAction(FMovieSceneSequenceLatentActio
 
 void UMovieSceneSequencePlayer::RunLatentActions()
 {
+	if (!Sequence)
+	{
+		return;
+	}
+
 	if (ensure(TickManager) && !EnumHasAnyFlags(Sequence->GetFlags(), EMovieSceneSequenceFlags::BlockingEvaluation))
 	{
 		TickManager->RunLatentActions();
