Sequencer: Fixed lack of component dependency between the master instantiator system and the material parameter collection system

diff --git a/Engine/Source/Runtime/MovieScene/Private/EntitySystem/MovieSceneMasterInstantiatorSystem.cpp b/Engine/Source/Runtime/MovieScene/Private/EntitySystem/MovieSceneMasterInstantiatorSystem.cpp
--- a/Engine/Source/Runtime/MovieScene/Private/EntitySystem/MovieSceneMasterInstantiatorSystem.cpp
+++ b/Engine/Source/Runtime/MovieScene/Private/EntitySystem/MovieSceneMasterInstantiatorSystem.cpp
@@ -12,6 +12,10 @@ UMovieSceneMasterInstantiatorSystem::UMovieSceneMasterInstantiatorSystem(const F
 	: Super(ObjInit)
 {
 	RelevantComponent = UE::MovieScene::FBuiltInComponentTypes::Get()->Tags.Master;
+	if (HasAnyFlags(RF_ClassDefaultObject))
+	{
+		DefineComponentProducer(GetClass(), UE::MovieScene::FBuiltInComponentTypes::Get()->Tags.Master);
+	}
 }
 
 void UMovieSceneMasterInstantiatorSystem::OnRun(FSystemTaskPrerequisites& InPrerequisites, FSystemSubsequentTasks& Subsequents)
diff --git a/Engine/Source/Runtime/MovieSceneTracks/Private/Systems/MovieSceneMaterialParameterCollectionSystem.cpp b/Engine/Source/Runtime/MovieSceneTracks/Private/Systems/MovieSceneMaterialParameterCollectionSystem.cpp
--- a/Engine/Source/Runtime/MovieSceneTracks/Private/Systems/MovieSceneMaterialParameterCollectionSystem.cpp
+++ b/Engine/Source/Runtime/MovieSceneTracks/Private/Systems/MovieSceneMaterialParameterCollectionSystem.cpp
@@ -136,6 +136,7 @@ UMovieSceneMaterialParameterCollectionSystem::UMovieSceneMaterialParameterCollec
 
 	if (HasAnyFlags(RF_ClassDefaultObject))
 	{
+		DefineComponentConsumer(GetClass(), BuiltInComponents->Tags.Master);
 		DefineComponentProducer(GetClass(), TracksComponents->BoundMaterial);
 		DefineImplicitPrerequisite(UMovieSceneCachePreAnimatedStateSystem::StaticClass(), GetClass());
 	}
