Sequencer: Fixed queued updates being lost when queued as part of an external callback during evaluation

 - The previous logic was flawed in that if there were no previous queued updates, but something queued an update during the re-entrancy window (such as a new animation being played), we would end up throwing the updates away by swapping the previous (empty) queue with the current (populated) queue. Now we always swap the queues to restore the old queue first, then append any new queued updates if necessary.

diff --git a/Engine/Source/Runtime/MovieScene/Private/EntitySystem/MovieSceneEntitySystemRunner.cpp b/Engine/Source/Runtime/MovieScene/Private/EntitySystem/MovieSceneEntitySystemRunner.cpp
--- a/Engine/Source/Runtime/MovieScene/Private/EntitySystem/MovieSceneEntitySystemRunner.cpp
+++ b/Engine/Source/Runtime/MovieScene/Private/EntitySystem/MovieSceneEntitySystemRunner.cpp
@@ -986,26 +986,21 @@ void FMovieSceneEntitySystemRunner::GameThread_PostEvaluationPhase(UMovieSceneEn
 		}
 	}
 
-	// Append our cached updates if there are outstanding updates from a re-entrant call, otherwise
-	// just swap them back
-	if (UpdateQueue.Num() > 0 && TmpUpdateQueue.Num() > 0)
+	// If we had any updates _before_ we called the external callbacks (which are stored in TmpUpdateQueue), make sure those are maintained first
+	Swap(TmpUpdateQueue, UpdateQueue);
+
+	// TmpUpdateQueue now contains any updates that were queued _during_ the external callbacks, so we append those to our old queue (which is now stored back in UpdateQueue)
+	if (TmpUpdateQueue.Num() > 0)
 	{
 		UpdateQueue.Append(TmpUpdateQueue);
 	}
-	else
-	{
-		Swap(TmpUpdateQueue, UpdateQueue);
-	}
 
-	// Same with the dissected updates
-	if (DissectedUpdates.Num() > 0 && TmpDissectedUpdates.Num() > 0)
+	// Do the same with the dissected updates
+	Swap(TmpDissectedUpdates, DissectedUpdates);
+	if (TmpDissectedUpdates.Num() > 0)
 	{
 		DissectedUpdates.Append(TmpDissectedUpdates);
 	}
-	else
-	{
-		Swap(TmpDissectedUpdates, DissectedUpdates);
-	}
 
 	// If we have any pending updates, we need to run another evaluation.
 	// This will cause us to effectively loop back to ERunnerFlushState::Import
