Make constructor of MovieScenePlayer thread-safe to avoid potential race when async loading is turned on

diff --git a/Engine/Source/Runtime/MovieScene/Private/MovieScenePlayer.cpp b/Engine/Source/Runtime/MovieScene/Private/MovieScenePlayer.cpp
--- a/Engine/Source/Runtime/MovieScene/Private/MovieScenePlayer.cpp
+++ b/Engine/Source/Runtime/MovieScene/Private/MovieScenePlayer.cpp
@@ -4,6 +4,7 @@
 #include "IMovieScenePlayer.h"
 #include "Evaluation/MovieSceneEvaluationTemplateInstance.h"
 #include "EntitySystem/MovieSceneSequenceInstance.h"
+#include "Misc/ScopeRWLock.h"
 #include "MovieSceneSequence.h"
 #include "MovieSceneSequenceID.h"
 
@@ -12,6 +13,7 @@ namespace UE
 namespace MovieScene
 {
 
+static FRWLock                          GGlobalPlayerRegistryLock;
 static TSparseArray<IMovieScenePlayer*> GGlobalPlayerRegistry;
 
 } // namespace MovieScene
@@ -19,18 +21,19 @@ static TSparseArray<IMovieScenePlayer*> GGlobalPlayerRegistry;
 
 IMovieScenePlayer::IMovieScenePlayer()
 {
+	FWriteScopeLock ScopeLock(UE::MovieScene::GGlobalPlayerRegistryLock);
 	UniqueIndex = UE::MovieScene::GGlobalPlayerRegistry.Add(this);
 }
 
 IMovieScenePlayer::~IMovieScenePlayer()
 {
-	ensureMsgf(IsInGameThread(), TEXT("Destruction must occur on the game thread"));
-
+	FWriteScopeLock ScopeLock(UE::MovieScene::GGlobalPlayerRegistryLock);
 	UE::MovieScene::GGlobalPlayerRegistry.RemoveAt(UniqueIndex, 1);
 }
 
 IMovieScenePlayer* IMovieScenePlayer::Get(uint16 InUniqueIndex)
 {
+	FReadScopeLock ScopeLock(UE::MovieScene::GGlobalPlayerRegistryLock);
 	check(UE::MovieScene::GGlobalPlayerRegistry.IsValidIndex(InUniqueIndex));
 	return UE::MovieScene::GGlobalPlayerRegistry[InUniqueIndex];
 }
