Sequencer: re-add last update before finishing a sequence instance so that systems can do any necessary tear down logic

This fixes a problem where the camera cut track instance wouldn't get its entities unlinked until the player was destroyed. But this time, the player's OnStopped had already been called, and some state needed for the track instance's OnDestroyed has been cleared.

diff --git a/Engine/Plugins/MovieScene/TemplateSequence/Source/TemplateSequence/Private/CameraAnimationSequencePlayer.cpp b/Engine/Plugins/MovieScene/TemplateSequence/Source/TemplateSequence/Private/CameraAnimationSequencePlayer.cpp
--- a/Engine/Plugins/MovieScene/TemplateSequence/Source/TemplateSequence/Private/CameraAnimationSequencePlayer.cpp
+++ b/Engine/Plugins/MovieScene/TemplateSequence/Source/TemplateSequence/Private/CameraAnimationSequencePlayer.cpp
@@ -420,7 +420,7 @@ void UCameraAnimationSequencePlayer::Stop()
 
 	if (TSharedPtr<FMovieSceneEntitySystemRunner> Runner = RootTemplateInstance.GetRunner())
 	{
-		if (Runner->QueueFinalUpdate(RootTemplateInstance.GetRootInstanceHandle(), UE::MovieScene::ERunnerUpdateFlags::Finish))
+		if (Runner->QueueFinalUpdate(RootTemplateInstance.GetRootInstanceHandle()))
 		{
 			// @todo: we should really only call flush _once_ for all camera anim sequences
 			Runner->Flush();
diff --git a/Engine/Source/Editor/Sequencer/Private/Sequencer.cpp b/Engine/Source/Editor/Sequencer/Private/Sequencer.cpp
--- a/Engine/Source/Editor/Sequencer/Private/Sequencer.cpp
+++ b/Engine/Source/Editor/Sequencer/Private/Sequencer.cpp
@@ -587,7 +587,7 @@ FSequencer::~FSequencer()
 {
 	if (Runner->IsAttachedToLinker())
 	{
-		Runner->QueueFinalUpdate(RootTemplateInstance.GetRootInstanceHandle(), UE::MovieScene::ERunnerUpdateFlags::Finish);
+		Runner->QueueFinalUpdate(RootTemplateInstance.GetRootInstanceHandle());
 		Runner->Flush();
 	}
 
@@ -630,7 +630,7 @@ void FSequencer::Close()
 
 	if (Runner->IsAttachedToLinker())
 	{
-		Runner->QueueFinalUpdate(RootTemplateInstance.GetRootInstanceHandle(), UE::MovieScene::ERunnerUpdateFlags::Finish);
+		Runner->QueueFinalUpdate(RootTemplateInstance.GetRootInstanceHandle());
 		Runner->Flush();
 	}
 
@@ -977,7 +977,7 @@ void FSequencer::ResetToNewRootSequence(UMovieSceneSequence& NewSequence)
 
 	if (Runner->IsAttachedToLinker())
 	{
-		Runner->QueueFinalUpdate(RootTemplateInstance.GetRootInstanceHandle(), UE::MovieScene::ERunnerUpdateFlags::Finish);
+		Runner->QueueFinalUpdate(RootTemplateInstance.GetRootInstanceHandle());
 		Runner->Flush();
 	}
 
diff --git a/Engine/Source/Runtime/MovieScene/Private/EntitySystem/MovieSceneEntitySystemRunner.cpp b/Engine/Source/Runtime/MovieScene/Private/EntitySystem/MovieSceneEntitySystemRunner.cpp
--- a/Engine/Source/Runtime/MovieScene/Private/EntitySystem/MovieSceneEntitySystemRunner.cpp
+++ b/Engine/Source/Runtime/MovieScene/Private/EntitySystem/MovieSceneEntitySystemRunner.cpp
@@ -185,12 +185,25 @@ void FMovieSceneEntitySystemRunner::QueueUpdate(const FMovieSceneContext& Contex
 	UpdateQueue.Add(FUpdateParamsAndContext{ OnFlushedDelegate, Context, { Instance, Flags } });
 }
 
-bool FMovieSceneEntitySystemRunner::QueueFinalUpdate(FInstanceHandle InInstanceHandle, UE::MovieScene::ERunnerUpdateFlags Flags)
+bool FMovieSceneEntitySystemRunner::QueueFinalUpdate(FInstanceHandle InInstanceHandle)
 {
-	using namespace UE::MovieScene;
+	return QueueFinalUpdateImpl(InInstanceHandle, FSimpleDelegate(), false);
+}
+
+bool FMovieSceneEntitySystemRunner::QueueFinalUpdate(FInstanceHandle InInstanceHandle, FSimpleDelegate&& InOnLastFlushDelegate)
+{
+	return QueueFinalUpdateImpl(InInstanceHandle, MoveTemp(InOnLastFlushDelegate), false);
+}
+
+bool FMovieSceneEntitySystemRunner::QueueFinalUpdateAndDestroy(FInstanceHandle InInstanceHandle)
+{
+	return QueueFinalUpdateImpl(InInstanceHandle, FSimpleDelegate(), true);
+}
+
 
-	// It is not well defined to call this function with no flags - that would signify the function should not have been called at all
-	check(Flags != ERunnerUpdateFlags::None);
+bool FMovieSceneEntitySystemRunner::QueueFinalUpdateImpl(FInstanceHandle InInstanceHandle, FSimpleDelegate&& InOnLastFlushDelegate, bool bDestroyInstance)
+{
+	using namespace UE::MovieScene;
 
 	UMovieSceneEntitySystemLinker* Linker = GetLinker();
 	if (!Linker)
@@ -206,47 +219,47 @@ bool FMovieSceneEntitySystemRunner::QueueFinalUpdate(FInstanceHandle InInstanceH
 
 	FSequenceInstance& Instance = InstanceRegistry->MutateInstance(InInstanceHandle);
 
-	if (EnumHasAnyFlags(Flags, UE::MovieScene::ERunnerUpdateFlags::Flush))
-	{
-		bRequireFullFlush = true;
-	}
-
-	ERunnerFlushState UnsafeDestroyMask = ERunnerFlushState::Everything & ~(ERunnerFlushState::PostEvaluation | ERunnerFlushState::End);
-
-	// If this instance has not current updates and we're not in the middle of an update loop
-	// we can finish and destroy this instance right away. The effects of the instance being torn down
-	// (ie entities being unlinked) will be seen once the runner is flushed
+	// If the following conditions are met, we can destroy this instance right away:
+	// 
+	// 1. the instance has no entities left to unlink
+	// 2. we're not in the middle of an update loop 
+	// 3. the instance has no current updates
+	//
+	const bool bHasAnyEntities = !Instance.Ledger.IsEmpty();
+	const ERunnerFlushState UnsafeDestroyMask = ERunnerFlushState::Everything & ~(ERunnerFlushState::PostEvaluation | ERunnerFlushState::End);
 	const bool bSafeToDestroyNow = !EnumHasAnyFlags(FlushState, UnsafeDestroyMask);
-	if (bSafeToDestroyNow && !HasQueuedUpdates(InInstanceHandle))
+	if (!bHasAnyEntities && bSafeToDestroyNow && !HasQueuedUpdates(InInstanceHandle))
 	{
-		if (EnumHasAnyFlags(Flags, ERunnerUpdateFlags::Finish))
-		{
-			Instance.Finish(Linker);
-		}
+		Instance.Finish(Linker);
 		Instance.PostEvaluation(Linker);
 
+		InOnLastFlushDelegate.ExecuteIfBound();
+
 		// PostEvaluation could have prompted a bunch of other logic that may have destroyed our instance handle so we have to check it for validity again
-		if (EnumHasAnyFlags(Flags, ERunnerUpdateFlags::Destroy) && InstanceRegistry->IsHandleValid(InInstanceHandle))
+		if (bDestroyInstance && InstanceRegistry->IsHandleValid(InInstanceHandle))
 		{
 			InstanceRegistry->DestroyInstance(InInstanceHandle);
 		}
 
 		return false;
 	}
-	else
-	{
-		FUpdateParamsAndContext Update;
-		Update.Params.InstanceHandle = InInstanceHandle;
-		Update.Params.UpdateFlags = Flags;
-		UpdateQueue.Add(Update);
 
-		return true;
-	}
+	// We queue up one last update that unlinks all of this instance's entities and lets systems tear down
+	// anything that has notable side-effects.
+	// It's possible that the instance already had an update request queued. We leave it in so it gets honored,
+	// and this second update request will be fulfilled in a second flush pass.
+	FUpdateParamsAndContext LastUpdate;
+	LastUpdate.Params.InstanceHandle = InInstanceHandle;
+	LastUpdate.Params.UpdateFlags = (bDestroyInstance ? (ERunnerUpdateFlags::Finish | ERunnerUpdateFlags::Destroy) : ERunnerUpdateFlags::Finish);
+	LastUpdate.OnFlushed = InOnLastFlushDelegate;
+	UpdateQueue.Add(LastUpdate);
+
+	return true;
 }
 
 void FMovieSceneEntitySystemRunner::AbandonAndDestroyInstance(FInstanceHandle Instance)
 {
-	if (QueueFinalUpdate(Instance, UE::MovieScene::ERunnerUpdateFlags::Finish | UE::MovieScene::ERunnerUpdateFlags::Destroy))
+	if (QueueFinalUpdateAndDestroy(Instance))
 	{
 		Flush();
 	}
@@ -535,9 +548,9 @@ bool FMovieSceneEntitySystemRunner::GameThread_UpdateSequenceInstances(UMovieSce
 			{
 				continue;
 			}
-			// Already have an update for this 
 			else if (UpdatedSequenceInstances.IsValidIndex(InstanceID) && UpdatedSequenceInstances[InstanceID] == true)
 			{
+				// We already have an update for this instance, so we queue it up again for the next flush
 				UpdateQueue.Add(Request);
 				continue;
 			}
diff --git a/Engine/Source/Runtime/MovieScene/Private/Evaluation/MovieSceneEvaluationTemplateInstance.cpp b/Engine/Source/Runtime/MovieScene/Private/Evaluation/MovieSceneEvaluationTemplateInstance.cpp
--- a/Engine/Source/Runtime/MovieScene/Private/Evaluation/MovieSceneEvaluationTemplateInstance.cpp
+++ b/Engine/Source/Runtime/MovieScene/Private/Evaluation/MovieSceneEvaluationTemplateInstance.cpp
@@ -385,7 +385,7 @@ void FMovieSceneRootEvaluationTemplateInstance::PlaybackContextChanged(IMovieSce
 
 		if (TSharedPtr<FMovieSceneEntitySystemRunner> Runner = WeakRunner.Pin())
 		{
-			Runner->QueueFinalUpdate(RootInstanceHandle, ERunnerUpdateFlags::Finish);
+			Runner->QueueFinalUpdate(RootInstanceHandle);
 			Runner->Flush();
 		}
 
diff --git a/Engine/Source/Runtime/MovieScene/Private/MovieSceneSequencePlayer.cpp b/Engine/Source/Runtime/MovieScene/Private/MovieSceneSequencePlayer.cpp
--- a/Engine/Source/Runtime/MovieScene/Private/MovieSceneSequencePlayer.cpp
+++ b/Engine/Source/Runtime/MovieScene/Private/MovieSceneSequencePlayer.cpp
@@ -442,14 +442,11 @@ void UMovieSceneSequencePlayer::StopInternal(FFrameTime TimeToResetTo)
 			this->RunLatentActions();
 		};
 
-		// Add the post eval callback that will be invoked when QueueFinalUpdate has been flushed
-		PostEvaluationCallbacks.Add(FOnEvaluationCallback::CreateWeakLambda(this, OnFlushed));
-
 		TSharedPtr<FMovieSceneEntitySystemRunner> Runner = RootTemplateInstance.GetRunner();
 		if (Runner)
 		{
 			// Finish but do not destroy
-			if (Runner->QueueFinalUpdate(RootTemplateInstance.GetRootInstanceHandle(), UE::MovieScene::ERunnerUpdateFlags::Finish))
+			if (Runner->QueueFinalUpdate(RootTemplateInstance.GetRootInstanceHandle(), FSimpleDelegate::CreateWeakLambda(this, OnFlushed)))
 			{
 				Runner->Flush();
 			}
@@ -466,7 +463,7 @@ void UMovieSceneSequencePlayer::StopInternal(FFrameTime TimeToResetTo)
 		if (Runner)
 		{
 			// Finish but do not destroy
-			if (Runner->QueueFinalUpdate(RootTemplateInstance.GetRootInstanceHandle(), UE::MovieScene::ERunnerUpdateFlags::Finish))
+			if (Runner->QueueFinalUpdate(RootTemplateInstance.GetRootInstanceHandle()))
 			{
 				Runner->Flush();
 			}
diff --git a/Engine/Source/Runtime/MovieScene/Public/EntitySystem/MovieSceneEntitySystemRunner.h b/Engine/Source/Runtime/MovieScene/Public/EntitySystem/MovieSceneEntitySystemRunner.h
--- a/Engine/Source/Runtime/MovieScene/Public/EntitySystem/MovieSceneEntitySystemRunner.h
+++ b/Engine/Source/Runtime/MovieScene/Public/EntitySystem/MovieSceneEntitySystemRunner.h
@@ -94,10 +94,12 @@ public:
 	void QueueUpdate(const FMovieSceneContext& Context, FInstanceHandle Instance, FSimpleDelegate&& InOnFlushedDelegate, UE::MovieScene::ERunnerUpdateFlags UpdateFlags = UE::MovieScene::ERunnerUpdateFlags::None);
 
 	/**
-	 * Queue a final update for the specified instance, optionally finishing and destroying it
+	 * Queue a final update for the specified instance, optionally destroying it after finishing it
 	 * @return true if the update requires a flush, or false if it was finished and/or destroyed immediately without requiring a flush
 	 */
-	bool QueueFinalUpdate(FInstanceHandle Instance, UE::MovieScene::ERunnerUpdateFlags Flags);
+	bool QueueFinalUpdate(FInstanceHandle Instance);
+	bool QueueFinalUpdate(FInstanceHandle Instance, FSimpleDelegate&& InOnLastFlushDelegate);
+	bool QueueFinalUpdateAndDestroy(FInstanceHandle Instance);
 
 	/**
 	 * Abandon the specified instance handle and destroy it immediately. May flush this runner if necessary
@@ -164,6 +166,11 @@ private:
 
 private:
 
+	/**
+	 * Queue the final update of a given instance, optionally destroying it after is finishes
+	 */
+	bool QueueFinalUpdateImpl(FInstanceHandle Instance, FSimpleDelegate&& InOnLastFlushDelegate, bool bDestroyInstance);
+
 	/**
 	 * Flush the next item in our update loop based off the contents of FlushState
 	 *
diff --git a/Engine/Source/Runtime/UMG/Private/Animation/UMGSequencePlayer.cpp b/Engine/Source/Runtime/UMG/Private/Animation/UMGSequencePlayer.cpp
--- a/Engine/Source/Runtime/UMG/Private/Animation/UMGSequencePlayer.cpp
+++ b/Engine/Source/Runtime/UMG/Private/Animation/UMGSequencePlayer.cpp
@@ -410,7 +410,7 @@ void UUMGSequencePlayer::HandleLatentStop()
 {
 	if (TSharedPtr<FMovieSceneEntitySystemRunner> Runner = RootTemplateInstance.GetRunner())
 	{
-		Runner->QueueFinalUpdate(RootTemplateInstance.GetRootInstanceHandle(), UE::MovieScene::ERunnerUpdateFlags::Finish);
+		Runner->QueueFinalUpdate(RootTemplateInstance.GetRootInstanceHandle());
 
 		// Even if our request to Finish the instance was queued, we can wait until the next flush for those effects to be seen
 		// This will most likely happen immediately anyway since the runner will keep looping until its queue is empty,
