Sequencer: Check the channel metadata to only add evaluation field entries for override channels that are enabled

diff --git a/Engine/Source/Runtime/MovieScene/Private/Channels/MovieSceneSectionChannelOverrideRegistry.cpp b/Engine/Source/Runtime/MovieScene/Private/Channels/MovieSceneSectionChannelOverrideRegistry.cpp
--- a/Engine/Source/Runtime/MovieScene/Private/Channels/MovieSceneSectionChannelOverrideRegistry.cpp
+++ b/Engine/Source/Runtime/MovieScene/Private/Channels/MovieSceneSectionChannelOverrideRegistry.cpp
@@ -1,6 +1,7 @@
 // Copyright Epic Games, Inc. All Rights Reserved.
 
 #include "Channels/MovieSceneSectionChannelOverrideRegistry.h"
+#include "Channels/MovieSceneChannelEditorData.h"
 #include "Channels/MovieSceneChannelProxy.h"
 #include "Channels/MovieSceneChannelOverrideContainer.h"
 #include "Channels/IMovieSceneChannelOverrideProvider.h"
@@ -53,14 +54,51 @@ void UMovieSceneSectionChannelOverrideRegistry::PopulateEvaluationFieldImpl(cons
 	IMovieSceneChannelOverrideProvider* OverrideProvider = Cast<IMovieSceneChannelOverrideProvider>(GetOuter());
 	if (ensure(OverrideProvider))
 	{
+		// We only want to allocate entities for those channels overrides that are actually in use.
+		// In order to know which channels are in use, we check the channel proxy.
+		TSet<FMovieSceneChannel*> ActiveChannels;
+		FMovieSceneChannelProxy& ChannelProxy = OwnerSection.GetChannelProxy();
+		for (const FMovieSceneChannelEntry& ChannelEntry : ChannelProxy.GetAllEntries())
+		{
+
+			TArrayView<FMovieSceneChannel* const> Channels = ChannelEntry.GetChannels();
+			const int32 NumChannels = Channels.Num();
+
+#if WITH_EDITOR
+			TArrayView<const FMovieSceneChannelMetaData> ChannelMetaDatas = ChannelEntry.GetMetaData();
+			checkf(Channels.Num() == ChannelMetaDatas.Num(), TEXT("Expected matching arrays between channels and channel metadatas!"));
+#endif
+
+			for (int32 ChannelIndex = 0; ChannelIndex < NumChannels; ++ChannelIndex)
+			{
+#if WITH_EDITOR
+				const FMovieSceneChannelMetaData& ChannelMetaData = ChannelMetaDatas[ChannelIndex];
+				const bool bIsChannelEnabled = ChannelMetaData.bEnabled;
+#else
+				// We don't have the metadata for the channel... we assume all channels are enabled. This will
+				// create potentially useless entries in the evaluation field, but hopefully the entity providers
+				// will return empty entity builders for those.
+				const bool bIsChannelEnabled = true;
+#endif
+
+				if (bIsChannelEnabled)
+				{
+					ActiveChannels.Add(Channels[ChannelIndex]);
+				}
+			}
+		}
+
+		FChannelOverrideProviderTraitsHandle ChannelOverrideTraits = OverrideProvider->GetChannelOverrideProviderTraits();
+		check(ChannelOverrideTraits.IsValid());
 		for (const TPair<FName, TObjectPtr<UMovieSceneChannelOverrideContainer>>& Override : Overrides)
 		{
-			FChannelOverrideProviderTraitsHandle ChannelOverrideTraits = OverrideProvider->GetChannelOverrideProviderTraits();
-			check(ChannelOverrideTraits.IsValid());
-			const int32 EntityID = ChannelOverrideTraits->GetChannelOverrideEntityID(Override.Key);
-			const int32 EntityIndex = OutFieldBuilder->FindOrAddEntity(&OwnerSection, EntityID);
-			const int32 MetaDataIndex = OutFieldBuilder->AddMetaData(InMetaData);
-			OutFieldBuilder->AddPersistentEntity(EffectiveRange, EntityIndex, MetaDataIndex);
+			if (ActiveChannels.Contains(Override.Value->GetChannel()))
+			{
+				const int32 EntityID = ChannelOverrideTraits->GetChannelOverrideEntityID(Override.Key);
+				const int32 EntityIndex = OutFieldBuilder->FindOrAddEntity(&OwnerSection, EntityID);
+				const int32 MetaDataIndex = OutFieldBuilder->AddMetaData(InMetaData);
+				OutFieldBuilder->AddPersistentEntity(EffectiveRange, EntityIndex, MetaDataIndex);
+			}
 		}
 	}
 }
