Sequencer: Fixed replicated sequences that begin at a non-zero start time causing a jump on the first frame

diff --git a/Engine/Source/Runtime/MovieScene/Private/MovieSceneSequencePlayer.cpp b/Engine/Source/Runtime/MovieScene/Private/MovieSceneSequencePlayer.cpp
--- a/Engine/Source/Runtime/MovieScene/Private/MovieSceneSequencePlayer.cpp
+++ b/Engine/Source/Runtime/MovieScene/Private/MovieSceneSequencePlayer.cpp
@@ -884,6 +884,8 @@ void UMovieSceneSequencePlayer::Initialize(UMovieSceneSequence* InSequence)
 	// Set up playback position (with offset) after Stop(), which will reset the starting time to StartTime
 	PlayPosition.Reset(StartTimeWithOffset);
 	TimeController->Reset(GetCurrentTime());
+
+	UpdateNetworkSyncProperties();
 }
 
 void UMovieSceneSequencePlayer::Update(const float DeltaSeconds)
