Sequencer: Use flush state instead of system Phase to protect against re-entrancy

 - Using the currently active flush state protects us against re-entrancy for all codepaths, even if they may not call out to external code currently
 - If the re-entrancy ensure is triggered, we no longer go ahead and continue flushing the linker anyway
 - Tidied up the manual bitmask handling for FlushState and added methods for entering a new flush state and skipping a flush state

diff --git a/Engine/Source/Runtime/MovieScene/Private/EntitySystem/MovieSceneEntitySystemRunner.cpp b/Engine/Source/Runtime/MovieScene/Private/EntitySystem/MovieSceneEntitySystemRunner.cpp
--- a/Engine/Source/Runtime/MovieScene/Private/EntitySystem/MovieSceneEntitySystemRunner.cpp
+++ b/Engine/Source/Runtime/MovieScene/Private/EntitySystem/MovieSceneEntitySystemRunner.cpp
@@ -28,12 +28,14 @@ struct FMovieSceneEntitySystemEvaluationReentrancyWindow
 	UMovieSceneEntitySystemLinker* Linker;
 	int32 RunnerIndex;
 	UE::MovieScene::ERunnerFlushState CachedFlushState;
+	UE::MovieScene::ERunnerFlushState CachedCurrentFlushState;
 
 	FMovieSceneEntitySystemEvaluationReentrancyWindow(FMovieSceneEntitySystemRunner* InRunner, UMovieSceneEntitySystemLinker* InLinker)
 		: Runner(InRunner)
 		, Linker(InLinker)
 		, RunnerIndex(Linker->ActiveRunners.Num() - 1)
 		, CachedFlushState(Runner->FlushState)
+		, CachedCurrentFlushState(Runner->CurrentFlushState)
 	{
 		check(RunnerIndex >= 0);
 		checkf(Linker->ActiveRunnerReentrancyFlags[RunnerIndex] == false, TEXT("Nested FMovieSceneEntitySystemEvaluationReentrancyWindows are not supported for the same active runner"));
@@ -43,6 +45,7 @@ struct FMovieSceneEntitySystemEvaluationReentrancyWindow
 		// Clear the current flush state to prevent inner scopes from flushing
 		//      our pending states and set the flag that allows us to be re-entrant
 		Runner->FlushState = UE::MovieScene::ERunnerFlushState::None;
+		Runner->CurrentFlushState = UE::MovieScene::ERunnerFlushState::None;
 	}
 
 	~FMovieSceneEntitySystemEvaluationReentrancyWindow()
@@ -54,6 +57,7 @@ struct FMovieSceneEntitySystemEvaluationReentrancyWindow
 		}
 
 		Runner->FlushState = CachedFlushState;
+		Runner->CurrentFlushState = CachedCurrentFlushState;
 		if (ensure(Linker->ActiveRunnerReentrancyFlags.IsValidIndex(RunnerIndex)))
 		{
 			Linker->ActiveRunnerReentrancyFlags[RunnerIndex] = false;
@@ -66,6 +70,7 @@ FMovieSceneEntitySystemRunner::FMovieSceneEntitySystemRunner()
 	: GameThread(ENamedThreads::GameThread_Local)
 	, CurrentPhase(UE::MovieScene::ESystemPhase::None)
 	, FlushState(UE::MovieScene::ERunnerFlushState::None)
+	, CurrentFlushState(UE::MovieScene::ERunnerFlushState::None)
 	, bRequireFullFlush(false)
 {
 }
@@ -285,8 +290,7 @@ bool FMovieSceneEntitySystemRunner::StartEvaluation(UMovieSceneEntitySystemLinke
 		return false;
 	}
 
-	// Don't run start again until explicitly told to
-	FlushState &= ~ERunnerFlushState::Start;
+	EnterFlushState(ERunnerFlushState::Start);
 
 	// Our entity manager cannot be locked down for us to continue. Something must have left it locked if 
 	// this check fails
@@ -302,9 +306,7 @@ bool FMovieSceneEntitySystemRunner::StartEvaluation(UMovieSceneEntitySystemLinke
 void FMovieSceneEntitySystemRunner::EndEvaluation(UMovieSceneEntitySystemLinker* Linker)
 {
 	using namespace UE::MovieScene;
-
 	Linker->EndEvaluation(*this);
-	FlushState &= ~ERunnerFlushState::End;
 }
 
 bool FMovieSceneEntitySystemRunner::FlushNext(UMovieSceneEntitySystemLinker* Linker)
@@ -319,11 +321,13 @@ bool FMovieSceneEntitySystemRunner::FlushNext(UMovieSceneEntitySystemLinker* Lin
 
 	// Guard the value of the CurrentPhase so it always returns back to ::None
 	TGuardValue<ESystemPhase> CurrentPhaseGuard(CurrentPhase, ESystemPhase::None);
+	TGuardValue<ERunnerFlushState> CurrentFlushStateGuard(CurrentFlushState, ERunnerFlushState::None);
 
 	// Step 1: Initialize all external flags and broadcast 'begin eval' events.
 	//         NOTE: This is only ever called once, regardless of how many iterations we perform
 	if (EnumHasAnyFlags(FlushState, ERunnerFlushState::Start))
 	{
+		// EnterFlushState is called by StartEvaluation itself since it can fail
 		return StartEvaluation(Linker);
 	}
 
@@ -331,6 +335,7 @@ bool FMovieSceneEntitySystemRunner::FlushNext(UMovieSceneEntitySystemLinker* Lin
 	//         This is the entry-point for for each iteration
 	if (EnumHasAnyFlags(FlushState, ERunnerFlushState::Import))
 	{
+		EnterFlushState(ERunnerFlushState::Import);
 		return GameThread_UpdateSequenceInstances(Linker);
 	}
 
@@ -338,6 +343,7 @@ bool FMovieSceneEntitySystemRunner::FlushNext(UMovieSceneEntitySystemLinker* Lin
 	//
 	if (EnumHasAnyFlags(FlushState, ERunnerFlushState::Spawn))
 	{
+		EnterFlushState(ERunnerFlushState::Spawn);
 		return GameThread_SpawnPhase(Linker);
 	}
 
@@ -345,6 +351,7 @@ bool FMovieSceneEntitySystemRunner::FlushNext(UMovieSceneEntitySystemLinker* Lin
 	//
 	if (EnumHasAnyFlags(FlushState, ERunnerFlushState::Instantiation))
 	{
+		EnterFlushState(ERunnerFlushState::Instantiation);
 		return GameThread_InstantiationPhase(Linker);
 	}
 
@@ -352,6 +359,7 @@ bool FMovieSceneEntitySystemRunner::FlushNext(UMovieSceneEntitySystemLinker* Lin
 	//
 	if (EnumHasAnyFlags(FlushState, ERunnerFlushState::Evaluation))
 	{
+		EnterFlushState(ERunnerFlushState::Evaluation);
 		return GameThread_EvaluationPhase(Linker);
 	}
 
@@ -359,6 +367,7 @@ bool FMovieSceneEntitySystemRunner::FlushNext(UMovieSceneEntitySystemLinker* Lin
 	//
 	if (EnumHasAnyFlags(FlushState, ERunnerFlushState::Finalization))
 	{
+		EnterFlushState(ERunnerFlushState::Finalization);
 		GameThread_EvaluationFinalizationPhase(Linker);
 		return true;
 	}
@@ -367,6 +376,7 @@ bool FMovieSceneEntitySystemRunner::FlushNext(UMovieSceneEntitySystemLinker* Lin
 	//
 	if (EnumHasAnyFlags(FlushState, ERunnerFlushState::EventTriggers))
 	{
+		EnterFlushState(ERunnerFlushState::EventTriggers);
 		GameThread_EventTriggerPhase(Linker);
 		return true;
 	}
@@ -375,6 +385,7 @@ bool FMovieSceneEntitySystemRunner::FlushNext(UMovieSceneEntitySystemLinker* Lin
 	//
 	if (EnumHasAnyFlags(FlushState, ERunnerFlushState::PostEvaluation))
 	{
+		EnterFlushState(ERunnerFlushState::PostEvaluation);
 		GameThread_PostEvaluationPhase(Linker);
 		return true;
 	}
@@ -383,6 +394,7 @@ bool FMovieSceneEntitySystemRunner::FlushNext(UMovieSceneEntitySystemLinker* Lin
 	//         NOTE: Only ever called once regardless of how many iterations we perform
 	if (EnumHasAnyFlags(FlushState, ERunnerFlushState::End))
 	{
+		EnterFlushState(ERunnerFlushState::End);
 		EndEvaluation(Linker);
 		return true;
 	}
@@ -421,9 +433,12 @@ void FMovieSceneEntitySystemRunner::FlushOutstanding(double BudgetMs, UE::MovieS
 		return;
 	}
 
-	SCOPE_CYCLE_COUNTER(MovieSceneEval_RunnerFlush);
+	if (!ensureMsgf(CurrentFlushState == ERunnerFlushState::None, TEXT("Cannot flush this runner while it is already being flushed outside of a re-entrancy window")))
+	{
+		return;
+	}
 
-	ensureMsgf(CurrentPhase == ESystemPhase::None, TEXT("Cannot flush this runner while it is already inside a system phase"));
+	SCOPE_CYCLE_COUNTER(MovieSceneEval_RunnerFlush);
 
 	// We need to run the system from the game thread so we know we can fire events and callbacks from here.
 	check(IsInGameThread());
@@ -514,6 +529,19 @@ void FMovieSceneEntitySystemRunner::DiscardQueuedUpdates(FInstanceHandle Instanc
 	}
 }
 
+void FMovieSceneEntitySystemRunner::EnterFlushState(UE::MovieScene::ERunnerFlushState EnteredFlushState)
+{
+	// Stop this state from being run again and allow the next one to run
+	FlushState &= ~EnteredFlushState;
+	// Mark this as the current state
+	CurrentFlushState = EnteredFlushState;
+}
+
+void FMovieSceneEntitySystemRunner::SkipFlushState(UE::MovieScene::ERunnerFlushState FlushStateToSkip)
+{
+	FlushState &= ~FlushStateToSkip;
+}
+
 bool FMovieSceneEntitySystemRunner::GameThread_UpdateSequenceInstances(UMovieSceneEntitySystemLinker* Linker)
 {
 	using namespace UE::MovieScene;
@@ -524,9 +552,6 @@ bool FMovieSceneEntitySystemRunner::GameThread_UpdateSequenceInstances(UMovieSce
 	// with a clean slate.
 	TGuardValue<FScopedPreAnimatedCaptureSource*> CaptureSourceGuard(FScopedPreAnimatedCaptureSource::GetCaptureSourcePtr(), nullptr);
 
-	// Allow the next step to run
-	FlushState &= ~ERunnerFlushState::Import;
-
 	FInstanceRegistry* InstanceRegistry = GetInstanceRegistry();
 
 	if (DissectedUpdates.Num() == 0)
@@ -714,8 +739,6 @@ bool FMovieSceneEntitySystemRunner::GameThread_SpawnPhase(UMovieSceneEntitySyste
 
 	check(GameThread == ENamedThreads::GameThread || GameThread == ENamedThreads::GameThread_Local);
 
-	// Allow the next step to run
-	FlushState &= ~ERunnerFlushState::Spawn;
 
 	CurrentPhase = ESystemPhase::Spawn;
 
@@ -768,7 +791,7 @@ bool FMovieSceneEntitySystemRunner::GameThread_SpawnPhase(UMovieSceneEntitySyste
 	const bool bAnyPending = Linker->EntityManager.ContainsAnyComponent(FBuiltInComponentTypes::Get()->RequiresInstantiationMask) || InstanceRegistry->HasInvalidatedBindings();
 	if (bInstantiationDirty == false && bAnyPending == false)
 	{
-		FlushState &= ~ERunnerFlushState::Instantiation;
+		SkipFlushState(ERunnerFlushState::Instantiation);
 	}
 
 	return true;
@@ -782,8 +805,6 @@ bool FMovieSceneEntitySystemRunner::GameThread_InstantiationPhase(UMovieSceneEnt
 
 	check(GameThread == ENamedThreads::GameThread || GameThread == ENamedThreads::GameThread_Local);
 
-	// Allow the next step to run
-	FlushState &= ~ERunnerFlushState::Instantiation;
 	CurrentPhase = ESystemPhase::Instantiation;
 
 	FGraphEventArray AllTasks;
@@ -839,7 +860,6 @@ bool FMovieSceneEntitySystemRunner::GameThread_EvaluationPhase(UMovieSceneEntity
 
 	SCOPE_CYCLE_COUNTER(MovieSceneEval_EvaluationPhase);
 
-	FlushState &= ~ERunnerFlushState::Evaluation;
 	CurrentPhase = ESystemPhase::Evaluation;
 
 	FGraphEventArray AllTasks;
@@ -872,7 +892,6 @@ void FMovieSceneEntitySystemRunner::GameThread_EvaluationFinalizationPhase(UMovi
 	check(GameThread == ENamedThreads::GameThread || GameThread == ENamedThreads::GameThread_Local);
 
 	CurrentPhase = ESystemPhase::Finalization;
-	FlushState &= ~ERunnerFlushState::Finalization;
 
 	// Post-eval events can be queued during the finalization phase so let's open that up.
 	// The events are actually executed a bit later, in GameThread_EventTriggerPhase.
@@ -905,7 +924,7 @@ void FMovieSceneEntitySystemRunner::GameThread_EvaluationFinalizationPhase(UMovi
 	if (!EventTriggers.IsBound())
 	{
 		// Skip event triggers if there are none
-		FlushState &= ~ERunnerFlushState::EventTriggers;
+		SkipFlushState(ERunnerFlushState::EventTriggers);
 	}
 }
 
@@ -913,8 +932,6 @@ void FMovieSceneEntitySystemRunner::GameThread_EventTriggerPhase(UMovieSceneEnti
 {
 	using namespace UE::MovieScene;
 
-	FlushState &= ~ERunnerFlushState::EventTriggers;
-
 	// Execute any queued events from the evaluation finalization phase.
 	if (EventTriggers.IsBound())
 	{
@@ -937,8 +954,6 @@ void FMovieSceneEntitySystemRunner::GameThread_PostEvaluationPhase(UMovieSceneEn
 
 	SCOPE_CYCLE_COUNTER(MovieSceneEval_PostEvaluationPhase);
 
-	FlushState &= ~ERunnerFlushState::PostEvaluation;
-
 	// Now run the post-evaluation logic so that we can safely handle broadcast events (like OnFinished)
 	// that trigger some new evaluations (such as connecting it to another sequence's Play in Blueprint).
 	//
diff --git a/Engine/Source/Runtime/MovieScene/Public/EntitySystem/MovieSceneEntitySystemRunner.h b/Engine/Source/Runtime/MovieScene/Public/EntitySystem/MovieSceneEntitySystemRunner.h
--- a/Engine/Source/Runtime/MovieScene/Public/EntitySystem/MovieSceneEntitySystemRunner.h
+++ b/Engine/Source/Runtime/MovieScene/Public/EntitySystem/MovieSceneEntitySystemRunner.h
@@ -208,6 +208,12 @@ private:
 	 */
 	void EndEvaluation(UMovieSceneEntitySystemLinker* Linker);
 
+	/** Called to set/unset the necessary flags to enter a new flush state and progress to the next */
+	void EnterFlushState(UE::MovieScene::ERunnerFlushState EnteredFlushState);
+
+	/** Skip the specified flush states if they are currently pending */
+	void SkipFlushState(UE::MovieScene::ERunnerFlushState FlushStateToSkip);
+
 private:
 
 	friend struct FMovieSceneEntitySystemEvaluationReentrancyWindow;
@@ -262,6 +268,9 @@ private:
 	 */
 	UE::MovieScene::ERunnerFlushState FlushState;
 
+	/** The current FlushState that we are running in the current callstack. Used to detect re-entrancy */
+	UE::MovieScene::ERunnerFlushState CurrentFlushState;
+
 	bool bCanQueueEventTriggers;
 
 	/** True if any update has been queued with the ERunnerUpdateFlags::Flush flag since our last flush */
