Sequencer: Added an interface for controlling whether the effects of a UMovieSceneSequencePlayer can be observed

diff --git a/Engine/Source/Runtime/MovieScene/Private/MovieSceneSequencePlayer.cpp b/Engine/Source/Runtime/MovieScene/Private/MovieSceneSequencePlayer.cpp
--- a/Engine/Source/Runtime/MovieScene/Private/MovieSceneSequencePlayer.cpp
+++ b/Engine/Source/Runtime/MovieScene/Private/MovieSceneSequencePlayer.cpp
@@ -155,6 +155,7 @@ void UMovieSceneSequencePlayer::GetLifetimeReplicatedProps(TArray<FLifetimePrope
 	DOREPLIFETIME(UMovieSceneSequencePlayer, DurationFrames);
 	DOREPLIFETIME(UMovieSceneSequencePlayer, DurationSubFrames);
 	DOREPLIFETIME(UMovieSceneSequencePlayer, PlaybackSettings);
+	DOREPLIFETIME(UMovieSceneSequencePlayer, Observer);
 }
 
 EMovieScenePlayerStatus::Type UMovieSceneSequencePlayer::GetPlaybackStatus() const
@@ -343,8 +344,12 @@ void UMovieSceneSequencePlayer::Pause()
 			}
 		};
 
+		if (Observer && !Observer->CanObserveSequence())
+		{
+			FinishPause();
+		}
 		// Evaluate the sequence at its current time, with a status of 'stopped' to ensure that animated state pauses correctly. (ie. audio sounds should stop/pause)
-		if (TSharedPtr<FMovieSceneEntitySystemRunner> Runner = RootTemplateInstance.GetRunner())
+		else if (TSharedPtr<FMovieSceneEntitySystemRunner> Runner = RootTemplateInstance.GetRunner())
 		{
 			FMovieSceneEvaluationRange CurrentTimeRange = PlayPosition.GetCurrentPositionAsRange();
 			const FMovieSceneContext Context(CurrentTimeRange, EMovieScenePlayerStatus::Stopped);
@@ -443,7 +448,12 @@ void UMovieSceneSequencePlayer::StopInternal(FFrameTime TimeToResetTo)
 		};
 
 		TSharedPtr<FMovieSceneEntitySystemRunner> Runner = RootTemplateInstance.GetRunner();
-		if (Runner)
+
+		if (Observer && !Observer->CanObserveSequence())
+		{
+			OnFlushed();
+		}
+		else if (Runner)
 		{
 			// Finish but do not destroy
 			if (Runner->QueueFinalUpdate(RootTemplateInstance.GetRootInstanceHandle(), FSimpleDelegate::CreateWeakLambda(this, OnFlushed)))
@@ -1155,6 +1165,11 @@ void UMovieSceneSequencePlayer::UpdateMovieSceneInstance(FMovieSceneEvaluationRa
 
 void UMovieSceneSequencePlayer::UpdateMovieSceneInstance(FMovieSceneEvaluationRange InRange, EMovieScenePlayerStatus::Type PlayerStatus, const FMovieSceneUpdateArgs& Args)
 {
+	if (Observer && !Observer->CanObserveSequence())
+	{
+		return;
+	}
+
 	UMovieSceneSequence* MovieSceneSequence = RootTemplateInstance.GetSequence(MovieSceneSequenceID::Root);
 	if (!MovieSceneSequence)
 	{
diff --git a/Engine/Source/Runtime/MovieScene/Public/IMovieSceneSequencePlayerObserver.h b/Engine/Source/Runtime/MovieScene/Public/IMovieSceneSequencePlayerObserver.h
new file mode 100644
--- /dev/null
+++ b/Engine/Source/Runtime/MovieScene/Public/IMovieSceneSequencePlayerObserver.h
@@ -0,0 +1,23 @@
+// Copyright Epic Games, Inc. All Rights Reserved.
+
+#pragma once
+
+#include "CoreTypes.h"
+#include "UObject/ObjectMacros.h"
+#include "IMovieSceneSequencePlayerObserver.generated.h"
+
+
+UINTERFACE()
+class MOVIESCENE_API UMovieSceneSequencePlayerObserver : public UInterface
+{
+public:
+	GENERATED_BODY()
+};
+
+class IMovieSceneSequencePlayerObserver
+{
+public:
+	GENERATED_BODY()
+
+	virtual bool CanObserveSequence() const = 0;
+};
\ No newline at end of file
diff --git a/Engine/Source/Runtime/MovieScene/Public/MovieSceneSequencePlayer.h b/Engine/Source/Runtime/MovieScene/Public/MovieSceneSequencePlayer.h
--- a/Engine/Source/Runtime/MovieScene/Public/MovieSceneSequencePlayer.h
+++ b/Engine/Source/Runtime/MovieScene/Public/MovieSceneSequencePlayer.h
@@ -16,6 +16,7 @@
 #include "MovieSceneSequenceTickManager.h"
 #include "Evaluation/MovieScenePlayback.h"
 #include "MovieSceneSequenceTickManager.h"
+#include "IMovieSceneSequencePlayerObserver.h"
 #include "MovieSceneSequencePlayer.generated.h"
 
 class UMovieSceneSequenceTickManager;
@@ -265,6 +266,10 @@ class MOVIESCENE_API UMovieSceneSequencePlayer
 public:
 	GENERATED_BODY()
 
+	/** Obeserver interface used for controlling whether the effects of this sequence can be seen even when it is playing back. */
+	UPROPERTY(replicated)
+	TScriptInterface<IMovieSceneSequencePlayerObserver> Observer;
+
 	UMovieSceneSequencePlayer(const FObjectInitializer&);
 	virtual ~UMovieSceneSequencePlayer();
 
