Added missing Interface.h include

diff --git a/Engine/Source/Runtime/MovieScene/Public/IMovieSceneSequencePlayerObserver.h b/Engine/Source/Runtime/MovieScene/Public/IMovieSceneSequencePlayerObserver.h
--- a/Engine/Source/Runtime/MovieScene/Public/IMovieSceneSequencePlayerObserver.h
+++ b/Engine/Source/Runtime/MovieScene/Public/IMovieSceneSequencePlayerObserver.h
@@ -4,6 +4,7 @@
 
 #include "CoreTypes.h"
 #include "UObject/ObjectMacros.h"
+#include "UObject/Interface.h"
 #include "IMovieSceneSequencePlayerObserver.generated.h"
 
 
