MovieScene: Added flag to ignore playback replication on the sequence player.

diff --git a/Engine/Source/Runtime/MovieScene/Private/MovieSceneSequencePlayer.cpp b/Engine/Source/Runtime/MovieScene/Private/MovieSceneSequencePlayer.cpp
--- a/Engine/Source/Runtime/MovieScene/Private/MovieSceneSequencePlayer.cpp
+++ b/Engine/Source/Runtime/MovieScene/Private/MovieSceneSequencePlayer.cpp
@@ -1294,6 +1294,11 @@ void UMovieSceneSequencePlayer::SetTimeController(TSharedPtr<FMovieSceneTimeCont
 	}
 }
 
+void UMovieSceneSequencePlayer::SetIgnorePlaybackReplication(bool bState)
+{
+	bIgnorePlaybackReplication = bState;
+}
+
 TArray<UObject*> UMovieSceneSequencePlayer::GetBoundObjects(FMovieSceneObjectBindingID ObjectBinding)
 {
 	TArray<UObject*> Objects;
@@ -1433,7 +1438,7 @@ void UMovieSceneSequencePlayer::RPC_ExplicitServerUpdateEvent_Implementation(EUp
 	// ~~~~~~~~~~~~~~~~~~~~~~~~~~~~~~~~~~~~~~~~~~~~~~~~~~~~~~~~~~~
 	// Handle an explicit jump/play/scrub command from the server.
 
-	if (HasAuthority() || !Sequence)
+	if (HasAuthority() || !Sequence || bIgnorePlaybackReplication)
 	{
 		// Never run network sync operations on authoritative players
 		return;
@@ -1467,7 +1472,7 @@ void UMovieSceneSequencePlayer::RPC_OnStopEvent_Implementation(FFrameTime Stoppe
 	// ~~~~~~~~~~~~~~~~~~~~~~~~~~~~~~~~~~~~~~~~~~~~~~~~
 	// Handle an explicit Stop command from the server.
 
-	if (HasAuthority() || !Sequence)
+	if (HasAuthority() || !Sequence || bIgnorePlaybackReplication)
 	{
 		// Never run network sync operations on authoritative players or players that have not been initialized yet
 		return;
@@ -1522,7 +1527,7 @@ void UMovieSceneSequencePlayer::RPC_OnStopEvent_Implementation(FFrameTime Stoppe
 
 void UMovieSceneSequencePlayer::RPC_OnFinishPlaybackEvent_Implementation(FFrameTime StoppedTime)
 {
-	if (HasAuthority() || !Sequence)
+	if (HasAuthority() || !Sequence || bIgnorePlaybackReplication)
 	{
 		// Never run network sync operations on authoritative players or players that have not been initialized yet
 		return;
@@ -1538,7 +1543,7 @@ void UMovieSceneSequencePlayer::PostNetReceive()
 
 	Super::PostNetReceive();
 
-	if (!ensure(!HasAuthority()) || !Sequence)
+	if (!ensure(!HasAuthority()) || !Sequence || bIgnorePlaybackReplication)
 	{
 		// Never run network sync operations on authoritative players or players that have not been initialized yet
 		return;
diff --git a/Engine/Source/Runtime/MovieScene/Public/MovieSceneSequencePlayer.h b/Engine/Source/Runtime/MovieScene/Public/MovieSceneSequencePlayer.h
--- a/Engine/Source/Runtime/MovieScene/Public/MovieSceneSequencePlayer.h
+++ b/Engine/Source/Runtime/MovieScene/Public/MovieSceneSequencePlayer.h
@@ -565,6 +565,12 @@ public:
 	 */
 	void SetTimeController(TSharedPtr<FMovieSceneTimeController> InTimeController);
 
+	/**
+	 * Sets whether to listen or ignore playback replication events.
+	 * @param bState If true, ignores playback replication.
+	 */
+	void SetIgnorePlaybackReplication(bool bState);
+
 protected:
 
 	void PlayInternal();
@@ -780,6 +786,9 @@ protected:
 	/** (Optional) Synchronous runner to use when no tick manager is in use */
 	TSharedPtr<FMovieSceneEntitySystemRunner> SynchronousRunner;
 
+	/** When true, ignore playback replication events. */
+	bool bIgnorePlaybackReplication = false;
+
 private:
 
 	/** The event that will be broadcast every time the sequence is updated */
