Sequencer: MovieScene, LevelSequence and UMG Animation header dependency cleanup

diff --git a/Engine/Plugins/Animation/ControlRig/Source/ControlRigEditor/Private/ControlRigBlueprintActions.cpp b/Engine/Plugins/Animation/ControlRig/Source/ControlRigEditor/Private/ControlRigBlueprintActions.cpp
--- a/Engine/Plugins/Animation/ControlRig/Source/ControlRigEditor/Private/ControlRigBlueprintActions.cpp
+++ b/Engine/Plugins/Animation/ControlRig/Source/ControlRigEditor/Private/ControlRigBlueprintActions.cpp
@@ -20,6 +20,8 @@
 #include "Components/SkeletalMeshComponent.h"
 #include "ScopedTransaction.h"
 #include "LevelSequenceActor.h"
+#include "LevelSequence.h"
+#include "LevelSequencePlayer.h"
 #include "EditorDirectories.h"
 #include "EngineUtils.h"
 #include "ISequencer.h"
diff --git a/Engine/Plugins/Animation/MLDeformer/MLDeformerFramework/Source/MLDeformerFrameworkEditor/Private/SMLDeformerTimeline.cpp b/Engine/Plugins/Animation/MLDeformer/MLDeformerFramework/Source/MLDeformerFrameworkEditor/Private/SMLDeformerTimeline.cpp
--- a/Engine/Plugins/Animation/MLDeformer/MLDeformerFramework/Source/MLDeformerFrameworkEditor/Private/SMLDeformerTimeline.cpp
+++ b/Engine/Plugins/Animation/MLDeformer/MLDeformerFramework/Source/MLDeformerFrameworkEditor/Private/SMLDeformerTimeline.cpp
@@ -5,6 +5,7 @@
 #include "Widgets/SWidget.h"
 #include "Widgets/Layout/SSplitter.h"
 #include "ISequencerWidgetsModule.h"
+#include "MovieSceneTimeHelpers.h"
 #include "FrameNumberNumericInterface.h"
 #include "Widgets/Input/SSearchBox.h"
 #include "Widgets/Layout/SScrollBorder.h"
diff --git a/Engine/Plugins/Cameras/GameplayCameras/Source/GameplayCameras/Public/LegacyCameraShake.h b/Engine/Plugins/Cameras/GameplayCameras/Source/GameplayCameras/Public/LegacyCameraShake.h
--- a/Engine/Plugins/Cameras/GameplayCameras/Source/GameplayCameras/Public/LegacyCameraShake.h
+++ b/Engine/Plugins/Cameras/GameplayCameras/Source/GameplayCameras/Public/LegacyCameraShake.h
@@ -4,10 +4,14 @@
 
 #include "CoreMinimal.h"
 #include "Camera/CameraShakeBase.h"
-#include "Evaluation/MovieSceneCameraShakeTemplate.h"
 #include "Kismet/BlueprintFunctionLibrary.h"
 #include "UObject/Object.h"
 #include "GameFramework/Actor.h"
+
+#if UE_ENABLE_INCLUDE_ORDER_DEPRECATED_IN_5_1
+	#include "Evaluation/MovieSceneCameraShakeTemplate.h"
+#endif
+
 #include "LegacyCameraShake.generated.h"
 
 class AActor;
diff --git a/Engine/Plugins/MovieScene/MovieRenderPipeline/Source/MovieRenderPipelineCore/Private/MoviePipeline.cpp b/Engine/Plugins/MovieScene/MovieRenderPipeline/Source/MovieRenderPipelineCore/Private/MoviePipeline.cpp
--- a/Engine/Plugins/MovieScene/MovieRenderPipeline/Source/MovieRenderPipelineCore/Private/MoviePipeline.cpp
+++ b/Engine/Plugins/MovieScene/MovieRenderPipeline/Source/MovieRenderPipelineCore/Private/MoviePipeline.cpp
@@ -6,6 +6,7 @@
 #include "Tracks/MovieSceneSubTrack.h"
 #include "Tracks/MovieSceneCameraCutTrack.h"
 #include "LevelSequenceActor.h"
+#include "LevelSequencePlayer.h"
 #include "EngineUtils.h"
 #include "MovieSceneSection.h"
 #include "Sections/MovieSceneCameraCutSection.h"
diff --git a/Engine/Plugins/MovieScene/MovieRenderPipeline/Source/MovieRenderPipelineCore/Private/MoviePipelineTiming.cpp b/Engine/Plugins/MovieScene/MovieRenderPipeline/Source/MovieRenderPipelineCore/Private/MoviePipelineTiming.cpp
--- a/Engine/Plugins/MovieScene/MovieRenderPipeline/Source/MovieRenderPipelineCore/Private/MoviePipelineTiming.cpp
+++ b/Engine/Plugins/MovieScene/MovieRenderPipeline/Source/MovieRenderPipelineCore/Private/MoviePipelineTiming.cpp
@@ -11,6 +11,7 @@
 #include "MovieScene.h"
 #include "LevelSequence.h"
 #include "LevelSequenceActor.h"
+#include "LevelSequencePlayer.h"
 #include "Engine/World.h"
 #include "GameFramework/PlayerController.h"
 #include "GameFramework/WorldSettings.h"
diff --git a/Engine/Plugins/MovieScene/TemplateSequence/Source/TemplateSequence/Private/CameraAnimationSequencePlayer.cpp b/Engine/Plugins/MovieScene/TemplateSequence/Source/TemplateSequence/Private/CameraAnimationSequencePlayer.cpp
--- a/Engine/Plugins/MovieScene/TemplateSequence/Source/TemplateSequence/Private/CameraAnimationSequencePlayer.cpp
+++ b/Engine/Plugins/MovieScene/TemplateSequence/Source/TemplateSequence/Private/CameraAnimationSequencePlayer.cpp
@@ -3,6 +3,7 @@
 #include "CameraAnimationSequencePlayer.h"
 #include "CineCameraComponent.h"
 #include "Engine/World.h"
+#include "EntitySystem/MovieSceneEntitySystemRunner.h"
 #include "EntitySystem/MovieSceneEntitySystemLinker.h"
 #include "EntitySystem/MovieSceneEntitySystemTask.h"
 #include "EntitySystem/MovieScenePropertyComponentHandler.h"
diff --git a/Engine/Plugins/MovieScene/TemplateSequence/Source/TemplateSequence/Public/CameraAnimationSequenceSubsystem.h b/Engine/Plugins/MovieScene/TemplateSequence/Source/TemplateSequence/Public/CameraAnimationSequenceSubsystem.h
--- a/Engine/Plugins/MovieScene/TemplateSequence/Source/TemplateSequence/Public/CameraAnimationSequenceSubsystem.h
+++ b/Engine/Plugins/MovieScene/TemplateSequence/Source/TemplateSequence/Public/CameraAnimationSequenceSubsystem.h
@@ -3,6 +3,7 @@
 #pragma once
 
 #include "CoreMinimal.h"
+#include "EntitySystem/MovieSceneEntitySystemLinker.h"
 #include "EntitySystem/MovieSceneEntityInstantiatorSystem.h"
 #include "Subsystems/WorldSubsystem.h"
 #include "UObject/ObjectMacros.h"
diff --git a/Engine/Plugins/VirtualProduction/Takes/Source/TakeRecorderSources/Private/TakeRecorderLevelSequenceSource.cpp b/Engine/Plugins/VirtualProduction/Takes/Source/TakeRecorderSources/Private/TakeRecorderLevelSequenceSource.cpp
--- a/Engine/Plugins/VirtualProduction/Takes/Source/TakeRecorderSources/Private/TakeRecorderLevelSequenceSource.cpp
+++ b/Engine/Plugins/VirtualProduction/Takes/Source/TakeRecorderSources/Private/TakeRecorderLevelSequenceSource.cpp
@@ -6,6 +6,7 @@
 #include "Editor.h"
 #include "GameFramework/Actor.h"
 #include "LevelSequenceActor.h"
+#include "LevelSequencePlayer.h"
 #include "Styling/SlateIconFinder.h"
 #include "ClassIconFinder.h"
 
diff --git a/Engine/Source/Editor/LevelEditor/Private/LevelEditorToolBar.cpp b/Engine/Source/Editor/LevelEditor/Private/LevelEditorToolBar.cpp
--- a/Engine/Source/Editor/LevelEditor/Private/LevelEditorToolBar.cpp
+++ b/Engine/Source/Editor/LevelEditor/Private/LevelEditorToolBar.cpp
@@ -34,6 +34,7 @@
 #include "IContentBrowserSingleton.h"
 #include "ContentBrowserModule.h"
 #include "LevelSequenceActor.h"
+#include "LevelSequence.h"
 #include "Engine/LevelScriptBlueprint.h"
 #include "ISettingsCategory.h"
 #include "ISettingsContainer.h"
diff --git a/Engine/Source/Editor/SequenceRecorder/Private/SequenceRecorder.cpp b/Engine/Source/Editor/SequenceRecorder/Private/SequenceRecorder.cpp
--- a/Engine/Source/Editor/SequenceRecorder/Private/SequenceRecorder.cpp
+++ b/Engine/Source/Editor/SequenceRecorder/Private/SequenceRecorder.cpp
@@ -42,6 +42,7 @@
 #include "ScopedTransaction.h"
 #include "Features/IModularFeatures.h"
 #include "ScopedTransaction.h"
+#include "LevelSequencePlayer.h"
 #include "Subsystems/AssetEditorSubsystem.h"
 
 #define LOCTEXT_NAMESPACE "SequenceRecorder"
diff --git a/Engine/Source/Editor/Sequencer/Private/Sequencer.cpp b/Engine/Source/Editor/Sequencer/Private/Sequencer.cpp
--- a/Engine/Source/Editor/Sequencer/Private/Sequencer.cpp
+++ b/Engine/Source/Editor/Sequencer/Private/Sequencer.cpp
@@ -155,6 +155,7 @@
 #include "CurveEditorCommands.h"
 #include "Tools/SequencerEditTool_Movement.h"
 #include "Tools/SequencerEditTool_Selection.h"
+#include "EntitySystem/MovieSceneEntitySystemRunner.h"
 #include "EntitySystem/MovieSceneEntitySystemLinker.h"
 #include "EntitySystem/MovieScenePreAnimatedStateSystem.h"
 #include "Systems/MovieSceneMotionVectorSimulationSystem.h"
diff --git a/Engine/Source/Editor/UMGEditor/Private/Customizations/UMGDetailCustomizations.cpp b/Engine/Source/Editor/UMGEditor/Private/Customizations/UMGDetailCustomizations.cpp
--- a/Engine/Source/Editor/UMGEditor/Private/Customizations/UMGDetailCustomizations.cpp
+++ b/Engine/Source/Editor/UMGEditor/Private/Customizations/UMGDetailCustomizations.cpp
@@ -34,6 +34,7 @@
 #include "Components/WidgetComponent.h"
 #include "UMGEditorModule.h"
 #include "WidgetBlueprintEditor.h"
+#include "Animation/WidgetAnimation.h"
 
 #define LOCTEXT_NAMESPACE "UMG"
 
diff --git a/Engine/Source/Editor/UMGEditor/Private/Widgets/SBindWidgetView.cpp b/Engine/Source/Editor/UMGEditor/Private/Widgets/SBindWidgetView.cpp
--- a/Engine/Source/Editor/UMGEditor/Private/Widgets/SBindWidgetView.cpp
+++ b/Engine/Source/Editor/UMGEditor/Private/Widgets/SBindWidgetView.cpp
@@ -17,6 +17,8 @@
 #include "Widgets/Views/STableRow.h"
 #include "Widgets/Views/STreeView.h"
 
+#include "Animation/WidgetAnimation.h"
+
 
 #define LOCTEXT_NAMESPACE "UMG.BindWidget"
 
diff --git a/Engine/Source/Editor/UnrealEd/Public/UnrealEdSharedPCH.h b/Engine/Source/Editor/UnrealEd/Public/UnrealEdSharedPCH.h
--- a/Engine/Source/Editor/UnrealEd/Public/UnrealEdSharedPCH.h
+++ b/Engine/Source/Editor/UnrealEd/Public/UnrealEdSharedPCH.h
@@ -617,7 +617,4 @@
 #include "Settings/EditorLoadingSavingSettings.h"
 
 // From ToolMenus
-#include "ToolMenus.h"
-
-// From Movie
-#include "IMovieScenePlayer.h"
\ No newline at end of file
+#include "ToolMenus.h"
\ No newline at end of file
diff --git a/Engine/Source/Runtime/LevelSequence/Private/LevelSequenceActor.cpp b/Engine/Source/Runtime/LevelSequence/Private/LevelSequenceActor.cpp
--- a/Engine/Source/Runtime/LevelSequence/Private/LevelSequenceActor.cpp
+++ b/Engine/Source/Runtime/LevelSequence/Private/LevelSequenceActor.cpp
@@ -544,6 +544,11 @@ void ALevelSequenceActor::UpdateObjectFromProxy(FStructOnScope& Proxy, IProperty
 	ObjectPropertyHandle.SetValue(BoundActor);
 }
 
+UMovieSceneSequence* ALevelSequenceActor::RetrieveOwnedSequence() const
+{
+	return GetSequence();
+}
+
 bool ALevelSequenceActor::GetReferencedContentObjects(TArray<UObject*>& Objects) const
 {
 	if (LevelSequenceAsset)
diff --git a/Engine/Source/Runtime/LevelSequence/Private/LevelSequenceDirector.cpp b/Engine/Source/Runtime/LevelSequence/Private/LevelSequenceDirector.cpp
--- a/Engine/Source/Runtime/LevelSequence/Private/LevelSequenceDirector.cpp
+++ b/Engine/Source/Runtime/LevelSequence/Private/LevelSequenceDirector.cpp
@@ -8,6 +8,7 @@
 #include "IMovieScenePlayer.h"
 #include "Evaluation/MovieSceneEvaluationTemplateInstance.h"
 #include "MovieSceneSequence.h"
+#include "EntitySystem/MovieSceneSequenceInstance.h"
 #include "Evaluation/MovieScenePlayback.h"
 
 #include UE_INLINE_GENERATED_CPP_BY_NAME(LevelSequenceDirector)
diff --git a/Engine/Source/Runtime/LevelSequence/Private/LevelSequenceModule.cpp b/Engine/Source/Runtime/LevelSequence/Private/LevelSequenceModule.cpp
--- a/Engine/Source/Runtime/LevelSequence/Private/LevelSequenceModule.cpp
+++ b/Engine/Source/Runtime/LevelSequence/Private/LevelSequenceModule.cpp
@@ -5,7 +5,7 @@
 #include "LevelSequence.h"
 #include "LevelSequenceActor.h"
 #include "LevelSequenceActorSpawner.h"
-#include "MovieSceneSequencePlayer.h"
+#include "LevelSequencePlayer.h"
 #include "UObject/UObjectBaseUtility.h"
 
 DEFINE_LOG_CATEGORY(LogLevelSequence);
diff --git a/Engine/Source/Runtime/LevelSequence/Private/SequenceMediaController.cpp b/Engine/Source/Runtime/LevelSequence/Private/SequenceMediaController.cpp
--- a/Engine/Source/Runtime/LevelSequence/Private/SequenceMediaController.cpp
+++ b/Engine/Source/Runtime/LevelSequence/Private/SequenceMediaController.cpp
@@ -12,6 +12,7 @@
 #include "MediaPlayer.h"
 #include "MediaPlaylist.h"
 #include "LevelSequenceActor.h"
+#include "LevelSequencePlayer.h"
 #include "LevelSequenceModule.h"
 
 #include UE_INLINE_GENERATED_CPP_BY_NAME(SequenceMediaController)
diff --git a/Engine/Source/Runtime/LevelSequence/Public/LevelSequenceActor.h b/Engine/Source/Runtime/LevelSequence/Public/LevelSequenceActor.h
--- a/Engine/Source/Runtime/LevelSequence/Public/LevelSequenceActor.h
+++ b/Engine/Source/Runtime/LevelSequence/Public/LevelSequenceActor.h
@@ -7,12 +7,20 @@
 #include "UObject/Object.h"
 #include "UObject/SoftObjectPath.h"
 #include "GameFramework/Actor.h"
-#include "LevelSequencePlayer.h"
+#include "IMovieScenePlaybackClient.h"
+#include "MovieSceneSequencePlaybackSettings.h"
 #include "MovieSceneBindingOwnerInterface.h"
 #include "MovieSceneBindingOverrides.h"
+#include "LevelSequenceCameraSettings.h"
+
+#if UE_ENABLE_INCLUDE_ORDER_DEPRECATED_IN_5_1
+	#include "LevelSequencePlayer.h"
+#endif
+
 #include "LevelSequenceActor.generated.h"
 
 class ULevelSequenceBurnIn;
+class ULevelSequencePlayer;
 class UMovieSceneSequenceTickManager;
 
 UCLASS(Blueprintable, DefaultToInstanced)
@@ -303,10 +311,7 @@ public:
 #if WITH_EDITOR
 	virtual TSharedPtr<FStructOnScope> GetObjectPickerProxy(TSharedPtr<IPropertyHandle> PropertyHandle) override;
 	virtual void UpdateObjectFromProxy(FStructOnScope& Proxy, IPropertyHandle& ObjectPropertyHandle) override;
-	virtual UMovieSceneSequence* RetrieveOwnedSequence() const override
-	{
-		return GetSequence();
-	}
+	virtual UMovieSceneSequence* RetrieveOwnedSequence() const override;
 #endif
 
 private:
diff --git a/Engine/Source/Runtime/LevelSequence/Public/LevelSequenceCameraSettings.h b/Engine/Source/Runtime/LevelSequence/Public/LevelSequenceCameraSettings.h
new file mode 100644
--- /dev/null
+++ b/Engine/Source/Runtime/LevelSequence/Public/LevelSequenceCameraSettings.h
@@ -0,0 +1,20 @@
+// Copyright Epic Games, Inc. All Rights Reserved.
+
+#pragma once
+
+#include "CoreTypes.h"
+#include "UObject/ObjectMacros.h"
+#include "Engine/EngineTypes.h"
+#include "LevelSequenceCameraSettings.generated.h"
+
+USTRUCT(BlueprintType)
+struct FLevelSequenceCameraSettings
+{
+	GENERATED_BODY()
+
+	UPROPERTY(EditAnywhere, BlueprintReadWrite, Category = "Aspect Ratio")
+	bool bOverrideAspectRatioAxisConstraint = true;
+
+	UPROPERTY(EditAnywhere, BlueprintReadWrite, Category = "Aspect Ratio", meta = (EditCondition = bOverrideAspectRatioAxisConstraint))
+	TEnumAsByte<EAspectRatioAxisConstraint> AspectRatioAxisConstraint = EAspectRatioAxisConstraint::AspectRatio_MaintainXFOV;
+};
\ No newline at end of file
diff --git a/Engine/Source/Runtime/LevelSequence/Public/LevelSequencePlayer.h b/Engine/Source/Runtime/LevelSequence/Public/LevelSequencePlayer.h
--- a/Engine/Source/Runtime/LevelSequence/Public/LevelSequencePlayer.h
+++ b/Engine/Source/Runtime/LevelSequence/Public/LevelSequencePlayer.h
@@ -16,6 +16,7 @@
 #include "MovieSceneSequencePlayer.h"
 #include "Misc/QualifiedFrameTime.h"
 #include "LevelSequence.h"
+#include "LevelSequenceCameraSettings.h"
 #include "LevelSequencePlayer.generated.h"
 
 class AActor;
@@ -69,18 +70,6 @@ struct FLevelSequencePlayerSnapshot
 	FMovieSceneSequenceID ShotID;
 };
 
-USTRUCT(BlueprintType)
-struct FLevelSequenceCameraSettings
-{
-	GENERATED_BODY()
-
-	UPROPERTY(EditAnywhere, BlueprintReadWrite, Category = "Aspect Ratio")
-	bool bOverrideAspectRatioAxisConstraint = true;
-
-	UPROPERTY(EditAnywhere, BlueprintReadWrite, Category = "Aspect Ratio", meta = (EditCondition = bOverrideAspectRatioAxisConstraint))
-	TEnumAsByte<enum EAspectRatioAxisConstraint> AspectRatioAxisConstraint = EAspectRatioAxisConstraint::AspectRatio_MaintainXFOV;
-};
-
 /**
  * ULevelSequencePlayer is used to actually "play" an level sequence asset at runtime.
  *
diff --git a/Engine/Source/Runtime/MovieScene/Private/MovieSceneSequenceTickManager.cpp b/Engine/Source/Runtime/MovieScene/Private/MovieSceneSequenceTickManager.cpp
--- a/Engine/Source/Runtime/MovieScene/Private/MovieSceneSequenceTickManager.cpp
+++ b/Engine/Source/Runtime/MovieScene/Private/MovieSceneSequenceTickManager.cpp
@@ -3,6 +3,8 @@
 #include "MovieSceneSequenceTickManager.h"
 #include "Engine/World.h"
 #include "EntitySystem/MovieSceneEntitySystemLinker.h"
+#include "MovieSceneSequenceTickManagerClient.h"
+#include "EntitySystem/MovieSceneEntitySystemRunner.h"
 #include "ProfilingDebugging/CountersTrace.h"
 
 #include "Logging/MessageLog.h"
diff --git a/Engine/Source/Runtime/MovieScene/Public/EntitySystem/MovieSceneEntitySystem.h b/Engine/Source/Runtime/MovieScene/Public/EntitySystem/MovieSceneEntitySystem.h
--- a/Engine/Source/Runtime/MovieScene/Public/EntitySystem/MovieSceneEntitySystem.h
+++ b/Engine/Source/Runtime/MovieScene/Public/EntitySystem/MovieSceneEntitySystem.h
@@ -5,7 +5,6 @@
 #include "Async/TaskGraphInterfaces.h"
 #include "EntitySystem/MovieSceneEntityIDs.h"
 #include "EntitySystem/MovieSceneEntityManager.h"
-#include "EntitySystem/MovieSceneEntitySystemLinker.h"
 #include "EntitySystem/MovieSceneEntitySystemTypes.h"
 #include "EntitySystem/MovieSceneSequenceInstanceHandle.h"
 #include "HAL/Platform.h"
@@ -17,6 +16,10 @@
 #include "UObject/ObjectPtr.h"
 #include "UObject/UObjectGlobals.h"
 
+#if UE_ENABLE_INCLUDE_ORDER_DEPRECATED_IN_5_1
+	#include "EntitySystem/MovieSceneEntitySystemLinker.h"
+#endif
+
 #include "MovieSceneEntitySystem.generated.h"
 
 class UMovieSceneEntitySystemLinker;
diff --git a/Engine/Source/Runtime/MovieScene/Public/Evaluation/MovieScenePreAnimatedState.h b/Engine/Source/Runtime/MovieScene/Public/Evaluation/MovieScenePreAnimatedState.h
--- a/Engine/Source/Runtime/MovieScene/Public/Evaluation/MovieScenePreAnimatedState.h
+++ b/Engine/Source/Runtime/MovieScene/Public/Evaluation/MovieScenePreAnimatedState.h
@@ -4,18 +4,21 @@
 
 #include "Containers/Map.h"
 #include "CoreTypes.h"
-#include "EntitySystem/MovieScenePropertySystemTypes.h"
 #include "EntitySystem/MovieSceneSequenceInstanceHandle.h"
-#include "EntitySystem/TrackInstance/MovieSceneTrackInstance.h"
-#include "Evaluation/IMovieSceneEvaluationHook.h"
-#include "Evaluation/MovieSceneAnimTypeID.h"
-#include "Evaluation/MovieSceneCompletionMode.h"
-#include "Evaluation/MovieSceneEvaluationKey.h"
-#include "Misc/TVariant.h"
 #include "Templates/SharedPointer.h"
 #include "UObject/WeakObjectPtr.h"
 #include "UObject/WeakObjectPtrTemplates.h"
 
+#if UE_ENABLE_INCLUDE_ORDER_DEPRECATED_IN_5_1
+	#include "EntitySystem/MovieScenePropertySystemTypes.h"
+	#include "EntitySystem/TrackInstance/MovieSceneTrackInstance.h"
+	#include "Evaluation/IMovieSceneEvaluationHook.h"
+	#include "Evaluation/MovieSceneAnimTypeID.h"
+	#include "Evaluation/MovieSceneCompletionMode.h"
+	#include "Evaluation/MovieSceneEvaluationKey.h"
+	#include "Misc/TVariant.h"
+#endif
+
 class FMovieScenePreAnimatedState;
 class IMovieScenePlayer;
 class UClass;
diff --git a/Engine/Source/Runtime/MovieScene/Public/IMovieScenePlayer.h b/Engine/Source/Runtime/MovieScene/Public/IMovieScenePlayer.h
--- a/Engine/Source/Runtime/MovieScene/Public/IMovieScenePlayer.h
+++ b/Engine/Source/Runtime/MovieScene/Public/IMovieScenePlayer.h
@@ -8,6 +8,7 @@
 #include "Misc/InlineValue.h"
 
 #include "MovieSceneSpawnRegister.h"
+#include "Evaluation/MovieSceneAnimTypeID.h"
 #include "Evaluation/MovieScenePreAnimatedState.h"
 #include "Evaluation/MovieSceneEvaluationState.h"
 #include "Evaluation/MovieSceneEvaluationOperand.h"
@@ -21,6 +22,7 @@
 
 enum class EMovieSceneBuiltInEasing : uint8;
 
+struct FMovieSceneContext;
 class UMovieSceneSequence;
 class FViewportClient;
 class IMovieScenePlaybackClient;
diff --git a/Engine/Source/Runtime/MovieScene/Public/MovieSceneLatentActionManager.h b/Engine/Source/Runtime/MovieScene/Public/MovieSceneLatentActionManager.h
new file mode 100644
--- /dev/null
+++ b/Engine/Source/Runtime/MovieScene/Public/MovieSceneLatentActionManager.h
@@ -0,0 +1,31 @@
+// Copyright Epic Games, Inc. All Rights Reserved.
+
+#pragma once
+
+#include "CoreTypes.h"
+#include "Templates/Function.h"
+#include "Containers/ContainersFwd.h"
+
+class UObject;
+
+DECLARE_DELEGATE(FMovieSceneSequenceLatentActionDelegate);
+
+/**
+ * Utility class for running latent actions created from sequence players.
+ */
+class MOVIESCENE_API FMovieSceneLatentActionManager
+{
+public:
+	void AddLatentAction(FMovieSceneSequenceLatentActionDelegate Delegate);
+	void ClearLatentActions(UObject* Object);
+	void ClearLatentActions();
+
+	void RunLatentActions(TFunctionRef<void()> FlushCallback);
+
+	bool IsEmpty() const { return LatentActions.Num() == 0; }
+
+private:
+	TArray<FMovieSceneSequenceLatentActionDelegate> LatentActions;
+
+	bool bIsRunningLatentActions = false;
+};
diff --git a/Engine/Source/Runtime/MovieScene/Public/MovieSceneSequencePlaybackSettings.h b/Engine/Source/Runtime/MovieScene/Public/MovieSceneSequencePlaybackSettings.h
new file mode 100644
--- /dev/null
+++ b/Engine/Source/Runtime/MovieScene/Public/MovieSceneSequencePlaybackSettings.h
@@ -0,0 +1,121 @@
+// Copyright Epic Games, Inc. All Rights Reserved.
+
+#pragma once
+
+#include "CoreTypes.h"
+#include "UObject/ObjectMacros.h"
+#include "MovieSceneSequenceTickInterval.h"
+#include "MovieSceneSequencePlaybackSettings.generated.h"
+
+
+/** POD struct that represents a number of loops where -1 signifies infinite looping, 0 means no loops, etc
+ * Defined as a struct rather than an int so a property type customization can be bound to it
+ */
+USTRUCT(BlueprintType)
+struct FMovieSceneSequenceLoopCount
+{
+	FMovieSceneSequenceLoopCount()
+		: Value(0)
+	{}
+
+	GENERATED_BODY()
+
+	/** Serialize this count from an int */
+	bool SerializeFromMismatchedTag(const FPropertyTag& Tag, FStructuredArchive::FSlot Slot );
+
+	/** Whether or not to loop playback. If Loop Exactly is chosen, you can specify the number of times to loop */
+	UPROPERTY(EditAnywhere, BlueprintReadWrite, Category="Playback", meta=(UIMin=1, DisplayName="Loop"))
+	int32 Value;
+};
+template<> struct TStructOpsTypeTraits<FMovieSceneSequenceLoopCount> : public TStructOpsTypeTraitsBase2<FMovieSceneSequenceLoopCount>
+{
+	enum { WithStructuredSerializeFromMismatchedTag = true };
+};
+
+
+/**
+ * Settings for the level sequence player actor.
+ */
+USTRUCT(BlueprintType)
+struct FMovieSceneSequencePlaybackSettings
+{
+	FMovieSceneSequencePlaybackSettings()
+		: bAutoPlay(false)
+		, PlayRate(1.f)
+		, StartTime(0.f)
+		, bRandomStartTime(false)
+		, bRestoreState(false)
+		, bDisableMovementInput(false)
+		, bDisableLookAtInput(false)
+		, bHidePlayer(false)
+		, bHideHud(false)
+		, bDisableCameraCuts(false)
+		, bPauseAtEnd(false)
+		, bInheritTickIntervalFromOwner(true)
+	{ }
+
+	GENERATED_BODY()
+
+	/** Auto-play the sequence when created */
+	UPROPERTY(EditAnywhere, BlueprintReadWrite, Category = "Playback")
+	uint32 bAutoPlay : 1;
+
+	/** Number of times to loop playback. -1 for infinite, else the number of times to loop before stopping */
+	UPROPERTY(EditAnywhere, BlueprintReadWrite, Category="Playback", meta=(UIMin=1, DisplayName="Loop"))
+	FMovieSceneSequenceLoopCount LoopCount;
+
+	/** Overridable tick interval for this sequence to update at. When not overridden, the owning actor or component's tick interval will be used */
+	UPROPERTY(EditAnywhere, Category="Playback", meta=(ShowOnlyInnerProperties, Units=s, EditCondition="!bInheritTickIntervalFromOwner"))
+	FMovieSceneSequenceTickInterval TickInterval;
+
+	/** The rate at which to playback the animation */
+	UPROPERTY(EditAnywhere, BlueprintReadWrite, Category="Playback", meta=(Units=Multiplier))
+	float PlayRate;
+
+	/** Start playback at the specified offset from the start of the sequence's playback range */
+	UPROPERTY(EditAnywhere, BlueprintReadWrite, Category="Playback", DisplayName="Start Offset", meta=(Units=s, EditCondition="!bRandomStartTime"))
+	float StartTime;
+
+	/** Start playback at a random time */
+	UPROPERTY(EditAnywhere, BlueprintReadWrite, Category="Playback")
+	uint32 bRandomStartTime : 1;
+
+	/** Flag used to specify whether actor states should be restored on stop */
+	UPROPERTY(EditAnywhere, BlueprintReadWrite, Category="Playback")
+	uint32 bRestoreState : 1;
+
+	/** Disable Input from player during play */
+	UPROPERTY(EditAnywhere, BlueprintReadWrite, Category="Cinematic")
+	uint32 bDisableMovementInput : 1;
+
+	/** Disable LookAt Input from player during play */
+	UPROPERTY(EditAnywhere, BlueprintReadWrite, Category="Cinematic")
+	uint32 bDisableLookAtInput : 1;
+
+	/** Hide Player Pawn during play */
+	UPROPERTY(EditAnywhere, BlueprintReadWrite, Category="Cinematic")
+	uint32 bHidePlayer : 1;
+
+	/** Hide HUD during play */
+	UPROPERTY(EditAnywhere, BlueprintReadWrite, Category="Cinematic")
+	uint32 bHideHud : 1;
+
+	/** Disable camera cuts */
+	UPROPERTY(EditAnywhere, BlueprintReadWrite, Category="Cinematic")
+	uint32 bDisableCameraCuts : 1;
+
+	/** Pause the sequence when playback reaches the end rather than stopping it */
+	UPROPERTY(EditAnywhere, BlueprintReadWrite, Category="Playback")
+	uint32 bPauseAtEnd : 1;
+
+	/** When checked, a custom tick interval can be provided to define how often to update this sequence */
+	UPROPERTY(EditAnywhere, Category="Playback", meta=(InlineEditConditionToggle))
+	uint32 bInheritTickIntervalFromOwner : 1;
+
+	MOVIESCENE_API bool SerializeFromMismatchedTag(const FPropertyTag& Tag, FStructuredArchive::FSlot Slot);
+};
+
+template<> struct TStructOpsTypeTraits<FMovieSceneSequencePlaybackSettings> : public TStructOpsTypeTraitsBase2<FMovieSceneSequencePlaybackSettings>
+{
+	enum { WithStructuredSerializeFromMismatchedTag = true };
+};
diff --git a/Engine/Source/Runtime/MovieScene/Public/MovieSceneSequencePlayer.h b/Engine/Source/Runtime/MovieScene/Public/MovieSceneSequencePlayer.h
--- a/Engine/Source/Runtime/MovieScene/Public/MovieSceneSequencePlayer.h
+++ b/Engine/Source/Runtime/MovieScene/Public/MovieSceneSequencePlayer.h
@@ -7,16 +7,22 @@
 #include "UObject/ScriptMacros.h"
 #include "IMovieScenePlayer.h"
 #include "MovieScene.h"
-#include "MovieSceneSequenceTickManager.h"
 #include "Evaluation/MovieSceneEvaluationTemplateInstance.h"
 #include "IMovieScenePlaybackClient.h"
 #include "Misc/QualifiedFrameTime.h"
 #include "MovieSceneTimeController.h"
 #include "Evaluation/MovieScenePlayback.h"
-#include "MovieSceneSequenceTickManager.h"
 #include "Evaluation/MovieScenePlayback.h"
-#include "MovieSceneSequenceTickManager.h"
+#include "MovieSceneSequenceTickManagerClient.h"
+#include "MovieSceneSequencePlaybackSettings.h"
+#include "MovieSceneLatentActionManager.h"
 #include "IMovieSceneSequencePlayerObserver.h"
+#include "EntitySystem/MovieSceneEntityIDs.h"
+
+#if UE_ENABLE_INCLUDE_ORDER_DEPRECATED_IN_5_1
+	#include "MovieSceneSequenceTickManager.h"
+#endif
+
 #include "MovieSceneSequencePlayer.generated.h"
 
 class UMovieSceneSequenceTickManager;
@@ -38,33 +44,6 @@ enum class EUpdatePositionMethod : uint8
 };
 
 
-
-/** POD struct that represents a number of loops where -1 signifies infinite looping, 0 means no loops, etc
- * Defined as a struct rather than an int so a property type customization can be bound to it
- */
-USTRUCT(BlueprintType)
-struct FMovieSceneSequenceLoopCount
-{
-	FMovieSceneSequenceLoopCount()
-		: Value(0)
-	{}
-
-	GENERATED_BODY()
-
-	/** Serialize this count from an int */
-	bool SerializeFromMismatchedTag(const FPropertyTag& Tag, FStructuredArchive::FSlot Slot );
-
-	/** Whether or not to loop playback. If Loop Exactly is chosen, you can specify the number of times to loop */
-	UPROPERTY(EditAnywhere, BlueprintReadWrite, Category="Playback", meta=(UIMin=1, DisplayName="Loop"))
-	int32 Value;
-};
-template<> struct TStructOpsTypeTraits<FMovieSceneSequenceLoopCount> : public TStructOpsTypeTraitsBase2<FMovieSceneSequenceLoopCount>
-{
-	enum { WithStructuredSerializeFromMismatchedTag = true };
-};
-
-
-
 /**
  * Properties that are broadcast from server->clients for time/state synchronization
  */
@@ -92,89 +71,6 @@ struct FMovieSceneSequenceReplProperties
 };
 
 
-
-/**
- * Settings for the level sequence player actor.
- */
-USTRUCT(BlueprintType)
-struct FMovieSceneSequencePlaybackSettings
-{
-	FMovieSceneSequencePlaybackSettings()
-		: bAutoPlay(false)
-		, PlayRate(1.f)
-		, StartTime(0.f)
-		, bRandomStartTime(false)
-		, bRestoreState(false)
-		, bDisableMovementInput(false)
-		, bDisableLookAtInput(false)
-		, bHidePlayer(false)
-		, bHideHud(false)
-		, bDisableCameraCuts(false)
-		, bPauseAtEnd(false)
-		, bInheritTickIntervalFromOwner(true)
-	{ }
-
-	GENERATED_BODY()
-
-	/** Auto-play the sequence when created */
-	UPROPERTY(EditAnywhere, BlueprintReadWrite, Category = "Playback")
-	uint32 bAutoPlay : 1;
-
-	/** Number of times to loop playback. -1 for infinite, else the number of times to loop before stopping */
-	UPROPERTY(EditAnywhere, BlueprintReadWrite, Category="Playback", meta=(UIMin=1, DisplayName="Loop"))
-	FMovieSceneSequenceLoopCount LoopCount;
-
-	/** Overridable tick interval for this sequence to update at. When not overridden, the owning actor or component's tick interval will be used */
-	UPROPERTY(EditAnywhere, Category="Playback", meta=(ShowOnlyInnerProperties, Units=s, EditCondition="!bInheritTickIntervalFromOwner"))
-	FMovieSceneSequenceTickInterval TickInterval;
-
-	/** The rate at which to playback the animation */
-	UPROPERTY(EditAnywhere, BlueprintReadWrite, Category="Playback", meta=(Units=Multiplier))
-	float PlayRate;
-
-	/** Start playback at the specified offset from the start of the sequence's playback range */
-	UPROPERTY(EditAnywhere, BlueprintReadWrite, Category="Playback", DisplayName="Start Offset", meta=(Units=s, EditCondition="!bRandomStartTime"))
-	float StartTime;
-
-	/** Start playback at a random time */
-	UPROPERTY(EditAnywhere, BlueprintReadWrite, Category="Playback")
-	uint32 bRandomStartTime : 1;
-
-	/** Flag used to specify whether actor states should be restored on stop */
-	UPROPERTY(EditAnywhere, BlueprintReadWrite, Category="Playback")
-	uint32 bRestoreState : 1;
-
-	/** Disable Input from player during play */
-	UPROPERTY(EditAnywhere, BlueprintReadWrite, Category="Cinematic")
-	uint32 bDisableMovementInput : 1;
-
-	/** Disable LookAt Input from player during play */
-	UPROPERTY(EditAnywhere, BlueprintReadWrite, Category="Cinematic")
-	uint32 bDisableLookAtInput : 1;
-
-	/** Hide Player Pawn during play */
-	UPROPERTY(EditAnywhere, BlueprintReadWrite, Category="Cinematic")
-	uint32 bHidePlayer : 1;
-
-	/** Hide HUD during play */
-	UPROPERTY(EditAnywhere, BlueprintReadWrite, Category="Cinematic")
-	uint32 bHideHud : 1;
-
-	/** Disable camera cuts */
-	UPROPERTY(EditAnywhere, BlueprintReadWrite, Category="Cinematic")
-	uint32 bDisableCameraCuts : 1;
-
-	/** Pause the sequence when playback reaches the end rather than stopping it */
-	UPROPERTY(EditAnywhere, BlueprintReadWrite, Category="Playback")
-	uint32 bPauseAtEnd : 1;
-
-	/** When checked, a custom tick interval can be provided to define how often to update this sequence */
-	UPROPERTY(EditAnywhere, Category="Playback", meta=(InlineEditConditionToggle))
-	uint32 bInheritTickIntervalFromOwner : 1;
-
-	MOVIESCENE_API bool SerializeFromMismatchedTag(const FPropertyTag& Tag, FStructuredArchive::FSlot Slot);
-};
-
 UENUM(BlueprintType)
 enum class EMovieScenePositionType : uint8
 {
@@ -249,11 +145,6 @@ struct FMovieSceneSequencePlayToParams
 	bool bExclusive = true;
 };
 
-template<> struct TStructOpsTypeTraits<FMovieSceneSequencePlaybackSettings> : public TStructOpsTypeTraitsBase2<FMovieSceneSequencePlaybackSettings>
-{
-	enum { WithStructuredSerializeFromMismatchedTag = true };
-};
-
 /**
  * Abstract class that provides consistent player behaviour for various animation players
  */
diff --git a/Engine/Source/Runtime/MovieScene/Public/MovieSceneSequenceTickManager.h b/Engine/Source/Runtime/MovieScene/Public/MovieSceneSequenceTickManager.h
--- a/Engine/Source/Runtime/MovieScene/Public/MovieSceneSequenceTickManager.h
+++ b/Engine/Source/Runtime/MovieScene/Public/MovieSceneSequenceTickManager.h
@@ -5,54 +5,19 @@
 #include "CoreMinimal.h"
 #include "UObject/ObjectMacros.h"
 #include "UObject/Object.h"
-#include "EntitySystem/MovieSceneEntitySystemRunner.h"
 #include "MovieSceneSequenceTickInterval.h"
-#include "MovieSceneSequenceTickManager.generated.h"
-
-
-class UMovieSceneEntitySystemLinker;
-
-DECLARE_DELEGATE(FMovieSceneSequenceLatentActionDelegate);
-
-/**
- * Utility class for running latent actions created from sequence players.
- */
-class MOVIESCENE_API FMovieSceneLatentActionManager
-{
-public:
-	void AddLatentAction(FMovieSceneSequenceLatentActionDelegate Delegate);
-	void ClearLatentActions(UObject* Object);
-	void ClearLatentActions();
-
-	void RunLatentActions(TFunctionRef<void()> FlushCallback);
-
-	bool IsEmpty() const { return LatentActions.Num() == 0; }
-
-private:
-	TArray<FMovieSceneSequenceLatentActionDelegate> LatentActions;
+#include "MovieSceneLatentActionManager.h"
 
-	bool bIsRunningLatentActions = false;
-};
+#if UE_ENABLE_INCLUDE_ORDER_DEPRECATED_IN_5_1
+	#include "EntitySystem/MovieSceneEntitySystemRunner.h"
+#endif
 
-/**
- * Interface for objects that are to be ticked by the tick manager.
- */
-UINTERFACE()
-class MOVIESCENE_API UMovieSceneSequenceTickManagerClient
-	: public UInterface
-{
-public:
-	GENERATED_BODY()
-};
-
-class MOVIESCENE_API IMovieSceneSequenceTickManagerClient
-{
-public:
-	GENERATED_BODY()
+#include "MovieSceneSequenceTickManager.generated.h"
 
-	virtual void TickFromSequenceTickManager(float DeltaSeconds, FMovieSceneEntitySystemRunner* Runner) = 0;
-};
 
+class FMovieSceneEntitySystemRunner;
+class UMovieSceneEntitySystemLinker;
+class IMovieSceneSequenceTickManagerClient;
 
 /**
  * Global (one per-UWorld) manager object that manages ticking and updating any and all Sequencer-based
@@ -239,16 +204,4 @@ private:
 
 	/** Latent action manager for managing actions that need performing after evaluation */
 	FMovieSceneLatentActionManager LatentActionManager;
-};
-
-
-/** Deprecated IMovieSceneSequenceActor interface that has been replaced with IMovieSceneSequenceTickManagerClient */
-class UE_DEPRECATED(5.1, "Please use IMovieSceneSequenceTickManagerClient directly") IMovieSceneSequenceActor
-	: public IMovieSceneSequenceTickManagerClient
-{
-	virtual void TickFromSequenceTickManager(float DeltaSeconds) = 0;
-	virtual void TickFromSequenceTickManager(float DeltaSeconds, FMovieSceneEntitySystemRunner* Runner) override
-	{
-		TickFromSequenceTickManager(DeltaSeconds);
-	}
-};
+};
\ No newline at end of file
diff --git a/Engine/Source/Runtime/MovieScene/Public/MovieSceneSequenceTickManagerClient.h b/Engine/Source/Runtime/MovieScene/Public/MovieSceneSequenceTickManagerClient.h
new file mode 100644
--- /dev/null
+++ b/Engine/Source/Runtime/MovieScene/Public/MovieSceneSequenceTickManagerClient.h
@@ -0,0 +1,41 @@
+// Copyright Epic Games, Inc. All Rights Reserved.
+
+#pragma once
+
+#include "CoreTypes.h"
+#include "UObject/ObjectMacros.h"
+#include "UObject/Interface.h"
+#include "MovieSceneSequenceTickManagerClient.generated.h"
+
+
+class FMovieSceneEntitySystemRunner;
+
+/**
+ * Interface for objects that are to be ticked by the tick manager.
+ */
+UINTERFACE()
+class MOVIESCENE_API UMovieSceneSequenceTickManagerClient
+	: public UInterface
+{
+public:
+	GENERATED_BODY()
+};
+
+class MOVIESCENE_API IMovieSceneSequenceTickManagerClient
+{
+public:
+	GENERATED_BODY()
+
+	virtual void TickFromSequenceTickManager(float DeltaSeconds, FMovieSceneEntitySystemRunner* Runner) = 0;
+};
+
+/** Deprecated IMovieSceneSequenceActor interface that has been replaced with IMovieSceneSequenceTickManagerClient */
+class UE_DEPRECATED(5.1, "Please use IMovieSceneSequenceTickManagerClient directly") IMovieSceneSequenceActor
+	: public IMovieSceneSequenceTickManagerClient
+{
+	virtual void TickFromSequenceTickManager(float DeltaSeconds) = 0;
+	virtual void TickFromSequenceTickManager(float DeltaSeconds, FMovieSceneEntitySystemRunner* Runner) override
+	{
+		TickFromSequenceTickManager(DeltaSeconds);
+	}
+};
diff --git a/Engine/Source/Runtime/MovieSceneTracks/Private/Systems/MovieScenePredictionSystem.cpp b/Engine/Source/Runtime/MovieSceneTracks/Private/Systems/MovieScenePredictionSystem.cpp
--- a/Engine/Source/Runtime/MovieSceneTracks/Private/Systems/MovieScenePredictionSystem.cpp
+++ b/Engine/Source/Runtime/MovieSceneTracks/Private/Systems/MovieScenePredictionSystem.cpp
@@ -7,6 +7,7 @@
 
 #include "EntitySystem/BuiltInComponentTypes.h"
 #include "EntitySystem/MovieSceneInitialValueCache.h"
+#include "EntitySystem/MovieSceneEntitySystemRunner.h"
 #include "EntitySystem/Interrogation/MovieSceneInterrogatedPropertyInstantiator.h"
 
 #include "Sections/MovieScene3DTransformSection.h"
diff --git a/Engine/Source/Runtime/UMG/Private/Animation/UMGSequencePlayer.cpp b/Engine/Source/Runtime/UMG/Private/Animation/UMGSequencePlayer.cpp
--- a/Engine/Source/Runtime/UMG/Private/Animation/UMGSequencePlayer.cpp
+++ b/Engine/Source/Runtime/UMG/Private/Animation/UMGSequencePlayer.cpp
@@ -6,6 +6,7 @@
 #include "Animation/WidgetAnimation.h"
 #include "MovieSceneTimeHelpers.h"
 #include "Evaluation/MovieScenePlayback.h"
+#include "EntitySystem/MovieSceneEntitySystemRunner.h"
 #include "EntitySystem/MovieSceneEntitySystemLinker.h"
 #include "Animation/UMGSequenceTickManager.h"
 
diff --git a/Engine/Source/Runtime/UMG/Private/Animation/UMGSequenceTickManager.cpp b/Engine/Source/Runtime/UMG/Private/Animation/UMGSequenceTickManager.cpp
--- a/Engine/Source/Runtime/UMG/Private/Animation/UMGSequenceTickManager.cpp
+++ b/Engine/Source/Runtime/UMG/Private/Animation/UMGSequenceTickManager.cpp
@@ -4,6 +4,7 @@
 #include "Animation/UMGSequencePlayer.h"
 #include "Blueprint/UserWidget.h"
 #include "Engine/World.h"
+#include "EntitySystem/MovieSceneEntitySystemRunner.h"
 #include "EntitySystem/MovieSceneEntitySystemLinker.h"
 #include "ProfilingDebugging/CountersTrace.h"
 #include "Framework/Application/SlateApplication.h"
diff --git a/Engine/Source/Runtime/UMG/Public/Animation/UMGSequenceTickManager.h b/Engine/Source/Runtime/UMG/Public/Animation/UMGSequenceTickManager.h
--- a/Engine/Source/Runtime/UMG/Public/Animation/UMGSequenceTickManager.h
+++ b/Engine/Source/Runtime/UMG/Public/Animation/UMGSequenceTickManager.h
@@ -5,10 +5,16 @@
 #include "CoreMinimal.h"
 #include "UObject/ObjectMacros.h"
 #include "UObject/Object.h"
-#include "MovieSceneSequenceTickManager.h"
-#include "EntitySystem/MovieSceneEntitySystemRunner.h"
+#include "MovieSceneLatentActionManager.h"
+
+#if UE_ENABLE_INCLUDE_ORDER_DEPRECATED_IN_5_1
+	#include "MovieSceneSequenceTickManager.h"
+	#include "EntitySystem/MovieSceneEntitySystemRunner.h"
+#endif
+
 #include "UMGSequenceTickManager.generated.h"
 
+class FMovieSceneEntitySystemRunner;
 class UMovieSceneEntitySystemLinker;
 class UUserWidget;
 
diff --git a/Engine/Source/Runtime/UMG/Public/Animation/WidgetAnimation.h b/Engine/Source/Runtime/UMG/Public/Animation/WidgetAnimation.h
--- a/Engine/Source/Runtime/UMG/Public/Animation/WidgetAnimation.h
+++ b/Engine/Source/Runtime/UMG/Public/Animation/WidgetAnimation.h
@@ -7,15 +7,12 @@
 #include "UObject/ScriptMacros.h"
 #include "MovieSceneSequence.h"
 #include "Animation/WidgetAnimationBinding.h"
+#include "Animation/WidgetAnimationEvents.h"
 #include "WidgetAnimation.generated.h"
 
 class UMovieScene;
 class UUserWidget;
 
-DECLARE_DYNAMIC_MULTICAST_DELEGATE(FOnWidgetAnimationPlaybackStatusChanged);
-
-DECLARE_DYNAMIC_DELEGATE(FWidgetAnimationDynamicEvent);
-DECLARE_DYNAMIC_MULTICAST_DELEGATE(FWidgetAnimationDynamicEvents);
 
 /**
  * 
diff --git a/Engine/Source/Runtime/UMG/Public/Animation/WidgetAnimationEvents.h b/Engine/Source/Runtime/UMG/Public/Animation/WidgetAnimationEvents.h
new file mode 100644
--- /dev/null
+++ b/Engine/Source/Runtime/UMG/Public/Animation/WidgetAnimationEvents.h
@@ -0,0 +1,16 @@
+// Copyright Epic Games, Inc. All Rights Reserved.
+
+#pragma once
+
+#include "Delegates/DelegateCombinations.h"
+
+#include "WidgetAnimationEvents.generated.h"
+
+UDELEGATE()
+DECLARE_DYNAMIC_MULTICAST_DELEGATE(FOnWidgetAnimationPlaybackStatusChanged);
+
+UDELEGATE()
+DECLARE_DYNAMIC_DELEGATE(FWidgetAnimationDynamicEvent);
+
+UDELEGATE()
+DECLARE_DYNAMIC_MULTICAST_DELEGATE(FWidgetAnimationDynamicEvents);
diff --git a/Engine/Source/Runtime/UMG/Public/Blueprint/UserWidget.h b/Engine/Source/Runtime/UMG/Public/Blueprint/UserWidget.h
--- a/Engine/Source/Runtime/UMG/Public/Blueprint/UserWidget.h
+++ b/Engine/Source/Runtime/UMG/Public/Blueprint/UserWidget.h
@@ -23,9 +23,13 @@
 #include "Stats/Stats.h"
 #include "EngineStats.h"
 #include "SlateGlobals.h"
-#include "Animation/WidgetAnimation.h"
+#include "Animation/WidgetAnimationEvents.h"
 #include "Blueprint/WidgetBlueprintGeneratedClass.h"
 
+#if UE_ENABLE_INCLUDE_ORDER_DEPRECATED_IN_5_1
+	#include "Animation/WidgetAnimation.h"
+#endif
+
 #include "UserWidget.generated.h"
 
 class Error;
