Sequencer: fix incorrect index used for matching composites

diff --git a/Engine/Source/Runtime/MovieScene/Public/EntitySystem/MovieScenePropertyComponentHandler.h b/Engine/Source/Runtime/MovieScene/Public/EntitySystem/MovieScenePropertyComponentHandler.h
--- a/Engine/Source/Runtime/MovieScene/Public/EntitySystem/MovieScenePropertyComponentHandler.h
+++ b/Engine/Source/Runtime/MovieScene/Public/EntitySystem/MovieScenePropertyComponentHandler.h
@@ -437,7 +437,7 @@ struct TPropertyComponentHandlerImpl<PropertyTraits, TPropertyMetaData<MetaDataT
 
 			for (int32 CompositeIndex = 0; CompositeIndex < NumComposites; ++CompositeIndex)
 			{
-				if ((PropertyDefinition.DoubleCompositeMask & ( 1 << Index)) != 0)
+				if ((PropertyDefinition.DoubleCompositeMask & ( 1 << CompositeIndex)) != 0)
 				{
 					const double* InitialValueComposite = nullptr;
 					FAlignedDecomposedValue& AlignedOutput = AlignedOutputs[CompositeIndex];
