Sequencer: Added automation test to verify MigrateEntity behavior

diff --git a/Engine/Source/Runtime/MovieScene/Private/Tests/MovieSceneEntitySystemTypesTests.cpp b/Engine/Source/Runtime/MovieScene/Private/Tests/MovieSceneEntitySystemTypesTests.cpp
--- a/Engine/Source/Runtime/MovieScene/Private/Tests/MovieSceneEntitySystemTypesTests.cpp
+++ b/Engine/Source/Runtime/MovieScene/Private/Tests/MovieSceneEntitySystemTypesTests.cpp
@@ -2,6 +2,8 @@
 
 #include "CoreMinimal.h"
 #include "EntitySystem/MovieSceneEntitySystemTypes.h"
+#include "EntitySystem/MovieSceneEntityManager.h"
+#include "EntitySystem/MovieSceneEntityFactoryTemplates.h"
 #include "Containers/ArrayView.h"
 #include "Misc/AutomationTest.h"
 
@@ -60,6 +62,195 @@ bool FMovieSceneEntityComponentFilterTest::RunTest(const FString& Parameters)
 	return true;
 }
 
+IMPLEMENT_SIMPLE_AUTOMATION_TEST(FMovieSceneEntityMigrationTest, 
+		"System.Engine.Sequencer.EntitySystem.Migration", 
+		EAutomationTestFlags::EditorContext | EAutomationTestFlags::EngineFilter)
+bool FMovieSceneEntityMigrationTest::RunTest(const FString& Parameters)
+{
+	/*~~~~~~~~~~~~~~~~~~~~~~~~~~~~~~~~~~~~~~~~~~~~~~~~~~~~~~~~~~~~~~~~~~~~~~~~~~~~~~~~~
+	This test is designed to guarantee the correctness of FEntityManager::MigrateEntity
+	with respect to relocation of entity-component data. It works by allocating a number
+	of entities with a 'random' (it's a seeded random stream for determinism) combination
+	of components, then randomly mutating them. After each mutation, the entity-component
+	data is verified against the expected component data, and the test will fail if any
+	of the component data, or the entity's ID differs from the expected result.
+	~~~~~~~~~~~~~~~~~~~~~~~~~~~~~~~~~~~~~~~~~~~~~~~~~~~~~~~~~~~~~~~~~~~~~~~~~~~~~~~~~~~*/
+
+	using namespace UE::MovieScene;
+
+	FComponentRegistry ComponentRegistry;
+	FEntityManager EntityManager;
+
+	TGuardValue<FEntityManager*> EntityManagerForDebugging(GEntityManagerForDebuggingVisualizers, &EntityManager);
+
+	EntityManager.SetComponentRegistry(&ComponentRegistry);
+
+	TComponentTypeID<int32> IntComponentTypes[2];
+	TComponentTypeID<FString> StringComponentTypes[2];
+
+	// We allocate entities with any combination of the following 4 components (including none)
+	// We keep the component count low in order to ensure that there are multiple entities in
+	// each allocation (without which, this test would be useless)
+	ComponentRegistry.NewComponentType(&IntComponentTypes[0], TEXT("Int 0"));
+	ComponentRegistry.NewComponentType(&IntComponentTypes[1], TEXT("Int 1"));
+	ComponentRegistry.NewComponentType(&StringComponentTypes[0], TEXT("String 0"));
+	ComponentRegistry.NewComponentType(&StringComponentTypes[1], TEXT("String 1"));
+
+	struct FEntityValue
+	{
+		FMovieSceneEntityID ID;
+
+		FString StringValues[2];
+
+		int32 IntValues[2];
+
+		uint8 EnabledMask;
+	};
+
+	constexpr int32 Seed = 0xfeedbeef;
+	FRandomStream RandomStream(Seed);
+
+	TArray<FEntityValue> AllEntities;
+
+	auto TestEntities = [this, &AllEntities, &EntityManager, &IntComponentTypes, &StringComponentTypes]
+	{
+		for (const FEntityValue& EntityValue : AllEntities)
+		{
+			FEntityInfo EntityInfo = EntityManager.GetEntity(EntityValue.ID);
+
+			if (EntityValue.EnabledMask == 0)
+			{
+				UTEST_TRUE(TEXT("Entity Is Empty"), EntityInfo.Data.Allocation == nullptr);
+			}
+			else if (EntityInfo.Data.Allocation)
+			{
+				UTEST_EQUAL(TEXT("Entity ID"), EntityValue.ID, EntityInfo.Data.Allocation->GetEntityIDs()[EntityInfo.Data.ComponentOffset]);
+			}
+			else
+			{
+				AddError(TEXT("Entity ID has no data but should"));
+				return false;
+			}
+
+			TComponentLock<TReadOptional<int32>> IntValue0 = EntityManager.ReadComponent(EntityValue.ID, IntComponentTypes[0]);
+			TComponentLock<TReadOptional<int32>> IntValue1 = EntityManager.ReadComponent(EntityValue.ID, IntComponentTypes[1]);
+
+			bool ShouldHaveIntComponents[] = {
+				static_cast<bool>(EntityValue.EnabledMask & (1 << 0)),
+				static_cast<bool>(EntityValue.EnabledMask & (1 << 1)),
+			};
+
+			UTEST_EQUAL(TEXT("Has Int Value 0"), IntValue0.IsValid(), ShouldHaveIntComponents[0]);
+			UTEST_EQUAL(TEXT("Has Int Value 1"), IntValue1.IsValid(), ShouldHaveIntComponents[1]);
+
+			if (ShouldHaveIntComponents[0] && IntValue0)
+			{
+				UTEST_EQUAL(TEXT("Int Value 0"), *IntValue0.AsPtr(), EntityValue.IntValues[0])
+			}
+			if (ShouldHaveIntComponents[1] && IntValue1)
+			{
+				UTEST_EQUAL(TEXT("Int Value 1"), *IntValue1.AsPtr(), EntityValue.IntValues[1])
+			}
+
+			TComponentLock<TReadOptional<FString>> StringValue0 = EntityManager.ReadComponent(EntityValue.ID, StringComponentTypes[0]);
+			TComponentLock<TReadOptional<FString>> StringValue1 = EntityManager.ReadComponent(EntityValue.ID, StringComponentTypes[1]);
+
+			bool ShouldHaveStringComponents[] = {
+				static_cast<bool>(EntityValue.EnabledMask & (1 << 2)),
+				static_cast<bool>(EntityValue.EnabledMask & (1 << 3)),
+			};
+
+			UTEST_EQUAL(TEXT("Has String Value 0"), StringValue0.IsValid(), ShouldHaveStringComponents[0]);
+			UTEST_EQUAL(TEXT("Has String Value 1"), StringValue1.IsValid(), ShouldHaveStringComponents[1]);
+
+			if (ShouldHaveStringComponents[0] && StringValue0)
+			{
+				UTEST_EQUAL(TEXT("String Value 0"), *StringValue0.AsPtr(), EntityValue.StringValues[0])
+			}
+			if (ShouldHaveStringComponents[1] && StringValue1)
+			{
+				UTEST_EQUAL(TEXT("String Value 1"), *StringValue1.AsPtr(), EntityValue.StringValues[1])
+			}
+		}
+		return true;
+	};
+
+	constexpr int32 NumTypes = 16;
+	constexpr int32 NumPerType = 10;
+	for (int32 TypeIndex = 0; TypeIndex < NumTypes; ++TypeIndex)
+	{
+		const uint8 ComponentMask = TypeIndex;
+
+		// Allocate NumPerType entities to ensure we have many entities in the allocation
+		for (int32 EntityIndex = 0; EntityIndex < NumPerType; ++EntityIndex)
+		{
+			FEntityValue NewValue;
+			NewValue.IntValues[0] = FMath::RandRange(0, 100);
+			NewValue.IntValues[1] = FMath::RandRange(0, 100);
+			NewValue.StringValues[0] = FString::Printf(TEXT("String Value %i"), FMath::RandRange(0, 100));
+			NewValue.StringValues[1] = FString::Printf(TEXT("String Value %i"), FMath::RandRange(0, 100));
+			NewValue.EnabledMask = ComponentMask;
+
+			NewValue.ID = FEntityBuilder()
+			.AddConditional(IntComponentTypes[0], NewValue.IntValues[0], (ComponentMask & 1 << 0))
+			.AddConditional(IntComponentTypes[1], NewValue.IntValues[1], (ComponentMask & 1 << 1))
+			.AddConditional(StringComponentTypes[0], NewValue.StringValues[0], (ComponentMask & 1 << 2))
+			.AddConditional(StringComponentTypes[1], NewValue.StringValues[1], (ComponentMask & 1 << 3))
+			.CreateEntity(&EntityManager);
+
+			AllEntities.Emplace(MoveTemp(NewValue));
+		}
+	}
+
+	// Shuffle the array of entities so we mutate them in a random order
+	for (int32 i = 0; i < AllEntities.Num()-2; i++)
+	{
+		const int32 j = RandomStream.RandRange(i+1, AllEntities.Num() - 1);
+		AllEntities.Swap(i, j);
+	}
+
+	// 'randomly' mutate entities
+	for (int32 Index = 0; Index < AllEntities.Num(); ++Index)
+	{
+		const uint8 NewType = RandomStream.RandRange(0, 0b1111);
+
+		AddInfo(FString::Printf(TEXT("Mutating entity index %i from %i to %i"), Index, (int32)AllEntities[Index].EnabledMask, (int32)NewType));
+
+		AllEntities[Index].EnabledMask = NewType;
+
+		for (int32 IntIndex = 0; IntIndex < 2; ++IntIndex)
+		{
+			if (NewType & (1 << IntIndex))
+			{
+				EntityManager.AddComponent(AllEntities[Index].ID, IntComponentTypes[IntIndex], AllEntities[Index].IntValues[IntIndex]);
+			}
+			else
+			{
+				EntityManager.RemoveComponent(AllEntities[Index].ID, IntComponentTypes[IntIndex]);
+			}
+		}
+
+		for (int32 StringIndex = 0; StringIndex < 2; ++StringIndex)
+		{
+			if (NewType & (0x4 << StringIndex))
+			{
+				EntityManager.AddComponent(AllEntities[Index].ID, StringComponentTypes[StringIndex], AllEntities[Index].StringValues[StringIndex]);
+			}
+			else
+			{
+				EntityManager.RemoveComponent(AllEntities[Index].ID, StringComponentTypes[StringIndex]);
+			}
+		}
+
+		if (!TestEntities())
+		{
+			return false;
+		}
+	}
+
+	return true;
+}
+
 #undef LOCTEXT_NAMESPACE
 
 #endif // WITH_DEV_AUTOMATION_TESTS
