Sequencer: Do not flush runners when there have been no updates

diff --git a/Engine/Source/Runtime/MovieScene/Private/MovieSceneSequenceTickManager.cpp b/Engine/Source/Runtime/MovieScene/Private/MovieSceneSequenceTickManager.cpp
--- a/Engine/Source/Runtime/MovieScene/Private/MovieSceneSequenceTickManager.cpp
+++ b/Engine/Source/Runtime/MovieScene/Private/MovieSceneSequenceTickManager.cpp
@@ -419,7 +419,11 @@ void UMovieSceneSequenceTickManager::TickSequenceActors(float DeltaSeconds)
 				FLinkerGroup& Group = LinkerGroups[Index];
 				// Hitting this check would indicate that the loop above that processes OutstandingLinkers either failed, or some other partial flush happened between then and now.
 				ensureMsgf(!Group.Runner->IsCurrentlyEvaluating(), TEXT("Linker is part-way thorugh a flush when a new flush is being instigated. This is undefined behavior."));
-				Group.Runner->Flush(Group.FrameBudgetMs);
+
+				if (Group.Runner->HasQueuedUpdates())
+				{
+					Group.Runner->Flush(Group.FrameBudgetMs);
+				}
 			}
 		}
 	}
