Sequencer: Fix CIS

diff --git a/Engine/Source/Runtime/MovieScene/Private/Tests/MovieSceneEntitySystemTypesTests.cpp b/Engine/Source/Runtime/MovieScene/Private/Tests/MovieSceneEntitySystemTypesTests.cpp
--- a/Engine/Source/Runtime/MovieScene/Private/Tests/MovieSceneEntitySystemTypesTests.cpp
+++ b/Engine/Source/Runtime/MovieScene/Private/Tests/MovieSceneEntitySystemTypesTests.cpp
@@ -1,6 +1,7 @@
 // Copyright Epic Games, Inc. All Rights Reserved.
 
 #include "CoreMinimal.h"
+#include "EntitySystem/MovieSceneEntityBuilder.h"
 #include "EntitySystem/MovieSceneEntitySystemTypes.h"
 #include "EntitySystem/MovieSceneEntityManager.h"
 #include "EntitySystem/MovieSceneEntityFactoryTemplates.h"
