Sequencer: Remove Master from Sequencer

diff --git a/Engine/Source/Runtime/MovieScene/Public/MovieSceneObjectBindingID.h b/Engine/Source/Runtime/MovieScene/Public/MovieSceneObjectBindingID.h
--- a/Engine/Source/Runtime/MovieScene/Public/MovieSceneObjectBindingID.h
+++ b/Engine/Source/Runtime/MovieScene/Public/MovieSceneObjectBindingID.h
@@ -57,17 +57,17 @@ struct FRelativeObjectBindingID
 	{}
 
 	/**
-	 * Construction from a master sequence asset
+	 * Construction from a root sequence asset
 	 *
-	 * @param SourceSequenceID      Absolute sequence ID within MasterSequence that represents the start of the path (ie, the sequence that this binding ID is serialized within)
-	 * @param TargetSequenceID      Absolute sequence ID within MasterSequence that holds the target object
+	 * @param SourceSequenceID      Absolute sequence ID within RootSequence that represents the start of the path (ie, the sequence that this binding ID is serialized within)
+	 * @param TargetSequenceID      Absolute sequence ID within RootSequence that holds the target object
 	 * @param TargetGuid            The GUID of the object binding within TargetSequenceID that represents the object at runtime
-	 * @param MasterSequence        The master sequence asset
+	 * @param RootSequence          The root sequence asset
 	 */
-	MOVIESCENE_API FRelativeObjectBindingID(FMovieSceneSequenceID SourceSequenceID, FMovieSceneSequenceID TargetSequenceID, const FGuid& TargetGuid, UMovieSceneSequence* MasterSequence);
+	MOVIESCENE_API FRelativeObjectBindingID(FMovieSceneSequenceID SourceSequenceID, FMovieSceneSequenceID TargetSequenceID, const FGuid& TargetGuid, UMovieSceneSequence* RootSequence);
 
 	/**
-	 * Construction from a master sequence Hierarchy
+	 * Construction from a root sequence Hierarchy
 	 *
 	 * @param SourceSequenceID      Absolute sequence ID within Hierarchy that represents the start of the path (ie, the sequence that this binding ID is serialized within)
 	 * @param TargetSequenceID      Absolute sequence ID within Hierarchy that holds the target object
@@ -77,12 +77,12 @@ struct FRelativeObjectBindingID
 	MOVIESCENE_API FRelativeObjectBindingID(FMovieSceneSequenceID SourceSequenceID, FMovieSceneSequenceID TargetSequenceID, const FGuid& TargetGuid, const FMovieSceneSequenceHierarchy* Hierarchy);
 
 	/**
-	 * Construction from a master sequence Player
+	 * Construction from a root sequence Player
 	 *
-	 * @param SourceSequenceID      Absolute sequence ID within Player's master sequence hierarchy that represents the start of the path (ie, the sequence that this binding ID is serialized within)
-	 * @param TargetSequenceID      Absolute sequence ID within Player's master sequence hierarchy that holds the target object
+	 * @param SourceSequenceID      Absolute sequence ID within Player's root sequence hierarchy that represents the start of the path (ie, the sequence that this binding ID is serialized within)
+	 * @param TargetSequenceID      Absolute sequence ID within Player's root sequence hierarchy that holds the target object
 	 * @param TargetGuid            The GUID of the object binding within TargetSequenceID that represents the object at runtime
-	 * @param Player                The movie scene player that is currently playing the master sequence
+	 * @param Player                The movie scene player that is currently playing the root sequence
 	 */
 	MOVIESCENE_API FRelativeObjectBindingID(FMovieSceneSequenceID SourceSequenceID, FMovieSceneSequenceID TargetSequenceID, const FGuid& TargetGuid, IMovieScenePlayer& Player);
 
@@ -140,7 +140,7 @@ public:
 	 * Convert this fixed binding ID to one that is resolved relative to the specified Sequence.
 	 *
 	 * @param SourceSequenceID    The sequence ID to make this fixed ID relative to
-	 * @param InPlayer            The active movie scene player to retrieve a master sequence hierarchy from
+	 * @param InPlayer            The active movie scene player to retrieve a root sequence hierarchy from
 	 * @return A new binding ID that will resolve relative to the specified sequence
 	 */
 	MOVIESCENE_API FRelativeObjectBindingID ConvertToRelative(FMovieSceneSequenceID SourceSequenceID, IMovieScenePlayer& InPlayer) const;
@@ -149,7 +149,7 @@ public:
 	 * Convert this fixed binding ID to one that is resolved relative to the specified Sequence.
 	 *
 	 * @param SourceSequenceID    The sequence ID to make this fixed ID relative to
-	 * @param Hierarchy           The hierarchy that represents the master sequence that SourceSequenceID is contained within
+	 * @param Hierarchy           The hierarchy that represents the root sequence that SourceSequenceID is contained within
 	 * @return A new binding ID that will resolve relative to the specified sequence
 	 */
 	MOVIESCENE_API FRelativeObjectBindingID ConvertToRelative(FMovieSceneSequenceID SourceSequenceID, const FMovieSceneSequenceHierarchy* Hierarchy) const;
@@ -186,12 +186,12 @@ enum class EMovieSceneObjectBindingSpace : uint8
 /**
  * Persistent identifier to a specific object binding within a sequence hierarchy.
  * 
- * Binding IDs come in 3 flavors with Local and External being preferred as they are reslilient towards sequences being authored in isolation or included in other master sequences:
+ * Binding IDs come in 3 flavors with Local and External being preferred as they are reslilient towards sequences being authored in isolation or included in other root sequences:
  *     Local: (ResolveParentIndex == 0) SequenceID relates to _this sequence's_ local hierarchy; represents an object binding within the same sequence as the ID is resolved, or inside one of its sub-sequences. Sequence ID must be remapped at runtime.
  *     External: (ResolveParentIndex > 0) SequenceID is local to the parent sequence of this one denoted by the parent index (ie, 1 = parent, 2 = grandparent etc). Sequence ID must be remapped at runtime.
  *     Fixed: Represents a binding anywhere in the sequence; always resolved from the root sequence.
  *
- * Fixed bindings will break if the sequence is evaluated inside a different master sequence.
+ * Fixed bindings will break if the sequence is evaluated inside a different root sequence.
  */
 USTRUCT(BlueprintType, meta=(HasNativeMake))
 struct FMovieSceneObjectBindingID
@@ -326,28 +326,28 @@ public:
 
 
 	/**
-	 * Resolve this binding ID to a fixed object binding ID resolvable from the root master sequence
+	 * Resolve this binding ID to a fixed object binding ID resolvable from the root sequence
 	 *
 	 * @param SourceSequenceID      The source sequence ID that owns this binding ID
-	 * @param Player                The movie scene player that is currently playing the master sequence
-	 * @return A fixed binding ID whose SequenceID relates to the master sequence hierarchy
+	 * @param Player                The movie scene player that is currently playing the root sequence
+	 * @return A fixed binding ID whose SequenceID relates to the root sequence hierarchy
 	 */
 	MOVIESCENE_API UE::MovieScene::FFixedObjectBindingID ResolveToFixed(FMovieSceneSequenceID SourceSequenceID, IMovieScenePlayer& Player) const;
 
 
 
 	/**
-	 * Resolve the sequence ID for this binding for the root master sequence
+	 * Resolve the sequence ID for this binding for the root sequence
 	 *
 	 * @param SourceSequenceID      The source sequence ID that owns this binding ID
-	 * @param Player                The movie scene player that is currently playing the master sequence
+	 * @param Player                The movie scene player that is currently playing the sequence
 	 * @return The sequence ID that holds the target object
 	 */
 	MOVIESCENE_API FMovieSceneSequenceID ResolveSequenceID(FMovieSceneSequenceID SourceSequenceID, IMovieScenePlayer& Player) const;
 
 
 	/**
-	 * Resolve the sequence ID for this binding for the root master sequence
+	 * Resolve the sequence ID for this binding for the root sequence
 	 *
 	 * @param SourceSequenceID      The source sequence ID that owns this binding ID
 	 * @param Hierarchy             The hierachy to use for computation of the relative path. Many be nullptr where Source and Targets are both MovieSceneSequenceID::Root
@@ -401,7 +401,7 @@ private:
 	UPROPERTY()
 	int32 SequenceID;
 
-	/** The parent index to resolve the sequence ID from. 0 signifies the sequence this binding ID is serialized within, -1 signifies the master sequence (previously EMovieSceneObjectBindingSpace::Root) */
+	/** The parent index to resolve the sequence ID from. 0 signifies the sequence this binding ID is serialized within, -1 signifies the root sequence (previously EMovieSceneObjectBindingSpace::Root) */
 	UPROPERTY()
 	int32 ResolveParentIndex;
 
diff --git a/Engine/Source/Runtime/MovieScene/Public/MovieSceneSequence.h b/Engine/Source/Runtime/MovieScene/Public/MovieSceneSequence.h
--- a/Engine/Source/Runtime/MovieScene/Public/MovieSceneSequence.h
+++ b/Engine/Source/Runtime/MovieScene/Public/MovieSceneSequence.h
@@ -266,7 +266,7 @@ public:
 	}
 
 	/**
-	 * Check whether this sequence is playable directly outside of a master sub sequence or not
+	 * Check whether this sequence is playable directly outside of a root sub sequence or not
 	 *
 	 * @return True if this sequences cooked data will include all the necessary information to be played back on its own, false if this data is not present in cooked builds
 	 */
@@ -276,7 +276,7 @@ public:
 	}
 
 	/**
-	 * Assign whether this sequence is playable directly outside of a master sub sequence or not
+	 * Assign whether this sequence is playable directly outside of a root sub sequence or not
 	 *
 	 * @param bInPlayableDirectly   When true, this sequence's cooked data will include all the necessary information to be played back on its own. When false this data will be culled resulting in less memory usage.
 	 */
@@ -329,7 +329,7 @@ protected:
 	bool bParentContextsAreSignificant;
 
 	/**
-	 * When true, this sequence should be compiled as if it is playable directly (outside of a master sequence). When false, various compiled data will be omitted, preventing direct playback at runtime (although will still play as a sub sequence)
+	 * When true, this sequence should be compiled as if it is playable directly (outside of a root sequence). When false, various compiled data will be omitted, preventing direct playback at runtime (although will still play as a sub sequence)
 	 */
 	UPROPERTY()
 	bool bPlayableDirectly;
diff --git a/Engine/Source/Runtime/MovieScene/Public/MovieSceneSequenceVisitor.h b/Engine/Source/Runtime/MovieScene/Public/MovieSceneSequenceVisitor.h
--- a/Engine/Source/Runtime/MovieScene/Public/MovieSceneSequenceVisitor.h
+++ b/Engine/Source/Runtime/MovieScene/Public/MovieSceneSequenceVisitor.h
@@ -49,7 +49,7 @@ struct FMovieSceneBinding;
 *	};
 *	
 *	UE::MovieScene::FSequenceVisitParams Params;
-*	Params.bVisitMasterTracks = true;
+*	Params.bVisitRootTracks = true;
 *	Params.bVisitSubSequences = true;
 *	FCameraCutVisitor CameraCutVisitor;
 *	
@@ -95,7 +95,7 @@ struct MOVIESCENE_API ISequenceVisitor
 struct FSequenceVisitParams
 {
 	FSequenceVisitParams()
-		: bVisitMasterTracks(false)
+		: bVisitRootTracks(false)
 		, bVisitObjectBindings(false)
 		, bVisitTracks(false)
 		, bVisitSections(false)
@@ -109,7 +109,7 @@ struct FSequenceVisitParams
 		return bVisitTracks || bVisitSections;
 	}
 
-	bool bVisitMasterTracks;
+	bool bVisitRootTracks;
 	bool bVisitObjectBindings;
 	bool bVisitTracks;
 	bool bVisitSections;
diff --git a/Engine/Source/Runtime/MovieScene/Public/MovieSceneSpawnable.h b/Engine/Source/Runtime/MovieScene/Public/MovieSceneSpawnable.h
--- a/Engine/Source/Runtime/MovieScene/Public/MovieSceneSpawnable.h
+++ b/Engine/Source/Runtime/MovieScene/Public/MovieSceneSpawnable.h
@@ -19,7 +19,7 @@ enum class ESpawnOwnership : uint8
 	InnerSequence,
 
 	/** The object's lifetime is managed by the outermost sequence */
-	MasterSequence,
+	RootSequence,
 
 	/** Once spawned, the object's lifetime is managed externally. */
 	External,
diff --git a/Engine/Source/Runtime/MovieSceneCapture/Public/MovieSceneCaptureSettings.h b/Engine/Source/Runtime/MovieSceneCapture/Public/MovieSceneCaptureSettings.h
--- a/Engine/Source/Runtime/MovieSceneCapture/Public/MovieSceneCaptureSettings.h
+++ b/Engine/Source/Runtime/MovieSceneCapture/Public/MovieSceneCaptureSettings.h
@@ -50,7 +50,7 @@ struct MOVIESCENECAPTURE_API FMovieSceneCaptureSettings
 	 * {quality}	- The image compression quality setting
 	 * {material}   - The material/render pass
 	 * {shot}       - The name of the level sequence asset shot being played
-	 * {sequence}   - The name of the level sequence asset (ie. master) being played
+	 * {sequence}   - The name of the level sequence asset being played
 	 * {camera}     - The name of the current camera
 	 * {date}       - The date in the format of {year}.{month}.{day}
 	 * {year}       - The current year
diff --git a/Engine/Source/Runtime/MovieSceneTracks/Private/Evaluation/MovieSceneAudioTemplate.cpp b/Engine/Source/Runtime/MovieSceneTracks/Private/Evaluation/MovieSceneAudioTemplate.cpp
--- a/Engine/Source/Runtime/MovieSceneTracks/Private/Evaluation/MovieSceneAudioTemplate.cpp
+++ b/Engine/Source/Runtime/MovieSceneTracks/Private/Evaluation/MovieSceneAudioTemplate.cpp
@@ -108,7 +108,7 @@ struct FCachedAudioTrackData : IPersistentEvaluationData
 	
 	FCachedAudioTrackData()
 	{
-		// Create the container for master tracks, which do not have an actor to attach to
+		// Create the container for root tracks, which do not have an actor to attach to
 		AudioComponentsByActorKey.Add(FObjectKey(), TMap<FObjectKey, TWeakObjectPtr<UAudioComponent>>());
 	}
 
@@ -208,7 +208,7 @@ struct FCachedAudioTrackData : IPersistentEvaluationData
 	}
 
 	/** Only to be called on the game thread */
-	UAudioComponent* AddMasterAudioComponentForRow(int32 RowIndex, FObjectKey SectionKey, UWorld* World, IMovieScenePlayer& Player)
+	UAudioComponent* AddRootAudioComponentForRow(int32 RowIndex, FObjectKey SectionKey, UWorld* World, IMovieScenePlayer& Player)
 	{
 		UAudioComponent* ExistingComponent = GetAudioComponent(FObjectKey(), SectionKey);
 		if (!ExistingComponent)
@@ -221,7 +221,7 @@ struct FCachedAudioTrackData : IPersistentEvaluationData
 
 			if (!ExistingComponent)
 			{
-				UE_LOG(LogMovieScene, Warning, TEXT("Failed to create audio component for master audio track (row %d)."), RowIndex);
+				UE_LOG(LogMovieScene, Warning, TEXT("Failed to create audio component for root audio track (row %d)."), RowIndex);
 				return nullptr;
 			}
 
@@ -275,7 +275,7 @@ struct FAudioSectionExecutionToken : IMovieSceneExecutionToken
 			TrackData.StopAllSounds();
 		}
 
-		// Master audio track
+		// Root audio track
 		else if (!Operand.ObjectBindingID.IsValid())
 		{
 			UObject* PlaybackContext = Player.GetPlaybackContext();
@@ -307,7 +307,7 @@ struct FAudioSectionExecutionToken : IMovieSceneExecutionToken
 			if (!AudioComponent)
 			{
 				// Initialize the sound
-				AudioComponent = TrackData.AddMasterAudioComponentForRow(AudioSection->GetRowIndex(), SectionKey, PlaybackContext ? PlaybackContext->GetWorld() : nullptr, Player);
+				AudioComponent = TrackData.AddRootAudioComponentForRow(AudioSection->GetRowIndex(), SectionKey, PlaybackContext ? PlaybackContext->GetWorld() : nullptr, Player);
 
 				if (AudioComponent)
 				{
diff --git a/Engine/Source/Runtime/MovieSceneTracks/Private/Sections/MovieSceneAudioSection.cpp b/Engine/Source/Runtime/MovieSceneTracks/Private/Sections/MovieSceneAudioSection.cpp
--- a/Engine/Source/Runtime/MovieSceneTracks/Private/Sections/MovieSceneAudioSection.cpp
+++ b/Engine/Source/Runtime/MovieSceneTracks/Private/Sections/MovieSceneAudioSection.cpp
@@ -103,6 +103,8 @@ EMovieSceneChannelProxyType  UMovieSceneAudioSection::CacheChannelProxy()
 	FMovieSceneChannelProxyData Channels;
 
 	UMovieSceneAudioTrack* AudioTrack = Cast<UMovieSceneAudioTrack>(GetOuter());
+	UMovieScene* MovieScene = AudioTrack ? Cast<UMovieScene>(AudioTrack->GetOuter()) : nullptr;
+	const bool bHasAttachData = MovieScene && MovieScene->ContainsTrack(*AudioTrack);
 
 #if WITH_EDITOR
 
@@ -110,7 +112,7 @@ EMovieSceneChannelProxyType  UMovieSceneAudioSection::CacheChannelProxy()
 	Channels.Add(SoundVolume,     EditorData.Data[0], TMovieSceneExternalValue<float>());
 	Channels.Add(PitchMultiplier, EditorData.Data[1], TMovieSceneExternalValue<float>());
 
-	if (AudioTrack && AudioTrack->IsAMasterTrack())
+	if (bHasAttachData)
 	{
 		Channels.Add(AttachActorData, EditorData.Data[2]);
 	}
@@ -119,7 +121,7 @@ EMovieSceneChannelProxyType  UMovieSceneAudioSection::CacheChannelProxy()
 
 	Channels.Add(SoundVolume);
 	Channels.Add(PitchMultiplier);
-	if (AudioTrack && AudioTrack->IsAMasterTrack())
+	if (bHasAttachData)
 	{
 		Channels.Add(AttachActorData);
 	}
diff --git a/Engine/Source/Runtime/MovieSceneTracks/Private/Sections/MovieSceneCVarSection.cpp b/Engine/Source/Runtime/MovieSceneTracks/Private/Sections/MovieSceneCVarSection.cpp
--- a/Engine/Source/Runtime/MovieSceneTracks/Private/Sections/MovieSceneCVarSection.cpp
+++ b/Engine/Source/Runtime/MovieSceneTracks/Private/Sections/MovieSceneCVarSection.cpp
@@ -77,7 +77,7 @@ void UMovieSceneCVarSection::ImportEntityImpl(UMovieSceneEntitySystemLinker* Ent
 
 	OutImportedEntity->AddBuilder(
 		FEntityBuilder()
-		.AddTag(FBuiltInComponentTypes::Get()->Tags.Master)
+		.AddTag(FBuiltInComponentTypes::Get()->Tags.Root)
 		.Add(FBuiltInComponentTypes::Get()->TrackInstance, TrackInstance)
 	);
 }
diff --git a/Engine/Source/Runtime/MovieSceneTracks/Private/Systems/MovieSceneEventSystems.cpp b/Engine/Source/Runtime/MovieSceneTracks/Private/Systems/MovieSceneEventSystems.cpp
--- a/Engine/Source/Runtime/MovieSceneTracks/Private/Systems/MovieSceneEventSystems.cpp
+++ b/Engine/Source/Runtime/MovieSceneTracks/Private/Systems/MovieSceneEventSystems.cpp
@@ -223,7 +223,7 @@ void UMovieSceneEventSystem::TriggerEventWithParameters(UObject* DirectorInstanc
 		}
 	}
 
-	// At this point we know the event is a master track with parameters - either trigger for global contexts, or just on its own
+	// At this point we know the event is a track with parameters - either trigger for global contexts, or just on its own
 	else if (GlobalContexts.Num() != 0)
 	{
 		for (UObject* Context : GlobalContexts)
@@ -278,7 +278,7 @@ bool UMovieSceneEventSystem::PatchBoundObject(uint8* Parameters, UObject* BoundO
 			FMessageLog("PIE").Warning()
 				->AddToken(FUObjectToken::Create(BoundObjectProperty->GetOwnerUObject()))
 				->AddToken(FUObjectToken::Create(Player->GetEvaluationTemplate().GetSequence(SequenceID)))
-				->AddToken(FTextToken::Create(LOCTEXT("LevelBP_LevelScriptActor_Error", "Failed to trigger event: only Interface pins are supported for master tracks within Level Sequences. Please remove the pin, or change it to an interface that is implemented on the desired level blueprint.")));
+				->AddToken(FTextToken::Create(LOCTEXT("LevelBP_LevelScriptActor_Error", "Failed to trigger event: only Interface pins are supported for root tracks within Level Sequences. Please remove the pin, or change it to an interface that is implemented on the desired level blueprint.")));
 
 			return false;
 		}
diff --git a/Engine/Source/Runtime/MovieSceneTracks/Private/Systems/MovieSceneLevelVisibilitySystem.cpp b/Engine/Source/Runtime/MovieSceneTracks/Private/Systems/MovieSceneLevelVisibilitySystem.cpp
--- a/Engine/Source/Runtime/MovieSceneTracks/Private/Systems/MovieSceneLevelVisibilitySystem.cpp
+++ b/Engine/Source/Runtime/MovieSceneTracks/Private/Systems/MovieSceneLevelVisibilitySystem.cpp
@@ -3,7 +3,7 @@
 #include "Systems/MovieSceneLevelVisibilitySystem.h"
 #include "EntitySystem/BuiltInComponentTypes.h"
 #include "EntitySystem/MovieSceneEntitySystemTask.h"
-#include "EntitySystem/MovieSceneMasterInstantiatorSystem.h"
+#include "EntitySystem/MovieSceneRootInstantiatorSystem.h"
 #include "EntitySystem/MovieSceneSpawnablesSystem.h"
 #include "EntitySystem/MovieSceneEntitySystemLinker.h"
 #include "MovieSceneTracksComponentTypes.h"
diff --git a/Engine/Source/Runtime/MovieSceneTracks/Private/Systems/MovieSceneMaterialParameterCollectionSystem.cpp b/Engine/Source/Runtime/MovieSceneTracks/Private/Systems/MovieSceneMaterialParameterCollectionSystem.cpp
--- a/Engine/Source/Runtime/MovieSceneTracks/Private/Systems/MovieSceneMaterialParameterCollectionSystem.cpp
+++ b/Engine/Source/Runtime/MovieSceneTracks/Private/Systems/MovieSceneMaterialParameterCollectionSystem.cpp
@@ -138,7 +138,7 @@ UMovieSceneMaterialParameterCollectionSystem::UMovieSceneMaterialParameterCollec
 
 	if (HasAnyFlags(RF_ClassDefaultObject))
 	{
-		DefineComponentConsumer(GetClass(), BuiltInComponents->Tags.Master);
+		DefineComponentConsumer(GetClass(), BuiltInComponents->Tags.Root);
 		DefineComponentProducer(GetClass(), TracksComponents->BoundMaterial);
 		DefineImplicitPrerequisite(UMovieSceneCachePreAnimatedStateSystem::StaticClass(), GetClass());
 	}
diff --git a/Engine/Source/Runtime/MovieSceneTracks/Private/Systems/MovieScenePiecewiseDoubleBlenderSystem.cpp b/Engine/Source/Runtime/MovieSceneTracks/Private/Systems/MovieScenePiecewiseDoubleBlenderSystem.cpp
--- a/Engine/Source/Runtime/MovieSceneTracks/Private/Systems/MovieScenePiecewiseDoubleBlenderSystem.cpp
+++ b/Engine/Source/Runtime/MovieSceneTracks/Private/Systems/MovieScenePiecewiseDoubleBlenderSystem.cpp
@@ -477,7 +477,7 @@ void UMovieScenePiecewiseDoubleBlenderSystem::OnRun(FSystemTaskPrerequisites& In
 
 		if (Task)
 		{
-			Prereqs.AddMasterTask(Task);
+			Prereqs.AddRootTask(Task);
 		}
 	}
 
@@ -494,7 +494,7 @@ void UMovieScenePiecewiseDoubleBlenderSystem::OnRun(FSystemTaskPrerequisites& In
 
 		if (Task)
 		{
-			Prereqs.AddMasterTask(Task);
+			Prereqs.AddRootTask(Task);
 		}
 	}
 
@@ -511,7 +511,7 @@ void UMovieScenePiecewiseDoubleBlenderSystem::OnRun(FSystemTaskPrerequisites& In
 
 		if (Task)
 		{
-			Prereqs.AddMasterTask(Task);
+			Prereqs.AddRootTask(Task);
 		}
 	}
 
@@ -528,7 +528,7 @@ void UMovieScenePiecewiseDoubleBlenderSystem::OnRun(FSystemTaskPrerequisites& In
 
 		if (Task)
 		{
-			Prereqs.AddMasterTask(Task);
+			Prereqs.AddRootTask(Task);
 		}
 	}
 
diff --git a/Engine/Source/Runtime/MovieSceneTracks/Private/Systems/MovieScenePiecewiseIntegerBlenderSystem.cpp b/Engine/Source/Runtime/MovieSceneTracks/Private/Systems/MovieScenePiecewiseIntegerBlenderSystem.cpp
--- a/Engine/Source/Runtime/MovieSceneTracks/Private/Systems/MovieScenePiecewiseIntegerBlenderSystem.cpp
+++ b/Engine/Source/Runtime/MovieSceneTracks/Private/Systems/MovieScenePiecewiseIntegerBlenderSystem.cpp
@@ -285,7 +285,7 @@ void UMovieScenePiecewiseIntegerBlenderSystem::OnRun(FSystemTaskPrerequisites& I
 
 		if (Task)
 		{
-			Prereqs.AddMasterTask(Task);
+			Prereqs.AddRootTask(Task);
 		}
 	}
 
@@ -301,7 +301,7 @@ void UMovieScenePiecewiseIntegerBlenderSystem::OnRun(FSystemTaskPrerequisites& I
 
 		if (Task)
 		{
-			Prereqs.AddMasterTask(Task);
+			Prereqs.AddRootTask(Task);
 		}
 	}
 
@@ -317,7 +317,7 @@ void UMovieScenePiecewiseIntegerBlenderSystem::OnRun(FSystemTaskPrerequisites& I
 
 		if (Task)
 		{
-			Prereqs.AddMasterTask(Task);
+			Prereqs.AddRootTask(Task);
 		}
 	}
 
@@ -334,11 +334,11 @@ void UMovieScenePiecewiseIntegerBlenderSystem::OnRun(FSystemTaskPrerequisites& I
 
 		if (Task)
 		{
-			Prereqs.AddMasterTask(Task);
+			Prereqs.AddRootTask(Task);
 		}
 	}
 
-	// Master task that performs the actual blends
+	// Root task that performs the actual blends
 	FEntityTaskBuilder()
 	.Read(BuiltInComponents->BlendChannelOutput)
 	.ReadOptional(TracksComponents->Integer.InitialValue)
diff --git a/Engine/Source/Runtime/MovieSceneTracks/Private/Systems/WeightAndEasingEvaluatorSystem.cpp b/Engine/Source/Runtime/MovieSceneTracks/Private/Systems/WeightAndEasingEvaluatorSystem.cpp
--- a/Engine/Source/Runtime/MovieSceneTracks/Private/Systems/WeightAndEasingEvaluatorSystem.cpp
+++ b/Engine/Source/Runtime/MovieSceneTracks/Private/Systems/WeightAndEasingEvaluatorSystem.cpp
@@ -393,7 +393,7 @@ void UWeightAndEasingEvaluatorSystem::OnRun(FSystemTaskPrerequisites& InPrerequi
 			FGraphEventRef AccumulateTask = TGraphTask<FAccumulateHierarchicalEasings>::CreateTask(&AccumulatePrereqs, Linker->EntityManager.GetDispatchThread())
 				.ConstructAndDispatchWhenReady(&EasingChannels);
 
-			PropagatePrereqs.AddMasterTask(AccumulateTask);
+			PropagatePrereqs.AddRootTask(AccumulateTask);
 		}
 
 		// Step 3: Apply hierarchical easing results to all entities inside affected sub-sequences.
diff --git a/Engine/Source/Runtime/MovieSceneTracks/Private/Tracks/MovieSceneAudioTrack.cpp b/Engine/Source/Runtime/MovieSceneTracks/Private/Tracks/MovieSceneAudioTrack.cpp
--- a/Engine/Source/Runtime/MovieSceneTracks/Private/Tracks/MovieSceneAudioTrack.cpp
+++ b/Engine/Source/Runtime/MovieSceneTracks/Private/Tracks/MovieSceneAudioTrack.cpp
@@ -145,14 +145,6 @@ UMovieSceneSection* UMovieSceneAudioTrack::AddNewSoundOnRow(USoundBase* Sound, F
 	return NewSection;
 }
 
-
-bool UMovieSceneAudioTrack::IsAMasterTrack() const
-{
-	UMovieScene* MovieScene = Cast<UMovieScene>(GetOuter());
-	return MovieScene ? MovieScene->IsAMasterTrack(*this) : false;
-}
-
-
 UMovieSceneSection* UMovieSceneAudioTrack::CreateNewSection()
 {
 	return NewObject<UMovieSceneAudioSection>(this, NAME_None, RF_Transactional);
diff --git a/Engine/Source/Runtime/MovieSceneTracks/Private/Tracks/MovieSceneMaterialParameterCollectionTrack.cpp b/Engine/Source/Runtime/MovieSceneTracks/Private/Tracks/MovieSceneMaterialParameterCollectionTrack.cpp
--- a/Engine/Source/Runtime/MovieSceneTracks/Private/Tracks/MovieSceneMaterialParameterCollectionTrack.cpp
+++ b/Engine/Source/Runtime/MovieSceneTracks/Private/Tracks/MovieSceneMaterialParameterCollectionTrack.cpp
@@ -49,7 +49,7 @@ void UMovieSceneMaterialParameterCollectionTrack::ExtendEntityImpl(UMovieSceneEn
 		FEntityBuilder()
 		.Add(TracksComponents->MPC, MPC)
 		.AddTag(BuiltInComponents->Tags.AbsoluteBlend)
-		.AddTag(BuiltInComponents->Tags.Master)
+		.AddTag(BuiltInComponents->Tags.Root)
 	);
 }
 
diff --git a/Engine/Source/Runtime/MovieSceneTracks/Public/Sections/MovieSceneAudioSection.h b/Engine/Source/Runtime/MovieSceneTracks/Public/Sections/MovieSceneAudioSection.h
--- a/Engine/Source/Runtime/MovieSceneTracks/Public/Sections/MovieSceneAudioSection.h
+++ b/Engine/Source/Runtime/MovieSceneTracks/Public/Sections/MovieSceneAudioSection.h
@@ -18,7 +18,7 @@
 class USoundBase;
 
 /**
- * Audio section, for use in the master audio, or by attached audio objects
+ * Audio section, for use in the audio track, or by attached audio objects
  */
 UCLASS()
 class MOVIESCENETRACKS_API UMovieSceneAudioSection
diff --git a/Engine/Source/Runtime/MovieSceneTracks/Public/Systems/MovieSceneBlenderSystemHelper.h b/Engine/Source/Runtime/MovieSceneTracks/Public/Systems/MovieSceneBlenderSystemHelper.h
--- a/Engine/Source/Runtime/MovieSceneTracks/Public/Systems/MovieSceneBlenderSystemHelper.h
+++ b/Engine/Source/Runtime/MovieSceneTracks/Public/Systems/MovieSceneBlenderSystemHelper.h
@@ -179,7 +179,7 @@ public:
 					&Linker->EntityManager, InPrerequisites, nullptr, BlenderSystem->GetBlenderSystemID(), BlendChannelResults);
 		if (Task)
 		{
-			Prereqs.AddMasterTask(Task);
+			Prereqs.AddRootTask(Task);
 		}
 
 		// Compute the final blended values and assign them to the blend channel output entities.
diff --git a/Engine/Source/Runtime/MovieSceneTracks/Public/Tracks/MovieSceneAudioTrack.h b/Engine/Source/Runtime/MovieSceneTracks/Public/Tracks/MovieSceneAudioTrack.h
--- a/Engine/Source/Runtime/MovieSceneTracks/Public/Tracks/MovieSceneAudioTrack.h
+++ b/Engine/Source/Runtime/MovieSceneTracks/Public/Tracks/MovieSceneAudioTrack.h
@@ -40,7 +40,7 @@ public:
 		return AudioSections;
 	}
 
-	/** @return true if this is a master audio track */
+	UE_DEPRECATED(5.2, "IsAMasterTrack is deprecated. Please use Cast<UMovieScene>(GetOuter())->ContainsTrack instead")
 	bool IsAMasterTrack() const;
 
 public:
@@ -63,7 +63,7 @@ public:
 
 private:
 
-	/** List of all master audio sections */
+	/** List of all root audio sections */
 	UPROPERTY()
 	TArray<TObjectPtr<UMovieSceneSection>> AudioSections;
 
