Sequencer: fix pointer arithemtic to not trip up AddressSanitizer

diff --git a/Engine/Source/Runtime/MovieScene/Public/EntitySystem/MovieScenePropertyComponentHandler.h b/Engine/Source/Runtime/MovieScene/Public/EntitySystem/MovieScenePropertyComponentHandler.h
--- a/Engine/Source/Runtime/MovieScene/Public/EntitySystem/MovieScenePropertyComponentHandler.h
+++ b/Engine/Source/Runtime/MovieScene/Public/EntitySystem/MovieScenePropertyComponentHandler.h
@@ -443,8 +443,8 @@ struct TPropertyComponentHandlerImpl<PropertyTraits, TPropertyMetaData<MetaDataT
 					FAlignedDecomposedValue& AlignedOutput = AlignedOutputs[CompositeIndex];
 					if (InitialValueComponent)
 					{
-						const StorageType& InitialValue = (*InitialValueComponent);
-						InitialValueComposite = reinterpret_cast<const double*>(reinterpret_cast<const uint8*>(&InitialValue) + Composites[CompositeIndex].CompositeOffset);
+						const StorageType* InitialValuePtr = InitialValueComponent.AsPtr();
+						InitialValueComposite = reinterpret_cast<const double*>(reinterpret_cast<const uint8*>(InitialValuePtr) + Composites[CompositeIndex].CompositeOffset);
 					}
 
 					const double NewComposite = *reinterpret_cast<const double*>(reinterpret_cast<const uint8*>(&InCurrentValue) + Composites[CompositeIndex].CompositeOffset);
@@ -503,8 +503,8 @@ struct TPropertyComponentHandlerImpl<PropertyTraits, TPropertyMetaData<MetaDataT
 				const double* InitialValueComposite = nullptr;
 				if (InitialValueComponent)
 				{
-					const StorageType& InitialValue = (*InitialValueComponent);
-					InitialValueComposite = reinterpret_cast<const double*>(reinterpret_cast<const uint8*>(&InitialValue) + Composite.CompositeOffset);
+					const StorageType* InitialValuePtr = InitialValueComponent.AsPtr();
+					InitialValueComposite = reinterpret_cast<const double*>(reinterpret_cast<const uint8*>(InitialValuePtr) + Composite.CompositeOffset);
 				}
 
 				const double RecomposedComposite = AlignedOutput.Value.Recompose(EntityID, InCurrentValue, InitialValueComposite);
