Sequencer: Override ConstructEntitySystemLinker so that it is unique.

This fixes an issue where having the Day Sequence Editor and Level Sequence Editor both active at the same time would have shared the EntitySystemLinker. That led to a spawnable being removed from the registry with the id entry from the day sequence instead of the id entry from the level sequence. Subsequently, that spawnable wouldn't be spawned again after an autosave/restore state because it was still found in the registry.

diff --git a/Engine/Source/Runtime/MovieScene/Private/Evaluation/MovieSceneEvaluationTemplateInstance.cpp b/Engine/Source/Runtime/MovieScene/Private/Evaluation/MovieSceneEvaluationTemplateInstance.cpp
--- a/Engine/Source/Runtime/MovieScene/Private/Evaluation/MovieSceneEvaluationTemplateInstance.cpp
+++ b/Engine/Source/Runtime/MovieScene/Private/Evaluation/MovieSceneEvaluationTemplateInstance.cpp
@@ -149,6 +149,11 @@ void FMovieSceneRootEvaluationTemplateInstance::Initialize(UMovieSceneSequence&
 		TSharedPtr<FMovieSceneEntitySystemRunner> Runner = WeakRunner.Pin();
 		if (ensure(Runner) && EntitySystemLinker)
 		{
+			if (Runner->IsAttachedToLinker() && Runner->GetLinker() != EntitySystemLinker)
+			{
+				Runner->DetachFromLinker();
+			}
+
 			if (!Runner->IsAttachedToLinker())
 			{
 				Runner->AttachToLinker(EntitySystemLinker);
