Sequencer: Guard against issuing update requests for sub-sequence instances

diff --git a/Engine/Source/Runtime/MovieScene/Private/EntitySystem/MovieSceneEntitySystemRunner.cpp b/Engine/Source/Runtime/MovieScene/Private/EntitySystem/MovieSceneEntitySystemRunner.cpp
--- a/Engine/Source/Runtime/MovieScene/Private/EntitySystem/MovieSceneEntitySystemRunner.cpp
+++ b/Engine/Source/Runtime/MovieScene/Private/EntitySystem/MovieSceneEntitySystemRunner.cpp
@@ -5,6 +5,8 @@
 #include "EntitySystem/MovieSceneEntityMutations.h"
 #include "EntitySystem/BuiltInComponentTypes.h"
 #include "Evaluation/PreAnimatedState/MovieScenePreAnimatedCaptureSource.h"
+#include "Evaluation/MovieSceneEvaluationTemplateInstance.h"
+#include "MovieSceneSequence.h"
 #include "ProfilingDebugging/CountersTrace.h"
 
 DECLARE_CYCLE_STAT(TEXT("Runner Flush"), 				MovieSceneEval_RunnerFlush, 				STATGROUP_MovieSceneEval);
@@ -595,6 +597,19 @@ bool FMovieSceneEntitySystemRunner::GameThread_UpdateSequenceInstances(UMovieSce
 
 			// Give the instance an opportunity to dissect the range into distinct evaluations
 			FSequenceInstance& Instance = InstanceRegistry->MutateInstance(Request.Params.InstanceHandle);
+			if (!Instance.IsRootSequence())
+			{
+				FMovieSceneRootEvaluationTemplateInstance& Template = Instance.GetPlayer()->GetEvaluationTemplate();
+				UMovieSceneSequence* RootSequence = Template.GetRootSequence();
+				UMovieSceneSequence* SubSequence  = Template.GetSequence(Instance.GetSequenceID());
+
+				ensureMsgf(Instance.IsRootSequence(), TEXT("Update request received for a non-root sequence ID 0x%08X (%s) in root-sequence %s. This is not supported."),
+					Instance.GetSequenceID().GetInternalValue(),
+					SubSequence  ? *SubSequence->GetName()  : TEXT("<nullptr>"),
+					RootSequence ? *RootSequence->GetName() : TEXT("<nullptr>")
+				);
+				continue;
+			}
 			Instance.DissectContext(Linker, Request.Context, Dissections);
 
 			if (Dissections.Num() != 0)
