Sequencer: Improved ensure message for upstream dependencies added during excution

diff --git a/Engine/Source/Runtime/MovieScene/Private/EntitySystem/MovieSceneEntitySystemGraphs.cpp b/Engine/Source/Runtime/MovieScene/Private/EntitySystem/MovieSceneEntitySystemGraphs.cpp
--- a/Engine/Source/Runtime/MovieScene/Private/EntitySystem/MovieSceneEntitySystemGraphs.cpp
+++ b/Engine/Source/Runtime/MovieScene/Private/EntitySystem/MovieSceneEntitySystemGraphs.cpp
@@ -781,7 +781,15 @@ void FMovieSceneEntitySystemGraph::ExecutePhase(const ArrayType& SortedEntries,
 #if DO_CHECK
 			for (int32 NewIndex = 0; NewIndex < CurrentIndex; ++NewIndex)
 			{
-				ensureMsgf(HeadList.Contains(SortedEntries[NewIndex]), TEXT("New system has been added upstream to the same execution phase as is currently in-flight - this will not be run this frame"));
+				if (!HeadList.Contains(SortedEntries[NewIndex]))
+				{
+					const uint16 NewNodeIndex     = SortedEntries[NewIndex];
+					const uint16 CurrentNodeIndex = SortedEntries[CurrentIndex];
+					ensureAlwaysMsgf(false, TEXT("New system %s has been added upstream of %s in the same execution phase that is currently in-flight - this will not be run this frame"),
+						*this->Nodes.Array[NewNodeIndex].System->GetName(),
+						*this->Nodes.Array[CurrentNodeIndex].System->GetName()
+					);
+				}
 			}
 #endif
 		}
