Sequence: Added explicit casts to appease some compilers

diff --git a/Engine/Source/Runtime/MovieScene/Private/EntitySystem/MovieSceneEntityManager.cpp b/Engine/Source/Runtime/MovieScene/Private/EntitySystem/MovieSceneEntityManager.cpp
--- a/Engine/Source/Runtime/MovieScene/Private/EntitySystem/MovieSceneEntityManager.cpp
+++ b/Engine/Source/Runtime/MovieScene/Private/EntitySystem/MovieSceneEntityManager.cpp
@@ -987,7 +987,7 @@ void FEntityManager::EnterIteration() const
 
 void FEntityManager::ExitIteration() const
 {
-	checkSlow(IterationCount > 0);
+	checkSlow(static_cast<uint16>(IterationCount) > 0);
 	--IterationCount;
 }
 
diff --git a/Engine/Source/Runtime/MovieScene/Public/EntitySystem/MovieSceneEntityManager.h b/Engine/Source/Runtime/MovieScene/Public/EntitySystem/MovieSceneEntityManager.h
--- a/Engine/Source/Runtime/MovieScene/Public/EntitySystem/MovieSceneEntityManager.h
+++ b/Engine/Source/Runtime/MovieScene/Public/EntitySystem/MovieSceneEntityManager.h
@@ -751,7 +751,7 @@ public:
 
 	void CheckCanChangeStructure() const
 	{
-		checkf(IterationCount == 0, TEXT("Mutation of entities is not permissible while entities are being iterated"));
+		checkf(static_cast<uint16>(IterationCount) == 0, TEXT("Mutation of entities is not permissible while entities are being iterated"));
 		checkf(LockdownState == ELockdownState::Unlocked, TEXT("Structural changes to the entity manager are not permitted while it is locked down"));
 	}
 
@@ -786,7 +786,7 @@ public:
 	 */
 	FEntityComponentFilter& ModifyGlobalIterationFilter()
 	{
-		ensureMsgf(!IsLockedDown() && IterationCount == 0, TEXT("Manipulating the global iteration filter while locked down or iterating is not recommended"));
+		ensureMsgf(!IsLockedDown() && static_cast<uint16>(IterationCount) == 0, TEXT("Manipulating the global iteration filter while locked down or iterating is not recommended"));
 		return GlobalIterationFilter;
 	}
 
