Sequencer: Added the ability to perform batch mutations on a per-entity basis

diff --git a/Engine/Source/Runtime/MovieScene/Private/EntitySystem/MovieSceneEntityManager.cpp b/Engine/Source/Runtime/MovieScene/Private/EntitySystem/MovieSceneEntityManager.cpp
--- a/Engine/Source/Runtime/MovieScene/Private/EntitySystem/MovieSceneEntityManager.cpp
+++ b/Engine/Source/Runtime/MovieScene/Private/EntitySystem/MovieSceneEntityManager.cpp
@@ -34,6 +34,8 @@ namespace UE
 namespace MovieScene
 {
 
+FComponentMask GEntityManagerEmptyMask;
+
 // @todo: this is a very rough initial guess at the break even point for when threaded evaluation becomes beneficial, and will vary highly between platforms and hardware.
 // We may wish to make this more flexible in future (such as only threading hot paths such as float channel evaluation) by enabling threading per-task, but more data is required to make such decisions
 int32 GThreadedEvaluationAllocationThreshold = 32;
@@ -271,9 +273,11 @@ FFreeEntityOperation::FCommitData FFreeEntityOperation::Commit() const
 {
 	FFreeEntityOperation::FCommitData CommitData;
 
+	FBuiltInComponentTypes* BuiltInComponents = FBuiltInComponentTypes::Get();
+
 	TArray<FMovieSceneEntityID, TInlineAllocator<16>> EntitiesScratch;
 
-	auto MarkChildrenForFree = [this, &EntitiesScratch, &CommitData](FMovieSceneEntityID Entity)
+	auto MarkChildrenForFree = [this, &EntitiesScratch, &CommitData, BuiltInComponents](FMovieSceneEntityID Entity)
 	{
 		for (auto ChildIt = this->EntityManager->ParentToChild.CreateKeyIterator(Entity); ChildIt; ++ChildIt)
 		{
@@ -291,6 +295,8 @@ FFreeEntityOperation::FCommitData FFreeEntityOperation::Commit() const
 				FEntityManager::FEntityLocation Location = EntityManager->EntityLocations[MarkedEntity.AsIndex()];
 				if (Location.IsValid())
 				{
+					ensureAlwaysMsgf(EntityManager->EntityAllocationMasks[Location.GetAllocationIndex()].Contains(BuiltInComponents->Tags.NeedsUnlink), TEXT("Attempting to free an entity that has not been unlinked - this might result in stale references"));
+
 					FAllocationMask& Mask = CommitData.AllocationsToEntities.FindOrAdd(Location.GetAllocationIndex());
 					if (!Mask.bDestroyAllocation)
 					{
@@ -1070,7 +1076,11 @@ int32 FEntityManager::GetOrCreateAllocationWithSlack(const FComponentMask& Entit
 			return AllocationIndex;
 		}
 	}
+	return CreateAllocationWithSlack(EntityComponentMask, InOutDesiredSlack);
+}
 
+int32 FEntityManager::CreateAllocationWithSlack(const FComponentMask& EntityComponentMask, int32* InOutDesiredSlack)
+{
 	static const uint16 MaxCapacity = 2048;
 
 	uint16 DefaultCapacity = 4;
@@ -1204,10 +1214,10 @@ bool FEntityManager::HasComponent(FMovieSceneEntityID EntityID, FComponentTypeID
 	return Location.IsValid() && EntityAllocationMasks[Location.GetAllocationIndex()].Contains(ComponentType);
 }
 
-FComponentMask FEntityManager::GetEntityType(FMovieSceneEntityID EntityID) const
+const FComponentMask& FEntityManager::GetEntityType(FMovieSceneEntityID EntityID) const
 {
 	FEntityLocation Location = EntityLocations[EntityID.AsIndex()];
-	return Location.IsValid() ? EntityAllocationMasks[Location.GetAllocationIndex()] : FComponentMask();
+	return Location.IsValid() ? EntityAllocationMasks[Location.GetAllocationIndex()] : GEntityManagerEmptyMask;
 }
 
 void FEntityManager::ChangeEntityType(FMovieSceneEntityID EntityID, const FComponentMask& InNewMask)
@@ -1619,6 +1629,11 @@ int32 FEntityManager::MutateAll(const FEntityComponentFilter& Filter, const IMov
 			{
 				AllocationMutations.Add(AllocationIndex, MoveTemp(NewMutation));
 			}
+			else
+			{
+				FEntityAllocation* SourceAllocation = EntityAllocations[AllocationIndex];
+				Mutation.InitializeUnmodifiedAllocation(SourceAllocation, NewMutation);
+			}
 		}
 	}
 
@@ -1673,6 +1688,109 @@ int32 FEntityManager::MutateAll(const FEntityComponentFilter& Filter, const IMov
 	return TotalNumMutations;
 }
 
+int32 FEntityManager::MutateConditional(const FEntityComponentFilter& Filter, const IMovieSceneConditionalEntityMutation& Mutation)
+{
+	CheckCanChangeStructure();
+
+	TMap<int32, TTuple<FComponentMask, TBitArray<>>> AllocationMutations;
+
+	for (int32 AllocationIndex = 0; AllocationIndex < EntityAllocationMasks.GetMaxIndex(); ++AllocationIndex)
+	{
+		if (EntityAllocationMasks.IsAllocated(AllocationIndex) && Filter.Match(EntityAllocationMasks[AllocationIndex]))
+		{
+			FEntityAllocation* SourceAllocation = EntityAllocations[AllocationIndex];
+
+			TBitArray<> MarkedEntities;
+			Mutation.MarkAllocation(SourceAllocation, MarkedEntities);
+
+			if (MarkedEntities.Num() == 0)
+			{
+				continue;
+			}
+
+			const int32 NumIrrelevantBits = MarkedEntities.Num() - SourceAllocation->Num();
+			if (NumIrrelevantBits > 1)
+			{
+				MarkedEntities.RemoveAt(SourceAllocation->Num(), NumIrrelevantBits);
+			}
+
+			FComponentMask NewMutation = EntityAllocationMasks[AllocationIndex];
+			Mutation.CreateMutation(this, &NewMutation);
+
+			if (!NewMutation.CompareSetBits(EntityAllocationMasks[AllocationIndex]))
+			{
+				AllocationMutations.Add(AllocationIndex, MakeTuple(MoveTemp(NewMutation), MoveTemp(MarkedEntities)));
+			}
+		}
+	}
+
+	if (AllocationMutations.Num() == 0)
+	{
+		return 0;
+	}
+
+	int32 TotalNumMutations = 0;
+
+	for (TTuple<int32, TTuple<FComponentMask, TBitArray<>>>& Pair : AllocationMutations)
+	{
+		const int32 SourceAllocationIndex = Pair.Key;
+		FEntityAllocation* SourceAllocation = EntityAllocations[SourceAllocationIndex];
+
+		const int32 NumEntities = Pair.Value.Value.CountSetBits();
+		const bool bAllEntitiesMarked = NumEntities == SourceAllocation->Num();
+
+		TotalNumMutations += NumEntities;
+
+		if (bAllEntitiesMarked)
+		{
+			// When adding a component to an entire allocation we just reallocate within the same allocation entry to avoid having to fix up 
+			// Specific entity entry indices
+			FEntityAllocation* NewAllocation = MigrateAllocation(SourceAllocationIndex, Pair.Value.Key);
+
+			EntityAllocationMasks[SourceAllocationIndex] = Pair.Value.Key;
+			EntityAllocations[SourceAllocationIndex] = NewAllocation;
+
+			FEntityRange Range { NewAllocation, 0, NewAllocation->Num() };
+
+			Mutation.InitializeEntities(Range, Pair.Value.Key);
+
+			DestroyAllocation(SourceAllocation);
+		}
+		else
+		{
+			int32 DesiredSlack = NumEntities;
+			const int32 NewAllocationIndex = CreateAllocationWithSlack(Pair.Value.Key, &DesiredSlack);
+
+			const FEntityAllocation* DestAllocation = EntityAllocations[NewAllocationIndex];
+			FEntityRange Range { DestAllocation, DestAllocation->Num(), NumEntities };
+
+			const FMovieSceneEntityID* EntityIDs = SourceAllocation->GetRawEntityIDs();
+
+			// Migrate entities to the new allocation - care is taken to iterate the allocation backwards
+			// since entities can shift around in the allocation as they are migrated
+			for (TBitArray<>::FConstReverseIterator It(Pair.Value.Value); It; ++It)
+			{
+				if (It.GetValue())
+				{
+					FEntityLocation& Location = EntityLocations[EntityIDs[It.GetIndex()].AsIndex()];
+
+					int32 NewEntityIndex = MigrateEntity(NewAllocationIndex, Location.GetAllocationIndex(), Location.GetEntryIndexWithinAllocation());
+					Location.Set(NewAllocationIndex, NewEntityIndex);
+				}
+			}
+
+			check(Range.ComponentStartOffset + Range.Num == Range.Allocation->Num());
+			Mutation.InitializeEntities(Range, Pair.Value.Key);
+		}
+	}
+
+	CheckInvariants();
+
+	OnStructureChanged();
+
+	return TotalNumMutations;
+}
+
 void FEntityManager::TouchEntity(FMovieSceneEntityID EntityID)
 {
 	FEntityLocation Location = EntityLocations[EntityID.AsIndex()];
diff --git a/Engine/Source/Runtime/MovieScene/Public/EntitySystem/MovieSceneEntityManager.h b/Engine/Source/Runtime/MovieScene/Public/EntitySystem/MovieSceneEntityManager.h
--- a/Engine/Source/Runtime/MovieScene/Public/EntitySystem/MovieSceneEntityManager.h
+++ b/Engine/Source/Runtime/MovieScene/Public/EntitySystem/MovieSceneEntityManager.h
@@ -52,6 +52,7 @@ struct IComponentTypeHandler;
 struct IMovieSceneEntityMutation;
 template <typename T> struct TReadOptional;
 template <typename T> struct TWriteOptional;
+struct IMovieSceneConditionalEntityMutation;
 
 
 enum class EEntityRecursion : uint8
@@ -324,7 +325,7 @@ public:
 	 * @param InEntity  The ID of the entity
 	 * @return The type mask for this entity, or an empty mask if it has no components
 	 */
-	FComponentMask GetEntityType(FMovieSceneEntityID InEntity) const;
+	const FComponentMask& GetEntityType(FMovieSceneEntityID InEntity) const;
 
 
 	/**
@@ -435,7 +436,7 @@ public:
 		}
 
 		FEntityLocation Location = EntityLocations[Entity.AsIndex()];
-		if (!Location.IsValid())
+		if (!Location.IsValid() || !EntityAllocationMasks[Location.GetAllocationIndex()].Contains(ComponentTypeID))
 		{
 			return TComponentLock<TReadOptional<T>>();
 		}
@@ -491,7 +492,7 @@ public:
 		}
 
 		FEntityLocation Location = EntityLocations[Entity.AsIndex()];
-		if (!Location.IsValid())
+		if (!Location.IsValid() || !EntityAllocationMasks[Location.GetAllocationIndex()].Contains(ComponentTypeID))
 		{
 			return TComponentLock<TWriteOptional<T>>();
 		}
@@ -565,6 +566,16 @@ public:
 	int32 MutateAll(const FEntityComponentFilter& Filter, const IMovieSceneEntityMutation& Mutation);
 
 
+	/**
+	 * Efficiently mutate all entities that match a filter. Mutations can add or remove components from batches of entity data.
+	 *
+	 * @param Filter      The filter to match entity allocations against. Only entities that match the filter will be mutated
+	 * @param Mutation    Implementation that defines how to mutate the entities that match the filter
+	 * @return The number of entities that were mutated, or 0 if none were matched
+	 */
+	int32 MutateConditional(const FEntityComponentFilter& Filter, const IMovieSceneConditionalEntityMutation& Mutation);
+
+
 	/**
 	 * Touch the specified entity, cause the allocation and component serial numbers to be incremented. Will invalidate any transient caches maintained for serial such numbers.
 	 * @note Allocation and component serial numbers relate to the entire allocation of entities that this entity resides within. As such, it will invalidate downstream caches
@@ -872,6 +883,7 @@ private:
 	int32 CreateEntityAllocationEntry(const FComponentMask& EntityComponentMask, uint16 InitialCapacity, uint16 MaxCapacity);
 
 	int32 GetOrCreateAllocationWithSlack(const FComponentMask& EntityComponentMask, int32* InOutDesiredSlack = nullptr);
+	int32 CreateAllocationWithSlack(const FComponentMask& EntityComponentMask, int32* InOutDesiredSlack = nullptr);
 	int32 MigrateEntity(int32 DestIndex, int32 SourceIndex, int32 SourceEntryIndexWithinAllocation);
 
 	void CopyComponents(int32 DestAllocationIndex, int32 DestEntityIndex, int32 SourceAllocationIndex, int32 SourceEntityIndex, const FComponentMask* OptionalMask = nullptr);
diff --git a/Engine/Source/Runtime/MovieScene/Public/EntitySystem/MovieSceneEntityMutations.h b/Engine/Source/Runtime/MovieScene/Public/EntitySystem/MovieSceneEntityMutations.h
--- a/Engine/Source/Runtime/MovieScene/Public/EntitySystem/MovieSceneEntityMutations.h
+++ b/Engine/Source/Runtime/MovieScene/Public/EntitySystem/MovieSceneEntityMutations.h
@@ -13,12 +13,27 @@ namespace MovieScene
 {
 class FEntityManager;
 struct FEntityAllocation;
+struct FEntityDataLocation;
 
 struct IMovieSceneEntityMutation
 {
 	virtual ~IMovieSceneEntityMutation() {}
 	virtual void CreateMutation(FEntityManager* EntityManager, FComponentMask* InOutEntityComponentTypes) const = 0;
 	virtual void InitializeAllocation(FEntityAllocation* Allocation, const FComponentMask& AllocationType) const {}
+	virtual void InitializeUnmodifiedAllocation(FEntityAllocation* Allocation, const FComponentMask& AllocationType) const {}
+};
+
+struct IMovieScenePerEntityMutation
+{
+	virtual ~IMovieScenePerEntityMutation() {}
+
+	virtual void CreateMutation(FEntityManager* EntityManager, FComponentMask* InOutEntityComponentTypes) const = 0;
+	virtual void InitializeEntities(const FEntityRange& EntityRange, const FComponentMask& AllocationType) const {}
+};
+
+struct IMovieSceneConditionalEntityMutation : IMovieScenePerEntityMutation
+{
+	virtual void MarkAllocation(FEntityAllocation* Allocation, TBitArray<>& OutEntitiesToMutate) const = 0;
 };
 
 struct FAddSingleMutation : IMovieSceneEntityMutation
