Sequencer: Fixed a crash when running latent actions for a player that has no sync runner or tick manager

diff --git a/Engine/Source/Runtime/MovieScene/Private/MovieSceneSequencePlayer.cpp b/Engine/Source/Runtime/MovieScene/Private/MovieSceneSequencePlayer.cpp
--- a/Engine/Source/Runtime/MovieScene/Private/MovieSceneSequencePlayer.cpp
+++ b/Engine/Source/Runtime/MovieScene/Private/MovieSceneSequencePlayer.cpp
@@ -1801,18 +1801,17 @@ void UMovieSceneSequencePlayer::RunLatentActions()
 		return;
 	}
 
-	if (ensure(TickManager) && !EnumHasAnyFlags(Sequence->GetFlags(), EMovieSceneSequenceFlags::BlockingEvaluation))
-	{
-		TickManager->RunLatentActions();
-	}
-	else
+	if (SynchronousRunner)
 	{
-		check(SynchronousRunner);
 		LatentActionManager.RunLatentActions([this]
 		{
 			this->SynchronousRunner->Flush();
 		});
 	}
+	else if (ensure(TickManager) && !EnumHasAnyFlags(Sequence->GetFlags(), EMovieSceneSequenceFlags::BlockingEvaluation))
+	{
+		TickManager->RunLatentActions();
+	}
 }
 
 void UMovieSceneSequencePlayer::SetWeight(double InWeight)
