Sequencer: Added safety check to ensure that invalid IDs are not supplied from overlapping entity trackers

diff --git a/Engine/Source/Runtime/MovieScene/Public/EntitySystem/MovieSceneOverlappingEntityTracker.h b/Engine/Source/Runtime/MovieScene/Public/EntitySystem/MovieSceneOverlappingEntityTracker.h
--- a/Engine/Source/Runtime/MovieScene/Public/EntitySystem/MovieSceneOverlappingEntityTracker.h
+++ b/Engine/Source/Runtime/MovieScene/Public/EntitySystem/MovieSceneOverlappingEntityTracker.h
@@ -220,9 +220,17 @@ struct TOverlappingEntityTrackerImpl
 				const uint16 OutputIndex = static_cast<uint16>(InvalidOutput.GetIndex());
 
 				InputArray.Reset();
-				for (auto Inputs = OutputToEntity.CreateConstKeyIterator(OutputIndex); Inputs; ++Inputs)
+				for (auto Inputs = OutputToEntity.CreateKeyIterator(OutputIndex); Inputs; ++Inputs)
 				{
-					InputArray.Add(Inputs.Value());
+					if (ensure(Linker->EntityManager.IsAllocated(Inputs.Value())))
+					{
+						InputArray.Add(Inputs.Value());
+					}
+					else
+					{
+						EntityToOutput.Remove(Inputs.Value());
+						Inputs.RemoveCurrent();
+					}
 				}
 
 				FOutput& Output = Outputs[OutputIndex];
