Sequencer: Added cachable interpolation implementations to the sequencer runtime

  - These implementations are now stored for double and float channel components which allow us to cache the interpolation for a given pair of keys, bypassing the more expensive full-evaluation codepath that results in multiple cache-misses
  - We may consider separating the range and implementation members from FCachedInterpolation in order to achieve better cache and/or serparate each implementation into its own component type.
  - Performance metrics show the optimized approach taking just 30% of the time for 1000 transform properties (~roughly 9000 channels)
  - Channel flags have been removed in favor of using the presence of the cached interpolation components as indication that it needs evaluating. Constant channels have this component removed on link when they are deemed to not require any updating at all. Systems that previously wrote to the flags now just add this component back on to ensure they get re-evaluated.

diff --git a/Engine/Plugins/MovieScene/TemplateSequence/Source/TemplateSequence/Private/Systems/TemplateSequenceSystem.cpp b/Engine/Plugins/MovieScene/TemplateSequence/Source/TemplateSequence/Private/Systems/TemplateSequenceSystem.cpp
--- a/Engine/Plugins/MovieScene/TemplateSequence/Source/TemplateSequence/Private/Systems/TemplateSequenceSystem.cpp
+++ b/Engine/Plugins/MovieScene/TemplateSequence/Source/TemplateSequence/Private/Systems/TemplateSequenceSystem.cpp
@@ -248,8 +248,11 @@ void UTemplateSequencePropertyScalingInstantiatorSystem::OnRun(FSystemTaskPrereq
 	// belong to an active template sub-sequence that has at least some properties scaled.
 	for (const TTuple<FMovieSceneEntityID, FGuid>& Pair : PropertyScaledEntities)
 	{
-		Linker->EntityManager.AddComponent(
-				Pair.Get<0>(), TemplateSequenceComponents->PropertyScaleReverseBindingLookup, Pair.Get<1>());
+		FEntityBuilder()
+		.Add(TemplateSequenceComponents->PropertyScaleReverseBindingLookup, Pair.Get<1>())
+		// Stop constants from being optimized since we need them to be re-computed every frame to remove the scale
+		.AddTag(BuiltInComponents->Tags.DontOptimizeConstants)
+		.MutateExisting(&Linker->EntityManager, Pair.Get<0>());
 	}
 }
 
@@ -375,12 +378,12 @@ struct FScaleTransformProperties
 			TRead<FGuid> ReverseBindingLookups,
 			TReadOptional<double> BaseLocationXs, TReadOptional<double> BaseLocationYs, TReadOptional<double> BaseLocationZs,
 			TReadOptional<double> BaseRotationXs, TReadOptional<double> BaseRotationYs, TReadOptional<double> BaseRotationZs,
-			TWriteOptional<double> LocationXs, TWriteOptional<FSourceDoubleChannelFlags> LocationXFlags,
-			TWriteOptional<double> LocationYs, TWriteOptional<FSourceDoubleChannelFlags> LocationYFlags,
-			TWriteOptional<double> LocationZs, TWriteOptional<FSourceDoubleChannelFlags> LocationZFlags,
-			TWriteOptional<double> RotationXs, TWriteOptional<FSourceDoubleChannelFlags> RotationXFlags,
-			TWriteOptional<double> RotationYs, TWriteOptional<FSourceDoubleChannelFlags> RotationYFlags,
-			TWriteOptional<double> RotationZs, TWriteOptional<FSourceDoubleChannelFlags> RotationZFlags)
+			TWriteOptional<double> LocationXs,
+			TWriteOptional<double> LocationYs,
+			TWriteOptional<double> LocationZs,
+			TWriteOptional<double> RotationXs,
+			TWriteOptional<double> RotationYs,
+			TWriteOptional<double> RotationZs)
 	{
 		for (int32 Index = 0; Index < Allocation->Num(); ++Index)
 		{
@@ -403,23 +406,12 @@ struct FScaleTransformProperties
 						ScalePropertyValue(ScaleFactor, BaseLocationXs, LocationXs, Index);
 						ScalePropertyValue(ScaleFactor, BaseLocationYs, LocationYs, Index);
 						ScalePropertyValue(ScaleFactor, BaseLocationZs, LocationZs, Index);
-
-						// We set the source channel to force re-evaluating next frame because we want it to be
-						// reset to its unscaled value... otherwise we will start accumulating multipliers quickly!
-						if (LocationXFlags) LocationXFlags[Index].bNeedsEvaluate = true;
-						if (LocationYFlags) LocationYFlags[Index].bNeedsEvaluate = true;
-						if (LocationZFlags) LocationZFlags[Index].bNeedsEvaluate = true;
 						break;
 
 					case ETemplateSectionPropertyScaleType::TransformPropertyRotationOnly:
 						ScalePropertyValue(ScaleFactor, BaseRotationXs, RotationXs, Index);
 						ScalePropertyValue(ScaleFactor, BaseRotationYs, RotationYs, Index);
 						ScalePropertyValue(ScaleFactor, BaseRotationZs, RotationZs, Index);
-
-						// See comment above.
-						if (RotationXFlags) RotationXFlags[Index].bNeedsEvaluate = true;
-						if (RotationYFlags) RotationYFlags[Index].bNeedsEvaluate = true;
-						if (RotationZFlags) RotationZFlags[Index].bNeedsEvaluate = true;
 						break;
 
 					default:
@@ -449,8 +441,7 @@ struct FScaleFloatProperties
 			TRead<FMovieScenePropertyBinding> PropertyBindings,
 			TRead<FGuid> ReverseBindingLookups,
 			TReadOptional<double> BasePropertyValues,
-			TWrite<double> PropertyValues,
-			TWrite<FSourceFloatChannelFlags> PropertyFlags)
+			TWrite<double> PropertyValues)
 	{
 		for (int32 Index = 0; Index < Allocation->Num(); ++Index)
 		{
@@ -470,8 +461,6 @@ struct FScaleFloatProperties
 				{
 					case ETemplateSectionPropertyScaleType::FloatProperty:
 						ScalePropertyValue(ScaleFactor, BasePropertyValues, PropertyValues, Index);
-						// See comment for the similar code in FScaleTransformProperties.
-						PropertyFlags[Index].bNeedsEvaluate = true;
 						break;
 					default:
 						ensureMsgf(false, TEXT("Unsupported or invalid float property scale type."));
@@ -529,17 +518,11 @@ void UTemplateSequencePropertyScalingEvaluatorSystem::OnRun(FSystemTaskPrerequis
 			.ReadOptional(BuiltInComponents->BaseDouble[4])
 			.ReadOptional(BuiltInComponents->BaseDouble[5])
 			.WriteOptional(BuiltInComponents->DoubleResult[0])
-			.WriteOptional(BuiltInComponents->DoubleChannelFlags[0])
 			.WriteOptional(BuiltInComponents->DoubleResult[1])
-			.WriteOptional(BuiltInComponents->DoubleChannelFlags[1])
 			.WriteOptional(BuiltInComponents->DoubleResult[2])
-			.WriteOptional(BuiltInComponents->DoubleChannelFlags[2])
 			.WriteOptional(BuiltInComponents->DoubleResult[3])
-			.WriteOptional(BuiltInComponents->DoubleChannelFlags[3])
 			.WriteOptional(BuiltInComponents->DoubleResult[4])
-			.WriteOptional(BuiltInComponents->DoubleChannelFlags[4])
 			.WriteOptional(BuiltInComponents->DoubleResult[5])
-			.WriteOptional(BuiltInComponents->DoubleChannelFlags[5])
 			.FilterAny({ TrackComponents->ComponentTransform.PropertyTag, TrackComponents->Transform.PropertyTag })
 			.Dispatch_PerAllocation<FScaleTransformProperties>(&Linker->EntityManager, InPrerequisites, &Subsequents, this);
 	}
@@ -553,7 +536,6 @@ void UTemplateSequencePropertyScalingEvaluatorSystem::OnRun(FSystemTaskPrerequis
 			.Read(TemplateSequenceComponents->PropertyScaleReverseBindingLookup)
 			.ReadOptional(BuiltInComponents->BaseDouble[0])
 			.Write(BuiltInComponents->DoubleResult[0])
-			.Write(BuiltInComponents->FloatChannelFlags[0])
 			.FilterAll({ TrackComponents->Float.PropertyTag })
 			.Dispatch_PerAllocation<FScaleFloatProperties>(&Linker->EntityManager, InPrerequisites, &Subsequents, this);
 	}
diff --git a/Engine/Source/Runtime/MovieScene/Private/Channels/MovieSceneChannelData.cpp b/Engine/Source/Runtime/MovieScene/Private/Channels/MovieSceneChannelData.cpp
--- a/Engine/Source/Runtime/MovieScene/Private/Channels/MovieSceneChannelData.cpp
+++ b/Engine/Source/Runtime/MovieScene/Private/Channels/MovieSceneChannelData.cpp
@@ -18,7 +18,7 @@ namespace MovieScene
 		OutIndex2 = Index2 < InTimes.Num() ? Index2 : INDEX_NONE;
 	}
 
-	void EvaluateTime(TArrayView<const FFrameNumber> InTimes, FFrameTime InTime, int32& OutIndex1, int32& OutIndex2, float& OutInterp)
+	void EvaluateTime(TArrayView<const FFrameNumber> InTimes, FFrameTime InTime, int32& OutIndex1, int32& OutIndex2, double& OutInterp)
 	{
 		const int32 Index2 = Algo::UpperBound(InTimes, InTime.FrameNumber);
 		const int32 Index1 = Index2 - 1;
@@ -32,7 +32,7 @@ namespace MovieScene
 			const int32 Time1 = InTimes[Index1].Value, Time2 = InTimes[Index2].Value;
 			const int32 Difference = Time2 - Time1;
 
-			OutInterp = (InTime.FrameNumber.Value - Time1) / float(Difference) + InTime.GetSubFrame() / Difference;
+			OutInterp = (InTime.AsDecimal() - Time1) / double(Difference);
 		}
 		else
 		{
diff --git a/Engine/Source/Runtime/MovieScene/Private/Channels/MovieSceneCurveChannelImpl.cpp b/Engine/Source/Runtime/MovieScene/Private/Channels/MovieSceneCurveChannelImpl.cpp
--- a/Engine/Source/Runtime/MovieScene/Private/Channels/MovieSceneCurveChannelImpl.cpp
+++ b/Engine/Source/Runtime/MovieScene/Private/Channels/MovieSceneCurveChannelImpl.cpp
@@ -3,6 +3,7 @@
 #include "Channels/MovieSceneCurveChannelImpl.h"
 #include "Channels/MovieSceneDoubleChannel.h"
 #include "Channels/MovieSceneFloatChannel.h"
+#include "Channels/MovieSceneInterpolation.h"
 #include "HAL/ConsoleManager.h"
 #include "MovieSceneFrameMigration.h"
 #include "Curves/CurveEvaluation.h"
@@ -28,6 +29,14 @@ namespace UE
 namespace MovieScene
 {
 
+MOVIESCENE_API float GCachedChannelEvaluationParity = 0.f;
+static FAutoConsoleVariableRef CVarTestCachedChannelEvaluationParity(
+	TEXT("Sequencer.TestCachedChannelEvaluationParity"),
+	GCachedChannelEvaluationParity,
+	TEXT("Toggles whether Evaluate will use cached or non-cached evaluation."),
+	ECVF_Default);
+
+
 template<typename ChannelType>
 static typename ChannelType::CurveValueType
 EvalForTwoKeys(
@@ -74,27 +83,34 @@ EvalForTwoKeys(
 
 struct FCycleParams
 {
+	double ValueOffset;
 	FFrameTime Time;
 	int32 CycleCount;
-	float ValueOffset;
+	int32 Duration;
 
-	FCycleParams(FFrameTime InTime)
-		: Time(InTime)
+	FCycleParams(FFrameTime InTime, int32 InDuration)
+		: ValueOffset(0.0)
+		, Time(InTime)
 		, CycleCount(0)
-		, ValueOffset(0.f)
+		, Duration(InDuration)
 	{}
 
-	FORCEINLINE void ComputePreValueOffset(float FirstValue, float LastValue)
+	FORCEINLINE void ComputePreValueOffset(double FirstValue, double LastValue)
 	{
-		ValueOffset = (FirstValue-LastValue) * CycleCount;
+		// CycleCount is negative for pre-extrap
+		ValueOffset = (LastValue-FirstValue) * CycleCount;
 	}
-	FORCEINLINE void ComputePostValueOffset(float FirstValue, float LastValue)
+	FORCEINLINE void ComputePostValueOffset(double FirstValue, double LastValue)
 	{
 		ValueOffset = (LastValue-FirstValue) * CycleCount;
 	}
+	FORCEINLINE bool IsOscillated() const
+	{
+		return FMath::Abs(CycleCount) % 2 == 1;
+	}
 	FORCEINLINE void Oscillate(int32 MinFrame, int32 MaxFrame)
 	{
-		if (CycleCount % 2 == 1)
+		if (IsOscillated())
 		{
 			Time = MinFrame + (FFrameTime(MaxFrame) - Time);
 		}
@@ -103,27 +119,25 @@ struct FCycleParams
 
 FCycleParams CycleTime(FFrameNumber MinFrame, FFrameNumber MaxFrame, FFrameTime InTime)
 {
-	FCycleParams Params(InTime);
-	
-	const int32 Duration = MaxFrame.Value - MinFrame.Value;
-	if (Duration == 0)
+	FCycleParams Params(InTime, MaxFrame.Value - MinFrame.Value);
+	if (Params.Duration == 0)
 	{
 		Params.Time = MaxFrame;
 		Params.CycleCount = 0;
 	}
 	else if (InTime < MinFrame)
 	{
-		const int32 CycleCount = ((MaxFrame - InTime) / Duration).FloorToFrame().Value;
+		const int32 CycleCount = ((MaxFrame - InTime) / Params.Duration).FloorToFrame().Value;
 
-		Params.Time = InTime + FFrameTime(Duration)*CycleCount;
-		Params.CycleCount = CycleCount;
+		Params.CycleCount = -CycleCount;
+		Params.Time = InTime + FFrameTime(Params.Duration*CycleCount);
 	}
 	else if (InTime > MaxFrame)
 	{
-		const int32 CycleCount = ((InTime - MinFrame) / Duration).FloorToFrame().Value;
+		const int32 CycleCount = ((InTime - MinFrame) / Params.Duration).FloorToFrame().Value;
 
-		Params.Time = InTime - FFrameTime(Duration)*CycleCount;
 		Params.CycleCount = CycleCount;
+		Params.Time = InTime - FFrameTime(Params.Duration*CycleCount);
 	}
 
 	return Params;
@@ -299,6 +313,285 @@ bool TMovieSceneCurveChannelImpl<ChannelType>::EvaluateExtrapolation(const Chann
 	return false;
 }
 
+template<typename ChannelType>
+bool TMovieSceneCurveChannelImpl<ChannelType>::CacheExtrapolation(const ChannelType* InChannel, FFrameTime InTime, UE::MovieScene::Interpolation::FCachedInterpolation& OutValue)
+{
+	using namespace UE::MovieScene::Interpolation;
+
+	// If the time is outside of the curve, deal with extrapolation
+	if (InTime < InChannel->Times[0])
+	{
+		if (InChannel->PreInfinityExtrap == RCCE_None)
+		{
+			return false;
+		}
+
+		FCachedInterpolationRange Range = FCachedInterpolationRange::Until(InChannel->Times[0]);
+
+		if (InChannel->PreInfinityExtrap == RCCE_Constant)
+		{
+			OutValue = FCachedInterpolation(Range, FConstantValue(InChannel->Values[0].Value));
+			return true;
+		}
+
+		if (InChannel->PreInfinityExtrap == RCCE_Linear)
+		{
+			const ChannelValueType FirstValue = InChannel->Values[0];
+
+			if (FirstValue.InterpMode == RCIM_Constant)
+			{
+				OutValue = FCachedInterpolation(Range, FConstantValue(FirstValue.Value));
+			}
+			else if(FirstValue.InterpMode == RCIM_Cubic)
+			{
+				OutValue = FCachedInterpolation(Range, FLinearInterpolation(InChannel->Times[0], FirstValue.Tangent.ArriveTangent, FirstValue.Value));
+			}
+			else if(FirstValue.InterpMode == RCIM_Linear)
+			{
+				const int32 InterpStartFrame = InChannel->Times[1].Value;
+
+				const double DY = InChannel->Values[1].Value - FirstValue.Value;
+				const double DX = InChannel->Times[1].Value - InChannel->Times[0].Value;
+				const double Coefficient = DX == 0.0 ? 0.0 : DY/DX;
+
+				OutValue = FCachedInterpolation(Range, FLinearInterpolation(InChannel->Times[0], Coefficient, FirstValue.Value));
+			}
+			return true;
+		}
+	}
+	else if (InTime > InChannel->Times.Last())
+	{
+		if (InChannel->PostInfinityExtrap == RCCE_None)
+		{
+			return false;
+		}
+
+		FCachedInterpolationRange Range = FCachedInterpolationRange::From(InChannel->Times.Last());
+
+		if (InChannel->PostInfinityExtrap == RCCE_Constant)
+		{
+			OutValue = FCachedInterpolation(Range, FConstantValue(InChannel->Values.Last().Value));
+			return true;
+		}
+
+		if (InChannel->PostInfinityExtrap == RCCE_Linear)
+		{
+			const ChannelValueType LastValue = InChannel->Values.Last();
+
+			if (LastValue.InterpMode == RCIM_Constant)
+			{
+				OutValue = FCachedInterpolation(Range, FConstantValue(LastValue.Value));
+			}
+			else if(LastValue.InterpMode == RCIM_Cubic)
+			{
+				OutValue = FCachedInterpolation(Range, FLinearInterpolation(InChannel->Times.Last(), LastValue.Tangent.LeaveTangent, LastValue.Value));
+			}
+			else if(LastValue.InterpMode == RCIM_Linear)
+			{
+				const int32 NumKeys          = InChannel->Times.Num();
+				const int32 InterpStartFrame = InChannel->Times[NumKeys-2].Value;
+				const int32 DeltaFrame       = InChannel->Times.Last().Value - InterpStartFrame;
+
+				const double DY = LastValue.Value - InChannel->Values[NumKeys-2].Value;
+				const double DX = InChannel->Times.Last().Value - InChannel->Times[NumKeys-2].Value;
+				const double Coefficient = DX == 0.0 ? 0.0 : DY/DX;
+
+				OutValue = FCachedInterpolation(Range, FLinearInterpolation(InChannel->Times.Last(), Coefficient, LastValue.Value));
+			}
+			return true;
+		}
+	}
+
+	return false;
+}
+
+template<typename ChannelType>
+UE::MovieScene::Interpolation::FCachedInterpolation TMovieSceneCurveChannelImpl<ChannelType>::GetInterpolationForTime(const ChannelType* InChannel, FFrameTime InTime)
+{
+	using namespace UE::MovieScene;
+	using namespace UE::MovieScene::Interpolation;
+
+	const int32 NumKeys = InChannel->Times.Num();
+
+	// No keys means default value, or nothing
+	if (NumKeys == 0)
+	{
+		if (InChannel->bHasDefaultValue)
+		{
+			return FCachedInterpolation(FCachedInterpolationRange::Infinite(), FConstantValue(InChannel->DefaultValue));
+		}
+
+		return FCachedInterpolation();
+	}
+
+	// For single keys, we can only ever return that value
+	if (NumKeys == 1)
+	{
+		return FCachedInterpolation(FCachedInterpolationRange::Infinite(), FConstantValue(InChannel->Values[0].Value));
+	}
+
+	// Cache extrapolation if we're outside the bounds of the curve
+	// for any non-piecewise extrapolation algorithms (constant or linear)
+	{
+		FCachedInterpolation ExtrapolatedCache;
+		if (CacheExtrapolation(InChannel, InTime, ExtrapolatedCache))
+		{
+			return ExtrapolatedCache;
+		}
+	}
+
+	// Piecewise spline evaluation from this point
+	const FFrameNumber MinFrame = InChannel->Times[0];
+	const FFrameNumber MaxFrame = InChannel->Times.Last();
+
+	// Compute the cycled time based on extrapolation
+	FCycleParams Params = CycleTime(MinFrame, MaxFrame, InTime);
+
+	// Deal with offset cycles and oscillation
+	if (InTime < FFrameTime(MinFrame))
+	{
+		switch (InChannel->PreInfinityExtrap)
+		{
+		case RCCE_CycleWithOffset: Params.ComputePreValueOffset(InChannel->Values[0].Value, InChannel->Values[NumKeys-1].Value); break;
+		case RCCE_Oscillate:       Params.Oscillate(MinFrame.Value, MaxFrame.Value);                       break;
+		}
+	}
+	else if (InTime > FFrameTime(MaxFrame))
+	{
+		switch (InChannel->PostInfinityExtrap)
+		{
+		case RCCE_CycleWithOffset: Params.ComputePostValueOffset(InChannel->Values[0].Value, InChannel->Values[NumKeys-1].Value); break;
+		case RCCE_Oscillate:       Params.Oscillate(MinFrame.Value, MaxFrame.Value);                        break;
+		}
+	}
+
+	if (!ensureMsgf(Params.Time.FrameNumber >= MinFrame && Params.Time.FrameNumber <= MaxFrame, TEXT("Invalid time computed for float channel evaluation")))
+	{
+		return FCachedInterpolation();
+	}
+
+	// Find the pair of keys that we need to evaluate
+	double Interp = 0.0;
+	int32 Index1 = INDEX_NONE, Index2 = INDEX_NONE;
+	UE::MovieScene::EvaluateTime(InChannel->Times, Params.Time, Index1, Index2, Interp);
+
+	if (Index1 == INDEX_NONE)
+	{
+		// No starting key - we are probably evaluating directly on the first or last key
+		//   we explicitly only cache this for the current time to ensure that subsequent caches can cache the correct pair
+		FCachedInterpolationRange Range = FCachedInterpolationRange::Only(Params.Time.GetFrame());
+		return FCachedInterpolation(Range, FConstantValue(Params.ValueOffset + InChannel->Values[Index2].Value));
+	}
+	else if (Index2 == INDEX_NONE)
+	{
+		// No ending key - we are probably evaluating directly on the first or last key
+		//   we explicitly only cache this for the current time to ensure that subsequent caches can cache the correct pair
+		FCachedInterpolationRange Range = FCachedInterpolationRange::Only(Params.Time.GetFrame());
+		return FCachedInterpolation(Range, FConstantValue(Params.ValueOffset + InChannel->Values[Index1].Value));
+	}
+	else
+	{
+		// We have a valid pair of keys to cache on the spline.
+		const double DX = InChannel->Times[Index2].Value - InChannel->Times[Index1].Value;
+
+		// Cache the pair of keys from the spline at the correct time
+		ChannelValueType Key1 = InChannel->Values[Index1];
+		ChannelValueType Key2 = InChannel->Values[Index2];
+
+		// Manipulate they keys by translating them by the cycle count in the time-domain.
+		// By caching the control points in this translated state we can avoid having to 
+		// do any manipulation on the input time
+		const FFrameNumber CycleOffset = Params.CycleCount*Params.Duration;
+
+		// Compute the start and end values and tangents.
+		FFrameNumber Time1, Time2;
+
+		// Control point 1
+		double V1 = Params.ValueOffset + Key1.Value;
+		double T1 = Key1.Tangent.LeaveTangent;
+		double W1 = Key1.Tangent.LeaveTangentWeight;
+		bool bIsWeighted1 = (Key1.Tangent.TangentWeightMode == RCTWM_WeightedBoth || Key1.Tangent.TangentWeightMode == RCTWM_WeightedLeave);
+
+		// Control point 2
+		double V2 = Params.ValueOffset + Key2.Value;
+		double T2 = Key2.Tangent.ArriveTangent;
+		double W2 = Key2.Tangent.ArriveTangentWeight;
+		bool bIsWeighted2 = (Key2.Tangent.TangentWeightMode == RCTWM_WeightedBoth || Key2.Tangent.TangentWeightMode == RCTWM_WeightedArrive);
+
+		// Mirror the curve for oscillation by swapping the control points
+		// and mirroring the times based on the width.
+		if (Params.IsOscillated())
+		{
+			Swap(V1, V2);
+			Swap(T1, T2);
+			Swap(W1, W2);
+			Swap(bIsWeighted1, bIsWeighted2);
+
+			T1 = -T1;
+			T2 = -T2;
+
+			// Mirror the times of the control points such that
+			// MinFrame                              MaxFrame
+			//  | x1                    x2   x3       |
+			//  |               becomes               |
+			//  |      x3   x2                     x1 |
+			// and offset them by the cycle offset
+			Time1 = MinFrame + (MaxFrame - InChannel->Times[Index2]) + CycleOffset;
+			Time2 = MinFrame + (MaxFrame - InChannel->Times[Index1]) + CycleOffset;
+		}
+		else
+		{
+			// No oscillation so just offset the control points to be in their final positions
+			Time1 = InChannel->Times[Index1] + CycleOffset;
+			Time2 = InChannel->Times[Index2] + CycleOffset;
+		}
+
+		// Cache this interpolation for any time between the two control points
+		FCachedInterpolationRange Range = FCachedInterpolationRange::Finite(Time1, Time2);
+
+		// Careful: We use the original Key1 rather than the oscillated interpmode to ensure that
+		//          we use the correct key to produce a mirror image
+		TEnumAsByte<ERichCurveInterpMode> InterpMode = Key1.InterpMode;
+		const int CheckBothLinear = GSequencerLinearCubicInterpolation;
+	 	if(InterpMode == RCIM_Linear && (CheckBothLinear  && Key2.InterpMode == RCIM_Cubic))
+		{
+			InterpMode = RCIM_Cubic;
+		}
+
+		switch (InterpMode)
+		{
+		case RCIM_Cubic:
+			{
+				// Weighted if either control point has weight on their tangent
+				if (bIsWeighted1 || bIsWeighted2)
+				{
+					return FCachedInterpolation(Range, FWeightedCubicInterpolation(
+						InChannel->TickResolution, Time1,
+						Time1, V1, T1, W1, bIsWeighted1,
+						Time2, V2, T2, W2, bIsWeighted2
+					));
+				}
+				else
+				{
+					return FCachedInterpolation(Range, FCubicInterpolation(Time1, DX, V1, V2, T1, T2));
+				}
+			}
+			break;
+
+		case RCIM_Linear:
+			{
+				const double DY = V2 - V1;
+				return FCachedInterpolation(Range, FLinearInterpolation(Time1, DY/DX, V1));
+			}
+
+		default:
+			return FCachedInterpolation(Range, FConstantValue(V1));
+		}
+	}
+
+	return FCachedInterpolation();
+}
+
 template<typename ChannelType>
 bool TMovieSceneCurveChannelImpl<ChannelType>::Evaluate(const ChannelType* InChannel, FFrameTime InTime, CurveValueType& OutValue) 
 {
@@ -310,6 +603,21 @@ bool TMovieSceneCurveChannelImpl<ChannelType>::EvaluateWithCache(const ChannelTy
 {
 	using namespace UE::MovieScene;
 
+	// ------------------------------------------------------------------------
+	// New cached codepath - existing implementation still exists as a fallback
+	if (GCachedChannelEvaluationParity == 0.f)
+	{
+		double Result = 0.0;
+		if (GetInterpolationForTime(InChannel, InTime).Evaluate(InTime, Result))
+		{
+			OutValue = static_cast<CurveValueType>(Result);
+			return true;
+		}
+		return false;
+	}
+
+	// ------------------------------------------------------------------------
+	// Legacy evaluate in-place codepath - only exists as a fallback
 	const int32 NumKeys = InChannel->Times.Num();
 
 	// No keys means default value, or nothing
@@ -366,7 +674,7 @@ bool TMovieSceneCurveChannelImpl<ChannelType>::EvaluateWithCache(const ChannelTy
 	}
 
 	// Evaluate the curve data
-	float Interp = 0.f;
+	double Interp = 0.0;
 	int32 Index1 = INDEX_NONE, Index2 = INDEX_NONE;
 
 	// Initialize cache if not yet performed
@@ -489,7 +797,7 @@ bool TMovieSceneCurveChannelImpl<ChannelType>::EvaluateWithCache(const ChannelTy
 				Coeff[0] = Coeff[0] - Interp;
 				
 				const int32 NumResults = UE::Curves::SolveCubic(Coeff, Results);
-				float NewInterp = Interp;
+				double NewInterp = Interp;
 				if (NumResults == 1)
 				{
 					NewInterp = Results[0];
diff --git a/Engine/Source/Runtime/MovieScene/Private/Channels/MovieSceneCurveChannelImpl.h b/Engine/Source/Runtime/MovieScene/Private/Channels/MovieSceneCurveChannelImpl.h
--- a/Engine/Source/Runtime/MovieScene/Private/Channels/MovieSceneCurveChannelImpl.h
+++ b/Engine/Source/Runtime/MovieScene/Private/Channels/MovieSceneCurveChannelImpl.h
@@ -13,6 +13,12 @@
 #include "MovieSceneFwd.h"
 
 
+namespace UE::MovieScene::Interpolation
+{
+	struct FCachedInterpolation;
+}
+
+
 /** Utility class for curve channels */
 template<typename ChannelType>
 struct MOVIESCENE_API TMovieSceneCurveChannelImpl
@@ -28,6 +34,11 @@ struct MOVIESCENE_API TMovieSceneCurveChannelImpl
 	/** Evaluate this channel with the frame resolution */
 	static bool Evaluate(const ChannelType* InChannel, FFrameTime InTime, CurveValueType& OutValue);
 
+	/**
+	 * Evaluate this channel by returning a cachable interpolation structure
+	 */
+	static UE::MovieScene::Interpolation::FCachedInterpolation GetInterpolationForTime(const ChannelType* InChannel, FFrameTime InTime);
+
 	/** Structure used to store the result of UE::MovieScene::EvaluateTime for a given channel/key distribution*/
 	struct FTimeEvaluationCache
 	{
@@ -126,6 +137,15 @@ private:
 	 */
 	static bool EvaluateExtrapolation(const ChannelType* InChannel, FFrameTime InTime, CurveValueType& OutValue);
 
+	/**
+	 * Evaluate this channel's extrapolation by populating a cachable structure. Assumes more than 1 key is present.
+	 *
+	 * @param InTime     The time to evaluate at
+	 * @param OutValue   A value to receive the result
+	 * @return true if the time was evaluated with extrapolation, false otherwise
+	 */
+	static bool CacheExtrapolation(const ChannelType* InChannel, FFrameTime InTime, UE::MovieScene::Interpolation::FCachedInterpolation& OutValue);
+
 	/**
 	 * Adds median points between each of the supplied points if their evaluated value is significantly different than the linear interpolation of those points
 	 *
diff --git a/Engine/Source/Runtime/MovieScene/Private/Channels/MovieSceneDoubleChannel.cpp b/Engine/Source/Runtime/MovieScene/Private/Channels/MovieSceneDoubleChannel.cpp
--- a/Engine/Source/Runtime/MovieScene/Private/Channels/MovieSceneDoubleChannel.cpp
+++ b/Engine/Source/Runtime/MovieScene/Private/Channels/MovieSceneDoubleChannel.cpp
@@ -4,6 +4,7 @@
 #include "Channels/MovieSceneChannelProxy.h"
 #include "Channels/MovieSceneCurveChannelImpl.h"
 #include "Channels/MovieSceneFloatChannel.h"
+#include "Channels/MovieSceneInterpolation.h"
 #include "HAL/Platform.h"
 #include "MovieSceneFrameMigration.h"
 #include "MovieSceneFwd.h"
@@ -60,6 +61,11 @@ bool FMovieSceneDoubleChannel::Evaluate(FFrameTime InTime, float& OutValue) cons
 	return bResult;
 }
 
+UE::MovieScene::Interpolation::FCachedInterpolation FMovieSceneDoubleChannel::GetInterpolationForTime(FFrameTime InTime) const
+{
+	return FMovieSceneDoubleChannelImpl::GetInterpolationForTime(this, InTime);
+}
+
 void FMovieSceneDoubleChannel::Set(TArray<FFrameNumber> InTimes, TArray<FMovieSceneDoubleValue> InValues)
 {
 	FMovieSceneDoubleChannelImpl::Set(this, InTimes, InValues);
diff --git a/Engine/Source/Runtime/MovieScene/Private/Channels/MovieSceneFloatChannel.cpp b/Engine/Source/Runtime/MovieScene/Private/Channels/MovieSceneFloatChannel.cpp
--- a/Engine/Source/Runtime/MovieScene/Private/Channels/MovieSceneFloatChannel.cpp
+++ b/Engine/Source/Runtime/MovieScene/Private/Channels/MovieSceneFloatChannel.cpp
@@ -3,6 +3,7 @@
 #include "Channels/MovieSceneFloatChannel.h"
 #include "Channels/MovieSceneChannelProxy.h"
 #include "Channels/MovieSceneCurveChannelImpl.h"
+#include "Channels/MovieSceneInterpolation.h"
 #include "MovieSceneFrameMigration.h"
 #include "MovieSceneFwd.h"
 #include "UObject/FortniteMainBranchObjectVersion.h"
@@ -50,6 +51,11 @@ bool FMovieSceneFloatChannel::Evaluate(FFrameTime InTime,  float& OutValue) cons
 	return FMovieSceneFloatChannelImpl::Evaluate(this, InTime, OutValue);
 }
 
+UE::MovieScene::Interpolation::FCachedInterpolation FMovieSceneFloatChannel::GetInterpolationForTime(FFrameTime InTime) const
+{
+	return FMovieSceneFloatChannelImpl::GetInterpolationForTime(this, InTime);
+}
+
 void FMovieSceneFloatChannel::Set(TArray<FFrameNumber> InTimes, TArray<FMovieSceneFloatValue> InValues)
 {
 	FMovieSceneFloatChannelImpl::Set(this, InTimes, InValues);
diff --git a/Engine/Source/Runtime/MovieScene/Private/Channels/MovieSceneInterpolation.cpp b/Engine/Source/Runtime/MovieScene/Private/Channels/MovieSceneInterpolation.cpp
new file mode 100644
--- /dev/null
+++ b/Engine/Source/Runtime/MovieScene/Private/Channels/MovieSceneInterpolation.cpp
@@ -0,0 +1,245 @@
+// Copyright Epic Games, Inc. All Rights Reserved.
+
+#include "Channels/MovieSceneInterpolation.h"
+#include "Curves/CurveEvaluation.h"
+
+namespace UE::MovieScene::Interpolation
+{
+
+FCachedInterpolationRange FCachedInterpolationRange::Empty()
+{
+	return FCachedInterpolationRange{ 0, -1 };
+}
+FCachedInterpolationRange FCachedInterpolationRange::Finite(FFrameNumber InStart, FFrameNumber InEnd)
+{
+	return FCachedInterpolationRange{ InStart, InEnd };
+}
+FCachedInterpolationRange FCachedInterpolationRange::Infinite()
+{
+	return FCachedInterpolationRange{ TNumericLimits<FFrameNumber>::Lowest(), TNumericLimits<FFrameNumber>::Max() };
+}
+FCachedInterpolationRange FCachedInterpolationRange::Only(FFrameNumber InTime)
+{
+	const FFrameNumber EndTime = InTime < TNumericLimits<FFrameNumber>::Max()
+		? InTime + 1
+		: InTime;
+
+	return FCachedInterpolationRange{ InTime, EndTime };
+}
+FCachedInterpolationRange FCachedInterpolationRange::From(FFrameNumber InStart)
+{
+	return FCachedInterpolationRange{ InStart, TNumericLimits<FFrameNumber>::Max() };
+}
+FCachedInterpolationRange FCachedInterpolationRange::Until(FFrameNumber InEnd) 
+{
+	return FCachedInterpolationRange{ TNumericLimits<FFrameNumber>::Lowest(), InEnd };
+}
+bool FCachedInterpolationRange::Contains(FFrameNumber FrameNumber) const
+{
+	return FrameNumber >= Start && (FrameNumber < End || End == TNumericLimits<FFrameNumber>::Max());
+}
+
+FCachedInterpolation::FCachedInterpolation()
+	: Data(TInPlaceType<FInvalidValue>())
+	, Range(FCachedInterpolationRange::Empty())
+{
+}
+
+FCachedInterpolation::FCachedInterpolation(const FCachedInterpolationRange& InRange, const FConstantValue& Constant)
+	: Data(TInPlaceType<FConstantValue>(), Constant)
+	, Range(InRange)
+{
+}
+
+FCachedInterpolation::FCachedInterpolation(const FCachedInterpolationRange& InRange, const FLinearInterpolation& Linear)
+	: Data(TInPlaceType<FLinearInterpolation>(), Linear)
+	, Range(InRange)
+{
+}
+
+FCachedInterpolation::FCachedInterpolation(const FCachedInterpolationRange& InRange, const FCubicInterpolation& Cubic)
+	: Data(TInPlaceType<FCubicInterpolation>(), Cubic)
+	, Range(InRange)
+{
+}
+
+FCachedInterpolation::FCachedInterpolation(const FCachedInterpolationRange& InRange, const FWeightedCubicInterpolation& WeightedCubic)
+	: Data(TInPlaceType<FWeightedCubicInterpolation>(), WeightedCubic)
+	, Range(InRange)
+{
+}
+
+bool FCachedInterpolation::IsCacheValidForTime(FFrameNumber FrameNumber) const
+{
+	return Range.Contains(FrameNumber);
+}
+
+bool FCachedInterpolation::Evaluate(FFrameTime Time, double& OutResult) const
+{
+	if (const FConstantValue* Constant = Data.TryGet<FConstantValue>())
+	{
+		OutResult = Constant->Value;
+		return true;
+	}
+	else if (const FLinearInterpolation* Linear = Data.TryGet<FLinearInterpolation>())
+	{
+		OutResult = Linear->Evaluate(Time);
+		return true;
+	}
+	else if (const FCubicInterpolation* Cubic = Data.TryGet<FCubicInterpolation>())
+	{
+		OutResult = Cubic->Evaluate(Time);
+		return true;
+	}
+	else if (const FWeightedCubicInterpolation* WeightedCubic = Data.TryGet<FWeightedCubicInterpolation>())
+	{
+		OutResult = WeightedCubic->Evaluate(Time);
+		return true;
+	}
+
+	return false;
+}
+
+FCubicInterpolation::FCubicInterpolation(FFrameNumber InOrigin, double InDX, double InStartValue, double InEndValue, double InStartTangent, double InEndTangent)
+	: DX(InDX)
+	, Origin(InOrigin)
+{
+	constexpr double OneThird = 1.0 / 3.0;
+
+	P0 = InStartValue;
+	P1 = P0 + (InStartTangent * DX * OneThird);
+	P3 = InEndValue;
+	P2 = P3 - (InEndTangent * DX * OneThird);
+}
+
+double FCubicInterpolation::Evaluate(FFrameTime InTime) const
+{
+	const float Interp = static_cast<float>((InTime - Origin).AsDecimal() / DX);
+	return UE::Curves::BezierInterp(P0, P1, P2, P3, Interp);
+}
+
+FWeightedCubicInterpolation::FWeightedCubicInterpolation(
+	FFrameRate TickResolution,
+	FFrameNumber InOrigin,
+
+	FFrameNumber StartTime,
+	double StartValue,
+	double StartTangent,
+	double StartTangentWeight,
+	bool bStartIsWeighted,
+
+	FFrameNumber EndTime,
+	double EndValue,
+	double EndTangent,
+	double EndTangentWeight,
+	bool bEndIsWeighted)
+{
+	constexpr double OneThird = 1.0 / 3.0;
+
+	const float TimeInterval = TickResolution.AsInterval();
+	const float ToSeconds = 1.0f / TimeInterval;
+
+	const double Time1 = TickResolution.AsSeconds(StartTime);
+	const double Time2 = TickResolution.AsSeconds(EndTime);
+	const double DXInSeconds = Time2 - Time1;
+
+	Origin = InOrigin;
+	DX = (EndTime - StartTime).Value;
+	StartKeyValue = StartValue;
+	EndKeyValue = EndValue;
+
+	double CosAngle, SinAngle, Angle;
+
+	// ---------------------------------------------------------------------------------
+	// Initialize the start key parameters
+	Angle = FMath::Atan(StartTangent * ToSeconds);
+	FMath::SinCos(&SinAngle, &CosAngle, Angle);
+
+	if (bStartIsWeighted)
+	{
+		const double LeaveTangentNormalized = StartTangent / TimeInterval;
+		const double DY = LeaveTangentNormalized * DXInSeconds;
+		StartWeight = FMath::Sqrt(DXInSeconds*DXInSeconds + DY*DY) * OneThird;
+	}
+	else
+	{
+		StartWeight = StartTangentWeight;
+	}
+
+	const double StartKeyTanX = CosAngle * StartWeight + Time1;
+	StartKeyTanY              = SinAngle * StartWeight + StartValue;
+	NormalizedStartTanDX = (StartKeyTanX - Time1) / DXInSeconds;
+
+	// ---------------------------------------------------------------------------------
+	// Initialize the end key parameters
+	Angle = FMath::Atan(EndTangent * ToSeconds);
+	FMath::SinCos(&SinAngle, &CosAngle, Angle);
+
+	if (bEndIsWeighted)
+	{
+		const double ArriveTangentNormalized = EndTangent / TimeInterval;
+		const double DY = ArriveTangentNormalized * DXInSeconds;
+		EndWeight = FMath::Sqrt(DXInSeconds*DXInSeconds + DY*DY) * OneThird;
+	}
+	else
+	{
+		EndWeight =  EndTangentWeight;
+	}
+
+	const double EndKeyTanX = -CosAngle * EndWeight + Time2;
+	EndKeyTanY              = -SinAngle * EndWeight + EndValue;
+
+	NormalizedEndTanDX = (EndKeyTanX - Time1) / DXInSeconds;
+}
+
+double FWeightedCubicInterpolation::Evaluate(FFrameTime InTime) const
+{
+	const double Interp = (InTime - Origin).AsDecimal() / DX;
+
+	double Coeff[4];
+	double Results[3];
+
+	//Convert Bezier to Power basis, also float to double for precision for root finding.
+	UE::Curves::BezierToPower(
+		0.0, NormalizedStartTanDX, NormalizedEndTanDX, 1.0,
+		&(Coeff[3]), &(Coeff[2]), &(Coeff[1]), &(Coeff[0])
+	);
+
+	Coeff[0] = Coeff[0] - Interp;
+
+	const int32 NumResults = UE::Curves::SolveCubic(Coeff, Results);
+	float NewInterp = Interp;
+	if (NumResults == 1)
+	{
+		NewInterp = Results[0];
+	}
+	else
+	{
+		NewInterp = TNumericLimits<double>::Lowest(); //just need to be out of range
+		for (double Result : Results)
+		{
+			if ((Result >= 0.0f) && (Result <= 1.0f))
+			{
+				if (NewInterp < 0.0f || Result > NewInterp)
+				{
+					NewInterp = Result;
+				}
+			}
+		}
+
+		if (NewInterp == TNumericLimits<double>::Lowest())
+		{
+			NewInterp = 0.f;
+		}
+	}
+
+	//now use NewInterp and adjusted tangents plugged into the Y (Value) part of the graph.
+	const double P0 = StartKeyValue;
+	const double P1 = StartKeyTanY;
+	const double P3 = EndKeyValue;
+	const double P2 = EndKeyTanY;
+
+	return UE::Curves::BezierInterp(P0, P1, P2, P3, static_cast<float>(NewInterp));
+}
+
+} // namespace UE::MovieScene
diff --git a/Engine/Source/Runtime/MovieScene/Private/EntitySystem/BuiltInComponentTypes.cpp b/Engine/Source/Runtime/MovieScene/Private/EntitySystem/BuiltInComponentTypes.cpp
--- a/Engine/Source/Runtime/MovieScene/Private/EntitySystem/BuiltInComponentTypes.cpp
+++ b/Engine/Source/Runtime/MovieScene/Private/EntitySystem/BuiltInComponentTypes.cpp
@@ -8,6 +8,7 @@
 #include "EntitySystem/MovieSceneComponentRegistry.h"
 #include "EntitySystem/MovieScenePropertyBinding.h"
 #include "EntitySystem/MovieSceneEntityFactoryTemplates.h"
+#include "Channels/MovieSceneInterpolation.h"
 
 #include UE_INLINE_GENERATED_CPP_BY_NAME(BuiltInComponentTypes)
 
@@ -63,25 +64,16 @@ FBuiltInComponentTypes::FBuiltInComponentTypes()
 	ComponentRegistry->NewComponentType(&WeightChannel,           TEXT("Weight Channel"));
 	ComponentRegistry->NewComponentType(&ObjectPathChannel,       TEXT("Object Path Channel"));
 
-	ComponentRegistry->NewComponentType(&FloatChannelFlags[0],    TEXT("Float Channel 0 Flags"));
-	ComponentRegistry->NewComponentType(&FloatChannelFlags[1],    TEXT("Float Channel 1 Flags"));
-	ComponentRegistry->NewComponentType(&FloatChannelFlags[2],    TEXT("Float Channel 2 Flags"));
-	ComponentRegistry->NewComponentType(&FloatChannelFlags[3],    TEXT("Float Channel 3 Flags"));
-	ComponentRegistry->NewComponentType(&FloatChannelFlags[4],    TEXT("Float Channel 4 Flags"));
-	ComponentRegistry->NewComponentType(&FloatChannelFlags[5],    TEXT("Float Channel 5 Flags"));
-	ComponentRegistry->NewComponentType(&FloatChannelFlags[6],    TEXT("Float Channel 6 Flags"));
-	ComponentRegistry->NewComponentType(&FloatChannelFlags[7],    TEXT("Float Channel 7 Flags"));
-	ComponentRegistry->NewComponentType(&FloatChannelFlags[8],    TEXT("Float Channel 8 Flags"));
-	ComponentRegistry->NewComponentType(&DoubleChannelFlags[0],   TEXT("Double Channel 0 Flags"));
-	ComponentRegistry->NewComponentType(&DoubleChannelFlags[1],   TEXT("Double Channel 1 Flags"));
-	ComponentRegistry->NewComponentType(&DoubleChannelFlags[2],   TEXT("Double Channel 2 Flags"));
-	ComponentRegistry->NewComponentType(&DoubleChannelFlags[3],   TEXT("Double Channel 3 Flags"));
-	ComponentRegistry->NewComponentType(&DoubleChannelFlags[4],   TEXT("Double Channel 4 Flags"));
-	ComponentRegistry->NewComponentType(&DoubleChannelFlags[5],   TEXT("Double Channel 5 Flags"));
-	ComponentRegistry->NewComponentType(&DoubleChannelFlags[6],   TEXT("Double Channel 6 Flags"));
-	ComponentRegistry->NewComponentType(&DoubleChannelFlags[7],   TEXT("Double Channel 7 Flags"));
-	ComponentRegistry->NewComponentType(&DoubleChannelFlags[8],   TEXT("Double Channel 8 Flags"));
-	ComponentRegistry->NewComponentType(&WeightChannelFlags,      TEXT("Weight Channel Flags"));
+	ComponentRegistry->NewComponentType(&CachedInterpolation[0],   TEXT("Cached Interpolation [0]"));
+	ComponentRegistry->NewComponentType(&CachedInterpolation[1],   TEXT("Cached Interpolation [1]"));
+	ComponentRegistry->NewComponentType(&CachedInterpolation[2],   TEXT("Cached Interpolation [2]"));
+	ComponentRegistry->NewComponentType(&CachedInterpolation[3],   TEXT("Cached Interpolation [3]"));
+	ComponentRegistry->NewComponentType(&CachedInterpolation[4],   TEXT("Cached Interpolation [4]"));
+	ComponentRegistry->NewComponentType(&CachedInterpolation[5],   TEXT("Cached Interpolation [5]"));
+	ComponentRegistry->NewComponentType(&CachedInterpolation[6],   TEXT("Cached Interpolation [6]"));
+	ComponentRegistry->NewComponentType(&CachedInterpolation[7],   TEXT("Cached Interpolation [7]"));
+	ComponentRegistry->NewComponentType(&CachedInterpolation[8],   TEXT("Cached Interpolation [8]"));
+	ComponentRegistry->NewComponentType(&CachedWeightChannelInterpolation, TEXT("Cached Weight Channel Interpolation"));
 
 	ComponentRegistry->NewComponentType(&Easing,                  TEXT("Easing"));
 	ComponentRegistry->NewComponentType(&EasingResult,            TEXT("Easing Result"));
@@ -150,6 +142,7 @@ FBuiltInComponentTypes::FBuiltInComponentTypes()
 	Tags.ImportedEntity          = ComponentRegistry->NewTag(TEXT("Imported Entity"));
 	Tags.Finished                = ComponentRegistry->NewTag(TEXT("Finished Evaluating"));
 	Tags.Ignored                 = ComponentRegistry->NewTag(TEXT("Ignored"));
+	Tags.DontOptimizeConstants   = ComponentRegistry->NewTag(TEXT("Don't Optimize Constants"));
 	Tags.FixedTime               = ComponentRegistry->NewTag(TEXT("Fixed Time"));
 	Tags.PreRoll                 = ComponentRegistry->NewTag(TEXT("Pre Roll"));
 	Tags.SectionPreRoll          = ComponentRegistry->NewTag(TEXT("Section Pre Roll"));
@@ -223,7 +216,7 @@ FBuiltInComponentTypes::FBuiltInComponentTypes()
 				UE_ARRAY_COUNT(FloatChannel) == UE_ARRAY_COUNT(BaseDouble), 
 				"Base floats and float results should have the same size.");
 		static_assert(
-				UE_ARRAY_COUNT(FloatChannel) == UE_ARRAY_COUNT(FloatChannelFlags),
+				UE_ARRAY_COUNT(FloatChannel) == UE_ARRAY_COUNT(CachedInterpolation),
 				"Float channel flags and flor channels should have the same size.");
 
 		// Duplicate float channels
@@ -232,7 +225,7 @@ FBuiltInComponentTypes::FBuiltInComponentTypes()
 			ComponentRegistry->Factories.DuplicateChildComponent(FloatChannel[Index]);
 			ComponentRegistry->Factories.DefineMutuallyInclusiveComponent(FloatChannel[Index], DoubleResult[Index]);
 			ComponentRegistry->Factories.DefineMutuallyInclusiveComponent(FloatChannel[Index], EvalTime);
-			ComponentRegistry->Factories.DefineMutuallyInclusiveComponent(FloatChannel[Index], FloatChannelFlags[Index]);
+			ComponentRegistry->Factories.DefineMutuallyInclusiveComponent(FloatChannel[Index], CachedInterpolation[Index]);
 		}
 
 		// Create base float components for float channels that are meant to be additive from base.
@@ -250,7 +243,7 @@ FBuiltInComponentTypes::FBuiltInComponentTypes()
 				UE_ARRAY_COUNT(DoubleChannel) == UE_ARRAY_COUNT(BaseDouble), 
 				"Base doubles and double results should have the same size.");
 		static_assert(
-				UE_ARRAY_COUNT(DoubleChannel) == UE_ARRAY_COUNT(DoubleChannelFlags),
+				UE_ARRAY_COUNT(DoubleChannel) == UE_ARRAY_COUNT(CachedInterpolation),
 				"Double channel flags and flor channels should have the same size.");
 
 		// Duplicate double channels
@@ -260,7 +253,7 @@ FBuiltInComponentTypes::FBuiltInComponentTypes()
 			ComponentRegistry->Factories.DuplicateChildComponent(DoubleResult[Index]);
 			ComponentRegistry->Factories.DefineMutuallyInclusiveComponent(DoubleChannel[Index], DoubleResult[Index]);
 			ComponentRegistry->Factories.DefineMutuallyInclusiveComponent(DoubleChannel[Index], EvalTime);
-			ComponentRegistry->Factories.DefineMutuallyInclusiveComponent(DoubleChannel[Index], DoubleChannelFlags[Index]);
+			ComponentRegistry->Factories.DefineMutuallyInclusiveComponent(DoubleChannel[Index], CachedInterpolation[Index]);
 		}
 
 		// Create base double components for double channels that are meant to be additive from base.
@@ -300,7 +293,7 @@ FBuiltInComponentTypes::FBuiltInComponentTypes()
 		// Weight channel components need a time and result to evaluate
 		ComponentRegistry->Factories.DefineMutuallyInclusiveComponent(WeightChannel, EvalTime);
 		ComponentRegistry->Factories.DefineMutuallyInclusiveComponent(WeightChannel, WeightResult);
-		ComponentRegistry->Factories.DefineMutuallyInclusiveComponent(WeightResult, WeightChannelFlags);
+		ComponentRegistry->Factories.DefineMutuallyInclusiveComponent(WeightResult, CachedWeightChannelInterpolation);
 	}
 
 	// Weight and easing result component relationship
diff --git a/Engine/Source/Runtime/MovieScene/Public/Channels/MovieSceneChannelData.h b/Engine/Source/Runtime/MovieScene/Public/Channels/MovieSceneChannelData.h
--- a/Engine/Source/Runtime/MovieScene/Public/Channels/MovieSceneChannelData.h
+++ b/Engine/Source/Runtime/MovieScene/Public/Channels/MovieSceneChannelData.h
@@ -72,7 +72,7 @@ namespace MovieScene
 	 * @param OutIndex2      OutIndex1 + 1 if it is a valid index in the array, INDEX_NONE otherwise
 	 * @param OutInterp      A value from 0.0 -> 1.0 specifying how a linear interpolation value from index 1 to index 2
 	 */
-	MOVIESCENE_API void EvaluateTime(TArrayView<const FFrameNumber> InTimes, FFrameTime InTime, int32& OutIndex1, int32& OutIndex2, float& OutInterp);
+	MOVIESCENE_API void EvaluateTime(TArrayView<const FFrameNumber> InTimes, FFrameTime InTime, int32& OutIndex1, int32& OutIndex2, double& OutInterp);
 
 
 	/**
diff --git a/Engine/Source/Runtime/MovieScene/Public/Channels/MovieSceneDoubleChannel.h b/Engine/Source/Runtime/MovieScene/Public/Channels/MovieSceneDoubleChannel.h
--- a/Engine/Source/Runtime/MovieScene/Public/Channels/MovieSceneDoubleChannel.h
+++ b/Engine/Source/Runtime/MovieScene/Public/Channels/MovieSceneDoubleChannel.h
@@ -34,6 +34,11 @@ struct FPropertyTag;
 template<typename> struct TMovieSceneCurveChannelImpl;
 template <typename T> struct TIsPODType;
 
+namespace UE::MovieScene::Interpolation
+{
+	struct FCachedInterpolation;
+}
+
 
 USTRUCT()
 struct FMovieSceneDoubleValue
@@ -169,6 +174,11 @@ struct MOVIESCENE_API FMovieSceneDoubleChannel : public FMovieSceneChannel
 	 */
 	bool Evaluate(FFrameTime InTime, float& OutValue) const;
 
+	/**
+	 * Retrieve a cached interpolation from this channel for the specified time
+	 */
+	UE::MovieScene::Interpolation::FCachedInterpolation GetInterpolationForTime(FFrameTime InTime) const;
+
 	/**
 	 * Set the channel's times and values to the requested values
 	 */
diff --git a/Engine/Source/Runtime/MovieScene/Public/Channels/MovieSceneFloatChannel.h b/Engine/Source/Runtime/MovieScene/Public/Channels/MovieSceneFloatChannel.h
--- a/Engine/Source/Runtime/MovieScene/Public/Channels/MovieSceneFloatChannel.h
+++ b/Engine/Source/Runtime/MovieScene/Public/Channels/MovieSceneFloatChannel.h
@@ -34,6 +34,10 @@ struct FPropertyTag;
 template<typename> struct TMovieSceneCurveChannelImpl;
 template <typename T> struct TIsPODType;
 
+namespace UE::MovieScene::Interpolation
+{
+	struct FCachedInterpolation;
+}
 
 USTRUCT()
 struct FMovieSceneFloatValue
@@ -164,6 +168,11 @@ struct MOVIESCENE_API FMovieSceneFloatChannel : public FMovieSceneChannel
 	*/
 	bool Evaluate(FFrameTime InTime, float& OutValue) const;
 
+	/**
+	 * Retrieve a cached interpolation from this channel for the specified time
+	 */
+	UE::MovieScene::Interpolation::FCachedInterpolation GetInterpolationForTime(FFrameTime InTime) const;
+
 	/**
 	 * Set the channel's times and values to the requested values
 	 */
diff --git a/Engine/Source/Runtime/MovieScene/Public/Channels/MovieSceneInterpolation.h b/Engine/Source/Runtime/MovieScene/Public/Channels/MovieSceneInterpolation.h
new file mode 100644
--- /dev/null
+++ b/Engine/Source/Runtime/MovieScene/Public/Channels/MovieSceneInterpolation.h
@@ -0,0 +1,208 @@
+// Copyright Epic Games, Inc. All Rights Reserved.
+
+#pragma once
+
+#include "CoreTypes.h"
+#include "Misc/FrameNumber.h"
+#include "Misc/FrameTime.h"
+#include "Misc/FrameRate.h"
+#include "Misc/TVariant.h"
+
+
+/**
+ * This file contains structures that represent specific interpolation algorithms for either a continuous or discrete range.
+ *
+ * They are used by Sequencer evaluation to bypass expensive piecewise data searching every frame by caching the resulting
+ * interpolation over the relevant time range which allows it to perform only the minimum computation required to find a result.
+ *
+ * FCachedInterpolation is a variant type that can represent any one of the supported interpolation modes in this file.
+ */
+namespace UE::MovieScene::Interpolation
+{
+
+/**
+ * Sentinal type that represents an invalid interpolation value.
+ * Only used when a curve has no data whatsoever, and therefore cannot be evaluated
+ */
+struct FInvalidValue
+{
+};
+
+
+/**
+ * Structure representing a constant value.
+ * Temporarily includes a flag to determine whether it needs to be re-evaluated for legacy fallback.
+ */
+struct FConstantValue
+{
+	/** The constant value */
+	double Value;
+
+	FConstantValue(double InValue)
+		: Value(InValue)
+	{}
+};
+
+/**
+ * Structure representing a linear interpolation of the form f(t) = a(t-o) + b.
+ */
+struct FLinearInterpolation
+{
+	/** The coeffient 'a' in f(t) = a(t-o) + b */
+	double Coefficient;
+	/** The constant 'c' in f(t) = a(t-o) + b */
+	double Constant;
+	/** The origin 'o' in f(t) = a(t-o) + b */
+	FFrameNumber Origin;
+
+	FLinearInterpolation(FFrameNumber InOrigin, double InCoefficient, double InConstant)
+		: Coefficient(InCoefficient)
+		, Constant(InConstant)
+		, Origin(InOrigin)
+	{}
+
+	/**
+	 * Evaluate the expression
+	 */
+	double Evaluate(FFrameTime InTime) const
+	{
+		return Coefficient*(InTime - Origin).AsDecimal() + Constant;
+	}
+};
+
+/**
+ * A cubic bezier interpolation between 2 control points with tangents, represented as 4 control points on a Bezier curve
+ */
+struct MOVIESCENE_API FCubicInterpolation
+{
+	/** The delta value between the two control points in the time-domain */
+	double DX;
+	/** The four control points that should be passed to BezierInterp */
+	double P0, P1, P2, P3;
+	/** The origin time of the first control point */
+	FFrameNumber Origin;
+
+	FCubicInterpolation(
+		FFrameNumber InOrigin,
+		double InDX,
+		double InStartValue,
+		double InEndValue,
+		double InStartTangent,
+		double InEndTangent);
+
+	/**
+	 * Evaluate the expression
+	 */
+	double Evaluate(FFrameTime InTime) const;
+};
+
+/**
+ * A weighted cubic bezier interpolation between 2 control points with weighted tangents
+ */
+struct FWeightedCubicInterpolation
+{
+	double DX;
+	double StartKeyValue;
+	double NormalizedStartTanDX;
+	double StartKeyTanY;
+	double StartWeight;
+
+	double EndKeyValue;
+	double NormalizedEndTanDX;
+	double EndKeyTanY;
+	double EndWeight;
+
+	FFrameNumber Origin;
+
+	FWeightedCubicInterpolation(
+		FFrameRate TickResolution,
+		FFrameNumber InOrigin,
+
+		FFrameNumber StartTime, 
+		double StartValue,
+		double StartTangent,
+		double StartTangentWeight,
+		bool bStartIsWeighted,
+
+		FFrameNumber EndTime, 
+		double EndValue,
+		double EndTangent,
+		double EndTangentWeight,
+		bool bEndIsWeighted);
+
+	double Evaluate(FFrameTime InTime) const;
+};
+
+/**
+ * Simple 1 dimensional range based off a FFrameNumber to define the range within which a cached interpolation is valid
+ */
+struct MOVIESCENE_API FCachedInterpolationRange
+{
+	/** Make an empty range */
+	static FCachedInterpolationRange Empty();
+	/** Make finite range from InStart to InEnd, not including the end frame */
+	static FCachedInterpolationRange Finite(FFrameNumber InStart, FFrameNumber InEnd);
+	/** Make an infinite range */
+	static FCachedInterpolationRange Infinite();
+
+	/** Make a range that only contains the specified time */
+	static FCachedInterpolationRange Only(FFrameNumber InTime);
+	/** Make a range that covers all times from (and including) the specified start */
+	static FCachedInterpolationRange From(FFrameNumber InStart);
+	/** Make a range that covers all times up to (but not including) the specified end */
+	static FCachedInterpolationRange Until(FFrameNumber InEnd);
+
+	bool Contains(FFrameNumber FrameNumber) const;
+
+	/** Inclusive start frame */
+	FFrameNumber Start;
+	/** Exclusive end frame (unless End == Max()) */
+	FFrameNumber End;
+};
+
+/**
+ * Variant structure that wraps an interpolation and the range within which it is valid.
+ * ~96 bytes
+ */
+struct MOVIESCENE_API FCachedInterpolation
+{
+	/** Default construction to an invalid state. Calling Evaluate will always return false */
+	FCachedInterpolation();
+	/** Construction as a constant value */
+	FCachedInterpolation(const FCachedInterpolationRange& InRange, const FConstantValue& Constant);
+	/** Construction as a linear interpolation */
+	FCachedInterpolation(const FCachedInterpolationRange& InRange, const FLinearInterpolation& Linear);
+	/** Construction as a cubic interpolation */
+	FCachedInterpolation(const FCachedInterpolationRange& InRange, const FCubicInterpolation& Cubic);
+	/** Construction as a weighted cubic interpolation */
+	FCachedInterpolation(const FCachedInterpolationRange& InRange, const FWeightedCubicInterpolation& WeightedCubic);
+
+	/**
+	 * Check whether this cache is still valid for the specified frame number.
+	 * @return true if this interpolation is relevant to the frame, or false if it should be re-generated
+	 */
+	bool IsCacheValidForTime(FFrameNumber FrameNumber) const;
+
+	/**
+	 * Evaluate this interpolation for the specified frame time
+	 *
+	 * @param FrameTime  The time to evaluate at
+	 * @param OutResult  Reference to recieve the resulting value. Only written to if this function returns true.
+	 * @return true if this interpolation evaluated successfully (and the result written to OutResult), false otherwise.
+	 */
+	bool Evaluate(FFrameTime FrameTime, double& OutResult) const;
+
+private:
+
+	/** Variant containing the actual interpolation implementation */
+	TVariant<FInvalidValue,
+		FConstantValue,
+		FLinearInterpolation,
+		FCubicInterpolation,
+		FWeightedCubicInterpolation> Data;
+
+	/** Structure representint the range of times this interpolation applies to */
+	FCachedInterpolationRange Range;
+};
+
+} // namespace UE::MovieScene
\ No newline at end of file
diff --git a/Engine/Source/Runtime/MovieScene/Public/EntitySystem/BuiltInComponentTypes.h b/Engine/Source/Runtime/MovieScene/Public/EntitySystem/BuiltInComponentTypes.h
--- a/Engine/Source/Runtime/MovieScene/Public/EntitySystem/BuiltInComponentTypes.h
+++ b/Engine/Source/Runtime/MovieScene/Public/EntitySystem/BuiltInComponentTypes.h
@@ -23,11 +23,20 @@
 #include "BuiltInComponentTypes.generated.h"
 
 class IMovieSceneEvaluationHook;
-namespace UE { namespace MovieScene { struct FCustomPropertyIndex; } }
-namespace UE { namespace MovieScene { struct FInitialValueIndex; } }
-namespace UE { namespace MovieScene { struct FInstanceHandle; } }
-namespace UE { namespace MovieScene { struct FInterrogationKey; } }
-namespace UE { namespace MovieScene { struct FRootInstanceHandle; } }
+
+namespace UE::MovieScene
+{
+	struct FCustomPropertyIndex;
+	struct FInitialValueIndex;
+	struct FInstanceHandle;
+	struct FInterrogationKey;
+	struct FRootInstanceHandle;
+	namespace Interpolation
+	{
+		struct FCachedInterpolation;
+	}
+}
+
 struct FFrameTime;
 struct FMovieSceneBlendChannelID;
 struct FMovieSceneSequenceID;
@@ -172,11 +181,6 @@ struct FSourceFloatChannel
 	const FMovieSceneFloatChannel* Source;
 };
 
-struct FSourceFloatChannelFlags
-{
-	bool bNeedsEvaluate = true;
-};
-
 /**
  * The component data for evaluation a double channel
  */
@@ -193,11 +197,6 @@ struct FSourceDoubleChannel
 	const FMovieSceneDoubleChannel* Source;
 };
 
-struct FSourceDoubleChannelFlags
-{
-	bool bNeedsEvaluate = true;
-};
-
 struct FEvaluationHookFlags
 {
 	bool bHasBegun = false;
@@ -296,15 +295,16 @@ public:
 
 	// An FMovieSceneFloatChannel considered to be at index N within the source structure (ie 0 = Location.X, Vector.X, Color.R; 1 = Location.Y, Vector.Y, Color.G)
 	TComponentTypeID<FSourceFloatChannel> FloatChannel[9];
-	TComponentTypeID<FSourceFloatChannelFlags> FloatChannelFlags[9];
 
 	// An FMovieSceneDoubleChannel considered to be at index N within the source structure (ie 0 = Location.X, Vector.X; 1 = Location.Y, Vector.Y)
 	TComponentTypeID<FSourceDoubleChannel> DoubleChannel[9];
-	TComponentTypeID<FSourceDoubleChannelFlags> DoubleChannelFlags[9];
+
+	// A cached interpolation structure relating to either float channels or double channels
+	TComponentTypeID<Interpolation::FCachedInterpolation> CachedInterpolation[9];
 
 	// An FMovieSceneFloatChannel that represents an arbitrary weight
 	TComponentTypeID<FSourceFloatChannel> WeightChannel;
-	TComponentTypeID<FSourceFloatChannelFlags> WeightChannelFlags;
+	TComponentTypeID<Interpolation::FCachedInterpolation> CachedWeightChannelInterpolation;
 
 	// FMovieSceneObjectPathChannel that represents a changing object path over time
 	TComponentTypeID<FSourceObjectPathChannel> ObjectPathChannel;
@@ -396,6 +396,8 @@ public:
 
 		FComponentTypeID AlwaysCacheInitialValue;
 
+		FComponentTypeID DontOptimizeConstants;
+
 	} Tags;
 
 	struct
diff --git a/Engine/Source/Runtime/MovieSceneTracks/Private/Systems/DoubleChannelEvaluatorSystem.cpp b/Engine/Source/Runtime/MovieSceneTracks/Private/Systems/DoubleChannelEvaluatorSystem.cpp
--- a/Engine/Source/Runtime/MovieSceneTracks/Private/Systems/DoubleChannelEvaluatorSystem.cpp
+++ b/Engine/Source/Runtime/MovieSceneTracks/Private/Systems/DoubleChannelEvaluatorSystem.cpp
@@ -5,53 +5,136 @@
 #include "EntitySystem/EntityAllocationIterator.h"
 #include "EntitySystem/MovieSceneEntitySystemTask.h"
 #include "EntitySystem/MovieSceneEntitySystemLinker.h"
+#include "EntitySystem/MovieSceneEntitySystemRunner.h"
 #include "EntitySystem/MovieSceneComponentRegistry.h"
 
 #include "EntitySystem/MovieSceneEvalTimeSystem.h"
+#include "EntitySystem/MovieSceneEntityMutations.h"
 #include "Channels/MovieSceneDoubleChannel.h"
+#include "Channels/MovieSceneInterpolation.h"
 
 #include "Algo/Find.h"
 #include "Math/NumericLimits.h"
 
 #include UE_INLINE_GENERATED_CPP_BY_NAME(DoubleChannelEvaluatorSystem)
 
-DECLARE_CYCLE_STAT(TEXT("MovieScene: Double channel system"), MovieSceneEval_DoubleChannelSystem, STATGROUP_MovieSceneECS);
-DECLARE_CYCLE_STAT(TEXT("MovieScene: Gather double channels"), MovieSceneEval_GatherDoubleChannelTask, STATGROUP_MovieSceneECS);
 DECLARE_CYCLE_STAT(TEXT("MovieScene: Evaluate double channels"), MovieSceneEval_EvaluateDoubleChannelTask, STATGROUP_MovieSceneECS);
 
-namespace UE
+namespace UE::MovieScene
 {
-namespace MovieScene
+
+MOVIESCENE_API extern float GCachedChannelEvaluationParity;
+
+struct FDoubleChannelTypeAssociation
 {
+	TComponentTypeID<FSourceDoubleChannel> ChannelType;
+	TComponentTypeID<Interpolation::FCachedInterpolation> CachedInterpolationType;
+	TComponentTypeID<double> ResultType;
+};
+
+TArray<FDoubleChannelTypeAssociation, TInlineAllocator<4>> GDoubleChannelTypeAssociations;
 
 // @todo: for multi-bindings we currently re-evaluate the double channel for each binding, even though the time is the same.
 // Do we need to optimize for this case using something like the code below, while pessimizing the common (non-multi-bind) codepath??
-struct FEvaluateDoubleChannels
+
+/** Entity-component task that evaluates using the non-cached codepath for testing parity */
+struct FEvaluateDoubleChannels_Parity
 {
-	void ForEachEntity(FSourceDoubleChannel DoubleChannel, FFrameTime FrameTime, double& OutResult, FSourceDoubleChannelFlags& OutFlags)
+	void ForEachEntity(FSourceDoubleChannel DoubleChannel, FFrameTime FrameTime, Interpolation::FCachedInterpolation& Cache, double& OutResult)
 	{
-		if (OutFlags.bNeedsEvaluate == false)
+		if (!Cache.IsCacheValidForTime(FrameTime.GetFrame()))
 		{
-			return;
+			Cache = DoubleChannel.Source->GetInterpolationForTime(FrameTime);
 		}
 
-		if (!DoubleChannel.Source->Evaluate(FrameTime, OutResult))
+		if (!Cache.Evaluate(FrameTime, OutResult))
 		{
 			OutResult = MIN_dbl;
 		}
 
-		if (DoubleChannel.Source->GetTimes().Num() <= 1)
+		double ParityResult;
+		if (!DoubleChannel.Source->Evaluate(FrameTime, ParityResult))
+		{
+			ParityResult = MIN_dbl;
+		}
+
+		if (!FMath::IsNearlyEqual(ParityResult, OutResult, (double)GCachedChannelEvaluationParity))
+		{
+			UE_LOG(LogMovieScene, Warning, TEXT("Parity mismatch between cached and non-cached evaluation %.16f != %.16f!"), OutResult, ParityResult);
+
+			static bool bBreakDebugger = false;
+			ensureAlways(!bBreakDebugger);
+		}
+	}
+};
+
+/** Entity-component task that evaluates using a cached interpolation if possible */
+struct FEvaluateDoubleChannels_Cached
+{
+	void ForEachEntity(FSourceDoubleChannel DoubleChannel, FFrameTime FrameTime, Interpolation::FCachedInterpolation& Cache, double& OutResult)
+	{
+		if (!Cache.IsCacheValidForTime(FrameTime.GetFrame()))
+		{
+			Cache = DoubleChannel.Source->GetInterpolationForTime(FrameTime);
+		}
+
+		if (!Cache.Evaluate(FrameTime, OutResult))
 		{
-			OutFlags.bNeedsEvaluate = false;
+			OutResult = MIN_dbl;
 		}
 	}
 };
 
-} // namespace MovieScene
-} // namespace UE
+/**
+ * Mutation that removes FCachedInterpolation components from any channels that have a constant value
+ * This effectively prevents the channel from ever being updated again
+ */
+struct FInitializeConstantDoubleChannelMutation : IMovieSceneConditionalEntityMutation
+{
+	FDoubleChannelTypeAssociation Association;
+
+	FInitializeConstantDoubleChannelMutation(const FDoubleChannelTypeAssociation& InAssociation)
+		: Association(InAssociation)
+	{}
+
+	virtual void MarkAllocation(FEntityAllocation* Allocation, TBitArray<>& OutEntitiesToMutate) const
+	{
+		TComponentReader<FSourceDoubleChannel> Channels = Allocation->ReadComponents(Association.ChannelType);
+
+		const int32 Num = Allocation->Num();
+		for (int32 Index = 0; Index < Num; ++Index)
+		{
+			if (Channels[Index].Source->GetTimes().Num() <= 1)
+			{
+				OutEntitiesToMutate.PadToNum(Index + 1, false);
+				OutEntitiesToMutate[Index] = true;
+			}
+		}
+	}
+
+	virtual void CreateMutation(FEntityManager* EntityManager, FComponentMask* InOutEntityComponentTypes) const
+	{
+		InOutEntityComponentTypes->Remove(Association.CachedInterpolationType);
+	}
+
+	virtual void InitializeEntities(const FEntityRange& EntityRange, const FComponentMask& AllocationType) const
+	{
+		TComponentReader<FSourceDoubleChannel> Channels = EntityRange.Allocation->ReadComponents(Association.ChannelType);
+		TComponentWriter<double>               Results = EntityRange.Allocation->WriteComponents(Association.ResultType, FEntityAllocationWriteContext::NewAllocation());
+
+		for (int32 Index = EntityRange.ComponentStartOffset; Index < EntityRange.ComponentStartOffset + EntityRange.Num; ++Index)
+		{
+			// Eval time shouldn't even matter if it's constant
+			if (!Channels[Index].Source->Evaluate(0, Results[Index]))
+			{
+				Results[Index] = MIN_dbl;
+			}
+		}
+	}
+};
 
+} // namespace UE::MovieScene
 
-TArray<UDoubleChannelEvaluatorSystem::FChannelType, TInlineAllocator<4>> UDoubleChannelEvaluatorSystem::StaticChannelTypes;
 
 UDoubleChannelEvaluatorSystem::UDoubleChannelEvaluatorSystem(const FObjectInitializer& ObjInit)
 	: Super(ObjInit)
@@ -59,37 +142,41 @@ UDoubleChannelEvaluatorSystem::UDoubleChannelEvaluatorSystem(const FObjectInitia
 	using namespace UE::MovieScene;
 
 	SystemCategories = EEntitySystemCategory::ChannelEvaluators;
+	Phase = ESystemPhase::Instantiation | ESystemPhase::Evaluation;
 
 	if (HasAnyFlags(RF_ClassDefaultObject))
 	{
 		DefineImplicitPrerequisite(UMovieSceneEvalTimeSystem::StaticClass(), GetClass());
 
 		FBuiltInComponentTypes* Components = FBuiltInComponentTypes::Get();
+		DefineComponentConsumer(GetClass(), Components->SymbolicTags.CreatesEntities);
 
 		for (int32 Index = 0; Index < UE_ARRAY_COUNT(Components->DoubleChannel); ++Index)
 		{
-			RegisterChannelType(Components->DoubleChannel[Index], Components->DoubleChannelFlags[Index], Components->DoubleResult[Index]);
+			RegisterChannelType(Components->DoubleChannel[Index], Components->CachedInterpolation[Index], Components->DoubleResult[Index]);
 		}
 	}
 }
 
-void UDoubleChannelEvaluatorSystem::RegisterChannelType(TComponentTypeID<UE::MovieScene::FSourceDoubleChannel> SourceChannelType, TComponentTypeID<UE::MovieScene::FSourceDoubleChannelFlags> ChannelFlagsType, TComponentTypeID<double> ResultType)
+void UDoubleChannelEvaluatorSystem::RegisterChannelType(TComponentTypeID<UE::MovieScene::FSourceDoubleChannel> SourceChannelType, TComponentTypeID<UE::MovieScene::Interpolation::FCachedInterpolation> CachedInterpolationType, TComponentTypeID<double> ResultType)
 {
 	using namespace UE::MovieScene;
 
-	FChannelType ChannelType;
+	FDoubleChannelTypeAssociation ChannelType;
 	ChannelType.ChannelType = SourceChannelType;
-	ChannelType.ChannelFlagsType = ChannelFlagsType;
+	ChannelType.CachedInterpolationType = CachedInterpolationType;
 	ChannelType.ResultType  = ResultType;
 
 	DefineComponentProducer(UDoubleChannelEvaluatorSystem::StaticClass(), ResultType);
 
-	StaticChannelTypes.Add(ChannelType);
+	GDoubleChannelTypeAssociations.Add(ChannelType);
 }
 
 bool UDoubleChannelEvaluatorSystem::IsRelevantImpl(UMovieSceneEntitySystemLinker* InLinker) const
 {
-	for (const FChannelType& ChannelType : StaticChannelTypes)
+	using namespace UE::MovieScene;
+
+	for (const FDoubleChannelTypeAssociation& ChannelType : GDoubleChannelTypeAssociations)
 	{
 		if (InLinker->EntityManager.ContainsComponent(ChannelType.ChannelType))
 		{
@@ -104,21 +191,57 @@ void UDoubleChannelEvaluatorSystem::OnRun(FSystemTaskPrerequisites& InPrerequisi
 {
 	using namespace UE::MovieScene;
 
-	SCOPE_CYCLE_COUNTER(MovieSceneEval_DoubleChannelSystem)
+	FMovieSceneEntitySystemRunner* Runner = Linker->GetActiveRunner();
+	if (!Runner)
+	{
+		return;
+	}
 
 	FBuiltInComponentTypes* BuiltInComponents = FBuiltInComponentTypes::Get();
 
-	for (const FChannelType& ChannelType : StaticChannelTypes)
+	if (Runner->GetCurrentPhase() == ESystemPhase::Instantiation)
 	{
-		// Evaluate double channels per instance and write the evaluated value into the output
-		FEntityTaskBuilder()
-		.Read(ChannelType.ChannelType)
-		.Read(BuiltInComponents->EvalTime)
-		.Write(ChannelType.ResultType)
-		.Write(ChannelType.ChannelFlagsType)
-		.FilterNone({ BuiltInComponents->Tags.Ignored })
-		.SetStat(GET_STATID(MovieSceneEval_EvaluateDoubleChannelTask))
-		.Dispatch_PerEntity<FEvaluateDoubleChannels>(&Linker->EntityManager, InPrerequisites, &Subsequents);
+		// Remove cache components for anything that is constant
+		for (const FDoubleChannelTypeAssociation& Association : GDoubleChannelTypeAssociations)
+		{
+			FInitializeConstantDoubleChannelMutation Mutation{ Association };
+
+			FEntityComponentFilter Filter;
+			Filter.All({ Association.ChannelType, Association.ResultType, BuiltInComponents->Tags.NeedsLink });
+			Filter.None({ BuiltInComponents->Tags.DontOptimizeConstants });
+
+			Linker->EntityManager.MutateConditional(Filter, Mutation);
+		}
 	}
-}
+	else if (Runner->GetCurrentPhase() == ESystemPhase::Evaluation)
+	{
+		if (GCachedChannelEvaluationParity != 0.f)
+		{
+			for (const FDoubleChannelTypeAssociation& ChannelType : GDoubleChannelTypeAssociations)
+			{
+				FEntityTaskBuilder()
+				.Read(ChannelType.ChannelType)
+				.Read(BuiltInComponents->EvalTime)
+				.Write(ChannelType.CachedInterpolationType)
+				.Write(ChannelType.ResultType)
+				.FilterNone({ BuiltInComponents->Tags.Ignored })
+				.SetStat(GET_STATID(MovieSceneEval_EvaluateDoubleChannelTask))
+				.Dispatch_PerEntity<FEvaluateDoubleChannels_Parity>(&Linker->EntityManager, InPrerequisites, &Subsequents);
+			}
+			return;
+		}
 
+		for (const FDoubleChannelTypeAssociation& ChannelType : GDoubleChannelTypeAssociations)
+		{
+			// Evaluate double channels per instance and write the evaluated value into the output
+			FEntityTaskBuilder()
+			.Read(ChannelType.ChannelType)
+			.Read(BuiltInComponents->EvalTime)
+			.Write(ChannelType.CachedInterpolationType)
+			.Write(ChannelType.ResultType)
+			.FilterNone({ BuiltInComponents->Tags.Ignored })
+			.SetStat(GET_STATID(MovieSceneEval_EvaluateDoubleChannelTask))
+			.Dispatch_PerEntity<FEvaluateDoubleChannels_Cached>(&Linker->EntityManager, InPrerequisites, &Subsequents);
+		}
+	}
+}
diff --git a/Engine/Source/Runtime/MovieSceneTracks/Private/Systems/FloatChannelEvaluatorSystem.cpp b/Engine/Source/Runtime/MovieSceneTracks/Private/Systems/FloatChannelEvaluatorSystem.cpp
--- a/Engine/Source/Runtime/MovieSceneTracks/Private/Systems/FloatChannelEvaluatorSystem.cpp
+++ b/Engine/Source/Runtime/MovieSceneTracks/Private/Systems/FloatChannelEvaluatorSystem.cpp
@@ -5,56 +5,142 @@
 #include "EntitySystem/EntityAllocationIterator.h"
 #include "EntitySystem/MovieSceneEntitySystemTask.h"
 #include "EntitySystem/MovieSceneEntitySystemLinker.h"
+#include "EntitySystem/MovieSceneEntitySystemRunner.h"
 #include "EntitySystem/MovieSceneComponentRegistry.h"
 
 #include "EntitySystem/MovieSceneEvalTimeSystem.h"
+#include "EntitySystem/MovieSceneEntityMutations.h"
 #include "Channels/MovieSceneFloatChannel.h"
+#include "Channels/MovieSceneInterpolation.h"
 #include "MovieSceneTracksComponentTypes.h"
 
 #include "Algo/Find.h"
 
 #include UE_INLINE_GENERATED_CPP_BY_NAME(FloatChannelEvaluatorSystem)
 
-DECLARE_CYCLE_STAT(TEXT("MovieScene: Float channel system"), MovieSceneEval_FloatChannelSystem, STATGROUP_MovieSceneECS);
-DECLARE_CYCLE_STAT(TEXT("MovieScene: Gather float channels"), MovieSceneEval_GatherFloatChannelTask, STATGROUP_MovieSceneECS);
 DECLARE_CYCLE_STAT(TEXT("MovieScene: Evaluate float channels"), MovieSceneEval_EvaluateFloatChannelTask, STATGROUP_MovieSceneECS);
 
-namespace UE
+namespace UE::MovieScene
 {
-namespace MovieScene
+
+MOVIESCENE_API extern float GCachedChannelEvaluationParity;
+
+struct FFloatChannelTypeAssociation
 {
+	TComponentTypeID<FSourceFloatChannel> ChannelType;
+	TComponentTypeID<Interpolation::FCachedInterpolation> CachedInterpolationType;
+	TComponentTypeID<double> ResultType;
+};
+
+TArray<FFloatChannelTypeAssociation, TInlineAllocator<4>> GFloatChannelTypeAssociations;
 
 // @todo: for multi-bindings we currently re-evaluate the float channel for each binding, even though the time is the same.
 // Do we need to optimize for this case using something like the code below, while pessimizing the common (non-multi-bind) codepath??
-struct FEvaluateFloatChannels
+
+/** Entity-component task that evaluates using the non-cached codepath for testing parity */
+struct FEvaluateFloatChannels_Parity
 {
-	void ForEachEntity(FSourceFloatChannel FloatChannel, FFrameTime FrameTime, double& OutResult, FSourceFloatChannelFlags& OutFlags)
+	void ForEachEntity(FSourceFloatChannel FloatChannel, FFrameTime FrameTime, Interpolation::FCachedInterpolation& Cache, double& OutResult)
 	{
-		if (OutFlags.bNeedsEvaluate == false)
+		if (!Cache.IsCacheValidForTime(FrameTime.GetFrame()))
+		{
+			Cache = FloatChannel.Source->GetInterpolationForTime(FrameTime);
+		}
+
+		if (!Cache.Evaluate(FrameTime, OutResult))
 		{
+			OutResult = MIN_dbl;
 			return;
 		}
 
-		float FloatResult = 0.f;
-		if (!FloatChannel.Source->Evaluate(FrameTime, FloatResult))
+		float ParityResult;
+		if (!FloatChannel.Source->Evaluate(FrameTime, ParityResult))
 		{
-			FloatResult = MIN_flt;
+			return;
 		}
-		OutResult = (double)FloatResult;
 
-		if (FloatChannel.Source->GetTimes().Num() <= 1)
+		if (!FMath::IsNearlyEqual(ParityResult, (float)OutResult, GCachedChannelEvaluationParity))
 		{
-			OutFlags.bNeedsEvaluate = false;
+			UE_LOG(LogMovieScene, Warning, TEXT("Parity mismatch between cached and non-cached evaluation %.8f != %.8f!"), OutResult, ParityResult);
+
+			static bool bBreakDebugger = false;
+			ensureAlways(!bBreakDebugger);
 		}
 	}
 };
 
+/** Entity-component task that evaluates using a cached interpolation if possible */
+struct FEvaluateFloatChannels_Cached
+{
+	void ForEachEntity(FSourceFloatChannel FloatChannel, FFrameTime FrameTime, Interpolation::FCachedInterpolation& Cache, double& OutResult)
+	{
+		if (!Cache.IsCacheValidForTime(FrameTime.GetFrame()))
+		{
+			Cache = FloatChannel.Source->GetInterpolationForTime(FrameTime);
+		}
+
+		if (!Cache.Evaluate(FrameTime, OutResult))
+		{
+			OutResult = MIN_dbl;
+		}
+	}
+};
+
+/**
+ * Mutation that removes FCachedInterpolation components from any channels that have a constant value
+ * This effectively prevents the channel from ever being updated again
+ */
+struct FInitializeConstantFloatChannelMutation : IMovieSceneConditionalEntityMutation
+{
+	FFloatChannelTypeAssociation Association;
+
+	FInitializeConstantFloatChannelMutation(const FFloatChannelTypeAssociation& InAssociation)
+		: Association(InAssociation)
+	{}
+
+	virtual void MarkAllocation(FEntityAllocation* Allocation, TBitArray<>& OutEntitiesToMutate) const
+	{
+		TComponentReader<FSourceFloatChannel> Channels = Allocation->ReadComponents(Association.ChannelType);
+
+		const int32 Num = Allocation->Num();
+		for (int32 Index = 0; Index < Num; ++Index)
+		{
+			if (Channels[Index].Source->GetTimes().Num() <= 1)
+			{
+				OutEntitiesToMutate.PadToNum(Index + 1, false);
+				OutEntitiesToMutate[Index] = true;
+			}
+		}
+	}
+
+	virtual void CreateMutation(FEntityManager* EntityManager, FComponentMask* InOutEntityComponentTypes) const
+	{
+		InOutEntityComponentTypes->Remove(Association.CachedInterpolationType);
+	}
 
-} // namespace MovieScene
-} // namespace UE
+	virtual void InitializeEntities(const FEntityRange& EntityRange, const FComponentMask& AllocationType) const
+	{
+		TComponentReader<FSourceFloatChannel> Channels = EntityRange.Allocation->ReadComponents(Association.ChannelType);
+		TComponentWriter<double>               Results = EntityRange.Allocation->WriteComponents(Association.ResultType, FEntityAllocationWriteContext::NewAllocation());
+
+		for (int32 Index = EntityRange.ComponentStartOffset; Index < EntityRange.ComponentStartOffset + EntityRange.Num; ++Index)
+		{
+			// Eval time shouldn't even matter if it's constant
+			float Result;
+			if (!Channels[Index].Source->Evaluate(0, Result))
+			{
+				Results[Index] = MIN_dbl;
+			}
+			else
+			{
+				Results[Index] = Result;
+			}
+		}
+	}
+};
 
+} // namespace UE::MovieScene
 
-TArray<UFloatChannelEvaluatorSystem::FChannelType, TInlineAllocator<16>> UFloatChannelEvaluatorSystem::StaticChannelTypes;
 
 UFloatChannelEvaluatorSystem::UFloatChannelEvaluatorSystem(const FObjectInitializer& ObjInit)
 	: Super(ObjInit)
@@ -62,38 +148,41 @@ UFloatChannelEvaluatorSystem::UFloatChannelEvaluatorSystem(const FObjectInitiali
 	using namespace UE::MovieScene;
 
 	SystemCategories = EEntitySystemCategory::ChannelEvaluators;
+	Phase = ESystemPhase::Instantiation | ESystemPhase::Evaluation;
 
 	if (HasAnyFlags(RF_ClassDefaultObject))
 	{
 		DefineImplicitPrerequisite(UMovieSceneEvalTimeSystem::StaticClass(), GetClass());
 
 		FBuiltInComponentTypes* Components = FBuiltInComponentTypes::Get();
+		DefineComponentConsumer(GetClass(), Components->SymbolicTags.CreatesEntities);
 
 		for (int32 Index = 0; Index < UE_ARRAY_COUNT(Components->FloatChannel); ++Index)
 		{
-			RegisterChannelType(Components->FloatChannel[Index], Components->FloatChannelFlags[Index], Components->DoubleResult[Index]);
+			RegisterChannelType(Components->FloatChannel[Index], Components->CachedInterpolation[Index], Components->DoubleResult[Index]);
 		}
-		RegisterChannelType(Components->WeightChannel, Components->WeightChannelFlags, Components->WeightResult);
 	}
 }
 
-void UFloatChannelEvaluatorSystem::RegisterChannelType(TComponentTypeID<UE::MovieScene::FSourceFloatChannel> SourceChannelType, TComponentTypeID<UE::MovieScene::FSourceFloatChannelFlags> ChannelFlagsType, TComponentTypeID<double> ResultType)
+void UFloatChannelEvaluatorSystem::RegisterChannelType(TComponentTypeID<UE::MovieScene::FSourceFloatChannel> SourceChannelType, TComponentTypeID<UE::MovieScene::Interpolation::FCachedInterpolation> CachedInterpolationType, TComponentTypeID<double> ResultType)
 {
 	using namespace UE::MovieScene;
 
-	FChannelType ChannelType;
+	FFloatChannelTypeAssociation ChannelType;
 	ChannelType.ChannelType = SourceChannelType;
-	ChannelType.ChannelFlagsType = ChannelFlagsType;
+	ChannelType.CachedInterpolationType = CachedInterpolationType;
 	ChannelType.ResultType  = ResultType;
 
 	DefineComponentProducer(UFloatChannelEvaluatorSystem::StaticClass(), ResultType);
 
-	StaticChannelTypes.Add(ChannelType);
+	GFloatChannelTypeAssociations.Add(ChannelType);
 }
 
 bool UFloatChannelEvaluatorSystem::IsRelevantImpl(UMovieSceneEntitySystemLinker* InLinker) const
 {
-	for (const FChannelType& ChannelType : StaticChannelTypes)
+	using namespace UE::MovieScene;
+
+	for (const FFloatChannelTypeAssociation& ChannelType : GFloatChannelTypeAssociations)
 	{
 		if (InLinker->EntityManager.ContainsComponent(ChannelType.ChannelType))
 		{
@@ -108,21 +197,57 @@ void UFloatChannelEvaluatorSystem::OnRun(FSystemTaskPrerequisites& InPrerequisit
 {
 	using namespace UE::MovieScene;
 
-	SCOPE_CYCLE_COUNTER(MovieSceneEval_FloatChannelSystem)
+	FMovieSceneEntitySystemRunner* Runner = Linker->GetActiveRunner();
+	if (!Runner)
+	{
+		return;
+	}
 
 	FBuiltInComponentTypes* BuiltInComponents = FBuiltInComponentTypes::Get();
 
-	for (const FChannelType& ChannelType : StaticChannelTypes)
+	if (Runner->GetCurrentPhase() == ESystemPhase::Instantiation)
 	{
-		// Evaluate float channels per instance and write the evaluated value into the output
-		FEntityTaskBuilder()
-		.Read(ChannelType.ChannelType)
-		.Read(BuiltInComponents->EvalTime)
-		.Write(ChannelType.ResultType)
-		.Write(ChannelType.ChannelFlagsType)
-		.FilterNone({ BuiltInComponents->Tags.Ignored })
-		.SetStat(GET_STATID(MovieSceneEval_EvaluateFloatChannelTask))
-		.Dispatch_PerEntity<FEvaluateFloatChannels>(&Linker->EntityManager, InPrerequisites, &Subsequents);
+		// Remove cache components for anything that is constant
+		for (const FFloatChannelTypeAssociation& Association : GFloatChannelTypeAssociations)
+		{
+			FInitializeConstantFloatChannelMutation Mutation{ Association };
+
+			FEntityComponentFilter Filter;
+			Filter.All({ Association.ChannelType, Association.ResultType, BuiltInComponents->Tags.NeedsLink });
+			Filter.None({ BuiltInComponents->Tags.DontOptimizeConstants });
+
+			Linker->EntityManager.MutateConditional(Filter, Mutation);
+		}
 	}
-}
+	else if (Runner->GetCurrentPhase() == ESystemPhase::Evaluation)
+	{
+		if (GCachedChannelEvaluationParity != 0.f)
+		{
+			for (const FFloatChannelTypeAssociation& ChannelType : GFloatChannelTypeAssociations)
+			{
+				FEntityTaskBuilder()
+				.Read(ChannelType.ChannelType)
+				.Read(BuiltInComponents->EvalTime)
+				.Write(ChannelType.CachedInterpolationType)
+				.Write(ChannelType.ResultType)
+				.FilterNone({ BuiltInComponents->Tags.Ignored })
+				.SetStat(GET_STATID(MovieSceneEval_EvaluateFloatChannelTask))
+				.Dispatch_PerEntity<FEvaluateFloatChannels_Parity>(&Linker->EntityManager, InPrerequisites, &Subsequents);
+			}
+			return;
+		}
 
+		for (const FFloatChannelTypeAssociation& ChannelType : GFloatChannelTypeAssociations)
+		{
+			// Evaluate float channels per instance and write the evaluated value into the output
+			FEntityTaskBuilder()
+			.Read(ChannelType.ChannelType)
+			.Read(BuiltInComponents->EvalTime)
+			.Write(ChannelType.CachedInterpolationType)
+			.Write(ChannelType.ResultType)
+			.FilterNone({ BuiltInComponents->Tags.Ignored })
+			.SetStat(GET_STATID(MovieSceneEval_EvaluateFloatChannelTask))
+			.Dispatch_PerEntity<FEvaluateFloatChannels_Cached>(&Linker->EntityManager, InPrerequisites, &Subsequents);
+		}
+	}
+}
diff --git a/Engine/Source/Runtime/MovieSceneTracks/Private/Systems/MovieSceneTransformOriginSystem.cpp b/Engine/Source/Runtime/MovieSceneTracks/Private/Systems/MovieSceneTransformOriginSystem.cpp
--- a/Engine/Source/Runtime/MovieSceneTracks/Private/Systems/MovieSceneTransformOriginSystem.cpp
+++ b/Engine/Source/Runtime/MovieSceneTracks/Private/Systems/MovieSceneTransformOriginSystem.cpp
@@ -2,6 +2,8 @@
 
 #include "Systems/MovieSceneTransformOriginSystem.h"
 #include "EntitySystem/MovieSceneEntitySystemLinker.h"
+#include "EntitySystem/MovieSceneEntitySystemRunner.h"
+#include "EntitySystem/MovieSceneEntityMutations.h"
 #include "EntitySystem/MovieSceneInstanceRegistry.h"
 #include "EntitySystem/BuiltInComponentTypes.h"
 #include "Tracks/IMovieSceneTransformOrigin.h"
@@ -27,18 +29,14 @@ struct FAssignTransformOrigin
 
 	void ForEachAllocation(const FEntityAllocation* Allocation, TRead<FInstanceHandle> Instances, TRead<UObject*> BoundObjects,
 		TWriteOptional<double> LocationX, TWriteOptional<double> LocationY, TWriteOptional<double> LocationZ,
-		TWriteOptional<double> RotationX, TWriteOptional<double> RotationY, TWriteOptional<double> RotationZ,
-		TWriteOptional<FSourceDoubleChannelFlags> FlagsLocationX, TWriteOptional<FSourceDoubleChannelFlags> FlagsLocationY, TWriteOptional<FSourceDoubleChannelFlags> FlagsLocationZ,
-		TWriteOptional<FSourceDoubleChannelFlags> FlagsRotationX, TWriteOptional<FSourceDoubleChannelFlags> FlagsRotationY, TWriteOptional<FSourceDoubleChannelFlags> FlagsRotationZ)
+		TWriteOptional<double> RotationX, TWriteOptional<double> RotationY, TWriteOptional<double> RotationZ)
 	{
-		TransformLocation(Instances, BoundObjects, LocationX, LocationY, LocationZ, RotationX, RotationY, RotationZ, FlagsLocationX, FlagsLocationY, FlagsLocationZ, FlagsRotationX, FlagsRotationY, FlagsRotationZ, Allocation->Num());
+		TransformLocation(Instances, BoundObjects, LocationX, LocationY, LocationZ, RotationX, RotationY, RotationZ, Allocation->Num());
 	}
 
 	void TransformLocation(const FInstanceHandle* Instances, const UObject* const * BoundObjects,
 		double* OutLocationX, double* OutLocationY, double* OutLocationZ,
 		double* OutRotationX, double* OutRotationY, double* OutRotationZ,
-		FSourceDoubleChannelFlags* OutFlagsLocationX, FSourceDoubleChannelFlags* OutFlagsLocationY, FSourceDoubleChannelFlags* OutFlagsLocationZ,
-		FSourceDoubleChannelFlags* OutFlagsRotationX, FSourceDoubleChannelFlags* OutFlagsRotationY, FSourceDoubleChannelFlags* OutFlagsRotationZ,
 		int32 Num)
 	{
 		for (int32 Index = 0; Index < Num; ++Index)
@@ -73,14 +71,6 @@ struct FAssignTransformOrigin
 			if (OutRotationX) { OutRotationX[Index] = NewRotation.Roll; }
 			if (OutRotationY) { OutRotationY[Index] = NewRotation.Pitch; }
 			if (OutRotationZ) { OutRotationZ[Index] = NewRotation.Yaw; }
-
-			if (OutFlagsLocationX) { OutFlagsLocationX[Index].bNeedsEvaluate = true; }
-			if (OutFlagsLocationY) { OutFlagsLocationY[Index].bNeedsEvaluate = true; }
-			if (OutFlagsLocationZ) { OutFlagsLocationZ[Index].bNeedsEvaluate = true; }
-
-			if (OutFlagsRotationX) { OutFlagsRotationX[Index].bNeedsEvaluate = true; }
-			if (OutFlagsRotationY) { OutFlagsRotationY[Index].bNeedsEvaluate = true; }
-			if (OutFlagsRotationZ) { OutFlagsRotationZ[Index].bNeedsEvaluate = true; }
 		}
 	}
 };
@@ -89,6 +79,47 @@ struct FAssignTransformOrigin
 } // namespace UE
 
 
+UMovieSceneTransformOriginInstantiatorSystem::UMovieSceneTransformOriginInstantiatorSystem(const FObjectInitializer& ObjInit)
+	: Super(ObjInit)
+{
+	using namespace UE::MovieScene;
+
+	Phase = ESystemPhase::Instantiation;
+
+	if (HasAnyFlags(RF_ClassDefaultObject))
+	{
+		DefineComponentConsumer(GetClass(), FBuiltInComponentTypes::Get()->SymbolicTags.CreatesEntities);
+
+		// This must be run before the double channel evaluator
+		DefineImplicitPrerequisite(GetClass(), UDoubleChannelEvaluatorSystem::StaticClass());
+	}
+}
+
+
+void UMovieSceneTransformOriginInstantiatorSystem::OnRun(FSystemTaskPrerequisites& InPrerequisites, FSystemSubsequentTasks& Subsequents)
+{
+	using namespace UE::MovieScene;
+
+	FBuiltInComponentTypes*          BuiltInComponents = FBuiltInComponentTypes::Get();
+	FMovieSceneTracksComponentTypes* TracksComponents  = FMovieSceneTracksComponentTypes::Get();
+
+	FEntityComponentFilter Filter;
+	Filter.All({ TracksComponents->ComponentTransform.PropertyTag, BuiltInComponents->Tags.AbsoluteBlend, BuiltInComponents->Tags.NeedsLink });
+	Filter.None({ BuiltInComponents->BlendChannelOutput });
+	Filter.Any({
+		BuiltInComponents->DoubleResult[0],
+		BuiltInComponents->DoubleResult[1],
+		BuiltInComponents->DoubleResult[2],
+		BuiltInComponents->DoubleResult[3],
+		BuiltInComponents->DoubleResult[4],
+		BuiltInComponents->DoubleResult[5]
+	});
+
+	// Stop constant values for transforms from being optimized - we need them to be re-evaluated every frame
+	// since this system will trample them all
+	Linker->EntityManager.MutateAll(Filter, FAddSingleMutation(BuiltInComponents->Tags.DontOptimizeConstants));
+}
+
 UMovieSceneTransformOriginSystem::UMovieSceneTransformOriginSystem(const FObjectInitializer& ObjInit)
 	: Super(ObjInit)
 {
@@ -106,13 +137,6 @@ UMovieSceneTransformOriginSystem::UMovieSceneTransformOriginSystem(const FObject
 		DefineComponentConsumer(GetClass(), FBuiltInComponentTypes::Get()->DoubleResult[3]);
 		DefineComponentConsumer(GetClass(), FBuiltInComponentTypes::Get()->DoubleResult[4]);
 		DefineComponentConsumer(GetClass(), FBuiltInComponentTypes::Get()->DoubleResult[5]);
-
-		DefineComponentConsumer(GetClass(), FBuiltInComponentTypes::Get()->DoubleChannelFlags[0]);
-		DefineComponentConsumer(GetClass(), FBuiltInComponentTypes::Get()->DoubleChannelFlags[1]);
-		DefineComponentConsumer(GetClass(), FBuiltInComponentTypes::Get()->DoubleChannelFlags[2]);
-		DefineComponentConsumer(GetClass(), FBuiltInComponentTypes::Get()->DoubleChannelFlags[3]);
-		DefineComponentConsumer(GetClass(), FBuiltInComponentTypes::Get()->DoubleChannelFlags[4]);
-		DefineComponentConsumer(GetClass(), FBuiltInComponentTypes::Get()->DoubleChannelFlags[5]);
 	}
 }
 
@@ -135,6 +159,13 @@ bool UMovieSceneTransformOriginSystem::IsRelevantImpl(UMovieSceneEntitySystemLin
 	return false;
 }
 
+void UMovieSceneTransformOriginSystem::OnLink()
+{
+	UMovieSceneTransformOriginInstantiatorSystem* Instantiator = Linker->LinkSystem<UMovieSceneTransformOriginInstantiatorSystem>();
+	// This system keeps the instantiator around
+	Linker->SystemGraph.AddReference(this, Instantiator);
+}
+
 void UMovieSceneTransformOriginSystem::OnRun(FSystemTaskPrerequisites& InPrerequisites, FSystemSubsequentTasks& Subsequents)
 {
 	using namespace UE::MovieScene;
@@ -185,12 +216,6 @@ void UMovieSceneTransformOriginSystem::OnRun(FSystemTaskPrerequisites& InPrerequ
 		.WriteOptional(BuiltInComponents->DoubleResult[3])
 		.WriteOptional(BuiltInComponents->DoubleResult[4])
 		.WriteOptional(BuiltInComponents->DoubleResult[5])
-		.WriteOptional(BuiltInComponents->DoubleChannelFlags[0])
-		.WriteOptional(BuiltInComponents->DoubleChannelFlags[1])
-		.WriteOptional(BuiltInComponents->DoubleChannelFlags[2])
-		.WriteOptional(BuiltInComponents->DoubleChannelFlags[3])
-		.WriteOptional(BuiltInComponents->DoubleChannelFlags[4])
-		.WriteOptional(BuiltInComponents->DoubleChannelFlags[5])
 		.CombineFilter(Filter)
 		// Must contain at least one double result
 		.FilterAny({ BuiltInComponents->DoubleResult[0], BuiltInComponents->DoubleResult[1], BuiltInComponents->DoubleResult[2],
diff --git a/Engine/Source/Runtime/MovieSceneTracks/Public/Systems/DoubleChannelEvaluatorSystem.h b/Engine/Source/Runtime/MovieSceneTracks/Public/Systems/DoubleChannelEvaluatorSystem.h
--- a/Engine/Source/Runtime/MovieSceneTracks/Public/Systems/DoubleChannelEvaluatorSystem.h
+++ b/Engine/Source/Runtime/MovieSceneTracks/Public/Systems/DoubleChannelEvaluatorSystem.h
@@ -16,16 +16,15 @@
 class UMovieSceneEntitySystemLinker;
 class UObject;
 
-namespace UE
+namespace UE::MovieScene
 {
-namespace MovieScene
-{
-
 	struct FSourceDoubleChannel;
-	struct FSourceDoubleChannelFlags;
 
-} // namespace MovieScene
-} // namespace UE
+	namespace Interpolation
+	{
+		struct FCachedInterpolation;
+	}
+} // namespace UE::MovieScene
 
 /**
  * System that is responsible for evaluating double channels.
@@ -43,16 +42,5 @@ public:
 	virtual void OnRun(FSystemTaskPrerequisites& InPrerequisites, FSystemSubsequentTasks& Subsequents) override;
 	virtual bool IsRelevantImpl(UMovieSceneEntitySystemLinker* InLinker) const override;
 
-	static void RegisterChannelType(TComponentTypeID<UE::MovieScene::FSourceDoubleChannel> SourceChannelType, TComponentTypeID<UE::MovieScene::FSourceDoubleChannelFlags> ChannelFlagsType, TComponentTypeID<double> ResultType);
-
-private:
-
-	struct FChannelType
-	{
-		TComponentTypeID<UE::MovieScene::FSourceDoubleChannel> ChannelType;
-		TComponentTypeID<UE::MovieScene::FSourceDoubleChannelFlags> ChannelFlagsType;
-		TComponentTypeID<double> ResultType;
-	};
-
-	static TArray<FChannelType, TInlineAllocator<4>> StaticChannelTypes;
+	static void RegisterChannelType(TComponentTypeID<UE::MovieScene::FSourceDoubleChannel> SourceChannelType, TComponentTypeID<UE::MovieScene::Interpolation::FCachedInterpolation> CachedInterpolationType, TComponentTypeID<double> ResultType);
 };
diff --git a/Engine/Source/Runtime/MovieSceneTracks/Public/Systems/FloatChannelEvaluatorSystem.h b/Engine/Source/Runtime/MovieSceneTracks/Public/Systems/FloatChannelEvaluatorSystem.h
--- a/Engine/Source/Runtime/MovieSceneTracks/Public/Systems/FloatChannelEvaluatorSystem.h
+++ b/Engine/Source/Runtime/MovieSceneTracks/Public/Systems/FloatChannelEvaluatorSystem.h
@@ -16,25 +16,25 @@
 class UMovieSceneEntitySystemLinker;
 class UObject;
 
-namespace UE
+namespace UE::MovieScene
 {
-namespace MovieScene
-{
-
 	struct FSourceFloatChannel;
-	struct FSourceFloatChannelFlags;
 
-} // namespace MovieScene
-} // namespace UE
+	namespace Interpolation
+	{
+		struct FCachedInterpolation;
+	}
+} // namespace UE::MovieScene
 
 /**
- * System that is responsible for evaluating float channels.
+ * System that is responsible for evaluating double channels.
  */
 UCLASS()
 class MOVIESCENETRACKS_API UFloatChannelEvaluatorSystem : public UMovieSceneEntitySystem
 {
 public:
 
+
 	GENERATED_BODY()
 
 	UFloatChannelEvaluatorSystem(const FObjectInitializer& ObjInit);
@@ -42,16 +42,5 @@ public:
 	virtual void OnRun(FSystemTaskPrerequisites& InPrerequisites, FSystemSubsequentTasks& Subsequents) override;
 	virtual bool IsRelevantImpl(UMovieSceneEntitySystemLinker* InLinker) const override;
 
-	static void RegisterChannelType(TComponentTypeID<UE::MovieScene::FSourceFloatChannel> SourceChannelType, TComponentTypeID<UE::MovieScene::FSourceFloatChannelFlags> ChannelFlagsType, TComponentTypeID<double> ResultType);
-
-private:
-
-	struct FChannelType
-	{
-		TComponentTypeID<UE::MovieScene::FSourceFloatChannel> ChannelType;
-		TComponentTypeID<UE::MovieScene::FSourceFloatChannelFlags> ChannelFlagsType;
-		TComponentTypeID<double> ResultType;
-	};
-
-	static TArray<FChannelType, TInlineAllocator<16>> StaticChannelTypes;
-};
\ No newline at end of file
+	static void RegisterChannelType(TComponentTypeID<UE::MovieScene::FSourceFloatChannel> SourceChannelType, TComponentTypeID<UE::MovieScene::Interpolation::FCachedInterpolation> CachedInterpolationType, TComponentTypeID<double> ResultType);
+};
diff --git a/Engine/Source/Runtime/MovieSceneTracks/Public/Systems/MovieSceneTransformOriginSystem.h b/Engine/Source/Runtime/MovieSceneTracks/Public/Systems/MovieSceneTransformOriginSystem.h
--- a/Engine/Source/Runtime/MovieSceneTracks/Public/Systems/MovieSceneTransformOriginSystem.h
+++ b/Engine/Source/Runtime/MovieSceneTracks/Public/Systems/MovieSceneTransformOriginSystem.h
@@ -9,6 +9,21 @@
 
 struct FMovieSceneAnimTypeID;
 
+UCLASS(MinimalAPI)
+class UMovieSceneTransformOriginInstantiatorSystem : public UMovieSceneEntitySystem
+{
+public:
+
+	GENERATED_BODY()
+
+	UMovieSceneTransformOriginInstantiatorSystem(const FObjectInitializer& ObjInit);
+
+private:
+
+	virtual void OnRun(FSystemTaskPrerequisites& InPrerequisites, FSystemSubsequentTasks& Subsequents) override final;
+};
+
+
 UCLASS(MinimalAPI)
 class UMovieSceneTransformOriginSystem : public UMovieSceneEntitySystem
 {
@@ -20,9 +35,12 @@ public:
 
 private:
 
+	virtual void OnLink() override;
 	virtual bool IsRelevantImpl(UMovieSceneEntitySystemLinker* InLinker) const override;
 	virtual void OnRun(FSystemTaskPrerequisites& InPrerequisites, FSystemSubsequentTasks& Subsequents) override final;
 
+private:
+
 	TSparseArray<FTransform> TransformOriginsByInstanceID;
 };
 
