CIS fix

diff --git a/Engine/Source/Runtime/MovieScene/Private/Channels/MovieSceneCurveChannelImpl.h b/Engine/Source/Runtime/MovieScene/Private/Channels/MovieSceneCurveChannelImpl.h
--- a/Engine/Source/Runtime/MovieScene/Private/Channels/MovieSceneCurveChannelImpl.h
+++ b/Engine/Source/Runtime/MovieScene/Private/Channels/MovieSceneCurveChannelImpl.h
@@ -42,7 +42,7 @@ struct MOVIESCENE_API TMovieSceneCurveChannelImpl
 	/** Structure used to store the result of UE::MovieScene::EvaluateTime for a given channel/key distribution*/
 	struct FTimeEvaluationCache
 	{
-		float InterpValue = 0.f;
+		double InterpValue = 0.0;
 		int32 Index1 = INDEX_NONE, Index2 = INDEX_NONE;
 		int32 CachedNumFrames = INDEX_NONE;
 		FFrameTime CacheFrameTime;
