Sequencer: add more ways to save pre-animated state

This adds the ability to:
- track and save pre-animated state on a single entity instead of on a whole allocation's worth of them.
- provide custom data when saving pre-animated state, instead of solely relying on existing component data on the entity being considered

diff --git a/Engine/Source/Runtime/MovieScene/Public/Evaluation/PreAnimatedState/IMovieScenePreAnimatedStorage.h b/Engine/Source/Runtime/MovieScene/Public/Evaluation/PreAnimatedState/IMovieScenePreAnimatedStorage.h
--- a/Engine/Source/Runtime/MovieScene/Public/Evaluation/PreAnimatedState/IMovieScenePreAnimatedStorage.h
+++ b/Engine/Source/Runtime/MovieScene/Public/Evaluation/PreAnimatedState/IMovieScenePreAnimatedStorage.h
@@ -51,6 +51,7 @@ struct IPreAnimatedStorage : TSharedFromThis<IPreAnimatedStorage>
 struct IPreAnimatedObjectEntityStorage
 {
 	virtual void BeginTrackingEntities(const FPreAnimatedTrackerParams& Params, TRead<FMovieSceneEntityID> EntityIDs, TRead<FRootInstanceHandle> InstanceHandles, TRead<UObject*> BoundObjects) = 0;
+	virtual void BeginTrackingEntity(FMovieSceneEntityID EntityID, bool bWantsRestoreState, FRootInstanceHandle InstanceHandle, UObject* BoundObject) = 0;
 	virtual void CachePreAnimatedValues(const FCachePreAnimatedValueParams& Params, TArrayView<UObject* const> BoundObjects) = 0;
 };
 
diff --git a/Engine/Source/Runtime/MovieScene/Public/Evaluation/PreAnimatedState/MovieScenePreAnimatedObjectStorage.h b/Engine/Source/Runtime/MovieScene/Public/Evaluation/PreAnimatedState/MovieScenePreAnimatedObjectStorage.h
--- a/Engine/Source/Runtime/MovieScene/Public/Evaluation/PreAnimatedState/MovieScenePreAnimatedObjectStorage.h
+++ b/Engine/Source/Runtime/MovieScene/Public/Evaluation/PreAnimatedState/MovieScenePreAnimatedObjectStorage.h
@@ -68,6 +68,24 @@ public:
 		}
 	}
 
+	void BeginTrackingEntity(FMovieSceneEntityID EntityID, bool bWantsRestoreState, FRootInstanceHandle RootInstanceHandle, UObject* BoundObject) override
+	{
+		if (!this->ParentExtension->IsCapturingGlobalState() && !bWantsRestoreState)
+		{
+			return;
+		}
+
+		FPreAnimatedEntityCaptureSource* EntityMetaData = this->ParentExtension->GetOrCreateEntityMetaData();
+
+		FObjectKey Key{ BoundObject };
+
+		FPreAnimatedStorageGroupHandle GroupHandle  = this->Traits.MakeGroup(BoundObject);
+		FPreAnimatedStorageIndex       StorageIndex = this->GetOrCreateStorageIndex(Key);
+
+		FPreAnimatedStateEntry Entry{ GroupHandle, FPreAnimatedStateCachedValueHandle{ this->StorageID, StorageIndex } };
+		EntityMetaData->BeginTrackingEntity(Entry, EntityID, RootInstanceHandle, bWantsRestoreState);
+	}
+
 	void CachePreAnimatedValues(const FCachePreAnimatedValueParams& Params, TArrayView<UObject* const> BoundObjects) override
 	{
 		for (UObject* BoundObject : BoundObjects)
@@ -77,6 +95,46 @@ public:
 	}
 
 	void CachePreAnimatedValue(const FCachePreAnimatedValueParams& Params, UObject* BoundObject)
+	{
+		FObjectKey Key{ BoundObject };
+	
+		FPreAnimatedStorageIndex       StorageIndex = this->GetOrCreateStorageIndex(Key);
+		FPreAnimatedStorageGroupHandle GroupHandle  = this->Traits.MakeGroup(BoundObject);
+		FPreAnimatedStateEntry         Entry        = FPreAnimatedStateEntry{ GroupHandle, FPreAnimatedStateCachedValueHandle{ this->StorageID, StorageIndex } };
+	
+		if (this->ParentExtension->IsCapturingGlobalState())
+		{
+			this->ParentExtension->EnsureMetaData(Entry);
+		}
+		else if (!this->ParentExtension->MetaDataExists(Entry))
+		{
+			return;
+		}
+	
+		CachePreAnimatedValue(Params, Entry, BoundObject);
+	}
+	
+	void CachePreAnimatedValue(const FCachePreAnimatedValueParams& Params, const FPreAnimatedStateEntry& Entry, UObject* BoundObject)
+	{
+		FPreAnimatedStorageIndex StorageIndex = Entry.ValueHandle.StorageIndex;
+	
+		EPreAnimatedStorageRequirement StorageRequirement = this->ParentExtension->GetStorageRequirement(Entry);
+		if (!this->IsStorageRequirementSatisfied(StorageIndex, StorageRequirement))
+		{
+			StorageType NewValue;
+			ObjectTraits::CachePreAnimatedValue(BoundObject, NewValue);
+	
+			this->AssignPreAnimatedValue(StorageIndex, StorageRequirement, MoveTemp(NewValue));
+		}
+	
+		if (Params.bForcePersist)
+		{
+			this->ForciblyPersistStorage(StorageIndex);
+		}
+	}
+
+	template<typename ValueInitializer>
+	void CachePreAnimatedValue(const FCachePreAnimatedValueParams& Params, UObject* BoundObject, ValueInitializer&& InitCallback)
 	{
 		FObjectKey Key{ BoundObject };
 
@@ -93,10 +151,11 @@ public:
 			return;
 		}
 
-		CachePreAnimatedValue(Params, Entry, BoundObject);
+		CachePreAnimatedValue(Params, Entry, BoundObject, Forward<ValueInitializer>(InitCallback));
 	}
 
-	void CachePreAnimatedValue(const FCachePreAnimatedValueParams& Params, const FPreAnimatedStateEntry& Entry, UObject* BoundObject)
+	template<typename ValueInitializer>
+	void CachePreAnimatedValue(const FCachePreAnimatedValueParams& Params, const FPreAnimatedStateEntry& Entry, UObject* BoundObject, ValueInitializer&& InitCallback)
 	{
 		FPreAnimatedStorageIndex StorageIndex = Entry.ValueHandle.StorageIndex;
 
@@ -104,6 +163,7 @@ public:
 		if (!this->IsStorageRequirementSatisfied(StorageIndex, StorageRequirement))
 		{
 			StorageType NewValue;
+			InitCallback(BoundObject, NewValue);
 			ObjectTraits::CachePreAnimatedValue(BoundObject, NewValue);
 
 			this->AssignPreAnimatedValue(StorageIndex, StorageRequirement, MoveTemp(NewValue));
diff --git a/Engine/Source/Runtime/MovieScene/Public/Evaluation/PreAnimatedState/MovieScenePreAnimatedStateStorage.h b/Engine/Source/Runtime/MovieScene/Public/Evaluation/PreAnimatedState/MovieScenePreAnimatedStateStorage.h
--- a/Engine/Source/Runtime/MovieScene/Public/Evaluation/PreAnimatedState/MovieScenePreAnimatedStateStorage.h
+++ b/Engine/Source/Runtime/MovieScene/Public/Evaluation/PreAnimatedState/MovieScenePreAnimatedStateStorage.h
@@ -596,6 +596,35 @@ public:
 		.Iterate_PerAllocation(&Linker->EntityManager, VisitAllocation);
 	}
 
+	/**
+	 * Set up a new associations with a storage index for the given entity
+	 * WARNING: Does not cache actual pre-animated values
+	 */
+	template<typename... ContributorTypes>
+	void BeginTrackingEntity(FMovieSceneEntityID EntityID, const bool bWantsRestoreState, FRootInstanceHandle RootInstanceHandle, ContributorTypes... InComponents)
+	{
+		if (!this->ParentExtension->IsCapturingGlobalState() && !bWantsRestoreState)
+		{
+			return;
+		}
+
+		TSavePreAnimatedStateParams<ContributorTypes...> Params;
+
+		if (!Params.ShouldSavePreAnimatedState(InComponents...))
+		{
+			return;
+		}
+
+		KeyType Key{ InComponents... };
+
+		FPreAnimatedStorageGroupHandle GroupHandle  = this->Traits.MakeGroup(InComponents...);
+		FPreAnimatedStorageIndex       StorageIndex = this->GetOrCreateStorageIndex(Key);
+
+		FPreAnimatedStateEntry Entry{ GroupHandle, FPreAnimatedStateCachedValueHandle{ this->StorageID, StorageIndex } };
+		FPreAnimatedEntityCaptureSource* EntityMetaData = this->ParentExtension->GetOrCreateEntityMetaData();
+		EntityMetaData->BeginTrackingEntity(Entry, EntityID, RootInstanceHandle, bWantsRestoreState);
+	}
+
 	/**
 	 * Cache pre-animated values for entities with the specified component types
 	 */
