Sequencer: Fixed weighted tangent evaluation

diff --git a/Engine/Source/Runtime/MovieScene/Private/Channels/MovieSceneInterpolation.cpp b/Engine/Source/Runtime/MovieScene/Private/Channels/MovieSceneInterpolation.cpp
--- a/Engine/Source/Runtime/MovieScene/Private/Channels/MovieSceneInterpolation.cpp
+++ b/Engine/Source/Runtime/MovieScene/Private/Channels/MovieSceneInterpolation.cpp
@@ -157,13 +157,13 @@ FWeightedCubicInterpolation::FWeightedCubicInterpolation(
 
 	if (bStartIsWeighted)
 	{
-		const double LeaveTangentNormalized = StartTangent / TimeInterval;
-		const double DY = LeaveTangentNormalized * DXInSeconds;
-		StartWeight = FMath::Sqrt(DXInSeconds*DXInSeconds + DY*DY) * OneThird;
+		StartWeight = StartTangentWeight;
 	}
 	else
 	{
-		StartWeight = StartTangentWeight;
+		const double LeaveTangentNormalized = StartTangent / TimeInterval;
+		const double DY = LeaveTangentNormalized * DXInSeconds;
+		StartWeight = FMath::Sqrt(DXInSeconds*DXInSeconds + DY*DY) * OneThird;
 	}
 
 	const double StartKeyTanX = CosAngle * StartWeight + Time1;
@@ -177,13 +177,13 @@ FWeightedCubicInterpolation::FWeightedCubicInterpolation(
 
 	if (bEndIsWeighted)
 	{
-		const double ArriveTangentNormalized = EndTangent / TimeInterval;
-		const double DY = ArriveTangentNormalized * DXInSeconds;
-		EndWeight = FMath::Sqrt(DXInSeconds*DXInSeconds + DY*DY) * OneThird;
+		EndWeight =  EndTangentWeight;
 	}
 	else
 	{
-		EndWeight =  EndTangentWeight;
+		const double ArriveTangentNormalized = EndTangent / TimeInterval;
+		const double DY = ArriveTangentNormalized * DXInSeconds;
+		EndWeight = FMath::Sqrt(DXInSeconds*DXInSeconds + DY*DY) * OneThird;
 	}
 
 	const double EndKeyTanX = -CosAngle * EndWeight + Time2;
