Sequencer: Reworked cvars around cached channel evaluation to allow for toggling cached evaluation and parity testing separately

diff --git a/Engine/Source/Runtime/MovieScene/Private/Channels/MovieSceneCurveChannelImpl.cpp b/Engine/Source/Runtime/MovieScene/Private/Channels/MovieSceneCurveChannelImpl.cpp
--- a/Engine/Source/Runtime/MovieScene/Private/Channels/MovieSceneCurveChannelImpl.cpp
+++ b/Engine/Source/Runtime/MovieScene/Private/Channels/MovieSceneCurveChannelImpl.cpp
@@ -29,13 +29,19 @@ namespace UE
 namespace MovieScene
 {
 
-MOVIESCENE_API float GCachedChannelEvaluationParity = 0.f;
-static FAutoConsoleVariableRef CVarTestCachedChannelEvaluationParity(
-	TEXT("Sequencer.TestCachedChannelEvaluationParity"),
-	GCachedChannelEvaluationParity,
-	TEXT("Toggles whether Evaluate will use cached or non-cached evaluation."),
+MOVIESCENE_API float GCachedChannelEvaluationParityThreshold = 0.f;
+static FAutoConsoleVariableRef CVarCachedChannelEvaluationParityThreshold(
+	TEXT("Sequencer.CachedChannelEvaluationParityThreshold"),
+	GCachedChannelEvaluationParityThreshold,
+	TEXT("Threshold for testing Evaluate parity with cached/uncached routines."),
 	ECVF_Default);
 
+MOVIESCENE_API bool GEnableCachedChannelEvaluation = true;
+static FAutoConsoleVariableRef CVarEnableCachedChannelEvaluation(
+	TEXT("Sequencer.EnableCachedChannelEvaluation"),
+	GEnableCachedChannelEvaluation,
+	TEXT("Toggles whether channel evaluation will use cached or non-cached evaluation."),
+	ECVF_Default);
 
 template<typename ChannelType>
 static typename ChannelType::CurveValueType
@@ -603,19 +609,44 @@ bool TMovieSceneCurveChannelImpl<ChannelType>::EvaluateWithCache(const ChannelTy
 {
 	using namespace UE::MovieScene;
 
+	const bool bResult = GEnableCachedChannelEvaluation ? EvaluateCached(InChannel, InTime, OutValue) : EvaluateLegacy(InChannel, InTime, OutValue);
+
 	// ------------------------------------------------------------------------
-	// New cached codepath - existing implementation still exists as a fallback
-	if (GCachedChannelEvaluationParity == 0.f)
+	// Check against new cached codepath - eventually this will replace the code above
+	if (GCachedChannelEvaluationParityThreshold != 0.f)
 	{
-		double Result = 0.0;
-		if (GetInterpolationForTime(InChannel, InTime).Evaluate(InTime, Result))
+		CurveValueType ParityValue;
+		const bool bParityResult = GEnableCachedChannelEvaluation ? EvaluateLegacy(InChannel, InTime, ParityValue) : EvaluateCached(InChannel, InTime, ParityValue);
+
+		if (bParityResult != bResult || !FMath::IsNearlyEqual((double)ParityValue, (double)OutValue, (double)GCachedChannelEvaluationParityThreshold))
 		{
-			OutValue = static_cast<CurveValueType>(Result);
-			return true;
+			UE_LOG(LogMovieScene, Warning, TEXT("Parity mismatch between cached and non-cached evaluation %.16f != %.16f!"), OutValue, ParityValue);
+
+			static bool bBreakDebugger = false;
+			ensureAlways(!bBreakDebugger);
 		}
-		return false;
 	}
 
+	return bResult;
+}
+
+template<typename ChannelType>
+bool TMovieSceneCurveChannelImpl<ChannelType>::EvaluateCached(const ChannelType* InChannel, FFrameTime InTime, CurveValueType& OutValue)
+{
+	double ResultValue = 0.0;
+	if (GetInterpolationForTime(InChannel, InTime).Evaluate(InTime, ResultValue))
+	{
+		OutValue = static_cast<CurveValueType>(ResultValue);
+		return true;
+	}
+	return false;
+}
+
+template<typename ChannelType>
+bool TMovieSceneCurveChannelImpl<ChannelType>::EvaluateLegacy(const ChannelType* InChannel, FFrameTime InTime, CurveValueType& OutValue)
+{
+	using namespace UE::MovieScene;
+
 	// ------------------------------------------------------------------------
 	// Legacy evaluate in-place codepath - only exists as a fallback
 	const int32 NumKeys = InChannel->Times.Num();
diff --git a/Engine/Source/Runtime/MovieScene/Private/Channels/MovieSceneCurveChannelImpl.h b/Engine/Source/Runtime/MovieScene/Private/Channels/MovieSceneCurveChannelImpl.h
--- a/Engine/Source/Runtime/MovieScene/Private/Channels/MovieSceneCurveChannelImpl.h
+++ b/Engine/Source/Runtime/MovieScene/Private/Channels/MovieSceneCurveChannelImpl.h
@@ -155,6 +155,9 @@ private:
 	 * @param InOutPoints           An array to populate with the evaluated points
 	 */
 	static void RefineCurvePoints(const ChannelType* InChannel, FFrameRate TickResolution, double TimeThreshold, CurveValueType ValueThreshold, TArray<TTuple<double, double>>& InOutPoints);
+
+	static bool EvaluateLegacy(const ChannelType* InChannel, FFrameTime InTime, CurveValueType& OutValue);
+	static bool EvaluateCached(const ChannelType* InChannel, FFrameTime InTime, CurveValueType& OutValue);
 };
 
 template<typename ChannelType>
diff --git a/Engine/Source/Runtime/MovieSceneTracks/Private/Systems/DoubleChannelEvaluatorSystem.cpp b/Engine/Source/Runtime/MovieSceneTracks/Private/Systems/DoubleChannelEvaluatorSystem.cpp
--- a/Engine/Source/Runtime/MovieSceneTracks/Private/Systems/DoubleChannelEvaluatorSystem.cpp
+++ b/Engine/Source/Runtime/MovieSceneTracks/Private/Systems/DoubleChannelEvaluatorSystem.cpp
@@ -23,7 +23,7 @@ DECLARE_CYCLE_STAT(TEXT("MovieScene: Evaluate double channels"), MovieSceneEval_
 namespace UE::MovieScene
 {
 
-MOVIESCENE_API extern float GCachedChannelEvaluationParity;
+MOVIESCENE_API extern bool GEnableCachedChannelEvaluation;
 
 struct FDoubleChannelTypeAssociation
 {
@@ -38,33 +38,14 @@ TArray<FDoubleChannelTypeAssociation, TInlineAllocator<4>> GDoubleChannelTypeAss
 // Do we need to optimize for this case using something like the code below, while pessimizing the common (non-multi-bind) codepath??
 
 /** Entity-component task that evaluates using the non-cached codepath for testing parity */
-struct FEvaluateDoubleChannels_Parity
+struct FEvaluateDoubleChannels_Uncached
 {
 	void ForEachEntity(FSourceDoubleChannel DoubleChannel, FFrameTime FrameTime, Interpolation::FCachedInterpolation& Cache, double& OutResult)
 	{
-		if (!Cache.IsCacheValidForTime(FrameTime.GetFrame()))
-		{
-			Cache = DoubleChannel.Source->GetInterpolationForTime(FrameTime);
-		}
-
-		if (!Cache.Evaluate(FrameTime, OutResult))
+		if (!DoubleChannel.Source->Evaluate(FrameTime, OutResult))
 		{
 			OutResult = MIN_dbl;
 		}
-
-		double ParityResult;
-		if (!DoubleChannel.Source->Evaluate(FrameTime, ParityResult))
-		{
-			ParityResult = MIN_dbl;
-		}
-
-		if (!FMath::IsNearlyEqual(ParityResult, OutResult, (double)GCachedChannelEvaluationParity))
-		{
-			UE_LOG(LogMovieScene, Warning, TEXT("Parity mismatch between cached and non-cached evaluation %.16f != %.16f!"), OutResult, ParityResult);
-
-			static bool bBreakDebugger = false;
-			ensureAlways(!bBreakDebugger);
-		}
 	}
 };
 
@@ -215,7 +196,7 @@ void UDoubleChannelEvaluatorSystem::OnRun(FSystemTaskPrerequisites& InPrerequisi
 	}
 	else if (Runner->GetCurrentPhase() == ESystemPhase::Evaluation)
 	{
-		if (GCachedChannelEvaluationParity != 0.f)
+		if (!GEnableCachedChannelEvaluation)
 		{
 			for (const FDoubleChannelTypeAssociation& ChannelType : GDoubleChannelTypeAssociations)
 			{
@@ -226,7 +207,7 @@ void UDoubleChannelEvaluatorSystem::OnRun(FSystemTaskPrerequisites& InPrerequisi
 				.Write(ChannelType.ResultType)
 				.FilterNone({ BuiltInComponents->Tags.Ignored })
 				.SetStat(GET_STATID(MovieSceneEval_EvaluateDoubleChannelTask))
-				.Dispatch_PerEntity<FEvaluateDoubleChannels_Parity>(&Linker->EntityManager, InPrerequisites, &Subsequents);
+				.Dispatch_PerEntity<FEvaluateDoubleChannels_Uncached>(&Linker->EntityManager, InPrerequisites, &Subsequents);
 			}
 			return;
 		}
diff --git a/Engine/Source/Runtime/MovieSceneTracks/Private/Systems/FloatChannelEvaluatorSystem.cpp b/Engine/Source/Runtime/MovieSceneTracks/Private/Systems/FloatChannelEvaluatorSystem.cpp
--- a/Engine/Source/Runtime/MovieSceneTracks/Private/Systems/FloatChannelEvaluatorSystem.cpp
+++ b/Engine/Source/Runtime/MovieSceneTracks/Private/Systems/FloatChannelEvaluatorSystem.cpp
@@ -23,7 +23,7 @@ DECLARE_CYCLE_STAT(TEXT("MovieScene: Evaluate float channels"), MovieSceneEval_E
 namespace UE::MovieScene
 {
 
-MOVIESCENE_API extern float GCachedChannelEvaluationParity;
+MOVIESCENE_API extern bool GEnableCachedChannelEvaluation;
 
 struct FFloatChannelTypeAssociation
 {
@@ -38,33 +38,18 @@ TArray<FFloatChannelTypeAssociation, TInlineAllocator<4>> GFloatChannelTypeAssoc
 // Do we need to optimize for this case using something like the code below, while pessimizing the common (non-multi-bind) codepath??
 
 /** Entity-component task that evaluates using the non-cached codepath for testing parity */
-struct FEvaluateFloatChannels_Parity
+struct FEvaluateFloatChannels_Uncached
 {
 	void ForEachEntity(FSourceFloatChannel FloatChannel, FFrameTime FrameTime, Interpolation::FCachedInterpolation& Cache, double& OutResult)
 	{
-		if (!Cache.IsCacheValidForTime(FrameTime.GetFrame()))
+		float Result;
+		if (FloatChannel.Source->Evaluate(FrameTime, Result))
 		{
-			Cache = FloatChannel.Source->GetInterpolationForTime(FrameTime);
+			OutResult = Result;
 		}
-
-		if (!Cache.Evaluate(FrameTime, OutResult))
+		else
 		{
 			OutResult = MIN_dbl;
-			return;
-		}
-
-		float ParityResult;
-		if (!FloatChannel.Source->Evaluate(FrameTime, ParityResult))
-		{
-			return;
-		}
-
-		if (!FMath::IsNearlyEqual(ParityResult, (float)OutResult, GCachedChannelEvaluationParity))
-		{
-			UE_LOG(LogMovieScene, Warning, TEXT("Parity mismatch between cached and non-cached evaluation %.8f != %.8f!"), OutResult, ParityResult);
-
-			static bool bBreakDebugger = false;
-			ensureAlways(!bBreakDebugger);
 		}
 	}
 };
@@ -222,7 +207,7 @@ void UFloatChannelEvaluatorSystem::OnRun(FSystemTaskPrerequisites& InPrerequisit
 	}
 	else if (Runner->GetCurrentPhase() == ESystemPhase::Evaluation)
 	{
-		if (GCachedChannelEvaluationParity != 0.f)
+		if (!GEnableCachedChannelEvaluation)
 		{
 			for (const FFloatChannelTypeAssociation& ChannelType : GFloatChannelTypeAssociations)
 			{
@@ -233,7 +218,7 @@ void UFloatChannelEvaluatorSystem::OnRun(FSystemTaskPrerequisites& InPrerequisit
 				.Write(ChannelType.ResultType)
 				.FilterNone({ BuiltInComponents->Tags.Ignored })
 				.SetStat(GET_STATID(MovieSceneEval_EvaluateFloatChannelTask))
-				.Dispatch_PerEntity<FEvaluateFloatChannels_Parity>(&Linker->EntityManager, InPrerequisites, &Subsequents);
+				.Dispatch_PerEntity<FEvaluateFloatChannels_Uncached>(&Linker->EntityManager, InPrerequisites, &Subsequents);
 			}
 			return;
 		}
