Sequencer: Fixup CL 23079040 to work with FTimeEvaluationCache that exists on UE5/Main

diff --git a/Engine/Source/Runtime/MovieScene/Private/Channels/MovieSceneCurveChannelImpl.cpp b/Engine/Source/Runtime/MovieScene/Private/Channels/MovieSceneCurveChannelImpl.cpp
--- a/Engine/Source/Runtime/MovieScene/Private/Channels/MovieSceneCurveChannelImpl.cpp
+++ b/Engine/Source/Runtime/MovieScene/Private/Channels/MovieSceneCurveChannelImpl.cpp
@@ -413,6 +413,12 @@ bool TMovieSceneCurveChannelImpl<ChannelType>::CacheExtrapolation(const ChannelT
 
 template<typename ChannelType>
 UE::MovieScene::Interpolation::FCachedInterpolation TMovieSceneCurveChannelImpl<ChannelType>::GetInterpolationForTime(const ChannelType* InChannel, FFrameTime InTime)
+{
+	return GetInterpolationForTime(InChannel, nullptr, InTime);
+}
+
+template<typename ChannelType>
+UE::MovieScene::Interpolation::FCachedInterpolation TMovieSceneCurveChannelImpl<ChannelType>::GetInterpolationForTime(const ChannelType* InChannel, FTimeEvaluationCache* InOutEvaluationCache, FFrameTime InTime)
 {
 	using namespace UE::MovieScene;
 	using namespace UE::MovieScene::Interpolation;
@@ -479,7 +485,26 @@ UE::MovieScene::Interpolation::FCachedInterpolation TMovieSceneCurveChannelImpl<
 	// Find the pair of keys that we need to evaluate
 	double Interp = 0.0;
 	int32 Index1 = INDEX_NONE, Index2 = INDEX_NONE;
-	UE::MovieScene::EvaluateTime(InChannel->Times, Params.Time, Index1, Index2, Interp);
+
+	// Initialize cache if not yet performed
+	if(InOutEvaluationCache && InOutEvaluationCache->CachedNumFrames == INDEX_NONE)
+	{
+		UE::MovieScene::EvaluateTime(InChannel->Times, Params.Time, InOutEvaluationCache->Index1, InOutEvaluationCache->Index2, InOutEvaluationCache->InterpValue);
+		InOutEvaluationCache->CachedNumFrames = InChannel->Times.Num();
+		InOutEvaluationCache->CacheFrameTime = Params.Time;
+	}
+
+	// If cache matches contained number of frames, copy data out rather than evaluating time data again
+	if (InOutEvaluationCache && InOutEvaluationCache->CachedNumFrames == InChannel->Times.Num() && InOutEvaluationCache->CacheFrameTime == Params.Time)
+	{
+		Interp = InOutEvaluationCache->InterpValue;
+		Index1 = InOutEvaluationCache->Index1;
+		Index2 = InOutEvaluationCache->Index2;
+	}
+	else
+	{
+		UE::MovieScene::EvaluateTime(InChannel->Times, Params.Time, Index1, Index2, Interp);
+	}
 
 	if (Index1 == INDEX_NONE)
 	{
@@ -609,14 +634,14 @@ bool TMovieSceneCurveChannelImpl<ChannelType>::EvaluateWithCache(const ChannelTy
 {
 	using namespace UE::MovieScene;
 
-	const bool bResult = GEnableCachedChannelEvaluation ? EvaluateCached(InChannel, InTime, OutValue) : EvaluateLegacy(InChannel, InTime, OutValue);
+	const bool bResult = GEnableCachedChannelEvaluation ? EvaluateCached(InChannel, InOutEvaluationCache, InTime, OutValue) : EvaluateLegacy(InChannel, InOutEvaluationCache, InTime, OutValue);
 
 	// ------------------------------------------------------------------------
 	// Check against new cached codepath - eventually this will replace the code above
 	if (GCachedChannelEvaluationParityThreshold != 0.f)
 	{
 		CurveValueType ParityValue;
-		const bool bParityResult = GEnableCachedChannelEvaluation ? EvaluateLegacy(InChannel, InTime, ParityValue) : EvaluateCached(InChannel, InTime, ParityValue);
+		const bool bParityResult = GEnableCachedChannelEvaluation ? EvaluateLegacy(InChannel, InOutEvaluationCache, InTime, ParityValue) : EvaluateCached(InChannel, InOutEvaluationCache, InTime, ParityValue);
 
 		if (bParityResult != bResult || !FMath::IsNearlyEqual((double)ParityValue, (double)OutValue, (double)GCachedChannelEvaluationParityThreshold))
 		{
@@ -631,10 +656,10 @@ bool TMovieSceneCurveChannelImpl<ChannelType>::EvaluateWithCache(const ChannelTy
 }
 
 template<typename ChannelType>
-bool TMovieSceneCurveChannelImpl<ChannelType>::EvaluateCached(const ChannelType* InChannel, FFrameTime InTime, CurveValueType& OutValue)
+bool TMovieSceneCurveChannelImpl<ChannelType>::EvaluateCached(const ChannelType* InChannel, FTimeEvaluationCache* InOutEvaluationCache, FFrameTime InTime, CurveValueType& OutValue)
 {
 	double ResultValue = 0.0;
-	if (GetInterpolationForTime(InChannel, InTime).Evaluate(InTime, ResultValue))
+	if (GetInterpolationForTime(InChannel, InOutEvaluationCache, InTime).Evaluate(InTime, ResultValue))
 	{
 		OutValue = static_cast<CurveValueType>(ResultValue);
 		return true;
@@ -643,7 +668,7 @@ bool TMovieSceneCurveChannelImpl<ChannelType>::EvaluateCached(const ChannelType*
 }
 
 template<typename ChannelType>
-bool TMovieSceneCurveChannelImpl<ChannelType>::EvaluateLegacy(const ChannelType* InChannel, FFrameTime InTime, CurveValueType& OutValue)
+bool TMovieSceneCurveChannelImpl<ChannelType>::EvaluateLegacy(const ChannelType* InChannel, FTimeEvaluationCache* InOutEvaluationCache, FFrameTime InTime, CurveValueType& OutValue)
 {
 	using namespace UE::MovieScene;
 
diff --git a/Engine/Source/Runtime/MovieScene/Private/Channels/MovieSceneCurveChannelImpl.h b/Engine/Source/Runtime/MovieScene/Private/Channels/MovieSceneCurveChannelImpl.h
--- a/Engine/Source/Runtime/MovieScene/Private/Channels/MovieSceneCurveChannelImpl.h
+++ b/Engine/Source/Runtime/MovieScene/Private/Channels/MovieSceneCurveChannelImpl.h
@@ -38,6 +38,7 @@ struct MOVIESCENE_API TMovieSceneCurveChannelImpl
 	 * Evaluate this channel by returning a cachable interpolation structure
 	 */
 	static UE::MovieScene::Interpolation::FCachedInterpolation GetInterpolationForTime(const ChannelType* InChannel, FFrameTime InTime);
+	static UE::MovieScene::Interpolation::FCachedInterpolation GetInterpolationForTime(const ChannelType* InChannel, FTimeEvaluationCache* InOutEvaluationCache, FFrameTime InTime);
 
 	/** Structure used to store the result of UE::MovieScene::EvaluateTime for a given channel/key distribution*/
 	struct FTimeEvaluationCache
@@ -156,8 +157,8 @@ private:
 	 */
 	static void RefineCurvePoints(const ChannelType* InChannel, FFrameRate TickResolution, double TimeThreshold, CurveValueType ValueThreshold, TArray<TTuple<double, double>>& InOutPoints);
 
-	static bool EvaluateLegacy(const ChannelType* InChannel, FFrameTime InTime, CurveValueType& OutValue);
-	static bool EvaluateCached(const ChannelType* InChannel, FFrameTime InTime, CurveValueType& OutValue);
+	static bool EvaluateLegacy(const ChannelType* InChannel, FTimeEvaluationCache* InOutEvaluationCache, FFrameTime InTime, CurveValueType& OutValue);
+	static bool EvaluateCached(const ChannelType* InChannel, FTimeEvaluationCache* InOutEvaluationCache, FFrameTime InTime, CurveValueType& OutValue);
 };
 
 template<typename ChannelType>
