Sequencer: Ignore task prerequisites when there is no threading model

The majority of this was already happening as a result of the branch inside the actual task, but there are several systems that manually set up prerequisites that can be ignored. This saves considerable time on the single threaded code path inside systems like the channel evaluators and blender systems.

diff --git a/Engine/Source/Runtime/MovieScene/Private/EntitySystem/MovieSceneEntityManager.cpp b/Engine/Source/Runtime/MovieScene/Private/EntitySystem/MovieSceneEntityManager.cpp
--- a/Engine/Source/Runtime/MovieScene/Private/EntitySystem/MovieSceneEntityManager.cpp
+++ b/Engine/Source/Runtime/MovieScene/Private/EntitySystem/MovieSceneEntityManager.cpp
@@ -374,6 +374,7 @@ FEntityManager::FEntityManager()
 	LockdownState = ELockdownState::Unlocked;
 	SystemSerialNumber = 1;
 	StructureMutationSystemSerialNumber = 0;
+	ThreadingModel = EEntityThreadingModel::NoThreading;
 }
 
 FEntityManager::~FEntityManager()
@@ -546,7 +547,7 @@ bool FEntityManager::IsHandleValid(FEntityHandle InEntityHandle) const
 	return Generation && *Generation == InEntityHandle.HandleGeneration;
 }
 
-EEntityThreadingModel FEntityManager::GetThreadingModel() const
+EEntityThreadingModel FEntityManager::ComputeThreadingModel() const
 {
 	const bool bCanThread = FPlatformProcess::SupportsMultithreading();
 
@@ -557,6 +558,16 @@ EEntityThreadingModel FEntityManager::GetThreadingModel() const
 	return bShouldThread ? EEntityThreadingModel::TaskGraph : EEntityThreadingModel::NoThreading;
 }
 
+EEntityThreadingModel FEntityManager::GetThreadingModel() const
+{
+	return ThreadingModel;
+}
+
+void FEntityManager::UpdateThreadingModel()
+{
+	ThreadingModel = ComputeThreadingModel();
+}
+
 const FComponentMask& FEntityManager::GetAccumulatedMask() const
 {
 	FEntityManager* This = const_cast<FEntityManager*>(this);
diff --git a/Engine/Source/Runtime/MovieScene/Private/EntitySystem/MovieSceneEntitySystemGraphs.cpp b/Engine/Source/Runtime/MovieScene/Private/EntitySystem/MovieSceneEntitySystemGraphs.cpp
--- a/Engine/Source/Runtime/MovieScene/Private/EntitySystem/MovieSceneEntitySystemGraphs.cpp
+++ b/Engine/Source/Runtime/MovieScene/Private/EntitySystem/MovieSceneEntitySystemGraphs.cpp
@@ -740,7 +740,11 @@ void FMovieSceneEntitySystemGraph::ExecutePhase(const ArrayType& SortedEntries,
 {
 	using namespace UE::MovieScene;
 
-	FSystemSubsequentTasks DownstreamTasks(this, &OutTasks);
+	Linker->EntityManager.UpdateThreadingModel();
+
+	const EEntityThreadingModel ThreadingModel = Linker->EntityManager.GetThreadingModel();
+
+	FSystemSubsequentTasks DownstreamTasks(this, &OutTasks, ThreadingModel);
 
 	FSystemTaskPrerequisites NoPrerequisites;
 
diff --git a/Engine/Source/Runtime/MovieScene/Private/EntitySystem/MovieSceneSystemTaskDependencies.cpp b/Engine/Source/Runtime/MovieScene/Private/EntitySystem/MovieSceneSystemTaskDependencies.cpp
--- a/Engine/Source/Runtime/MovieScene/Private/EntitySystem/MovieSceneSystemTaskDependencies.cpp
+++ b/Engine/Source/Runtime/MovieScene/Private/EntitySystem/MovieSceneSystemTaskDependencies.cpp
@@ -2,6 +2,7 @@
 
 #include "EntitySystem/MovieSceneSystemTaskDependencies.h"
 #include "EntitySystem/MovieSceneEntitySystemGraphs.h"
+#include "EntitySystem/MovieSceneEntityManager.h"
 #include "Algo/Find.h"
 
 
@@ -13,6 +14,14 @@ namespace MovieScene
 {
 
 
+bool GIgnoreDependenciesWhenNotThreading = true;
+FAutoConsoleVariableRef CVarIgnoreDependenciesWhenNotThreading(
+	TEXT("Sequencer.IgnoreDependenciesWhenNotThreading"),
+	GIgnoreDependenciesWhenNotThreading,
+	TEXT("(Default: true) Whether to ignore task dependencies when there is no threading.")
+);
+
+
 void FSystemTaskPrerequisites::FilterByComponent(FGraphEventArray& OutArray, std::initializer_list<FComponentTypeID> ComponentTypes) const
 {
 	for (const FPrerequisite& Prereq : Prereqs)
@@ -37,13 +46,20 @@ void FSystemTaskPrerequisites::Consume(const FSystemTaskPrerequisites& Other)
 
 
 
-FSystemSubsequentTasks::FSystemSubsequentTasks(FMovieSceneEntitySystemGraph* InGraph, FGraphEventArray* InAllTasks)
+FSystemSubsequentTasks::FSystemSubsequentTasks(FMovieSceneEntitySystemGraph* InGraph, FGraphEventArray* InAllTasks, EEntityThreadingModel InThreadingModel)
 	: Graph(InGraph)
 	, AllTasks(InAllTasks)
+	, ThreadingModel(InThreadingModel)
 {}
 
 void FSystemSubsequentTasks::ResetNode(uint16 InNodeID)
 {
+	if (ThreadingModel == EEntityThreadingModel::NoThreading && GIgnoreDependenciesWhenNotThreading)
+	{
+		return;
+	}
+
+
 	Subsequents = Graph->Nodes.Array[InNodeID].SubsequentTasks;
 	NodeID = InNodeID;
 
@@ -55,8 +71,12 @@ void FSystemSubsequentTasks::ResetNode(uint16 InNodeID)
 
 void FSystemSubsequentTasks::AddRootTask(FGraphEventRef RootTask)
 {
-	SCOPE_CYCLE_COUNTER(MovieSceneEval_SystemDependencyCost)
+	if (ThreadingModel == EEntityThreadingModel::NoThreading && GIgnoreDependenciesWhenNotThreading)
+	{
+		return;
+	}
 
+	SCOPE_CYCLE_COUNTER(MovieSceneEval_SystemDependencyCost)
 	if (RootTask)
 	{
 		if (!Subsequents)
@@ -71,7 +91,12 @@ void FSystemSubsequentTasks::AddRootTask(FGraphEventRef RootTask)
 }
 
 void FSystemSubsequentTasks::AddComponentTask(UE::MovieScene::FComponentTypeID ComponentType, FGraphEventRef ComponentTask)
-{	
+{
+	if (ThreadingModel == EEntityThreadingModel::NoThreading && GIgnoreDependenciesWhenNotThreading)
+	{
+		return;
+	}
+
 	SCOPE_CYCLE_COUNTER(MovieSceneEval_SystemDependencyCost)
 
 	if (ComponentTask)
diff --git a/Engine/Source/Runtime/MovieScene/Public/EntitySystem/MovieSceneEntityManager.h b/Engine/Source/Runtime/MovieScene/Public/EntitySystem/MovieSceneEntityManager.h
--- a/Engine/Source/Runtime/MovieScene/Public/EntitySystem/MovieSceneEntityManager.h
+++ b/Engine/Source/Runtime/MovieScene/Public/EntitySystem/MovieSceneEntityManager.h
@@ -217,7 +217,19 @@ public:
 
 
 	/**
-	 * Compute and retrieve this entity manager's threading model
+	 * Compute and return this entity manager's threading model. Does not change the current cached threading model.
+	 */
+	EEntityThreadingModel ComputeThreadingModel() const;
+
+
+	/**
+	 * Compute and store the current threading model.
+	 */
+	void UpdateThreadingModel();
+
+
+	/**
+	 * Get this entitiy manager's current threading model based on the last time UpdateThreadingModel was called.
 	 */
 	EEntityThreadingModel GetThreadingModel() const;
 
@@ -1018,6 +1030,7 @@ private:
 
 	ENamedThreads::Type GatherThread;
 	ENamedThreads::Type DispatchThread;
+	EEntityThreadingModel ThreadingModel;
 
 	enum class ELockdownState
 	{
diff --git a/Engine/Source/Runtime/MovieScene/Public/EntitySystem/MovieSceneSystemTaskDependencies.h b/Engine/Source/Runtime/MovieScene/Public/EntitySystem/MovieSceneSystemTaskDependencies.h
--- a/Engine/Source/Runtime/MovieScene/Public/EntitySystem/MovieSceneSystemTaskDependencies.h
+++ b/Engine/Source/Runtime/MovieScene/Public/EntitySystem/MovieSceneSystemTaskDependencies.h
@@ -23,6 +23,7 @@ namespace UE
 namespace MovieScene
 {
 
+enum class EEntityThreadingModel : uint8;
 
 struct FSystemTaskPrerequisites
 {
@@ -105,7 +106,7 @@ private:
 
 	friend FMovieSceneEntitySystemGraph;
 
-	FSystemSubsequentTasks(FMovieSceneEntitySystemGraph* InGraph, FGraphEventArray* InAllTasks);
+	FSystemSubsequentTasks(FMovieSceneEntitySystemGraph* InGraph, FGraphEventArray* InAllTasks, EEntityThreadingModel InThreadingModel);
 
 	void ResetNode(uint16 InNodeID);
 
@@ -113,6 +114,7 @@ private:
 	FMovieSceneEntitySystemGraph* Graph;
 	FGraphEventArray* AllTasks;
 	uint16 NodeID;
+	EEntityThreadingModel ThreadingModel;
 };
 
 
