Sequencer: Fixed track instance system not honoring the current threading model

diff --git a/Engine/Source/Runtime/MovieScene/Private/EntitySystem/TrackInstance/MovieSceneTrackInstanceSystem.cpp b/Engine/Source/Runtime/MovieScene/Private/EntitySystem/TrackInstance/MovieSceneTrackInstanceSystem.cpp
--- a/Engine/Source/Runtime/MovieScene/Private/EntitySystem/TrackInstance/MovieSceneTrackInstanceSystem.cpp
+++ b/Engine/Source/Runtime/MovieScene/Private/EntitySystem/TrackInstance/MovieSceneTrackInstanceSystem.cpp
@@ -271,6 +271,8 @@ void UMovieSceneTrackInstanceSystem::OnLink()
 
 void UMovieSceneTrackInstanceSystem::OnRun(FSystemTaskPrerequisites& InPrerequisites, FSystemSubsequentTasks& Subsequents)
 {
+	using namespace UE::MovieScene;
+
 	SCOPE_CYCLE_COUNTER(MovieSceneEval_GenericTrackInstances)
 
 	if (this->Instantiator->GetTrackInstances().Num() != 0)
@@ -286,7 +288,14 @@ void UMovieSceneTrackInstanceSystem::OnRun(FSystemTaskPrerequisites& InPrerequis
 			}
 		};
 
-		FGraphEventRef Task = FFunctionGraphTask::CreateAndDispatchWhenReady(MoveTemp(Run), GET_STATID(MovieSceneEval_GenericTrackInstanceTask), InPrerequisites.All(), Linker->EntityManager.GetGatherThread());
-		Subsequents.AddRootTask(Task);
+		if (Linker->EntityManager.GetThreadingModel() == EEntityThreadingModel::NoThreading)
+		{
+			Run();
+		}
+		else
+		{
+			FGraphEventRef Task = FFunctionGraphTask::CreateAndDispatchWhenReady(MoveTemp(Run), GET_STATID(MovieSceneEval_GenericTrackInstanceTask), InPrerequisites.All(), Linker->EntityManager.GetGatherThread());
+			Subsequents.AddRootTask(Task);
+		}
 	}
 }
