Sequencer: Fixed cycle extrapolation being mistaken as oscillation

diff --git a/Engine/Source/Runtime/MovieScene/Private/Channels/MovieSceneCurveChannelImpl.cpp b/Engine/Source/Runtime/MovieScene/Private/Channels/MovieSceneCurveChannelImpl.cpp
--- a/Engine/Source/Runtime/MovieScene/Private/Channels/MovieSceneCurveChannelImpl.cpp
+++ b/Engine/Source/Runtime/MovieScene/Private/Channels/MovieSceneCurveChannelImpl.cpp
@@ -93,12 +93,14 @@ struct FCycleParams
 	FFrameTime Time;
 	int32 CycleCount;
 	int32 Duration;
+	bool bMirrorCurve;
 
 	FCycleParams(FFrameTime InTime, int32 InDuration)
 		: ValueOffset(0.0)
 		, Time(InTime)
 		, CycleCount(0)
 		, Duration(InDuration)
+		, bMirrorCurve(false)
 	{}
 
 	FORCEINLINE void ComputePreValueOffset(double FirstValue, double LastValue)
@@ -110,14 +112,15 @@ struct FCycleParams
 	{
 		ValueOffset = (LastValue-FirstValue) * CycleCount;
 	}
-	FORCEINLINE bool IsOscillated() const
+	FORCEINLINE bool ShouldMirrorCurve() const
 	{
-		return FMath::Abs(CycleCount) % 2 == 1;
+		return bMirrorCurve;
 	}
 	FORCEINLINE void Oscillate(int32 MinFrame, int32 MaxFrame)
 	{
-		if (IsOscillated())
+		if (FMath::Abs(CycleCount) % 2 == 1)
 		{
+			bMirrorCurve = true;
 			Time = MinFrame + (FFrameTime(MaxFrame) - Time);
 		}
 	}
@@ -551,7 +554,7 @@ UE::MovieScene::Interpolation::FCachedInterpolation TMovieSceneCurveChannelImpl<
 
 		// Mirror the curve for oscillation by swapping the control points
 		// and mirroring the times based on the width.
-		if (Params.IsOscillated())
+		if (Params.ShouldMirrorCurve())
 		{
 			Swap(V1, V2);
 			Swap(T1, T2);
