Sequencer: new API for running "unstructured" tasks in ECS systems

diff --git a/Engine/Source/Runtime/MovieScene/Public/EntitySystem/MovieSceneEntitySystemTask.h b/Engine/Source/Runtime/MovieScene/Public/EntitySystem/MovieSceneEntitySystemTask.h
--- a/Engine/Source/Runtime/MovieScene/Public/EntitySystem/MovieSceneEntitySystemTask.h
+++ b/Engine/Source/Runtime/MovieScene/Public/EntitySystem/MovieSceneEntitySystemTask.h
@@ -2,6 +2,7 @@
 
 #pragma once
 
+#include "Async/TaskGraphInterfaces.h"
 #include "MovieSceneEntityIDs.h"
 #include "EntitySystem/MovieSceneEntityManager.h"
 #include "EntitySystem/EntityAllocationIterator.h"
@@ -38,6 +39,7 @@ template<typename, typename...> struct TEntityTask;
 template<typename, typename...> struct TEntityTaskBase;
 template<typename, typename...> struct TEntityAllocationTask;
 template<typename, typename...> struct TEntityAllocationTaskBase;
+template<typename> struct TUnstructuredTask;
 
 
 template<typename TaskImpl, int32 NumComponents, bool AutoExpandAccessors>
@@ -228,6 +230,48 @@ struct TEntityTaskComponents : TEntityTaskComponentsImpl<TMakeIntegerSequence<in
 		return *this;
 	}
 
+	/**
+	 * Dispatch a custom task that runs non-structured logic.
+	 * Tasks must implement a Run function that doesn't take any argument.
+	 *
+	 * @param EntityManager    The entity manager to run the task on
+	 * @param Prerequisites    Prerequisite tasks that must run before this one, or nullptr if there are no prerequisites
+	 * @param Subsequents      (Optional) Subsequent task tracking that this task should be added to for each writable component type
+	 * @param InArgs           Optional arguments that are forwarded to the constructor of TaskImpl
+	 * @return A pointer to the graph event for the task, or nullptr if this task is not valid (ie contains invalid component types that would be necessary for the task to run), or threading is disabled
+	 */
+	template<typename TaskImpl, typename... TaskConstructionArgs>
+	FGraphEventRef Dispatch(FEntityManager* EntityManager, const FSystemTaskPrerequisites& Prerequisites, FSystemSubsequentTasks* Subsequents, TaskConstructionArgs&&... InArgs) const
+	{
+		checkfSlow(IsInGameThread(), TEXT("Tasks can only be dispatched from the game thread."));
+
+		const bool bRunInline = !ensure(EntityManager->IsLockedDown()) || EntityManager->GetThreadingModel() == EEntityThreadingModel::NoThreading;
+		if (bRunInline)
+		{
+			TaskImpl Task{ Forward<TaskConstructionArgs>(InArgs)... };
+			Task.Run();
+			return nullptr;
+		}
+		else
+		{
+			FGraphEventArray GatheredPrereqs;
+			this->PopulatePrerequisites(Prerequisites, &GatheredPrereqs);
+
+			ENamedThreads::Type ThisThread = CurrentThread == ENamedThreads::AnyThread ? EntityManager->GetDispatchThread() : CurrentThread;
+			checkSlow(ThisThread != ENamedThreads::AnyThread);
+
+			FGraphEventRef NewTask = TGraphTask< TUnstructuredTask<TaskImpl> >::CreateTask(GatheredPrereqs.Num() != 0 ? &GatheredPrereqs : nullptr, ThisThread)
+				.ConstructAndDispatchWhenReady( DesiredThread, StatId, bBreakOnRun, Forward<TaskConstructionArgs>(InArgs)... );
+
+			if (Subsequents)
+			{
+				this->PopulateSubsequents(NewTask, *Subsequents);
+			}
+
+			return NewTask;
+		}
+	}
+
 	/**
 	 * Dispatch a task for every allocation that matches the filters and component types. Must be explicitly instantiated with the task type to dispatch. Construction arguments are deduced.
 	 * Tasks must implement a ForEachAllocation function that matches this task's component accessor types.
@@ -358,6 +402,7 @@ struct TEntityTaskComponents : TEntityTaskComponentsImpl<TMakeIntegerSequence<in
 			if (Subsequents)
 			{
 				this->PopulateSubsequents(NewTask, *Subsequents);
+				Subsequents->AddRootTask(NewTask);
 			}
 
 			return NewTask;
@@ -896,7 +941,6 @@ struct TFilteredEntityTask
 		return Components;
 	}
 
-
 	/**
 	 * Dispatch a task for every entity that matches the filters and component types. Must be explicitly instantiated with the task type to dispatch. Construction arguments are deduced.
 	 * Tasks must implement a ForEachEntity function that matches this task's component accessor types.
@@ -1308,7 +1352,59 @@ private:
 	bool bBreakOnRun;
 };
 
+template<typename TaskImpl>
+struct TUnstructuredTask
+{
+	template<typename... ArgTypes>
+	explicit TUnstructuredTask(ENamedThreads::Type InDesiredThread, TStatId InStatId, bool bInBreakOnRun, ArgTypes&&... InArgs)
+		: TaskImplInstance{ Forward<ArgTypes>(InArgs)... }
+		, DesiredThread(InDesiredThread)
+		, StatId(InStatId)
+		, bBreakOnRun(bInBreakOnRun)
+	{}
+
+	TStatId GetStatId() const
+	{
+		return StatId;
+	}
+
+	ENamedThreads::Type GetDesiredThread()
+	{
+		return DesiredThread;
+	}
 
+	static ESubsequentsMode::Type GetSubsequentsMode()
+	{
+		return ESubsequentsMode::TrackSubsequents;
+	}
+
+	void DoTask(ENamedThreads::Type CurrentThread, const FGraphEventRef& CompletionGraphEvent)
+	{
+		if (bBreakOnRun)
+		{
+			UE_DEBUG_BREAK();
+		}
+
+		if ((DesiredThread & ENamedThreads::AnyThread) == 0)
+		{
+			checkf(CurrentThread == DesiredThread, TEXT("MovieScene evaluation task is not being run on its desired thread"));
+		}
+
+		UE_LOG(LogMovieSceneECS, VeryVerbose, TEXT("Running unstructured task"));
+
+		TaskImplInstance.Run();
+	}
+
+private:
+
+	TaskImpl TaskImplInstance;
+
+	ENamedThreads::Type DesiredThread;
+
+	TStatId StatId;
+
+	bool bBreakOnRun;
+};
 
 template<typename...> struct TEntityTaskCaller_AutoExpansion;
 
@@ -1360,4 +1456,4 @@ struct TEntityTaskCaller<TaskImpl, NumComponents, false>
 
 
 } // namespace MovieScene
-} // namespace UE
\ No newline at end of file
+} // namespace UE
