Sequencer: fix weight channel component dependency

The cached interpolation component is needed only if we have a weight channel component. If someone produces weight via other means than a curve channel, we don't need the cached interpolation.

diff --git a/Engine/Source/Runtime/MovieScene/Private/EntitySystem/BuiltInComponentTypes.cpp b/Engine/Source/Runtime/MovieScene/Private/EntitySystem/BuiltInComponentTypes.cpp
--- a/Engine/Source/Runtime/MovieScene/Private/EntitySystem/BuiltInComponentTypes.cpp
+++ b/Engine/Source/Runtime/MovieScene/Private/EntitySystem/BuiltInComponentTypes.cpp
@@ -293,7 +293,7 @@ FBuiltInComponentTypes::FBuiltInComponentTypes()
 		// Weight channel components need a time and result to evaluate
 		ComponentRegistry->Factories.DefineMutuallyInclusiveComponent(WeightChannel, EvalTime);
 		ComponentRegistry->Factories.DefineMutuallyInclusiveComponent(WeightChannel, WeightResult);
-		ComponentRegistry->Factories.DefineMutuallyInclusiveComponent(WeightResult, CachedWeightChannelInterpolation);
+		ComponentRegistry->Factories.DefineMutuallyInclusiveComponent(WeightChannel, CachedWeightChannelInterpolation);
 	}
 
 	// Weight and easing result component relationship
