Sequencer: Compile out Versbose logging in shipping and test builds

diff --git a/Engine/Source/Runtime/MovieScene/Public/MovieSceneFwd.h b/Engine/Source/Runtime/MovieScene/Public/MovieSceneFwd.h
--- a/Engine/Source/Runtime/MovieScene/Public/MovieSceneFwd.h
+++ b/Engine/Source/Runtime/MovieScene/Public/MovieSceneFwd.h
@@ -123,8 +123,13 @@ enum class EMovieSceneServerClientMask : uint8
 };
 ENUM_CLASS_FLAGS(EMovieSceneServerClientMask)
 
-MOVIESCENE_API DECLARE_LOG_CATEGORY_EXTERN(LogMovieScene, Log, All);
-MOVIESCENE_API DECLARE_LOG_CATEGORY_EXTERN(LogMovieSceneECS, Log, All);
+#if UE_BUILD_SHIPPING || UE_BUILD_TEST
+	MOVIESCENE_API DECLARE_LOG_CATEGORY_EXTERN(LogMovieScene, Log, Warning);
+	MOVIESCENE_API DECLARE_LOG_CATEGORY_EXTERN(LogMovieSceneECS, Log, Warning);
+#else
+	MOVIESCENE_API DECLARE_LOG_CATEGORY_EXTERN(LogMovieScene, Log, All);
+	MOVIESCENE_API DECLARE_LOG_CATEGORY_EXTERN(LogMovieSceneECS, Log, All);
+#endif
 DECLARE_STATS_GROUP(TEXT("Movie Scene Evaluation"), STATGROUP_MovieSceneEval, STATCAT_Advanced);
 
 MOVIESCENE_API FFrameRate GetLegacyConversionFrameRate();
