Sequencer: Disable mutex locking when not using a threading model

diff --git a/Engine/Source/Runtime/MovieScene/Private/EntitySystem/MovieSceneEntityManager.cpp b/Engine/Source/Runtime/MovieScene/Private/EntitySystem/MovieSceneEntityManager.cpp
--- a/Engine/Source/Runtime/MovieScene/Private/EntitySystem/MovieSceneEntityManager.cpp
+++ b/Engine/Source/Runtime/MovieScene/Private/EntitySystem/MovieSceneEntityManager.cpp
@@ -1669,6 +1669,8 @@ int32 FEntityManager::MutateAll(const FEntityComponentFilter& Filter, const IMov
 		EntityAllocationMasks[Pair.Key] = Pair.Value;
 		EntityAllocations[Pair.Key] = NewAllocation;
 
+		FEntityAllocationMutexGuard LockGuard(NewAllocation, EComponentHeaderLockMode::LockFree);
+
 		// Default construct all the new components in the allocation, and then allow the mutation to further initialize this data if needed
 		for (FComponentMaskIterator Component = Pair.Value.Iterate(); Component; ++Component)
 		{
diff --git a/Engine/Source/Runtime/MovieScene/Private/EntitySystem/MovieSceneEntitySystemTypes.cpp b/Engine/Source/Runtime/MovieScene/Private/EntitySystem/MovieSceneEntitySystemTypes.cpp
--- a/Engine/Source/Runtime/MovieScene/Private/EntitySystem/MovieSceneEntitySystemTypes.cpp
+++ b/Engine/Source/Runtime/MovieScene/Private/EntitySystem/MovieSceneEntitySystemTypes.cpp
@@ -111,20 +111,20 @@ FEntityAllocationWriteContext::FEntityAllocationWriteContext(const FEntityManage
 FComponentReader FEntityAllocation::ReadComponentsErased(FComponentTypeID ComponentType) const
 {
 	const FComponentHeader& Header = GetComponentHeaderChecked(ComponentType);
-	return FComponentReader(&Header);
+	return FComponentReader(&Header, LockMode);
 }
 
 FComponentWriter FEntityAllocation::WriteComponentsErased(FComponentTypeID ComponentType, FEntityAllocationWriteContext InWriteContext) const
 {
 	const FComponentHeader& Header = GetComponentHeaderChecked(ComponentType);
-	return FComponentWriter(&Header, InWriteContext);
+	return FComponentWriter(&Header, LockMode, InWriteContext);
 }
 
 FOptionalComponentReader FEntityAllocation::TryReadComponentsErased(FComponentTypeID ComponentType) const
 {
 	if (const FComponentHeader* Header = FindComponentHeader(ComponentType))
 	{
-		return FOptionalComponentReader(Header);
+		return FOptionalComponentReader(Header, LockMode);
 	}
 
 	return FOptionalComponentReader();
@@ -134,7 +134,7 @@ FOptionalComponentWriter FEntityAllocation::TryWriteComponentsErased(FComponentT
 {
 	if (const FComponentHeader* Header = FindComponentHeader(ComponentType))
 	{
-		return FOptionalComponentWriter(Header, InWriteContext);
+		return FOptionalComponentWriter(Header, LockMode, InWriteContext);
 	}
 
 	return FOptionalComponentWriter();
diff --git a/Engine/Source/Runtime/MovieScene/Public/EntitySystem/MovieSceneComponentPtr.h b/Engine/Source/Runtime/MovieScene/Public/EntitySystem/MovieSceneComponentPtr.h
--- a/Engine/Source/Runtime/MovieScene/Public/EntitySystem/MovieSceneComponentPtr.h
+++ b/Engine/Source/Runtime/MovieScene/Public/EntitySystem/MovieSceneComponentPtr.h
@@ -105,8 +105,8 @@ template<>
 struct TComponentLockMixin<FScopedHeaderReadLock>
 {
 	TComponentLockMixin() = default;
-	explicit TComponentLockMixin(const FComponentHeader* InHeader)
-		: Lock(InHeader)
+	explicit TComponentLockMixin(const FComponentHeader* InHeader, EComponentHeaderLockMode InLockMode)
+		: Lock(InHeader, InLockMode)
 	{}
 
 private:
@@ -117,8 +117,8 @@ template<>
 struct TComponentLockMixin<FScopedHeaderWriteLock>
 {
 	TComponentLockMixin() = default;
-	explicit TComponentLockMixin(const FComponentHeader* InHeader, FEntityAllocationWriteContext InWriteContext)
-		: Lock(InHeader, InWriteContext)
+	explicit TComponentLockMixin(const FComponentHeader* InHeader, EComponentHeaderLockMode InLockMode, FEntityAllocationWriteContext InWriteContext)
+		: Lock(InHeader, InLockMode, InWriteContext)
 	{}
 
 private:
@@ -325,8 +325,8 @@ template<typename AccessorType> struct TComponentLock;
 template<>
 struct TComponentLock<FReadErased> : TComponentLockMixin<FScopedHeaderReadLock>, FReadErased
 {
-	explicit TComponentLock(const FComponentHeader* InHeader, int32 ComponentOffset = 0)
-		: TComponentLockMixin<FScopedHeaderReadLock>(InHeader)
+	explicit TComponentLock(const FComponentHeader* InHeader, EComponentHeaderLockMode InLockMode, int32 ComponentOffset = 0)
+		: TComponentLockMixin<FScopedHeaderReadLock>(InHeader, InLockMode)
 		, FReadErased(InHeader, ComponentOffset)
 	{}
 
@@ -339,11 +339,11 @@ template<>
 struct TComponentLock<FReadErasedOptional> : TComponentLockMixin<FScopedHeaderReadLock>, FReadErasedOptional
 {
 	TComponentLock() = default;
-	explicit TComponentLock(const FComponentHeader* InHeader, int32 ComponentOffset = 0)
+	explicit TComponentLock(const FComponentHeader* InHeader, EComponentHeaderLockMode InLockMode, int32 ComponentOffset = 0)
 	{
 		if (InHeader)
 		{
-			*static_cast<TComponentLockMixin<FScopedHeaderReadLock>*>(this) = TComponentLockMixin<FScopedHeaderReadLock>(InHeader);
+			*static_cast<TComponentLockMixin<FScopedHeaderReadLock>*>(this) = TComponentLockMixin<FScopedHeaderReadLock>(InHeader, InLockMode);
 			*static_cast<FReadErasedOptional*>(this) = FReadErasedOptional(InHeader, ComponentOffset);
 		}
 	}
@@ -361,8 +361,8 @@ struct TComponentLock<FReadErasedOptional> : TComponentLockMixin<FScopedHeaderRe
 template<>
 struct TComponentLock<FWriteErased> : TComponentLockMixin<FScopedHeaderWriteLock>, FWriteErased
 {
-	explicit TComponentLock(const FComponentHeader* InHeader, FEntityAllocationWriteContext InWriteContext, int32 ComponentOffset = 0)
-		: TComponentLockMixin<FScopedHeaderWriteLock>(InHeader, InWriteContext)
+	explicit TComponentLock(const FComponentHeader* InHeader, EComponentHeaderLockMode InLockMode, FEntityAllocationWriteContext InWriteContext, int32 ComponentOffset = 0)
+		: TComponentLockMixin<FScopedHeaderWriteLock>(InHeader, InLockMode, InWriteContext)
 		, FWriteErased(InHeader, ComponentOffset)
 	{}
 
@@ -375,11 +375,11 @@ template<>
 struct TComponentLock<FWriteErasedOptional> : TComponentLockMixin<FScopedHeaderWriteLock>, FWriteErasedOptional
 {
 	TComponentLock() = default;
-	explicit TComponentLock(const FComponentHeader* InHeader, FEntityAllocationWriteContext InWriteContext, int32 ComponentOffset = 0)
+	explicit TComponentLock(const FComponentHeader* InHeader, EComponentHeaderLockMode InLockMode, FEntityAllocationWriteContext InWriteContext, int32 ComponentOffset = 0)
 	{
 		if (InHeader)
 		{
-			*static_cast<TComponentLockMixin<FScopedHeaderWriteLock>*>(this) = TComponentLockMixin<FScopedHeaderWriteLock>(InHeader, InWriteContext);
+			*static_cast<TComponentLockMixin<FScopedHeaderWriteLock>*>(this) = TComponentLockMixin<FScopedHeaderWriteLock>(InHeader, InLockMode, InWriteContext);
 			*static_cast<FWriteErasedOptional*>(this) = FWriteErasedOptional(InHeader, ComponentOffset);
 		}
 	}
@@ -397,8 +397,8 @@ struct TComponentLock<FWriteErasedOptional> : TComponentLockMixin<FScopedHeaderW
 template<typename T>
 struct TComponentLock<TRead<T>> : TComponentLockMixin<FScopedHeaderReadLock>, TRead<T>
 {
-	explicit TComponentLock(const FComponentHeader* InHeader, int32 ComponentOffset = 0)
-		: TComponentLockMixin<FScopedHeaderReadLock>(InHeader)
+	explicit TComponentLock(const FComponentHeader* InHeader, EComponentHeaderLockMode InLockMode , int32 ComponentOffset = 0)
+		: TComponentLockMixin<FScopedHeaderReadLock>(InHeader, InLockMode)
 		, TRead<T>(InHeader, ComponentOffset)
 	{}
 };
@@ -407,11 +407,11 @@ template<typename T>
 struct TComponentLock<TReadOptional<T>> : TComponentLockMixin<FScopedHeaderReadLock>, TReadOptional<T>
 {
 	TComponentLock() = default;
-	explicit TComponentLock(const FComponentHeader* InHeader, int32 ComponentOffset = 0)
+	explicit TComponentLock(const FComponentHeader* InHeader, EComponentHeaderLockMode InLockMode, int32 ComponentOffset = 0)
 	{
 		if (InHeader)
 		{
-			*static_cast<TComponentLockMixin<FScopedHeaderReadLock>*>(this) = TComponentLockMixin<FScopedHeaderReadLock>(InHeader);
+			*static_cast<TComponentLockMixin<FScopedHeaderReadLock>*>(this) = TComponentLockMixin<FScopedHeaderReadLock>(InHeader, InLockMode);
 			*static_cast<TReadOptional<T>*>(this) = TReadOptional<T>(InHeader, ComponentOffset);
 		}
 	}
@@ -429,8 +429,8 @@ struct TComponentLock<TReadOptional<T>> : TComponentLockMixin<FScopedHeaderReadL
 template<typename T>
 struct TComponentLock<TWrite<T>> : TComponentLockMixin<FScopedHeaderWriteLock>, TWrite<T>
 {
-	explicit TComponentLock(const FComponentHeader* InHeader, FEntityAllocationWriteContext InWriteContext, int32 ComponentOffset = 0)
-		: TComponentLockMixin<FScopedHeaderWriteLock>(InHeader, InWriteContext)
+	explicit TComponentLock(const FComponentHeader* InHeader, EComponentHeaderLockMode InLockMode, FEntityAllocationWriteContext InWriteContext, int32 ComponentOffset = 0)
+		: TComponentLockMixin<FScopedHeaderWriteLock>(InHeader, InLockMode, InWriteContext)
 		, TWrite<T>(InHeader, ComponentOffset)
 	{}
 
@@ -444,11 +444,11 @@ template<typename T>
 struct TComponentLock<TWriteOptional<T>> : TComponentLockMixin<FScopedHeaderWriteLock>, TWriteOptional<T>
 {
 	TComponentLock() = default;
-	explicit TComponentLock(const FComponentHeader* InHeader, FEntityAllocationWriteContext InWriteContext, int32 ComponentOffset = 0)
+	explicit TComponentLock(const FComponentHeader* InHeader, EComponentHeaderLockMode InLockMode, FEntityAllocationWriteContext InWriteContext, int32 ComponentOffset = 0)
 	{
 		if (InHeader)
 		{
-			*static_cast<TComponentLockMixin<FScopedHeaderWriteLock>*>(this) = TComponentLockMixin<FScopedHeaderWriteLock>(InHeader, InWriteContext);
+			*static_cast<TComponentLockMixin<FScopedHeaderWriteLock>*>(this) = TComponentLockMixin<FScopedHeaderWriteLock>(InHeader, InLockMode, InWriteContext);
 			*static_cast<TWriteOptional<T>*>(this) = TWriteOptional<T>(InHeader, ComponentOffset);
 		}
 	}
diff --git a/Engine/Source/Runtime/MovieScene/Public/EntitySystem/MovieSceneEntityBuilder.h b/Engine/Source/Runtime/MovieScene/Public/EntitySystem/MovieSceneEntityBuilder.h
--- a/Engine/Source/Runtime/MovieScene/Public/EntitySystem/MovieSceneEntityBuilder.h
+++ b/Engine/Source/Runtime/MovieScene/Public/EntitySystem/MovieSceneEntityBuilder.h
@@ -482,11 +482,20 @@ struct TAdd : FAdd
 	{
 		const FComponentHeader& Header = Allocation->GetComponentHeaderChecked(ComponentTypeID);
 		check(!Header.IsTag());
-		Header.ReadWriteLock.WriteLock();
+
+		if (Allocation->GetCurrentLockMode() != EComponentHeaderLockMode::LockFree)
+		{
+			Header.ReadWriteLock.WriteLock();
+		}
+
 		T* ComponentPtr = static_cast<T*>(Header.GetValuePtr(ComponentOffset));
 		*ComponentPtr = MoveTemp(Payload.GetValue());
 		Payload.Reset();
-		Header.ReadWriteLock.WriteUnlock();
+
+		if (Allocation->GetCurrentLockMode() != EComponentHeaderLockMode::LockFree)
+		{
+			Header.ReadWriteLock.WriteUnlock();
+		}
 	}
 };
 
@@ -508,11 +517,20 @@ struct TAddConditional : FAddConditional
 		{
 			const FComponentHeader& Header = Allocation->GetComponentHeaderChecked(ComponentTypeID);
 			check(!Header.IsTag());
-			Header.ReadWriteLock.WriteLock();
+
+			if (Allocation->GetCurrentLockMode() != EComponentHeaderLockMode::LockFree)
+			{
+				Header.ReadWriteLock.WriteLock();
+			}
+
 			T* ComponentPtr = static_cast<T*>(Header.GetValuePtr(ComponentOffset));
 			*ComponentPtr = MoveTemp(Payload.GetValue());
 			Payload.Reset();
-			Header.ReadWriteLock.WriteUnlock();
+
+			if (Allocation->GetCurrentLockMode() != EComponentHeaderLockMode::LockFree)
+			{
+				Header.ReadWriteLock.WriteUnlock();
+			}
 		}
 	}
 };
diff --git a/Engine/Source/Runtime/MovieScene/Public/EntitySystem/MovieSceneEntityManager.h b/Engine/Source/Runtime/MovieScene/Public/EntitySystem/MovieSceneEntityManager.h
--- a/Engine/Source/Runtime/MovieScene/Public/EntitySystem/MovieSceneEntityManager.h
+++ b/Engine/Source/Runtime/MovieScene/Public/EntitySystem/MovieSceneEntityManager.h
@@ -252,12 +252,18 @@ public:
 		if (Entry.Data.Allocation != nullptr)
 		{
 			const FComponentHeader& Header = Entry.Data.Allocation->GetComponentHeaderChecked(ComponentTypeID);
-			Header.ReadWriteLock.WriteLock();
+			if (Entry.Data.Allocation->GetCurrentLockMode() != EComponentHeaderLockMode::LockFree)
+			{
+				Header.ReadWriteLock.WriteLock();
+			}
 
 			T* Component = reinterpret_cast<T*>(Header.GetValuePtr(Entry.Data.ComponentOffset));
 			*Component = Forward<ValueType>(InValue);
 
-			Header.ReadWriteLock.WriteUnlock();
+			if (Entry.Data.Allocation->GetCurrentLockMode() != EComponentHeaderLockMode::LockFree)
+			{
+				Header.ReadWriteLock.WriteUnlock();
+			}
 		}
 	}
 
@@ -456,11 +462,15 @@ public:
 		FEntityAllocation* Allocation = EntityAllocations[Location.GetAllocationIndex()];
 		const int32 ComponentOffset = Location.GetEntryIndexWithinAllocation();
 
+		EComponentHeaderLockMode LockMode = GetThreadingModel() == EEntityThreadingModel::NoThreading
+			? EComponentHeaderLockMode::LockFree
+			: EComponentHeaderLockMode::Mutex;
+
 		for (FComponentHeader& Header : Allocation->GetComponentHeaders())
 		{
 			if (Header.ComponentType == ComponentTypeID)
 			{
-				return TComponentLock<TReadOptional<T>>(&Header, ComponentOffset);
+				return TComponentLock<TReadOptional<T>>(&Header, LockMode, ComponentOffset);
 			}
 		}
 
@@ -512,11 +522,15 @@ public:
 		FEntityAllocation* Allocation = EntityAllocations[Location.GetAllocationIndex()];
 		const int32 ComponentOffset = Location.GetEntryIndexWithinAllocation();
 
+		EComponentHeaderLockMode LockMode = GetThreadingModel() == EEntityThreadingModel::NoThreading
+			? EComponentHeaderLockMode::LockFree
+			: EComponentHeaderLockMode::Mutex;
+
 		for (FComponentHeader& Header : Allocation->GetComponentHeaders())
 		{
 			if (Header.ComponentType == ComponentTypeID)
 			{
-				return TComponentLock<TWriteOptional<T>>(&Header, FEntityAllocationWriteContext(*this), ComponentOffset);
+				return TComponentLock<TWriteOptional<T>>(&Header, LockMode, FEntityAllocationWriteContext(*this), ComponentOffset);
 			}
 		}
 
diff --git a/Engine/Source/Runtime/MovieScene/Public/EntitySystem/MovieSceneEntitySystemTask.h b/Engine/Source/Runtime/MovieScene/Public/EntitySystem/MovieSceneEntitySystemTask.h
--- a/Engine/Source/Runtime/MovieScene/Public/EntitySystem/MovieSceneEntitySystemTask.h
+++ b/Engine/Source/Runtime/MovieScene/Public/EntitySystem/MovieSceneEntitySystemTask.h
@@ -620,7 +620,7 @@ struct TEntityTaskComponentsImpl<TIntegerSequence<int, Indices...>, T...>
 	}
 
 	/**
-	 * Perform a thread-safe iteration of the specified allocation using this task, inline on the current thread
+	 * Perform a direct iteration of the specified allocation using this task, inline on the current thread
 	 * @note: This is highly unsafe as it circumvents all the thread-safety mechanisms that protect component data
 	 *
 	 * @param Allocation  The allocation to iterate
@@ -637,7 +637,7 @@ struct TEntityTaskComponentsImpl<TIntegerSequence<int, Indices...>, T...>
 	}
 
 	/**
-	 * Perform a thread-safe iteration of the specified entity range using this task, inline on the current thread
+	 * Perform a direct iteration of the specified entity range using this task, inline on the current thread
 	 * @note: This is highly unsafe as it circumvents all the thread-safety mechanisms that protect component data
 	 *
 	 * @param Allocation  The allocation to iterate
@@ -696,10 +696,16 @@ struct TEntityTaskComponentsImpl<TIntegerSequence<int, Indices...>, T...>
 		{
 			FEntityAllocationWriteContext WriteContext(*EntityManager);
 
+			EComponentHeaderLockMode LockMode = EntityManager->GetThreadingModel() == EEntityThreadingModel::NoThreading
+				? EComponentHeaderLockMode::LockFree
+				: EComponentHeaderLockMode::Mutex;
+
 			for (FEntityAllocation* Allocation : EntityManager->Iterate(&Filter))
 			{
 				FEntityIterationResult Result;
 
+				FEntityAllocationMutexGuard LockGuard(Allocation, LockMode);
+
 				// Lock the components we want to access
 				TupleType ComponentData( Accessors.template Get<Indices>().LockComponentData(Allocation, WriteContext)... );
 
@@ -725,9 +731,15 @@ struct TEntityTaskComponentsImpl<TIntegerSequence<int, Indices...>, T...>
 		if (IsValid())
 		{
 			FEntityAllocationWriteContext WriteContext(*EntityManager);
+
+			EComponentHeaderLockMode LockMode = EntityManager->GetThreadingModel() == EEntityThreadingModel::NoThreading
+				? EComponentHeaderLockMode::LockFree
+				: EComponentHeaderLockMode::Mutex;
+
 			for (FEntityAllocationIteratorItem Item : EntityManager->Iterate(&Filter))
 			{
 				FEntityAllocation* Allocation = Item;
+				FEntityAllocationMutexGuard LockGuard(Item.GetAllocation(), LockMode);
 
 				// Lock on the components we want to access
 				auto ComponentData = MakeTuple( Accessors.template Get<Indices>().LockComponentData(Allocation, WriteContext)... );
@@ -1145,8 +1157,13 @@ struct TEntityTaskBase
 
 		PreTask(&TaskImplInstance);
 
+		EComponentHeaderLockMode LockMode = EntityManager->GetThreadingModel() == EEntityThreadingModel::NoThreading
+			? EComponentHeaderLockMode::LockFree
+			: EComponentHeaderLockMode::Mutex;
+
 		for (FEntityAllocation* Allocation : EntityManager->Iterate(&FilteredTask.GetFilter()))
 		{
+			FEntityAllocationMutexGuard LockGuard(Allocation, LockMode);
 			Caller::ForEachEntityImpl(TaskImplInstance, Allocation, WriteContext, FilteredTask.GetComponents());
 		}
 
@@ -1259,8 +1276,13 @@ struct TEntityAllocationTaskBase
 
 		PreTask(&TaskImplInstance);
 
+		EComponentHeaderLockMode LockMode = EntityManager->GetThreadingModel() == EEntityThreadingModel::NoThreading
+			? EComponentHeaderLockMode::LockFree
+			: EComponentHeaderLockMode::Mutex;
+
 		for (FEntityAllocationIteratorItem Item : EntityManager->Iterate(&ComponentFilter.GetFilter()))
 		{
+			FEntityAllocationMutexGuard LockGuard(Item.GetAllocation(), LockMode);
 			Caller::ForEachAllocationImpl(TaskImplInstance, Item, WriteContext, ComponentFilter.GetComponents());
 		}
 
diff --git a/Engine/Source/Runtime/MovieScene/Public/EntitySystem/MovieSceneEntitySystemTypes.h b/Engine/Source/Runtime/MovieScene/Public/EntitySystem/MovieSceneEntitySystemTypes.h
--- a/Engine/Source/Runtime/MovieScene/Public/EntitySystem/MovieSceneEntitySystemTypes.h
+++ b/Engine/Source/Runtime/MovieScene/Public/EntitySystem/MovieSceneEntitySystemTypes.h
@@ -92,6 +92,14 @@ enum class EComplexFilterMode : uint8
 ENUM_CLASS_FLAGS(EComplexFilterMode);
 
 
+/**
+ * Enumeration specifying the locking mechanism to use when accessing component data
+ */
+enum class EComponentHeaderLockMode
+{
+	LockFree,
+	Mutex,
+};
 
 /**
  * A numeric identifier used to represent a specific 'channel' within an interrogation linker.
@@ -473,7 +481,7 @@ public:
 struct FScopedHeaderReadLock
 {
 	FScopedHeaderReadLock();
-	FScopedHeaderReadLock(const FComponentHeader* InHeader);
+	FScopedHeaderReadLock(const FComponentHeader* InHeader, EComponentHeaderLockMode InLockMode);
 
 	FScopedHeaderReadLock(const FScopedHeaderReadLock& RHS) = delete;
 	void operator=(const FScopedHeaderReadLock& RHS) = delete;
@@ -494,7 +502,7 @@ private:
 struct FScopedHeaderWriteLock
 {
 	FScopedHeaderWriteLock();
-	FScopedHeaderWriteLock(const FComponentHeader* InHeader, FEntityAllocationWriteContext InWriteContext);
+	FScopedHeaderWriteLock(const FComponentHeader* InHeader, EComponentHeaderLockMode InLockMode, FEntityAllocationWriteContext InWriteContext);
 
 	FScopedHeaderWriteLock(const FScopedHeaderWriteLock& RHS) = delete;
 	void operator=(const FScopedHeaderWriteLock& RHS) = delete;
@@ -507,6 +515,7 @@ struct FScopedHeaderWriteLock
 private:
 	const FComponentHeader* Header;
 	FEntityAllocationWriteContext WriteContext;
+	EComponentHeaderLockMode LockMode;
 };
 
 
@@ -539,6 +548,7 @@ struct FEntityAllocation
 		, Size(0)
 		, Capacity(0)
 		, MaxCapacity(0)
+		, LockMode(EComponentHeaderLockMode::Mutex)
 	{}
 
 	/**
@@ -739,6 +749,14 @@ struct FEntityAllocation
 		return int32(Capacity) - int32(Size);
 	}
 
+	/**
+	 * Get this allocation's current lock mode
+	 */
+	EComponentHeaderLockMode GetCurrentLockMode() const
+	{
+		return LockMode;
+	}
+
 	/**
 	 * Read type-erased component data for the specified component type
 	 */
@@ -766,7 +784,7 @@ struct FEntityAllocation
 	UE_NODISCARD TComponentLock<TRead<T>> ReadComponents(TComponentTypeID<T> ComponentType) const
 	{
 		const FComponentHeader& Header = GetComponentHeaderChecked(ComponentType);
-		return TComponentLock<TRead<T>>(&Header);
+		return TComponentLock<TRead<T>>(&Header, LockMode);
 	}
 
 	/**
@@ -777,7 +795,7 @@ struct FEntityAllocation
 	{
 		if (const FComponentHeader* Header = FindComponentHeader(ComponentType))
 		{
-			return TComponentLock<TReadOptional<T>>(Header);
+			return TComponentLock<TReadOptional<T>>(Header, LockMode);
 		}
 		return TComponentLock<TReadOptional<T>>();
 	}
@@ -786,21 +804,21 @@ struct FEntityAllocation
 	 * Write typed component data for the specified component type
 	 */
 	template<typename T>
-	UE_NODISCARD TComponentLock<TWrite<T>> WriteComponents(TComponentTypeID<T> ComponentType, FEntityAllocationWriteContext InWriteContext) const
+	UE_NODISCARD TComponentLock<TWrite<T>> WriteComponents(TComponentTypeID<T> ComponentType, FEntityAllocationWriteContext InWriteContext, EComponentHeaderLockMode LockMode = EComponentHeaderLockMode::Mutex) const
 	{
 		const FComponentHeader& Header = GetComponentHeaderChecked(ComponentType);
-		return TComponentLock<TWrite<T>>(&Header, InWriteContext);
+		return TComponentLock<TWrite<T>>(&Header, LockMode, InWriteContext);
 	}
 
 	/**
 	 * Attempt to write typed component data for the specified component type
 	 */
 	template<typename T>
-	UE_NODISCARD TComponentLock<TWriteOptional<T>> TryWriteComponents(TComponentTypeID<T> ComponentType, FEntityAllocationWriteContext InWriteContext) const
+	UE_NODISCARD TComponentLock<TWriteOptional<T>> TryWriteComponents(TComponentTypeID<T> ComponentType, FEntityAllocationWriteContext InWriteContext, EComponentHeaderLockMode LockMode = EComponentHeaderLockMode::Mutex) const
 	{
 		if (const FComponentHeader* Header = FindComponentHeader(ComponentType))
 		{
-			return TComponentLock<TWriteOptional<T>>(Header, InWriteContext);
+			return TComponentLock<TWriteOptional<T>>(Header, LockMode, InWriteContext);
 		}
 		return TComponentLock<TWriteOptional<T>>();
 	}
@@ -808,6 +826,7 @@ struct FEntityAllocation
 private:
 
 	friend struct FEntityInitializer;
+	friend struct FEntityAllocationMutexGuard;
 
 	/** Assigned to FEntityManager::GetSystemSerial whenever this allocation is written to */
 	uint64 SerialNumber;
@@ -822,6 +841,9 @@ private:
 	/** The maximum number of entities that this entity is allowed to reallocate to accomodate for. */
 	uint16 MaxCapacity;
 
+	/** Lock mode for access - under threading models this will be set to EComponentHeaderLockMode::Mutex, EComponentHeaderLockMode::LockFree otherwise */
+	EComponentHeaderLockMode LockMode;
+
 	/** Pointer to the entity ID array (stored in the end padding of this structure). */
 	FMovieSceneEntityID* EntityIDs;
 
@@ -834,6 +856,26 @@ public:
 };
 
 
+/**
+ * Scoped guard that temporarily overrides the locking mechanism for a specific allocation.
+ * In order to guarantee thread-safety, this structure only actually does anything if InLockMode is specified as LockFree
+ */
+struct FEntityAllocationMutexGuard
+{
+	FEntityAllocationMutexGuard(FEntityAllocation* InAllocation, EComponentHeaderLockMode InLockMode);
+	~FEntityAllocationMutexGuard();
+
+	FEntityAllocationMutexGuard(const FEntityAllocationMutexGuard&) = delete;
+	void operator=(const FEntityAllocationMutexGuard&) = delete;
+
+	FEntityAllocationMutexGuard(const FEntityAllocationMutexGuard&&) = delete;
+	void operator=(const FEntityAllocationMutexGuard&&) = delete;
+
+private:
+	FEntityAllocation* Allocation;
+};
+
+
 /**
  * A wrapper around an FEntityAllocation that provides access to other pieces of
  * information such as its component mask.
@@ -920,13 +962,44 @@ struct FEntityInfo
 	FMovieSceneEntityID EntityID;
 };
 
+inline FEntityAllocationMutexGuard::FEntityAllocationMutexGuard(FEntityAllocation* InAllocation, EComponentHeaderLockMode InLockMode)
+{
+	// Since FEntityAllocation always defaults to Mutex locking, we only 
+	// read/write to the allocation if LockFree is specified here (implying we are always in single-thread mode)
+	if (InLockMode == EComponentHeaderLockMode::LockFree)
+	{
+		Allocation = InAllocation;
+		InAllocation->LockMode = EComponentHeaderLockMode::LockFree;
+	}
+	else
+	{
+		Allocation = nullptr;
+	}
+}
+
+inline FEntityAllocationMutexGuard::~FEntityAllocationMutexGuard()
+{
+	// Always reset back to Mutex
+	if (Allocation)
+	{
+		Allocation->LockMode = EComponentHeaderLockMode::Mutex;
+	}
+}
+
 inline FScopedHeaderReadLock::FScopedHeaderReadLock()
 	: Header(nullptr)
 {}
-inline FScopedHeaderReadLock::FScopedHeaderReadLock(const FComponentHeader* InHeader)
-	: Header(InHeader)
+inline FScopedHeaderReadLock::FScopedHeaderReadLock(const FComponentHeader* InHeader, EComponentHeaderLockMode InLockMode)
 {
-	InHeader->ReadWriteLock.ReadLock();
+	if (InLockMode == EComponentHeaderLockMode::Mutex)
+	{
+		Header = InHeader;
+		InHeader->ReadWriteLock.ReadLock();
+	}
+	else
+	{
+		Header = nullptr;
+	}
 }
 inline FScopedHeaderReadLock::FScopedHeaderReadLock(FScopedHeaderReadLock&& RHS)
 	: Header(RHS.Header)
@@ -956,18 +1029,24 @@ inline FScopedHeaderReadLock::~FScopedHeaderReadLock()
 inline FScopedHeaderWriteLock::FScopedHeaderWriteLock()
 	: Header(nullptr)
 	, WriteContext(FEntityAllocationWriteContext::NewAllocation())
+	, LockMode(EComponentHeaderLockMode::Mutex)
 {}
 
-inline FScopedHeaderWriteLock::FScopedHeaderWriteLock(const FComponentHeader* InHeader, FEntityAllocationWriteContext InWriteContext)
+inline FScopedHeaderWriteLock::FScopedHeaderWriteLock(const FComponentHeader* InHeader, EComponentHeaderLockMode InLockMode, FEntityAllocationWriteContext InWriteContext)
 	: Header(InHeader)
 	, WriteContext(InWriteContext)
+	, LockMode(InLockMode)
 {
-	InHeader->ReadWriteLock.WriteLock();
+	if (InLockMode == EComponentHeaderLockMode::Mutex)
+	{
+		InHeader->ReadWriteLock.WriteLock();
+	}
 }
 
 inline FScopedHeaderWriteLock::FScopedHeaderWriteLock(FScopedHeaderWriteLock&& RHS)
 	: Header(RHS.Header)
 	, WriteContext(RHS.WriteContext)
+	, LockMode(RHS.LockMode)
 {
 	RHS.Header = nullptr;
 }
@@ -976,11 +1055,16 @@ inline FScopedHeaderWriteLock& FScopedHeaderWriteLock::operator=(FScopedHeaderWr
 	if (Header)
 	{
 		Header->PostWriteComponents(WriteContext);
-		Header->ReadWriteLock.WriteUnlock();
+
+		if (LockMode == EComponentHeaderLockMode::Mutex)
+		{
+			Header->ReadWriteLock.WriteUnlock();
+		}
 	}
 
 	Header = RHS.Header;
 	WriteContext = RHS.WriteContext;
+	LockMode = RHS.LockMode;
 
 	RHS.Header = nullptr;
 	return *this;
@@ -990,7 +1074,11 @@ inline FScopedHeaderWriteLock::~FScopedHeaderWriteLock()
 	if (Header)
 	{
 		Header->PostWriteComponents(WriteContext);
-		Header->ReadWriteLock.WriteUnlock();
+
+		if (LockMode == EComponentHeaderLockMode::Mutex)
+		{
+			Header->ReadWriteLock.WriteUnlock();
+		}
 	}
 }
 
diff --git a/Engine/Source/Runtime/MovieSceneTracks/Private/Systems/MovieScenePiecewiseDoubleBlenderSystem.cpp b/Engine/Source/Runtime/MovieSceneTracks/Private/Systems/MovieScenePiecewiseDoubleBlenderSystem.cpp
--- a/Engine/Source/Runtime/MovieSceneTracks/Private/Systems/MovieScenePiecewiseDoubleBlenderSystem.cpp
+++ b/Engine/Source/Runtime/MovieSceneTracks/Private/Systems/MovieScenePiecewiseDoubleBlenderSystem.cpp
@@ -40,12 +40,18 @@ struct FAccumulationTask
 		{
 			if (TArray<FBlendResult>* AccumulationBuffer = AccumulationBuffers->Find(ComponentHeader.ComponentType))
 			{
-				ComponentHeader.ReadWriteLock.ReadLock();
+				if (Allocation->GetCurrentLockMode() != EComponentHeaderLockMode::LockFree)
+				{
+					ComponentHeader.ReadWriteLock.ReadLock();
+				}
 
 				const double* Results = static_cast<const double*>(ComponentHeader.GetValuePtr(0));
 				AccumulateResults(Allocation, Results, BlendIDs, OptionalEasingAndWeights, *AccumulationBuffer);
 
-				ComponentHeader.ReadWriteLock.ReadUnlock();
+				if (Allocation->GetCurrentLockMode() != EComponentHeaderLockMode::LockFree)
+				{
+					ComponentHeader.ReadWriteLock.ReadUnlock();
+				}
 			}
 		}
 	}
@@ -120,7 +126,7 @@ struct FAdditiveFromBaseBlendTask
 			if (FAdditiveFromBaseBuffer* Buffer = AccumulationBuffers->Find(ComponentHeader.ComponentType))
 			{
 				TComponentReader<double> BaseValues = Allocation->ReadComponents(Buffer->BaseComponent.template ReinterpretCast<double>());
-				TComponentReader<double> Results(&ComponentHeader);
+				TComponentReader<double> Results(&ComponentHeader, Allocation->GetCurrentLockMode());
 
 				AccumulateResults(Allocation, Results.AsPtr(), BaseValues.AsPtr(), BlendIDs, EasingAndWeightResults, Buffer->Buffer);
 			}
