Removed superfluous input parameter

diff --git a/Engine/Source/Runtime/MovieScene/Public/EntitySystem/MovieSceneEntitySystemTypes.h b/Engine/Source/Runtime/MovieScene/Public/EntitySystem/MovieSceneEntitySystemTypes.h
--- a/Engine/Source/Runtime/MovieScene/Public/EntitySystem/MovieSceneEntitySystemTypes.h
+++ b/Engine/Source/Runtime/MovieScene/Public/EntitySystem/MovieSceneEntitySystemTypes.h
@@ -804,7 +804,7 @@ struct FEntityAllocation
 	 * Write typed component data for the specified component type
 	 */
 	template<typename T>
-	UE_NODISCARD TComponentLock<TWrite<T>> WriteComponents(TComponentTypeID<T> ComponentType, FEntityAllocationWriteContext InWriteContext, EComponentHeaderLockMode LockMode = EComponentHeaderLockMode::Mutex) const
+	UE_NODISCARD TComponentLock<TWrite<T>> WriteComponents(TComponentTypeID<T> ComponentType, FEntityAllocationWriteContext InWriteContext) const
 	{
 		const FComponentHeader& Header = GetComponentHeaderChecked(ComponentType);
 		return TComponentLock<TWrite<T>>(&Header, LockMode, InWriteContext);
@@ -814,7 +814,7 @@ struct FEntityAllocation
 	 * Attempt to write typed component data for the specified component type
 	 */
 	template<typename T>
-	UE_NODISCARD TComponentLock<TWriteOptional<T>> TryWriteComponents(TComponentTypeID<T> ComponentType, FEntityAllocationWriteContext InWriteContext, EComponentHeaderLockMode LockMode = EComponentHeaderLockMode::Mutex) const
+	UE_NODISCARD TComponentLock<TWriteOptional<T>> TryWriteComponents(TComponentTypeID<T> ComponentType, FEntityAllocationWriteContext InWriteContext) const
 	{
 		if (const FComponentHeader* Header = FindComponentHeader(ComponentType))
 		{
