Sequencer: properly run motion vector evaluations

Before, simply calling the ExecutePhase method didn't set the phase on the runner, which meant that systems that do different things based on the current phase weren't doing the righ thing.

diff --git a/Engine/Source/Runtime/MovieScene/Private/EntitySystem/MovieSceneEntitySystemRunner.cpp b/Engine/Source/Runtime/MovieScene/Private/EntitySystem/MovieSceneEntitySystemRunner.cpp
--- a/Engine/Source/Runtime/MovieScene/Private/EntitySystem/MovieSceneEntitySystemRunner.cpp
+++ b/Engine/Source/Runtime/MovieScene/Private/EntitySystem/MovieSceneEntitySystemRunner.cpp
@@ -4,6 +4,7 @@
 #include "EntitySystem/MovieSceneEntitySystemLinker.h"
 #include "EntitySystem/MovieSceneEntityMutations.h"
 #include "EntitySystem/BuiltInComponentTypes.h"
+#include "EntitySystem/MovieSceneEntitySystemTypes.h"
 #include "Evaluation/PreAnimatedState/MovieScenePreAnimatedCaptureSource.h"
 #include "Evaluation/MovieSceneEvaluationTemplateInstance.h"
 #include "IMovieScenePlayer.h"
@@ -1075,3 +1076,37 @@ FMovieSceneEntitySystemEventTriggers& FMovieSceneEntitySystemRunner::GetQueuedEv
 	return EventTriggers;
 }
 
+bool FMovieSceneEntitySystemRunner::FlushSingleEvaluationPhase()
+{
+	using namespace UE::MovieScene;
+
+	if (!ensureMsgf(
+			!IsCurrentlyEvaluating() && CurrentPhase == ESystemPhase::None,
+			TEXT("Can't run nested flush phase while the runner is evaluating and no re-entrancy window is open")))
+	{
+		return false;
+	}
+
+	UMovieSceneEntitySystemLinker* Linker = GetLinker();
+	if (!ensureMsgf(Linker, TEXT("Runner isn't attached to a valid linker")))
+	{
+		return false;
+	}
+
+	TGuardValue<ESystemPhase> PhaseGuard(CurrentPhase, ESystemPhase::Evaluation);
+
+	Linker->EntityManager.LockDown();
+
+	FGraphEventArray AllTasks;
+	Linker->SystemGraph.ExecutePhase(ESystemPhase::Evaluation, Linker, AllTasks);
+
+	if (AllTasks.Num() != 0)
+	{
+		FTaskGraphInterface::Get().WaitUntilTasksComplete(AllTasks, ENamedThreads::GameThread_Local);
+	}
+
+	Linker->EntityManager.ReleaseLockDown();
+
+	return true;
+}
+
diff --git a/Engine/Source/Runtime/MovieScene/Public/EntitySystem/MovieSceneEntitySystemRunner.h b/Engine/Source/Runtime/MovieScene/Public/EntitySystem/MovieSceneEntitySystemRunner.h
--- a/Engine/Source/Runtime/MovieScene/Public/EntitySystem/MovieSceneEntitySystemRunner.h
+++ b/Engine/Source/Runtime/MovieScene/Public/EntitySystem/MovieSceneEntitySystemRunner.h
@@ -147,6 +147,13 @@ public:
 	 */
 	bool IsCurrentlyEvaluating() const;
 
+	/**
+	 * Run a single evaluation phase
+	 *
+	 * @return Whether the phase was run
+	 */
+	bool FlushSingleEvaluationPhase();
+
 public:
 
 	UMovieSceneEntitySystemLinker* GetLinker() const;
@@ -275,4 +282,4 @@ private:
 
 	/** True if any update has been queued with the ERunnerUpdateFlags::Flush flag since our last flush */
 	bool bRequireFullFlush;
-};
\ No newline at end of file
+};
diff --git a/Engine/Source/Runtime/MovieSceneTracks/Private/Systems/MovieSceneMotionVectorSimulationSystem.cpp b/Engine/Source/Runtime/MovieSceneTracks/Private/Systems/MovieSceneMotionVectorSimulationSystem.cpp
--- a/Engine/Source/Runtime/MovieSceneTracks/Private/Systems/MovieSceneMotionVectorSimulationSystem.cpp
+++ b/Engine/Source/Runtime/MovieSceneTracks/Private/Systems/MovieSceneMotionVectorSimulationSystem.cpp
@@ -135,18 +135,10 @@ void UMovieSceneMotionVectorSimulationSystem::ComputeSimulatedMotion()
 
 	// --------------------------------------------------------------------------------------------------------------------------------------------
 	// Re-execute the evaluation phase to flush the new transforms
+	FMovieSceneEntitySystemRunner* ActiveRunner = Linker->GetActiveRunner();
+	if (ensure(ActiveRunner))
 	{
-		Linker->EntityManager.LockDown();
-
-		FGraphEventArray AllTasks;
-		Linker->SystemGraph.ExecutePhase(ESystemPhase::Evaluation, Linker, AllTasks);
-
-		if (AllTasks.Num() != 0)
-		{
-			FTaskGraphInterface::Get().WaitUntilTasksComplete(AllTasks, ENamedThreads::GameThread_Local);
-		}
-
-		Linker->EntityManager.ReleaseLockDown();
+		ActiveRunner->FlushSingleEvaluationPhase();
 	}
 
 	// --------------------------------------------------------------------------------------------------------------------------------------------
