Sequencer: invalidate channel proxy when channel overrides are changed as the result of an undo/redo operation

diff --git a/Engine/Source/Runtime/MovieScene/Private/Channels/MovieSceneSectionChannelOverrideRegistry.cpp b/Engine/Source/Runtime/MovieScene/Private/Channels/MovieSceneSectionChannelOverrideRegistry.cpp
--- a/Engine/Source/Runtime/MovieScene/Private/Channels/MovieSceneSectionChannelOverrideRegistry.cpp
+++ b/Engine/Source/Runtime/MovieScene/Private/Channels/MovieSceneSectionChannelOverrideRegistry.cpp
@@ -6,6 +6,7 @@
 #include "Channels/MovieSceneChannelOverrideContainer.h"
 #include "Channels/IMovieSceneChannelOverrideProvider.h"
 #include "EntitySystem/MovieSceneEntityBuilder.h"
+#include "MovieSceneSection.h"
 
 #include UE_INLINE_GENERATED_CPP_BY_NAME(MovieSceneSectionChannelOverrideRegistry)
 
@@ -120,5 +121,18 @@ void UMovieSceneSectionChannelOverrideRegistry::OnPostPaste()
 	}
 }
 
+void UMovieSceneSectionChannelOverrideRegistry::PostEditUndo()
+{
+	Super::PostEditUndo();
+
+	if (UMovieSceneSection* Section = GetTypedOuter<UMovieSceneSection>())
+	{
+		if (IMovieSceneChannelOverrideProvider* OverrideProvider = Cast<IMovieSceneChannelOverrideProvider>(Section))
+		{
+			OverrideProvider->OnChannelOverridesChanged();
+		}
+	}
+}
+
 #endif
 
diff --git a/Engine/Source/Runtime/MovieScene/Public/Channels/MovieSceneSectionChannelOverrideRegistry.h b/Engine/Source/Runtime/MovieScene/Public/Channels/MovieSceneSectionChannelOverrideRegistry.h
--- a/Engine/Source/Runtime/MovieScene/Public/Channels/MovieSceneSectionChannelOverrideRegistry.h
+++ b/Engine/Source/Runtime/MovieScene/Public/Channels/MovieSceneSectionChannelOverrideRegistry.h
@@ -113,6 +113,11 @@ public:
 	 */
 	void OnPostPaste();
 
+	/**
+	 * Called when an undo/redo operation has affected this override registry.
+	 */
+	virtual void PostEditUndo() override;
+
 #endif
 
 private:
