Sequencer: do some basic check before dispatching a task that might not have anything to do.

diff --git a/Engine/Source/Runtime/MovieScene/Public/EntitySystem/MovieSceneComponentAccessors.h b/Engine/Source/Runtime/MovieScene/Public/EntitySystem/MovieSceneComponentAccessors.h
--- a/Engine/Source/Runtime/MovieScene/Public/EntitySystem/MovieSceneComponentAccessors.h
+++ b/Engine/Source/Runtime/MovieScene/Public/EntitySystem/MovieSceneComponentAccessors.h
@@ -7,6 +7,7 @@
 #include "Containers/UnrealString.h"
 #include "EntitySystem/EntityAllocationIterator.h"
 #include "EntitySystem/MovieSceneComponentPtr.h"
+#include "EntitySystem/MovieSceneEntityManager.h"
 #include "EntitySystem/MovieSceneEntitySystemTypes.h"
 #include "EntitySystem/MovieSceneSystemTaskDependencies.h"
 #include "HAL/Platform.h"
@@ -345,6 +346,35 @@ inline bool IsAccessorValid(const TReadOneOrMoreOfAccessor<T...>* In)
 	return bValid;
 }
 
+
+
+inline bool HasAccessorWork(const FEntityManager*, const FEntityIDAccess*)
+{
+	return true;
+}
+inline bool HasAccessorWork(const FEntityManager* EntityManager, const FComponentAccess* In)
+{
+	return EntityManager->ContainsComponent(In->ComponentType);
+}
+inline bool HasAccessorWork(const FEntityManager* EntityManager, const FOptionalComponentAccess* In)
+{
+	return true;
+}
+template<typename... T>
+inline bool HasAccessorWork(const FEntityManager* EntityManager, const TReadOneOfAccessor<T...>* In)
+{
+	bool bAnyWork = false;
+	VisitTupleElements([&bAnyWork, EntityManager](FOptionalReadAccess It){ bAnyWork |= HasAccessorWork(EntityManager , &It); }, In->ComponentTypes);
+	return bAnyWork;
+}
+template<typename... T>
+inline bool HasAccessorWork(const FEntityManager* EntityManager, const TReadOneOrMoreOfAccessor<T...>* In)
+{
+	bool bAnyWork = false;
+	VisitTupleElements([&bAnyWork, EntityManager](FOptionalReadAccess It) { bAnyWork |= HasAccessorWork(EntityManager , &It); }, In->ComponentTypes);
+	return bAnyWork;
+}
+
 #if UE_MOVIESCENE_ENTITY_DEBUG
 
 	MOVIESCENE_API void AccessorToString(const FReadAccess* In, FEntityManager* EntityManager, FString& OutString);
diff --git a/Engine/Source/Runtime/MovieScene/Public/EntitySystem/MovieSceneEntitySystemTask.h b/Engine/Source/Runtime/MovieScene/Public/EntitySystem/MovieSceneEntitySystemTask.h
--- a/Engine/Source/Runtime/MovieScene/Public/EntitySystem/MovieSceneEntitySystemTask.h
+++ b/Engine/Source/Runtime/MovieScene/Public/EntitySystem/MovieSceneEntitySystemTask.h
@@ -243,6 +243,8 @@ struct TEntityTaskComponents : TEntityTaskComponentsImpl<TMakeIntegerSequence<in
 	template<typename TaskImpl, typename... TaskConstructionArgs>
 	FGraphEventRef Dispatch(FEntityManager* EntityManager, const FSystemTaskPrerequisites& Prerequisites, FSystemSubsequentTasks* Subsequents, TaskConstructionArgs&&... InArgs) const
 	{
+		static_assert(sizeof...(T) == 0, "Dispatch() is only for non-structured logic, which means that any call to Read or Write (or their variants) won't do anything -- please remove them");
+
 		checkfSlow(IsInGameThread(), TEXT("Tasks can only be dispatched from the game thread."));
 
 		const bool bRunInline = !ensure(EntityManager->IsLockedDown()) || EntityManager->GetThreadingModel() == EEntityThreadingModel::NoThreading;
@@ -307,6 +309,14 @@ struct TEntityTaskComponents : TEntityTaskComponentsImpl<TMakeIntegerSequence<in
 			return nullptr;
 		}
 
+		// Quick check to prevent dispatching a task that might not have any work to do. We only check the accessors
+		// here, so the task could still early-return with no work done if some custom filters end up matching
+		// nothing. Callers should in this case do their best to prevent useless tasks being dispatched.
+		if (!this->HasAnyWork(EntityManager))
+		{
+			return nullptr;
+		}
+
 		// If this ensure triggers, we are not in the evaluation phase - the callee should be using RunInline_ or Iterate_ variants
 		const bool bRunInline = !ensure(EntityManager->IsLockedDown()) || EntityManager->GetThreadingModel() == EEntityThreadingModel::NoThreading;
 		if (bRunInline)
@@ -317,7 +327,6 @@ struct TEntityTaskComponents : TEntityTaskComponentsImpl<TMakeIntegerSequence<in
 		}
 		else
 		{
-
 			FGraphEventArray GatheredPrereqs;
 			this->PopulatePrerequisites(Prerequisites, &GatheredPrereqs);
 
@@ -380,6 +389,12 @@ struct TEntityTaskComponents : TEntityTaskComponentsImpl<TMakeIntegerSequence<in
 			return nullptr;
 		}
 
+		// See comment in Dispatch_PerAllocation()
+		if (!this->HasAnyWork(EntityManager))
+		{
+			return nullptr;
+		}
+
 		// If this ensure triggers, we are not in the evaluation phase - the callee should be using RunInline_ or Iterate_ variants
 		const bool bRunInline = !ensure(EntityManager->IsLockedDown()) || EntityManager->GetThreadingModel() == EEntityThreadingModel::NoThreading;
 		if (bRunInline)
@@ -593,6 +608,19 @@ struct TEntityTaskComponentsImpl<TIntegerSequence<int, Indices...>, T...>
 		return bAllValid;
 	}
 
+	/**
+	 * Check whether all required accessors correspond to component types that are present in the given entity manager.
+	 */
+	bool HasAnyWork(const FEntityManager* EntityManager) const
+	{
+		bool bAllHaveWork = true;
+
+		int Temp[] = { ( bAllHaveWork &= HasAccessorWork(EntityManager, &Accessors.template Get<Indices>()), 0)..., 0 };
+		(void)Temp;
+
+		return bAllHaveWork;
+	}
+
 	/** Utility function called when the task is dispatched to populate the filter based on our component typs */
 	void PopulateFilter(FEntityComponentFilter* OutFilter) const
 	{
@@ -988,6 +1016,12 @@ struct TFilteredEntityTask
 			return nullptr;
 		}
 
+		// See comment in TEntityTaskComponents' Dispatch_PerAllocation()
+		if (!Components.HasAnyWork(EntityManager))
+		{
+			return nullptr;
+		}
+
 		// If this ensure triggers, we are not in the evaluation phase - the callee should be using RunInline_ or Iterate_ variants
 		const bool bRunInline = !ensure(EntityManager->IsLockedDown()) || EntityManager->GetThreadingModel() == EEntityThreadingModel::NoThreading;
 		if (bRunInline)
@@ -1060,6 +1094,12 @@ struct TFilteredEntityTask
 			return nullptr;
 		}
 
+		// See comment in TEntityTaskComponents' Dispatch_PerAllocation()
+		if (!Components.HasAnyWork(EntityManager))
+		{
+			return nullptr;
+		}
+
 		// If this ensure triggers, we are not in the evaluation phase - the callee should be using RunInline_ or Iterate_ variants
 		const bool bRunInline = !ensure(EntityManager->IsLockedDown()) || EntityManager->GetThreadingModel() == EEntityThreadingModel::NoThreading;
 		if (bRunInline)
