Sequencer: make the perlin noise channel evaluators use a new "EvalSeconds" component

This prevents these evaluators from having to look up the tick resolution on the sequence instance. First, this isn't very good for cache coherency, and second the instance handles aren't available during interrogations, so none of the perlin noise channels could evaluate, resulting in potentially incorrect keyed values.

diff --git a/Engine/Source/Runtime/MovieScene/Private/EntitySystem/BuiltInComponentTypes.cpp b/Engine/Source/Runtime/MovieScene/Private/EntitySystem/BuiltInComponentTypes.cpp
--- a/Engine/Source/Runtime/MovieScene/Private/EntitySystem/BuiltInComponentTypes.cpp
+++ b/Engine/Source/Runtime/MovieScene/Private/EntitySystem/BuiltInComponentTypes.cpp
@@ -32,6 +32,7 @@ FBuiltInComponentTypes::FBuiltInComponentTypes()
 	ComponentRegistry->NewComponentType(&RootInstanceHandle,    TEXT("Root Instance Handle"));
 
 	ComponentRegistry->NewComponentType(&EvalTime,              TEXT("Eval Time"));
+	ComponentRegistry->NewComponentType(&EvalSeconds,           TEXT("Eval Seconds"));
 
 	ComponentRegistry->NewComponentType(&BoundObject,           TEXT("Bound Object"));
 
@@ -126,6 +127,7 @@ FBuiltInComponentTypes::FBuiltInComponentTypes()
 	ComponentRegistry->NewComponentType(&EvaluationHookFlags,   TEXT("Evaluation Hook Flags"), EComponentTypeFlags::Preserved);
 
 	ComponentRegistry->NewComponentType(&Interrogation.InputKey,  TEXT("Interrogation Input"));
+	ComponentRegistry->NewComponentType(&Interrogation.Instance,  TEXT("Interrogation Instance"));
 	ComponentRegistry->NewComponentType(&Interrogation.OutputKey, TEXT("Interrogation Output"));
 
 	Tags.RestoreState            = ComponentRegistry->NewTag(TEXT("Is Restore State Entity"));
@@ -169,6 +171,7 @@ FBuiltInComponentTypes::FBuiltInComponentTypes()
 	ComponentRegistry->Factories.DefineMutuallyInclusiveComponent(Tags.SectionPreRoll,Tags.PreRoll);
 
 	ComponentRegistry->Factories.DuplicateChildComponent(EvalTime);
+	ComponentRegistry->Factories.DuplicateChildComponent(EvalSeconds);
 	ComponentRegistry->Factories.DuplicateChildComponent(BaseValueEvalTime);
 
 	ComponentRegistry->Factories.DuplicateChildComponent(SequenceID);
diff --git a/Engine/Source/Runtime/MovieScene/Private/EntitySystem/MovieSceneEvalTimeSystem.cpp b/Engine/Source/Runtime/MovieScene/Private/EntitySystem/MovieSceneEvalTimeSystem.cpp
--- a/Engine/Source/Runtime/MovieScene/Private/EntitySystem/MovieSceneEvalTimeSystem.cpp
+++ b/Engine/Source/Runtime/MovieScene/Private/EntitySystem/MovieSceneEvalTimeSystem.cpp
@@ -14,34 +14,28 @@
 DECLARE_CYCLE_STAT(TEXT("EvalTime System"), MovieSceneEval_EvalTimeSystem, STATGROUP_MovieSceneECS);
 DECLARE_CYCLE_STAT(TEXT("MovieScene: Assign evaluation times"), MovieSceneEval_EvalTimes, STATGROUP_MovieSceneECS);
 
-struct FAssignEvalTimesTask
-{
-	TArray<FFrameTime>* FrameTimes;
-
-	explicit FAssignEvalTimesTask(TArray<FFrameTime>* InFrameTimes)
-		: FrameTimes(InFrameTimes)
-	{}
-
-	void ForEachEntity(UE::MovieScene::FInstanceHandle InstanceHandle, FFrameTime& EvalTime)
-	{
-		EvalTime = (*FrameTimes)[InstanceHandle.InstanceID];
-	}
-};
-
 UMovieSceneEvalTimeSystem::UMovieSceneEvalTimeSystem(const FObjectInitializer& ObjInit)
 	: Super(ObjInit)
 {
 	using namespace UE::MovieScene;
 
-	RelevantComponent = FBuiltInComponentTypes::Get()->EvalTime;
+	const FBuiltInComponentTypes* BuiltInComponents = FBuiltInComponentTypes::Get();
+
 	SystemCategories = EEntitySystemCategory::Core;
+	RelevantFilter.Any({ BuiltInComponents->EvalTime, BuiltInComponents->EvalSeconds });
 
 	if (HasAnyFlags(RF_ClassDefaultObject))
 	{
-		DefineComponentProducer(GetClass(), FBuiltInComponentTypes::Get()->EvalTime);
+		DefineComponentProducer(GetClass(), BuiltInComponents->EvalTime);
+		DefineComponentProducer(GetClass(), BuiltInComponents->EvalSeconds);
 	}
 }
 
+bool UMovieSceneEvalTimeSystem::IsRelevantImpl(UMovieSceneEntitySystemLinker* InLinker) const
+{
+	return InLinker->EntityManager.Contains(RelevantFilter);
+}
+
 void UMovieSceneEvalTimeSystem::OnRun(FSystemTaskPrerequisites& InPrerequisites, FSystemSubsequentTasks& Subsequents)
 {
 	using namespace UE::MovieScene;
@@ -50,13 +44,13 @@ void UMovieSceneEvalTimeSystem::OnRun(FSystemTaskPrerequisites& InPrerequisites,
 
 	struct FGatherTimes
 	{
-		FGatherTimes(const FInstanceRegistry* InInstanceRegistry, TArray<FFrameTime>* InFrameTimes)
+		FGatherTimes(const FInstanceRegistry* InInstanceRegistry, TArray<FEvaluatedTime>* InEvaluatedTimes)
 			: InstanceRegistry(InInstanceRegistry)
-			, FrameTimes(InFrameTimes)
+			, EvaluatedTimes(InEvaluatedTimes)
 		{}
 
 		const FInstanceRegistry* InstanceRegistry;
-		TArray<FFrameTime>* FrameTimes;
+		TArray<FEvaluatedTime>* EvaluatedTimes;
 
 		FORCEINLINE TStatId           GetStatId() const    { RETURN_QUICK_DECLARE_CYCLE_STAT(FGenericTask, STATGROUP_TaskGraphTasks); }
 		static ENamedThreads::Type    GetDesiredThread()   { return ENamedThreads::AnyHiPriThreadHiPriTask; }
@@ -69,26 +63,60 @@ void UMovieSceneEvalTimeSystem::OnRun(FSystemTaskPrerequisites& InPrerequisites,
 
 		void Run()
 		{
-			FrameTimes->SetNum(InstanceRegistry->GetSparseInstances().GetMaxIndex());
+			EvaluatedTimes->SetNum(InstanceRegistry->GetSparseInstances().GetMaxIndex());
 			for (auto It = InstanceRegistry->GetSparseInstances().CreateConstIterator(); It; ++It)
 			{
-				(*FrameTimes)[It.GetIndex()] = It->GetContext().GetTime();
+				const FMovieSceneContext& Context = It->GetContext();
+
+				FEvaluatedTime EvaluatedTime;
+				EvaluatedTime.FrameTime = Context.GetTime();
+				EvaluatedTime.Seconds = Context.GetFrameRate().AsSeconds(EvaluatedTime.FrameTime);
+
+				(*EvaluatedTimes)[It.GetIndex()] = EvaluatedTime;
 			}
 		}
 	};
 
+	struct FAssignEvalTimesTask
+	{
+		TArrayView<const FEvaluatedTime> EvaluatedTimes;
+
+		explicit FAssignEvalTimesTask(TArrayView<const FEvaluatedTime> InEvaluatedTimes)
+			: EvaluatedTimes(InEvaluatedTimes)
+		{}
+
+		void ForEachEntity(UE::MovieScene::FInstanceHandle InstanceHandle, FFrameTime& EvalTime)
+		{
+			EvalTime = EvaluatedTimes[InstanceHandle.InstanceID].FrameTime;
+		}
+	};
+
+	struct FAssignEvalSecondsTask
+	{
+		TArrayView<const FEvaluatedTime> EvaluatedTimes;
+
+		explicit FAssignEvalSecondsTask(TArrayView<const FEvaluatedTime> InEvaluatedTimes)
+			: EvaluatedTimes(InEvaluatedTimes)
+		{}
+
+		void ForEachEntity(UE::MovieScene::FInstanceHandle InstanceHandle, double& EvalSeconds)
+		{
+			EvalSeconds = EvaluatedTimes[InstanceHandle.InstanceID].Seconds;
+		}
+	};
+
 	FSystemTaskPrerequisites EvalPrereqs;
 	FBuiltInComponentTypes* BuiltInComponents = FBuiltInComponentTypes::Get();
 
 	if (Linker->EntityManager.GetThreadingModel() == EEntityThreadingModel::NoThreading)
 	{
-		FGatherTimes(Linker->GetInstanceRegistry(), &FrameTimes).Run();
+		FGatherTimes(Linker->GetInstanceRegistry(), &EvaluatedTimes).Run();
 	}
 	else
 	{
 		// The only thing we depend on is the gather task
 		FGraphEventRef GatherEvalTimesEvent = TGraphTask<FGatherTimes>::CreateTask(nullptr, Linker->EntityManager.GetDispatchThread())
-		.ConstructAndDispatchWhenReady(Linker->GetInstanceRegistry(), &FrameTimes);
+		.ConstructAndDispatchWhenReady(Linker->GetInstanceRegistry(), &EvaluatedTimes);
 
 		EvalPrereqs.AddComponentTask(BuiltInComponents->EvalTime, GatherEvalTimesEvent);
 	}
@@ -98,6 +126,13 @@ void UMovieSceneEvalTimeSystem::OnRun(FSystemTaskPrerequisites& InPrerequisites,
 	.Write(BuiltInComponents->EvalTime)
 	.FilterNone({ FBuiltInComponentTypes::Get()->Tags.FixedTime })
 	.SetStat(GET_STATID(MovieSceneEval_EvalTimes))
-	.Dispatch_PerEntity<FAssignEvalTimesTask>(&Linker->EntityManager, EvalPrereqs, &Subsequents, &FrameTimes);
+	.Dispatch_PerEntity<FAssignEvalTimesTask>(&Linker->EntityManager, EvalPrereqs, &Subsequents, EvaluatedTimes);
+
+	FEntityTaskBuilder()
+	.Read(BuiltInComponents->InstanceHandle)
+	.Write(BuiltInComponents->EvalSeconds)
+	.FilterNone({ FBuiltInComponentTypes::Get()->Tags.FixedTime })
+	.SetStat(GET_STATID(MovieSceneEval_EvalTimes))
+	.Dispatch_PerEntity<FAssignEvalSecondsTask>(&Linker->EntityManager, EvalPrereqs, &Subsequents, EvaluatedTimes);
 }
 
diff --git a/Engine/Source/Runtime/MovieScene/Private/Evaluation/MovieSceneEvaluationField.cpp b/Engine/Source/Runtime/MovieScene/Private/Evaluation/MovieSceneEvaluationField.cpp
--- a/Engine/Source/Runtime/MovieScene/Private/Evaluation/MovieSceneEvaluationField.cpp
+++ b/Engine/Source/Runtime/MovieScene/Private/Evaluation/MovieSceneEvaluationField.cpp
@@ -40,6 +40,11 @@ FMovieSceneEvaluationFieldSharedEntityMetaData& FMovieSceneEntityComponentFieldB
 	return Field->SharedMetaData[SharedMetaDataIndex];
 }
 
+int32 FMovieSceneEntityComponentFieldBuilder::GetSharedMetaDataIndex() const
+{
+	return SharedMetaDataIndex;
+}
+
 int32 FMovieSceneEntityComponentFieldBuilder::FindOrAddEntity(UObject* EntityOwner, uint32 EntityID)
 {
 	FMovieSceneEvaluationFieldEntityKey Key = { EntityOwner, EntityID };
diff --git a/Engine/Source/Runtime/MovieScene/Private/MovieSceneSection.cpp b/Engine/Source/Runtime/MovieScene/Private/MovieSceneSection.cpp
--- a/Engine/Source/Runtime/MovieScene/Private/MovieSceneSection.cpp
+++ b/Engine/Source/Runtime/MovieScene/Private/MovieSceneSection.cpp
@@ -324,6 +324,7 @@ void UMovieSceneSection::BuildDefaultComponents(UMovieSceneEntitySystemLinker* E
 		.AddConditional(Components->Easing,                     FEasingComponentData{ decltype(FEasingComponentData::Section)(this) }, bHasEasing)
 		.AddConditional(Components->HierarchicalBias,           Params.Sequence.HierarchicalBias, Params.Sequence.HierarchicalBias != 0)
 		.AddConditional(Components->Interrogation.InputKey,     Params.InterrogationKey, Params.InterrogationKey.IsValid())
+		.AddConditional(Components->Interrogation.Instance,     Params.InterrogationInstance, Params.InterrogationInstance.IsValid())
 		.AddConditional(Components->EvalTime,                   Params.EntityMetaData ? Params.EntityMetaData->ForcedTime : 0, bHasForcedTime)
 		.AddTagConditional(Components->Tags.RestoreState,       bShouldRestoreState)
 		.AddTagConditional(Components->Tags.FixedTime,          bHasForcedTime)
diff --git a/Engine/Source/Runtime/MovieScene/Public/EntitySystem/BuiltInComponentTypes.h b/Engine/Source/Runtime/MovieScene/Public/EntitySystem/BuiltInComponentTypes.h
--- a/Engine/Source/Runtime/MovieScene/Public/EntitySystem/BuiltInComponentTypes.h
+++ b/Engine/Source/Runtime/MovieScene/Public/EntitySystem/BuiltInComponentTypes.h
@@ -30,6 +30,7 @@ namespace UE::MovieScene
 	struct FInitialValueIndex;
 	struct FInstanceHandle;
 	struct FInterrogationKey;
+	struct FInterrogationInstance;
 	struct FRootInstanceHandle;
 	namespace Interpolation
 	{
@@ -247,6 +248,8 @@ public:
 
 	TComponentTypeID<FFrameTime>          EvalTime;
 
+	TComponentTypeID<double>              EvalSeconds;
+
 public:
 
 	TComponentTypeID<FMovieSceneBlendChannelID> BlendChannelInput;
@@ -403,6 +406,7 @@ public:
 	struct
 	{
 		TComponentTypeID<FInterrogationKey> InputKey;
+		TComponentTypeID<FInterrogationInstance> Instance;
 		TComponentTypeID<FInterrogationKey> OutputKey;
 	} Interrogation;
 
diff --git a/Engine/Source/Runtime/MovieScene/Public/EntitySystem/IMovieSceneEntityProvider.h b/Engine/Source/Runtime/MovieScene/Public/EntitySystem/IMovieSceneEntityProvider.h
--- a/Engine/Source/Runtime/MovieScene/Public/EntitySystem/IMovieSceneEntityProvider.h
+++ b/Engine/Source/Runtime/MovieScene/Public/EntitySystem/IMovieSceneEntityProvider.h
@@ -94,6 +94,7 @@ struct FEntityImportParams
 	uint32 EntityID = 0;
 
 	FInterrogationKey InterrogationKey;
+	FInterrogationInstance InterrogationInstance;
 
 	FEntityImportSequenceParams Sequence;
 
@@ -118,8 +119,8 @@ class MOVIESCENE_API IMovieSceneEntityProvider
 {
 public:
 
-	using FEntityImportParams   = UE::MovieScene::FEntityImportParams;
-	using FImportedEntity       = UE::MovieScene::FImportedEntity;
+	using FEntityImportParams        = UE::MovieScene::FEntityImportParams;
+	using FImportedEntity            = UE::MovieScene::FImportedEntity;
 
 
 	GENERATED_BODY()
diff --git a/Engine/Source/Runtime/MovieScene/Public/EntitySystem/MovieSceneEntitySystemTypes.h b/Engine/Source/Runtime/MovieScene/Public/EntitySystem/MovieSceneEntitySystemTypes.h
--- a/Engine/Source/Runtime/MovieScene/Public/EntitySystem/MovieSceneEntitySystemTypes.h
+++ b/Engine/Source/Runtime/MovieScene/Public/EntitySystem/MovieSceneEntitySystemTypes.h
@@ -18,6 +18,7 @@
 
 #include <initializer_list>
 
+class UMovieSceneSequence;
 namespace UE { namespace MovieScene { class FEntityManager; } }
 
 #ifndef UE_MOVIESCENE_ENTITY_DEBUG
@@ -264,6 +265,20 @@ struct FInterrogationKey
 };
 
 
+/**
+ * Sequence instance information for interrogations
+ */
+struct FInterrogationInstance
+{
+	UMovieSceneSequence* Sequence = nullptr;
+
+	bool IsValid() const
+	{
+		return Sequence != nullptr;
+	}
+};
+
+
 struct MOVIESCENE_API FEntityComponentFilter
 {
 
diff --git a/Engine/Source/Runtime/MovieScene/Public/EntitySystem/MovieSceneEvalTimeSystem.h b/Engine/Source/Runtime/MovieScene/Public/EntitySystem/MovieSceneEvalTimeSystem.h
--- a/Engine/Source/Runtime/MovieScene/Public/EntitySystem/MovieSceneEvalTimeSystem.h
+++ b/Engine/Source/Runtime/MovieScene/Public/EntitySystem/MovieSceneEvalTimeSystem.h
@@ -5,6 +5,7 @@
 #include "Containers/Array.h"
 #include "EntitySystem/MovieSceneEntityIDs.h"
 #include "EntitySystem/MovieSceneEntitySystem.h"
+#include "EntitySystem/MovieSceneEntitySystemTypes.h"
 #include "Misc/FrameTime.h"
 #include "UObject/ObjectMacros.h"
 #include "UObject/UObjectGlobals.h"
@@ -23,8 +24,16 @@ public:
 	UMovieSceneEvalTimeSystem(const FObjectInitializer& ObjInit);
 
 	virtual void OnRun(FSystemTaskPrerequisites& InPrerequisites, FSystemSubsequentTasks& Subsequents) override;
+	virtual bool IsRelevantImpl(UMovieSceneEntitySystemLinker* InLinker) const override;
 
 private:
-	TArray<FFrameTime> FrameTimes;
+	struct FEvaluatedTime
+	{
+		FFrameTime FrameTime;
+		double Seconds;
+	};
+	TArray<FEvaluatedTime> EvaluatedTimes;
+
+	UE::MovieScene::FEntityComponentFilter RelevantFilter;
 };
 
diff --git a/Engine/Source/Runtime/MovieScene/Public/Evaluation/MovieSceneEvaluationField.h b/Engine/Source/Runtime/MovieScene/Public/Evaluation/MovieSceneEvaluationField.h
--- a/Engine/Source/Runtime/MovieScene/Public/Evaluation/MovieSceneEvaluationField.h
+++ b/Engine/Source/Runtime/MovieScene/Public/Evaluation/MovieSceneEvaluationField.h
@@ -123,7 +123,6 @@ struct FMovieSceneEvaluationFieldEntityMetaData
 	friend bool operator==(const FMovieSceneEvaluationFieldEntityMetaData& A, const FMovieSceneEvaluationFieldEntityMetaData& B)
 	{
 		return A.ForcedTime == B.ForcedTime && 
-			A.InterrogationChannel == B.InterrogationChannel && 
 			A.Flags == B.Flags && 
 			A.bEvaluateInSequencePreRoll == B.bEvaluateInSequencePreRoll && 
 			A.bEvaluateInSequencePostRoll == B.bEvaluateInSequencePostRoll &&
@@ -137,9 +136,6 @@ struct FMovieSceneEvaluationFieldEntityMetaData
 	UPROPERTY()
 	FFrameNumber ForcedTime;
 
-	/** Never serialized */
-	UE::MovieScene::FInterrogationChannel InterrogationChannel;
-
 	UPROPERTY()
 	ESectionEvaluationFlags Flags;
 
@@ -334,6 +330,11 @@ struct MOVIESCENE_API FMovieSceneEntityComponentFieldBuilder
 	 */
 	FMovieSceneEvaluationFieldSharedEntityMetaData& GetSharedMetaData();
 
+	/**
+	 * Access the index of the shared meta-data of this builder.
+	 */
+	int32 GetSharedMetaDataIndex() const;
+
 	/**
 	 * Add meta-data to this tree returning its index within this builder
 	 *
diff --git a/Engine/Source/Runtime/MovieSceneTracks/Private/EntitySystem/Interrogation/MovieSceneInterrogationLinker.cpp b/Engine/Source/Runtime/MovieSceneTracks/Private/EntitySystem/Interrogation/MovieSceneInterrogationLinker.cpp
--- a/Engine/Source/Runtime/MovieSceneTracks/Private/EntitySystem/Interrogation/MovieSceneInterrogationLinker.cpp
+++ b/Engine/Source/Runtime/MovieSceneTracks/Private/EntitySystem/Interrogation/MovieSceneInterrogationLinker.cpp
@@ -684,13 +684,25 @@ void FSystemInterrogator::ImportTrack(UMovieSceneTrack* Track, const FGuid& Obje
 
 	TGuardValue<FEntityManager*> DebugVizGuard(GEntityManagerForDebuggingVisualizers, &Linker->EntityManager);
 
-	FFrameRate TickResolution = Track->GetTypedOuter<UMovieScene>()->GetTickResolution();
-
 	const FMovieSceneTrackEvaluationField& EvaluationField = Track->GetEvaluationField();
 
 	FMovieSceneEntityComponentFieldBuilder FieldBuilder(&EntityComponentField);
 	FieldBuilder.GetSharedMetaData().ObjectBindingID = ObjectBindingID;
 
+	// Creating a new builder allocates a new piece of shared metadata. We will use that shared metadata's index
+	// to keep track of the extra metadata, since the extra metadata is also shared.
+	const int32 SharedMetaDataIndex = FieldBuilder.GetSharedMetaDataIndex();
+	UMovieSceneSequence* OwningSequence = Track->GetTypedOuter<UMovieSceneSequence>();
+	if (!ExtraMetaData.IsValidIndex(SharedMetaDataIndex))
+	{
+		ExtraMetaData.Insert(SharedMetaDataIndex, FExtraMetaData{ FInterrogationInstance { OwningSequence }, InChannel });
+	}
+	else
+	{
+		const FExtraMetaData& InterrogationMetaData = ExtraMetaData[SharedMetaDataIndex];
+		ensure(InterrogationMetaData.InterrogationInstance.Sequence == OwningSequence && InterrogationMetaData.InterrogationChannel == InChannel);
+	}
+
 	for (const FMovieSceneTrackEvaluationFieldEntry& Entry : EvaluationField.Entries)
 	{
 		IMovieSceneEntityProvider* EntityProvider = Cast<IMovieSceneEntityProvider>(Entry.Section);
@@ -712,11 +724,6 @@ void FSystemInterrogator::ImportTrack(UMovieSceneTrack* Track, const FGuid& Obje
 		MetaData.bEvaluateInSequencePreRoll  = Track->EvalOptions.bEvaluateInPreroll;
 		MetaData.bEvaluateInSequencePostRoll = Track->EvalOptions.bEvaluateInPostroll;
 
-		if (InChannel != FInterrogationChannel::Default())
-		{
-			MetaData.InterrogationChannel = InChannel;
-		}
-
 		if (!EntityProvider->PopulateEvaluationField(Entry.Range, MetaData, &FieldBuilder))
 		{
 			const int32 EntityIndex   = FieldBuilder.FindOrAddEntity(Entry.Section, 0);
@@ -769,9 +776,12 @@ void FSystemInterrogator::InterrogateEntity(int32 InterrogationIndex, const FMov
 	Params.InterrogationKey.Channel = FInterrogationChannel::Default();
 	Params.InterrogationKey.InterrogationIndex = InterrogationIndex;
 
-	if (Params.EntityMetaData && Params.EntityMetaData->InterrogationChannel.IsValid())
+	const int32 SharedMetaDataIndex = Query.Entity.SharedMetaDataIndex;
+	if (SharedMetaDataIndex != INDEX_NONE && ensure(ExtraMetaData.IsValidIndex(SharedMetaDataIndex)))
 	{
-		Params.InterrogationKey.Channel = Params.EntityMetaData->InterrogationChannel;
+		const FExtraMetaData& InterrogationMetaData = ExtraMetaData[SharedMetaDataIndex];
+		Params.InterrogationKey.Channel = InterrogationMetaData.InterrogationChannel;
+		Params.InterrogationInstance = InterrogationMetaData.InterrogationInstance;
 	}
 
 	FImportedEntity ImportedEntity;
@@ -800,11 +810,37 @@ void FSystemInterrogator::Update()
 	Linker->LinkRelevantSystems();
 
 	TArrayView<const FInterrogationParams> Interrogations = Channels.GetInterrogations();
+
+	// Compute evaluation time (in frames) for all entities that need it.
 	FEntityTaskBuilder()
 	.Read(FBuiltInComponentTypes::Get()->Interrogation.InputKey)
 	.Write(FBuiltInComponentTypes::Get()->EvalTime)
 	.FilterNone({ FBuiltInComponentTypes::Get()->Tags.FixedTime })
-	.Iterate_PerEntity(&Linker->EntityManager, [&Interrogations](const FInterrogationKey& InterrogationKey, FFrameTime& OutEvalTime) { OutEvalTime = Interrogations[InterrogationKey.InterrogationIndex].Time; });
+	.Iterate_PerEntity(
+		&Linker->EntityManager, 
+		[&Interrogations](const FInterrogationKey& InterrogationKey, FFrameTime& OutEvalTime)
+		{
+			const FInterrogationParams& InterrogationParams(Interrogations[InterrogationKey.InterrogationIndex]);
+			OutEvalTime = InterrogationParams.Time;
+		});
+
+	// Compute evaluation time (in seconds) for all entities that need it.
+	FEntityTaskBuilder()
+	.Read(FBuiltInComponentTypes::Get()->Interrogation.InputKey)
+	.Read(FBuiltInComponentTypes::Get()->Interrogation.Instance)
+	.Write(FBuiltInComponentTypes::Get()->EvalSeconds)
+	.FilterNone({ FBuiltInComponentTypes::Get()->Tags.FixedTime })
+	.Iterate_PerEntity(
+		&Linker->EntityManager, 
+		[&Interrogations](const FInterrogationKey& InterrogationKey, const FInterrogationInstance& InterrogationInstance, double& OutEvalSeconds)
+		{
+			const FInterrogationParams& InterrogationParams(Interrogations[InterrogationKey.InterrogationIndex]);
+			if (ensure(InterrogationInstance.Sequence && InterrogationInstance.Sequence->GetMovieScene()))
+			{
+				const FFrameRate& TickResolution = InterrogationInstance.Sequence->GetMovieScene()->GetTickResolution();
+				OutEvalSeconds = TickResolution.AsSeconds(InterrogationParams.Time);
+			}
+		});
 
 	FMovieSceneEntitySystemRunner Runner;
 	Runner.AttachToLinker(Linker);
diff --git a/Engine/Source/Runtime/MovieSceneTracks/Private/MovieSceneTracksComponentTypes.cpp b/Engine/Source/Runtime/MovieSceneTracks/Private/MovieSceneTracksComponentTypes.cpp
--- a/Engine/Source/Runtime/MovieSceneTracks/Private/MovieSceneTracksComponentTypes.cpp
+++ b/Engine/Source/Runtime/MovieSceneTracks/Private/MovieSceneTracksComponentTypes.cpp
@@ -798,7 +798,10 @@ FMovieSceneTracksComponentTypes::FMovieSceneTracksComponentTypes()
 	// --------------------------------------------------------------------------------------------
 	// Set up PerlinNoise components
 	ComponentRegistry->Factories.DuplicateChildComponent(FloatPerlinNoiseChannel);
+	ComponentRegistry->Factories.DefineMutuallyInclusiveComponent(FloatPerlinNoiseChannel, BuiltInComponents->EvalSeconds);
+
 	ComponentRegistry->Factories.DuplicateChildComponent(DoublePerlinNoiseChannel);
+	ComponentRegistry->Factories.DefineMutuallyInclusiveComponent(DoublePerlinNoiseChannel, BuiltInComponents->EvalSeconds);
 
 	// --------------------------------------------------------------------------------------------
 	// Set up SkeletalAnimation components
diff --git a/Engine/Source/Runtime/MovieSceneTracks/Private/Systems/DoublePerlinNoiseChannelEvaluatorSystem.cpp b/Engine/Source/Runtime/MovieSceneTracks/Private/Systems/DoublePerlinNoiseChannelEvaluatorSystem.cpp
--- a/Engine/Source/Runtime/MovieSceneTracks/Private/Systems/DoublePerlinNoiseChannelEvaluatorSystem.cpp
+++ b/Engine/Source/Runtime/MovieSceneTracks/Private/Systems/DoublePerlinNoiseChannelEvaluatorSystem.cpp
@@ -20,21 +20,9 @@ namespace UE
 	{
 		struct FEvaluateDoublePerlinNoiseChannels
 		{
-			FInstanceRegistry* InstanceRegistry;
-
-			FEvaluateDoublePerlinNoiseChannels(FInstanceRegistry* InInstanceRegistry)
-				: InstanceRegistry{ InInstanceRegistry }
+			void ForEachEntity(double EvalSeconds, const FPerlinNoiseParams& PerlinNoiseParams, double& OutResult)
 			{
-				check(InstanceRegistry);
-			}
-
-			void ForEachEntity(const FPerlinNoiseParams& PerlinNoiseParams, FInstanceHandle InstanceHandle, double& OutResult)
-			{
-				const FMovieSceneContext& Context = InstanceRegistry->GetContext(InstanceHandle);
-				FFrameTime Time = Context.GetTime();
-				FFrameRate Rate = Context.GetFrameRate();
-
-				OutResult = FMovieSceneDoublePerlinNoiseChannel::Evaluate(PerlinNoiseParams, Rate.AsSeconds(Time));
+				OutResult = FMovieSceneDoublePerlinNoiseChannel::Evaluate(PerlinNoiseParams, EvalSeconds);
 			}
 		};
 	} // namespace MovieScene
@@ -75,12 +63,12 @@ void UDoublePerlinNoiseChannelEvaluatorSystem::OnRun(FSystemTaskPrerequisites& I
 	for (int32 i = 0; i < UE_ARRAY_COUNT(BuiltInComponents->DoubleResult); ++i)
 	{
 		FEntityTaskBuilder()
+			.Read(BuiltInComponents->EvalSeconds)
 			.Read(TrackComponents->DoublePerlinNoiseChannel)
-			.Read(BuiltInComponents->InstanceHandle)
 			.Write(BuiltInComponents->DoubleResult[i])
 			.FilterNone({ BuiltInComponents->Tags.Ignored })
 			.SetStat(GET_STATID(MovieSceneEval_EvaluateDoublePerlinNoiseChannelTask))
-			.Dispatch_PerEntity<FEvaluateDoublePerlinNoiseChannels>(&Linker->EntityManager, InPrerequisites, &Subsequents, GetLinker()->GetInstanceRegistry());
+			.Dispatch_PerEntity<FEvaluateDoublePerlinNoiseChannels>(&Linker->EntityManager, InPrerequisites, &Subsequents);
 	}
 }
 
diff --git a/Engine/Source/Runtime/MovieSceneTracks/Private/Systems/FloatPerlinNoiseChannelEvaluatorSystem.cpp b/Engine/Source/Runtime/MovieSceneTracks/Private/Systems/FloatPerlinNoiseChannelEvaluatorSystem.cpp
--- a/Engine/Source/Runtime/MovieSceneTracks/Private/Systems/FloatPerlinNoiseChannelEvaluatorSystem.cpp
+++ b/Engine/Source/Runtime/MovieSceneTracks/Private/Systems/FloatPerlinNoiseChannelEvaluatorSystem.cpp
@@ -20,21 +20,9 @@ namespace UE
 	{
 		struct FEvaluateFloatPerlinNoiseChannels
 		{
-			FInstanceRegistry* InstanceRegistry;
-
-			FEvaluateFloatPerlinNoiseChannels(FInstanceRegistry* InInstanceRegistry)
-				: InstanceRegistry{ InInstanceRegistry }
+			void ForEachEntity(double EvalSeconds, const FPerlinNoiseParams& PerlinNoiseParams, double& OutResult)
 			{
-				check(InstanceRegistry);
-			}
-
-			void ForEachEntity(const FPerlinNoiseParams& PerlinNoiseParams, FInstanceHandle InstanceHandle, double& OutResult)
-			{
-				const FMovieSceneContext& Context = InstanceRegistry->GetContext(InstanceHandle);
-				FFrameTime Time = Context.GetTime();
-				FFrameRate Rate = Context.GetFrameRate();
-
-				const float NoiseResult = FMovieSceneFloatPerlinNoiseChannel::Evaluate(PerlinNoiseParams, Rate.AsSeconds(Time));
+				const float NoiseResult = FMovieSceneFloatPerlinNoiseChannel::Evaluate(PerlinNoiseParams, EvalSeconds);
 				OutResult = (double)NoiseResult;
 			}
 		};
@@ -76,12 +64,12 @@ void UFloatPerlinNoiseChannelEvaluatorSystem::OnRun(FSystemTaskPrerequisites& In
 	for (int32 i = 0; i < UE_ARRAY_COUNT(BuiltInComponents->DoubleResult); ++i)
 	{
 		FEntityTaskBuilder()
+			.Read(BuiltInComponents->EvalSeconds)
 			.Read(TrackComponents->FloatPerlinNoiseChannel)
-			.Read(BuiltInComponents->InstanceHandle)
 			.Write(BuiltInComponents->DoubleResult[i])
 			.FilterNone({ BuiltInComponents->Tags.Ignored })
 			.SetStat(GET_STATID(MovieSceneEval_EvaluateFloatPerlinNoiseChannelTask))
-			.Dispatch_PerEntity<FEvaluateFloatPerlinNoiseChannels>(&Linker->EntityManager, InPrerequisites, &Subsequents, GetLinker()->GetInstanceRegistry());
+			.Dispatch_PerEntity<FEvaluateFloatPerlinNoiseChannels>(&Linker->EntityManager, InPrerequisites, &Subsequents);
 	}
 }
 
diff --git a/Engine/Source/Runtime/MovieSceneTracks/Public/EntitySystem/Interrogation/MovieSceneInterrogationLinker.h b/Engine/Source/Runtime/MovieSceneTracks/Public/EntitySystem/Interrogation/MovieSceneInterrogationLinker.h
--- a/Engine/Source/Runtime/MovieSceneTracks/Public/EntitySystem/Interrogation/MovieSceneInterrogationLinker.h
+++ b/Engine/Source/Runtime/MovieSceneTracks/Public/EntitySystem/Interrogation/MovieSceneInterrogationLinker.h
@@ -3,6 +3,7 @@
 #pragma once
 
 #include "CoreTypes.h"
+#include "Containers/SparseArray.h"
 #include "Containers/SortedMap.h"
 #include "Evaluation/MovieSceneEvaluationField.h"
 #include "MovieSceneTracksComponentTypes.h"
@@ -605,7 +606,16 @@ protected:
 	/** The linker we own */
 	UMovieSceneEntitySystemLinker* Linker;
 
+	/** Initial value cache */
 	TSharedPtr<FInitialValueCache> InitialValueCache;
+
+	/** Array of shared extra metadata on entities */
+	struct FExtraMetaData
+	{
+		UE::MovieScene::FInterrogationInstance InterrogationInstance;
+		UE::MovieScene::FInterrogationChannel InterrogationChannel;
+	};
+	TSparseArray<FExtraMetaData> ExtraMetaData;
 };
 
 
