Sequencer: Protect against crash on renaming property array

diff --git a/Engine/Source/Runtime/MovieScene/Private/MovieSceneCommonHelpers.cpp b/Engine/Source/Runtime/MovieScene/Private/MovieSceneCommonHelpers.cpp
--- a/Engine/Source/Runtime/MovieScene/Private/MovieSceneCommonHelpers.cpp
+++ b/Engine/Source/Runtime/MovieScene/Private/MovieSceneCommonHelpers.cpp
@@ -680,28 +680,31 @@ FTrackInstancePropertyBindings::FPropertyAddress FTrackInstancePropertyBindings:
 
 	if (PropertyAndIndex.ArrayIndex != INDEX_NONE)
 	{
-		if (PropertyAndIndex.Property->IsA(FArrayProperty::StaticClass()))
+		if (PropertyAndIndex.Property)
 		{
-			FArrayProperty* ArrayProp = CastFieldChecked<FArrayProperty>(PropertyAndIndex.Property);
-
-			FScriptArrayHelper ArrayHelper(ArrayProp, ArrayProp->ContainerPtrToValuePtr<void>(BasePointer));
-			if (ArrayHelper.IsValidIndex(PropertyAndIndex.ArrayIndex))
+			if (PropertyAndIndex.Property->IsA(FArrayProperty::StaticClass()))
 			{
-				FStructProperty* InnerStructProp = CastField<FStructProperty>(ArrayProp->Inner);
-				if (InnerStructProp && InPropertyNames.IsValidIndex(Index + 1))
-				{
-					return FindPropertyAddressRecursive(ArrayHelper.GetRawPtr(PropertyAndIndex.ArrayIndex), InnerStructProp->Struct, InPropertyNames, Index + 1);
-				}
-				else
+				FArrayProperty* ArrayProp = CastFieldChecked<FArrayProperty>(PropertyAndIndex.Property);
+
+				FScriptArrayHelper ArrayHelper(ArrayProp, ArrayProp->ContainerPtrToValuePtr<void>(BasePointer));
+				if (ArrayHelper.IsValidIndex(PropertyAndIndex.ArrayIndex))
 				{
-					NewAddress.Property = ArrayProp->Inner;
-					NewAddress.Address = ArrayHelper.GetRawPtr(PropertyAndIndex.ArrayIndex);
+					FStructProperty* InnerStructProp = CastField<FStructProperty>(ArrayProp->Inner);
+					if (InnerStructProp && InPropertyNames.IsValidIndex(Index + 1))
+					{
+						return FindPropertyAddressRecursive(ArrayHelper.GetRawPtr(PropertyAndIndex.ArrayIndex), InnerStructProp->Struct, InPropertyNames, Index + 1);
+					}
+					else
+					{
+						NewAddress.Property = ArrayProp->Inner;
+						NewAddress.Address = ArrayHelper.GetRawPtr(PropertyAndIndex.ArrayIndex);
+					}
 				}
 			}
-		}
-		else if (PropertyAndIndex.Property)
-		{
-			UE_LOG(LogMovieScene, Error, TEXT("Mismatch in property evaluation. %s is not of type: %s"), *PropertyAndIndex.Property->GetName(), *FArrayProperty::StaticClass()->GetName());
+			else
+			{
+				UE_LOG(LogMovieScene, Error, TEXT("Mismatch in property evaluation. %s is not of type: %s"), *PropertyAndIndex.Property->GetName(), *FArrayProperty::StaticClass()->GetName());
+			}
 		}
 	}
 	else if (FStructProperty* StructProp = CastField<FStructProperty>(PropertyAndIndex.Property))
