[Backout] - CL23240016
Original CL Desc
-----------------------------------------------------------------
Sequencer: make the perlin noise channel evaluators use a new "EvalSeconds" component

This prevents these evaluators from having to look up the tick resolution on the sequence instance. First, this isn't very good for cache coherency, and second the instance handles aren't available during interrogations, so none of the perlin noise channels could evaluate, resulting in potentially incorrect keyed values.

diff --git a/Engine/Source/Runtime/MovieScene/Private/EntitySystem/BuiltInComponentTypes.cpp b/Engine/Source/Runtime/MovieScene/Private/EntitySystem/BuiltInComponentTypes.cpp
--- a/Engine/Source/Runtime/MovieScene/Private/EntitySystem/BuiltInComponentTypes.cpp
+++ b/Engine/Source/Runtime/MovieScene/Private/EntitySystem/BuiltInComponentTypes.cpp
@@ -32,7 +32,6 @@ FBuiltInComponentTypes::FBuiltInComponentTypes()
 	ComponentRegistry->NewComponentType(&RootInstanceHandle,    TEXT("Root Instance Handle"));
 
 	ComponentRegistry->NewComponentType(&EvalTime,              TEXT("Eval Time"));
-	ComponentRegistry->NewComponentType(&EvalSeconds,           TEXT("Eval Seconds"));
 
 	ComponentRegistry->NewComponentType(&BoundObject,           TEXT("Bound Object"));
 
@@ -127,7 +126,6 @@ FBuiltInComponentTypes::FBuiltInComponentTypes()
 	ComponentRegistry->NewComponentType(&EvaluationHookFlags,   TEXT("Evaluation Hook Flags"), EComponentTypeFlags::Preserved);
 
 	ComponentRegistry->NewComponentType(&Interrogation.InputKey,  TEXT("Interrogation Input"));
-	ComponentRegistry->NewComponentType(&Interrogation.Instance,  TEXT("Interrogation Instance"));
 	ComponentRegistry->NewComponentType(&Interrogation.OutputKey, TEXT("Interrogation Output"));
 
 	Tags.RestoreState            = ComponentRegistry->NewTag(TEXT("Is Restore State Entity"));
@@ -171,7 +169,6 @@ FBuiltInComponentTypes::FBuiltInComponentTypes()
 	ComponentRegistry->Factories.DefineMutuallyInclusiveComponent(Tags.SectionPreRoll,Tags.PreRoll);
 
 	ComponentRegistry->Factories.DuplicateChildComponent(EvalTime);
-	ComponentRegistry->Factories.DuplicateChildComponent(EvalSeconds);
 	ComponentRegistry->Factories.DuplicateChildComponent(BaseValueEvalTime);
 
 	ComponentRegistry->Factories.DuplicateChildComponent(SequenceID);
diff --git a/Engine/Source/Runtime/MovieScene/Private/EntitySystem/MovieSceneEvalTimeSystem.cpp b/Engine/Source/Runtime/MovieScene/Private/EntitySystem/MovieSceneEvalTimeSystem.cpp
--- a/Engine/Source/Runtime/MovieScene/Private/EntitySystem/MovieSceneEvalTimeSystem.cpp
+++ b/Engine/Source/Runtime/MovieScene/Private/EntitySystem/MovieSceneEvalTimeSystem.cpp
@@ -14,28 +14,34 @@
 DECLARE_CYCLE_STAT(TEXT("EvalTime System"), MovieSceneEval_EvalTimeSystem, STATGROUP_MovieSceneECS);
 DECLARE_CYCLE_STAT(TEXT("MovieScene: Assign evaluation times"), MovieSceneEval_EvalTimes, STATGROUP_MovieSceneECS);
 
+struct FAssignEvalTimesTask
+{
+	TArray<FFrameTime>* FrameTimes;
+
+	explicit FAssignEvalTimesTask(TArray<FFrameTime>* InFrameTimes)
+		: FrameTimes(InFrameTimes)
+	{}
+
+	void ForEachEntity(UE::MovieScene::FInstanceHandle InstanceHandle, FFrameTime& EvalTime)
+	{
+		EvalTime = (*FrameTimes)[InstanceHandle.InstanceID];
+	}
+};
+
 UMovieSceneEvalTimeSystem::UMovieSceneEvalTimeSystem(const FObjectInitializer& ObjInit)
 	: Super(ObjInit)
 {
 	using namespace UE::MovieScene;
 
-	const FBuiltInComponentTypes* BuiltInComponents = FBuiltInComponentTypes::Get();
-
+	RelevantComponent = FBuiltInComponentTypes::Get()->EvalTime;
 	SystemCategories = EEntitySystemCategory::Core;
-	RelevantFilter.Any({ BuiltInComponents->EvalTime, BuiltInComponents->EvalSeconds });
 
 	if (HasAnyFlags(RF_ClassDefaultObject))
 	{
-		DefineComponentProducer(GetClass(), BuiltInComponents->EvalTime);
-		DefineComponentProducer(GetClass(), BuiltInComponents->EvalSeconds);
+		DefineComponentProducer(GetClass(), FBuiltInComponentTypes::Get()->EvalTime);
 	}
 }
 
-bool UMovieSceneEvalTimeSystem::IsRelevantImpl(UMovieSceneEntitySystemLinker* InLinker) const
-{
-	return InLinker->EntityManager.Contains(RelevantFilter);
-}
-
 void UMovieSceneEvalTimeSystem::OnRun(FSystemTaskPrerequisites& InPrerequisites, FSystemSubsequentTasks& Subsequents)
 {
 	using namespace UE::MovieScene;
@@ -44,13 +50,13 @@ void UMovieSceneEvalTimeSystem::OnRun(FSystemTaskPrerequisites& InPrerequisites,
 
 	struct FGatherTimes
 	{
-		FGatherTimes(const FInstanceRegistry* InInstanceRegistry, TArray<FEvaluatedTime>* InEvaluatedTimes)
+		FGatherTimes(const FInstanceRegistry* InInstanceRegistry, TArray<FFrameTime>* InFrameTimes)
 			: InstanceRegistry(InInstanceRegistry)
-			, EvaluatedTimes(InEvaluatedTimes)
+			, FrameTimes(InFrameTimes)
 		{}
 
 		const FInstanceRegistry* InstanceRegistry;
-		TArray<FEvaluatedTime>* EvaluatedTimes;
+		TArray<FFrameTime>* FrameTimes;
 
 		FORCEINLINE TStatId           GetStatId() const    { RETURN_QUICK_DECLARE_CYCLE_STAT(FGenericTask, STATGROUP_TaskGraphTasks); }
 		static ENamedThreads::Type    GetDesiredThread()   { return ENamedThreads::AnyHiPriThreadHiPriTask; }
@@ -63,60 +69,26 @@ void UMovieSceneEvalTimeSystem::OnRun(FSystemTaskPrerequisites& InPrerequisites,
 
 		void Run()
 		{
-			EvaluatedTimes->SetNum(InstanceRegistry->GetSparseInstances().GetMaxIndex());
+			FrameTimes->SetNum(InstanceRegistry->GetSparseInstances().GetMaxIndex());
 			for (auto It = InstanceRegistry->GetSparseInstances().CreateConstIterator(); It; ++It)
 			{
-				const FMovieSceneContext& Context = It->GetContext();
-
-				FEvaluatedTime EvaluatedTime;
-				EvaluatedTime.FrameTime = Context.GetTime();
-				EvaluatedTime.Seconds = Context.GetFrameRate().AsSeconds(EvaluatedTime.FrameTime);
-
-				(*EvaluatedTimes)[It.GetIndex()] = EvaluatedTime;
+				(*FrameTimes)[It.GetIndex()] = It->GetContext().GetTime();
 			}
 		}
 	};
 
-	struct FAssignEvalTimesTask
-	{
-		TArrayView<const FEvaluatedTime> EvaluatedTimes;
-
-		explicit FAssignEvalTimesTask(TArrayView<const FEvaluatedTime> InEvaluatedTimes)
-			: EvaluatedTimes(InEvaluatedTimes)
-		{}
-
-		void ForEachEntity(UE::MovieScene::FInstanceHandle InstanceHandle, FFrameTime& EvalTime)
-		{
-			EvalTime = EvaluatedTimes[InstanceHandle.InstanceID].FrameTime;
-		}
-	};
-
-	struct FAssignEvalSecondsTask
-	{
-		TArrayView<const FEvaluatedTime> EvaluatedTimes;
-
-		explicit FAssignEvalSecondsTask(TArrayView<const FEvaluatedTime> InEvaluatedTimes)
-			: EvaluatedTimes(InEvaluatedTimes)
-		{}
-
-		void ForEachEntity(UE::MovieScene::FInstanceHandle InstanceHandle, double& EvalSeconds)
-		{
-			EvalSeconds = EvaluatedTimes[InstanceHandle.InstanceID].Seconds;
-		}
-	};
-
 	FSystemTaskPrerequisites EvalPrereqs;
 	FBuiltInComponentTypes* BuiltInComponents = FBuiltInComponentTypes::Get();
 
 	if (Linker->EntityManager.GetThreadingModel() == EEntityThreadingModel::NoThreading)
 	{
-		FGatherTimes(Linker->GetInstanceRegistry(), &EvaluatedTimes).Run();
+		FGatherTimes(Linker->GetInstanceRegistry(), &FrameTimes).Run();
 	}
 	else
 	{
 		// The only thing we depend on is the gather task
 		FGraphEventRef GatherEvalTimesEvent = TGraphTask<FGatherTimes>::CreateTask(nullptr, Linker->EntityManager.GetDispatchThread())
-		.ConstructAndDispatchWhenReady(Linker->GetInstanceRegistry(), &EvaluatedTimes);
+		.ConstructAndDispatchWhenReady(Linker->GetInstanceRegistry(), &FrameTimes);
 
 		EvalPrereqs.AddComponentTask(BuiltInComponents->EvalTime, GatherEvalTimesEvent);
 	}
@@ -126,13 +98,6 @@ void UMovieSceneEvalTimeSystem::OnRun(FSystemTaskPrerequisites& InPrerequisites,
 	.Write(BuiltInComponents->EvalTime)
 	.FilterNone({ FBuiltInComponentTypes::Get()->Tags.FixedTime })
 	.SetStat(GET_STATID(MovieSceneEval_EvalTimes))
-	.Dispatch_PerEntity<FAssignEvalTimesTask>(&Linker->EntityManager, EvalPrereqs, &Subsequents, EvaluatedTimes);
-
-	FEntityTaskBuilder()
-	.Read(BuiltInComponents->InstanceHandle)
-	.Write(BuiltInComponents->EvalSeconds)
-	.FilterNone({ FBuiltInComponentTypes::Get()->Tags.FixedTime })
-	.SetStat(GET_STATID(MovieSceneEval_EvalTimes))
-	.Dispatch_PerEntity<FAssignEvalSecondsTask>(&Linker->EntityManager, EvalPrereqs, &Subsequents, EvaluatedTimes);
+	.Dispatch_PerEntity<FAssignEvalTimesTask>(&Linker->EntityManager, EvalPrereqs, &Subsequents, &FrameTimes);
 }
 
diff --git a/Engine/Source/Runtime/MovieScene/Private/Evaluation/MovieSceneEvaluationField.cpp b/Engine/Source/Runtime/MovieScene/Private/Evaluation/MovieSceneEvaluationField.cpp
--- a/Engine/Source/Runtime/MovieScene/Private/Evaluation/MovieSceneEvaluationField.cpp
+++ b/Engine/Source/Runtime/MovieScene/Private/Evaluation/MovieSceneEvaluationField.cpp
@@ -40,11 +40,6 @@ FMovieSceneEvaluationFieldSharedEntityMetaData& FMovieSceneEntityComponentFieldB
 	return Field->SharedMetaData[SharedMetaDataIndex];
 }
 
-int32 FMovieSceneEntityComponentFieldBuilder::GetSharedMetaDataIndex() const
-{
-	return SharedMetaDataIndex;
-}
-
 int32 FMovieSceneEntityComponentFieldBuilder::FindOrAddEntity(UObject* EntityOwner, uint32 EntityID)
 {
 	FMovieSceneEvaluationFieldEntityKey Key = { EntityOwner, EntityID };
diff --git a/Engine/Source/Runtime/MovieScene/Private/MovieSceneSection.cpp b/Engine/Source/Runtime/MovieScene/Private/MovieSceneSection.cpp
--- a/Engine/Source/Runtime/MovieScene/Private/MovieSceneSection.cpp
+++ b/Engine/Source/Runtime/MovieScene/Private/MovieSceneSection.cpp
@@ -324,7 +324,6 @@ void UMovieSceneSection::BuildDefaultComponents(UMovieSceneEntitySystemLinker* E
 		.AddConditional(Components->Easing,                     FEasingComponentData{ decltype(FEasingComponentData::Section)(this) }, bHasEasing)
 		.AddConditional(Components->HierarchicalBias,           Params.Sequence.HierarchicalBias, Params.Sequence.HierarchicalBias != 0)
 		.AddConditional(Components->Interrogation.InputKey,     Params.InterrogationKey, Params.InterrogationKey.IsValid())
-		.AddConditional(Components->Interrogation.Instance,     Params.InterrogationInstance, Params.InterrogationInstance.IsValid())
 		.AddConditional(Components->EvalTime,                   Params.EntityMetaData ? Params.EntityMetaData->ForcedTime : 0, bHasForcedTime)
 		.AddTagConditional(Components->Tags.RestoreState,       bShouldRestoreState)
 		.AddTagConditional(Components->Tags.FixedTime,          bHasForcedTime)
diff --git a/Engine/Source/Runtime/MovieScene/Public/EntitySystem/BuiltInComponentTypes.h b/Engine/Source/Runtime/MovieScene/Public/EntitySystem/BuiltInComponentTypes.h
--- a/Engine/Source/Runtime/MovieScene/Public/EntitySystem/BuiltInComponentTypes.h
+++ b/Engine/Source/Runtime/MovieScene/Public/EntitySystem/BuiltInComponentTypes.h
@@ -30,7 +30,6 @@ namespace UE::MovieScene
 	struct FInitialValueIndex;
 	struct FInstanceHandle;
 	struct FInterrogationKey;
-	struct FInterrogationInstance;
 	struct FRootInstanceHandle;
 	namespace Interpolation
 	{
@@ -248,8 +247,6 @@ public:
 
 	TComponentTypeID<FFrameTime>          EvalTime;
 
-	TComponentTypeID<double>              EvalSeconds;
-
 public:
 
 	TComponentTypeID<FMovieSceneBlendChannelID> BlendChannelInput;
@@ -406,7 +403,6 @@ public:
 	struct
 	{
 		TComponentTypeID<FInterrogationKey> InputKey;
-		TComponentTypeID<FInterrogationInstance> Instance;
 		TComponentTypeID<FInterrogationKey> OutputKey;
 	} Interrogation;
 
diff --git a/Engine/Source/Runtime/MovieScene/Public/EntitySystem/IMovieSceneEntityProvider.h b/Engine/Source/Runtime/MovieScene/Public/EntitySystem/IMovieSceneEntityProvider.h
--- a/Engine/Source/Runtime/MovieScene/Public/EntitySystem/IMovieSceneEntityProvider.h
+++ b/Engine/Source/Runtime/MovieScene/Public/EntitySystem/IMovieSceneEntityProvider.h
@@ -94,7 +94,6 @@ struct FEntityImportParams
 	uint32 EntityID = 0;
 
 	FInterrogationKey InterrogationKey;
-	FInterrogationInstance InterrogationInstance;
 
 	FEntityImportSequenceParams Sequence;
 
@@ -119,8 +118,8 @@ class MOVIESCENE_API IMovieSceneEntityProvider
 {
 public:
 
-	using FEntityImportParams        = UE::MovieScene::FEntityImportParams;
-	using FImportedEntity            = UE::MovieScene::FImportedEntity;
+	using FEntityImportParams   = UE::MovieScene::FEntityImportParams;
+	using FImportedEntity       = UE::MovieScene::FImportedEntity;
 
 
 	GENERATED_BODY()
diff --git a/Engine/Source/Runtime/MovieScene/Public/EntitySystem/MovieSceneEntitySystemTypes.h b/Engine/Source/Runtime/MovieScene/Public/EntitySystem/MovieSceneEntitySystemTypes.h
--- a/Engine/Source/Runtime/MovieScene/Public/EntitySystem/MovieSceneEntitySystemTypes.h
+++ b/Engine/Source/Runtime/MovieScene/Public/EntitySystem/MovieSceneEntitySystemTypes.h
@@ -18,7 +18,6 @@
 
 #include <initializer_list>
 
-class UMovieSceneSequence;
 namespace UE { namespace MovieScene { class FEntityManager; } }
 
 #ifndef UE_MOVIESCENE_ENTITY_DEBUG
@@ -265,20 +264,6 @@ struct FInterrogationKey
 };
 
 
-/**
- * Sequence instance information for interrogations
- */
-struct FInterrogationInstance
-{
-	UMovieSceneSequence* Sequence = nullptr;
-
-	bool IsValid() const
-	{
-		return Sequence != nullptr;
-	}
-};
-
-
 struct MOVIESCENE_API FEntityComponentFilter
 {
 
diff --git a/Engine/Source/Runtime/MovieScene/Public/EntitySystem/MovieSceneEvalTimeSystem.h b/Engine/Source/Runtime/MovieScene/Public/EntitySystem/MovieSceneEvalTimeSystem.h
--- a/Engine/Source/Runtime/MovieScene/Public/EntitySystem/MovieSceneEvalTimeSystem.h
+++ b/Engine/Source/Runtime/MovieScene/Public/EntitySystem/MovieSceneEvalTimeSystem.h
@@ -5,7 +5,6 @@
 #include "Containers/Array.h"
 #include "EntitySystem/MovieSceneEntityIDs.h"
 #include "EntitySystem/MovieSceneEntitySystem.h"
-#include "EntitySystem/MovieSceneEntitySystemTypes.h"
 #include "Misc/FrameTime.h"
 #include "UObject/ObjectMacros.h"
 #include "UObject/UObjectGlobals.h"
@@ -24,16 +23,8 @@ public:
 	UMovieSceneEvalTimeSystem(const FObjectInitializer& ObjInit);
 
 	virtual void OnRun(FSystemTaskPrerequisites& InPrerequisites, FSystemSubsequentTasks& Subsequents) override;
-	virtual bool IsRelevantImpl(UMovieSceneEntitySystemLinker* InLinker) const override;
 
 private:
-	struct FEvaluatedTime
-	{
-		FFrameTime FrameTime;
-		double Seconds;
-	};
-	TArray<FEvaluatedTime> EvaluatedTimes;
-
-	UE::MovieScene::FEntityComponentFilter RelevantFilter;
+	TArray<FFrameTime> FrameTimes;
 };
 
diff --git a/Engine/Source/Runtime/MovieScene/Public/Evaluation/MovieSceneEvaluationField.h b/Engine/Source/Runtime/MovieScene/Public/Evaluation/MovieSceneEvaluationField.h
--- a/Engine/Source/Runtime/MovieScene/Public/Evaluation/MovieSceneEvaluationField.h
+++ b/Engine/Source/Runtime/MovieScene/Public/Evaluation/MovieSceneEvaluationField.h
@@ -123,6 +123,7 @@ struct FMovieSceneEvaluationFieldEntityMetaData
 	friend bool operator==(const FMovieSceneEvaluationFieldEntityMetaData& A, const FMovieSceneEvaluationFieldEntityMetaData& B)
 	{
 		return A.ForcedTime == B.ForcedTime && 
+			A.InterrogationChannel == B.InterrogationChannel && 
 			A.Flags == B.Flags && 
 			A.bEvaluateInSequencePreRoll == B.bEvaluateInSequencePreRoll && 
 			A.bEvaluateInSequencePostRoll == B.bEvaluateInSequencePostRoll &&
@@ -136,6 +137,9 @@ struct FMovieSceneEvaluationFieldEntityMetaData
 	UPROPERTY()
 	FFrameNumber ForcedTime;
 
+	/** Never serialized */
+	UE::MovieScene::FInterrogationChannel InterrogationChannel;
+
 	UPROPERTY()
 	ESectionEvaluationFlags Flags;
 
@@ -330,11 +334,6 @@ struct MOVIESCENE_API FMovieSceneEntityComponentFieldBuilder
 	 */
 	FMovieSceneEvaluationFieldSharedEntityMetaData& GetSharedMetaData();
 
-	/**
-	 * Access the index of the shared meta-data of this builder.
-	 */
-	int32 GetSharedMetaDataIndex() const;
-
 	/**
 	 * Add meta-data to this tree returning its index within this builder
 	 *
diff --git a/Engine/Source/Runtime/MovieSceneTracks/Private/EntitySystem/Interrogation/MovieSceneInterrogationLinker.cpp b/Engine/Source/Runtime/MovieSceneTracks/Private/EntitySystem/Interrogation/MovieSceneInterrogationLinker.cpp
--- a/Engine/Source/Runtime/MovieSceneTracks/Private/EntitySystem/Interrogation/MovieSceneInterrogationLinker.cpp
+++ b/Engine/Source/Runtime/MovieSceneTracks/Private/EntitySystem/Interrogation/MovieSceneInterrogationLinker.cpp
@@ -684,25 +684,13 @@ void FSystemInterrogator::ImportTrack(UMovieSceneTrack* Track, const FGuid& Obje
 
 	TGuardValue<FEntityManager*> DebugVizGuard(GEntityManagerForDebuggingVisualizers, &Linker->EntityManager);
 
+	FFrameRate TickResolution = Track->GetTypedOuter<UMovieScene>()->GetTickResolution();
+
 	const FMovieSceneTrackEvaluationField& EvaluationField = Track->GetEvaluationField();
 
 	FMovieSceneEntityComponentFieldBuilder FieldBuilder(&EntityComponentField);
 	FieldBuilder.GetSharedMetaData().ObjectBindingID = ObjectBindingID;
 
-	// Creating a new builder allocates a new piece of shared metadata. We will use that shared metadata's index
-	// to keep track of the extra metadata, since the extra metadata is also shared.
-	const int32 SharedMetaDataIndex = FieldBuilder.GetSharedMetaDataIndex();
-	UMovieSceneSequence* OwningSequence = Track->GetTypedOuter<UMovieSceneSequence>();
-	if (!ExtraMetaData.IsValidIndex(SharedMetaDataIndex))
-	{
-		ExtraMetaData.Insert(SharedMetaDataIndex, FExtraMetaData{ FInterrogationInstance { OwningSequence }, InChannel });
-	}
-	else
-	{
-		const FExtraMetaData& InterrogationMetaData = ExtraMetaData[SharedMetaDataIndex];
-		ensure(InterrogationMetaData.InterrogationInstance.Sequence == OwningSequence && InterrogationMetaData.InterrogationChannel == InChannel);
-	}
-
 	for (const FMovieSceneTrackEvaluationFieldEntry& Entry : EvaluationField.Entries)
 	{
 		IMovieSceneEntityProvider* EntityProvider = Cast<IMovieSceneEntityProvider>(Entry.Section);
@@ -724,6 +712,11 @@ void FSystemInterrogator::ImportTrack(UMovieSceneTrack* Track, const FGuid& Obje
 		MetaData.bEvaluateInSequencePreRoll  = Track->EvalOptions.bEvaluateInPreroll;
 		MetaData.bEvaluateInSequencePostRoll = Track->EvalOptions.bEvaluateInPostroll;
 
+		if (InChannel != FInterrogationChannel::Default())
+		{
+			MetaData.InterrogationChannel = InChannel;
+		}
+
 		if (!EntityProvider->PopulateEvaluationField(Entry.Range, MetaData, &FieldBuilder))
 		{
 			const int32 EntityIndex   = FieldBuilder.FindOrAddEntity(Entry.Section, 0);
@@ -776,12 +769,9 @@ void FSystemInterrogator::InterrogateEntity(int32 InterrogationIndex, const FMov
 	Params.InterrogationKey.Channel = FInterrogationChannel::Default();
 	Params.InterrogationKey.InterrogationIndex = InterrogationIndex;
 
-	const int32 SharedMetaDataIndex = Query.Entity.SharedMetaDataIndex;
-	if (SharedMetaDataIndex != INDEX_NONE && ensure(ExtraMetaData.IsValidIndex(SharedMetaDataIndex)))
+	if (Params.EntityMetaData && Params.EntityMetaData->InterrogationChannel.IsValid())
 	{
-		const FExtraMetaData& InterrogationMetaData = ExtraMetaData[SharedMetaDataIndex];
-		Params.InterrogationKey.Channel = InterrogationMetaData.InterrogationChannel;
-		Params.InterrogationInstance = InterrogationMetaData.InterrogationInstance;
+		Params.InterrogationKey.Channel = Params.EntityMetaData->InterrogationChannel;
 	}
 
 	FImportedEntity ImportedEntity;
@@ -810,37 +800,11 @@ void FSystemInterrogator::Update()
 	Linker->LinkRelevantSystems();
 
 	TArrayView<const FInterrogationParams> Interrogations = Channels.GetInterrogations();
-
-	// Compute evaluation time (in frames) for all entities that need it.
 	FEntityTaskBuilder()
 	.Read(FBuiltInComponentTypes::Get()->Interrogation.InputKey)
 	.Write(FBuiltInComponentTypes::Get()->EvalTime)
 	.FilterNone({ FBuiltInComponentTypes::Get()->Tags.FixedTime })
-	.Iterate_PerEntity(
-		&Linker->EntityManager, 
-		[&Interrogations](const FInterrogationKey& InterrogationKey, FFrameTime& OutEvalTime)
-		{
-			const FInterrogationParams& InterrogationParams(Interrogations[InterrogationKey.InterrogationIndex]);
-			OutEvalTime = InterrogationParams.Time;
-		});
-
-	// Compute evaluation time (in seconds) for all entities that need it.
-	FEntityTaskBuilder()
-	.Read(FBuiltInComponentTypes::Get()->Interrogation.InputKey)
-	.Read(FBuiltInComponentTypes::Get()->Interrogation.Instance)
-	.Write(FBuiltInComponentTypes::Get()->EvalSeconds)
-	.FilterNone({ FBuiltInComponentTypes::Get()->Tags.FixedTime })
-	.Iterate_PerEntity(
-		&Linker->EntityManager, 
-		[&Interrogations](const FInterrogationKey& InterrogationKey, const FInterrogationInstance& InterrogationInstance, double& OutEvalSeconds)
-		{
-			const FInterrogationParams& InterrogationParams(Interrogations[InterrogationKey.InterrogationIndex]);
-			if (ensure(InterrogationInstance.Sequence && InterrogationInstance.Sequence->GetMovieScene()))
-			{
-				const FFrameRate& TickResolution = InterrogationInstance.Sequence->GetMovieScene()->GetTickResolution();
-				OutEvalSeconds = TickResolution.AsSeconds(InterrogationParams.Time);
-			}
-		});
+	.Iterate_PerEntity(&Linker->EntityManager, [&Interrogations](const FInterrogationKey& InterrogationKey, FFrameTime& OutEvalTime) { OutEvalTime = Interrogations[InterrogationKey.InterrogationIndex].Time; });
 
 	FMovieSceneEntitySystemRunner Runner;
 	Runner.AttachToLinker(Linker);
diff --git a/Engine/Source/Runtime/MovieSceneTracks/Private/MovieSceneTracksComponentTypes.cpp b/Engine/Source/Runtime/MovieSceneTracks/Private/MovieSceneTracksComponentTypes.cpp
--- a/Engine/Source/Runtime/MovieSceneTracks/Private/MovieSceneTracksComponentTypes.cpp
+++ b/Engine/Source/Runtime/MovieSceneTracks/Private/MovieSceneTracksComponentTypes.cpp
@@ -798,10 +798,7 @@ FMovieSceneTracksComponentTypes::FMovieSceneTracksComponentTypes()
 	// --------------------------------------------------------------------------------------------
 	// Set up PerlinNoise components
 	ComponentRegistry->Factories.DuplicateChildComponent(FloatPerlinNoiseChannel);
-	ComponentRegistry->Factories.DefineMutuallyInclusiveComponent(FloatPerlinNoiseChannel, BuiltInComponents->EvalSeconds);
-
 	ComponentRegistry->Factories.DuplicateChildComponent(DoublePerlinNoiseChannel);
-	ComponentRegistry->Factories.DefineMutuallyInclusiveComponent(DoublePerlinNoiseChannel, BuiltInComponents->EvalSeconds);
 
 	// --------------------------------------------------------------------------------------------
 	// Set up SkeletalAnimation components
diff --git a/Engine/Source/Runtime/MovieSceneTracks/Private/Systems/DoublePerlinNoiseChannelEvaluatorSystem.cpp b/Engine/Source/Runtime/MovieSceneTracks/Private/Systems/DoublePerlinNoiseChannelEvaluatorSystem.cpp
--- a/Engine/Source/Runtime/MovieSceneTracks/Private/Systems/DoublePerlinNoiseChannelEvaluatorSystem.cpp
+++ b/Engine/Source/Runtime/MovieSceneTracks/Private/Systems/DoublePerlinNoiseChannelEvaluatorSystem.cpp
@@ -20,9 +20,21 @@ namespace UE
 	{
 		struct FEvaluateDoublePerlinNoiseChannels
 		{
-			void ForEachEntity(double EvalSeconds, const FPerlinNoiseParams& PerlinNoiseParams, double& OutResult)
+			FInstanceRegistry* InstanceRegistry;
+
+			FEvaluateDoublePerlinNoiseChannels(FInstanceRegistry* InInstanceRegistry)
+				: InstanceRegistry{ InInstanceRegistry }
 			{
-				OutResult = FMovieSceneDoublePerlinNoiseChannel::Evaluate(PerlinNoiseParams, EvalSeconds);
+				check(InstanceRegistry);
+			}
+
+			void ForEachEntity(const FPerlinNoiseParams& PerlinNoiseParams, FInstanceHandle InstanceHandle, double& OutResult)
+			{
+				const FMovieSceneContext& Context = InstanceRegistry->GetContext(InstanceHandle);
+				FFrameTime Time = Context.GetTime();
+				FFrameRate Rate = Context.GetFrameRate();
+
+				OutResult = FMovieSceneDoublePerlinNoiseChannel::Evaluate(PerlinNoiseParams, Rate.AsSeconds(Time));
 			}
 		};
 	} // namespace MovieScene
@@ -63,12 +75,12 @@ void UDoublePerlinNoiseChannelEvaluatorSystem::OnRun(FSystemTaskPrerequisites& I
 	for (int32 i = 0; i < UE_ARRAY_COUNT(BuiltInComponents->DoubleResult); ++i)
 	{
 		FEntityTaskBuilder()
-			.Read(BuiltInComponents->EvalSeconds)
 			.Read(TrackComponents->DoublePerlinNoiseChannel)
+			.Read(BuiltInComponents->InstanceHandle)
 			.Write(BuiltInComponents->DoubleResult[i])
 			.FilterNone({ BuiltInComponents->Tags.Ignored })
 			.SetStat(GET_STATID(MovieSceneEval_EvaluateDoublePerlinNoiseChannelTask))
-			.Dispatch_PerEntity<FEvaluateDoublePerlinNoiseChannels>(&Linker->EntityManager, InPrerequisites, &Subsequents);
+			.Dispatch_PerEntity<FEvaluateDoublePerlinNoiseChannels>(&Linker->EntityManager, InPrerequisites, &Subsequents, GetLinker()->GetInstanceRegistry());
 	}
 }
 
diff --git a/Engine/Source/Runtime/MovieSceneTracks/Private/Systems/FloatPerlinNoiseChannelEvaluatorSystem.cpp b/Engine/Source/Runtime/MovieSceneTracks/Private/Systems/FloatPerlinNoiseChannelEvaluatorSystem.cpp
--- a/Engine/Source/Runtime/MovieSceneTracks/Private/Systems/FloatPerlinNoiseChannelEvaluatorSystem.cpp
+++ b/Engine/Source/Runtime/MovieSceneTracks/Private/Systems/FloatPerlinNoiseChannelEvaluatorSystem.cpp
@@ -20,9 +20,21 @@ namespace UE
 	{
 		struct FEvaluateFloatPerlinNoiseChannels
 		{
-			void ForEachEntity(double EvalSeconds, const FPerlinNoiseParams& PerlinNoiseParams, double& OutResult)
+			FInstanceRegistry* InstanceRegistry;
+
+			FEvaluateFloatPerlinNoiseChannels(FInstanceRegistry* InInstanceRegistry)
+				: InstanceRegistry{ InInstanceRegistry }
 			{
-				const float NoiseResult = FMovieSceneFloatPerlinNoiseChannel::Evaluate(PerlinNoiseParams, EvalSeconds);
+				check(InstanceRegistry);
+			}
+
+			void ForEachEntity(const FPerlinNoiseParams& PerlinNoiseParams, FInstanceHandle InstanceHandle, double& OutResult)
+			{
+				const FMovieSceneContext& Context = InstanceRegistry->GetContext(InstanceHandle);
+				FFrameTime Time = Context.GetTime();
+				FFrameRate Rate = Context.GetFrameRate();
+
+				const float NoiseResult = FMovieSceneFloatPerlinNoiseChannel::Evaluate(PerlinNoiseParams, Rate.AsSeconds(Time));
 				OutResult = (double)NoiseResult;
 			}
 		};
@@ -64,12 +76,12 @@ void UFloatPerlinNoiseChannelEvaluatorSystem::OnRun(FSystemTaskPrerequisites& In
 	for (int32 i = 0; i < UE_ARRAY_COUNT(BuiltInComponents->DoubleResult); ++i)
 	{
 		FEntityTaskBuilder()
-			.Read(BuiltInComponents->EvalSeconds)
 			.Read(TrackComponents->FloatPerlinNoiseChannel)
+			.Read(BuiltInComponents->InstanceHandle)
 			.Write(BuiltInComponents->DoubleResult[i])
 			.FilterNone({ BuiltInComponents->Tags.Ignored })
 			.SetStat(GET_STATID(MovieSceneEval_EvaluateFloatPerlinNoiseChannelTask))
-			.Dispatch_PerEntity<FEvaluateFloatPerlinNoiseChannels>(&Linker->EntityManager, InPrerequisites, &Subsequents);
+			.Dispatch_PerEntity<FEvaluateFloatPerlinNoiseChannels>(&Linker->EntityManager, InPrerequisites, &Subsequents, GetLinker()->GetInstanceRegistry());
 	}
 }
 
diff --git a/Engine/Source/Runtime/MovieSceneTracks/Public/EntitySystem/Interrogation/MovieSceneInterrogationLinker.h b/Engine/Source/Runtime/MovieSceneTracks/Public/EntitySystem/Interrogation/MovieSceneInterrogationLinker.h
--- a/Engine/Source/Runtime/MovieSceneTracks/Public/EntitySystem/Interrogation/MovieSceneInterrogationLinker.h
+++ b/Engine/Source/Runtime/MovieSceneTracks/Public/EntitySystem/Interrogation/MovieSceneInterrogationLinker.h
@@ -3,7 +3,6 @@
 #pragma once
 
 #include "CoreTypes.h"
-#include "Containers/SparseArray.h"
 #include "Containers/SortedMap.h"
 #include "Evaluation/MovieSceneEvaluationField.h"
 #include "MovieSceneTracksComponentTypes.h"
@@ -606,16 +605,7 @@ protected:
 	/** The linker we own */
 	UMovieSceneEntitySystemLinker* Linker;
 
-	/** Initial value cache */
 	TSharedPtr<FInitialValueCache> InitialValueCache;
-
-	/** Array of shared extra metadata on entities */
-	struct FExtraMetaData
-	{
-		UE::MovieScene::FInterrogationInstance InterrogationInstance;
-		UE::MovieScene::FInterrogationChannel InterrogationChannel;
-	};
-	TSparseArray<FExtraMetaData> ExtraMetaData;
 };
 
 
