Sequencer: Skip past empty ranges when populating the material evaluation field

diff --git a/Engine/Source/Runtime/MovieScene/Private/Sections/MovieSceneSpawnSection.cpp b/Engine/Source/Runtime/MovieScene/Private/Sections/MovieSceneSpawnSection.cpp
--- a/Engine/Source/Runtime/MovieScene/Private/Sections/MovieSceneSpawnSection.cpp
+++ b/Engine/Source/Runtime/MovieScene/Private/Sections/MovieSceneSpawnSection.cpp
@@ -83,7 +83,10 @@ bool UMovieSceneSpawnSection::PopulateEvaluationFieldImpl(const TRange<FFrameNum
 			{
 				// Add the last range to the tree
 				TRange<FFrameNumber> Range(StartBound, TRangeBound<FFrameNumber>::Exclusive(Times[Index]));
-				OutFieldBuilder->AddPersistentEntity(Range, EntityIndex, MetaDataIndex);
+				if (!Range.IsEmpty())
+				{
+					OutFieldBuilder->AddPersistentEntity(Range, EntityIndex, MetaDataIndex);
+				}
 			}
 
 			bIsSpawned = Values[Index];
