Sequencer: Fixed conditional mutations not default-constructing new components when the entire allocation is mutated

diff --git a/Engine/Source/Runtime/MovieScene/Private/EntitySystem/MovieSceneEntityManager.cpp b/Engine/Source/Runtime/MovieScene/Private/EntitySystem/MovieSceneEntityManager.cpp
--- a/Engine/Source/Runtime/MovieScene/Private/EntitySystem/MovieSceneEntityManager.cpp
+++ b/Engine/Source/Runtime/MovieScene/Private/EntitySystem/MovieSceneEntityManager.cpp
@@ -1760,6 +1760,25 @@ int32 FEntityManager::MutateConditional(const FEntityComponentFilter& Filter, co
 			// Specific entity entry indices
 			FEntityAllocation* NewAllocation = MigrateAllocation(SourceAllocationIndex, Pair.Value.Key);
 
+			// Default construct all the new components in the allocation, and then allow the mutation to further initialize this data if needed
+			for (FComponentMaskIterator Component = Pair.Value.Key.Iterate(); Component; ++Component)
+			{
+				FComponentTypeID ComponentTypeID = FComponentTypeID::FromBitIndex(Component.GetIndex());
+				if (EntityAllocationMasks[SourceAllocationIndex].Contains(ComponentTypeID))
+				{
+					continue;
+				}
+
+				const FComponentHeader& ComponentHeader = NewAllocation->GetComponentHeaderChecked(ComponentTypeID);
+				if (!ComponentHeader.IsTag())
+				{
+					const FComponentTypeInfo& ComponentTypeInfo = ComponentRegistry->GetComponentTypeChecked(ComponentTypeID);
+
+					void* Components = ComponentHeader.GetValuePtr(0);
+					ComponentTypeInfo.ConstructItems(Components, NewAllocation->Num());
+				}
+			}
+
 			EntityAllocationMasks[SourceAllocationIndex] = Pair.Value.Key;
 			EntityAllocations[SourceAllocationIndex] = NewAllocation;
 
