Sequencer: Fix FindPossessableObjectId to use a dummy player so that it can find child bindings

diff --git a/Engine/Source/Runtime/MovieScene/Private/MovieSceneSequence.cpp b/Engine/Source/Runtime/MovieScene/Private/MovieSceneSequence.cpp
--- a/Engine/Source/Runtime/MovieScene/Private/MovieSceneSequence.cpp
+++ b/Engine/Source/Runtime/MovieScene/Private/MovieSceneSequence.cpp
@@ -3,6 +3,7 @@
 #include "MovieSceneSequence.h"
 
 #include "Evaluation/MovieSceneEvaluationCustomVersion.h"
+#include "Evaluation/MovieSceneEvaluationTemplateInstance.h"
 #include "MovieScene.h"
 #include "UObject/EditorObjectVersion.h"
 #include "UObject/ReleaseObjectVersion.h"
@@ -300,22 +301,23 @@ UMovieSceneCompiledData* UMovieSceneSequence::GetOrCreateCompiledData()
 
 FGuid UMovieSceneSequence::FindPossessableObjectId(UObject& Object, UObject* Context) const
 {
-	UMovieScene* MovieScene = GetMovieScene();
-	if (!MovieScene)
-	{
-		return FGuid();
-	}
-
-	// Search all possessables
-	for (int32 Index = 0; Index < MovieScene->GetPossessableCount(); ++Index)
+	class FTransientPlayer : public IMovieScenePlayer
 	{
-		FGuid ThisGuid = MovieScene->GetPossessable(Index).GetGuid();
-		if (LocateBoundObjects(ThisGuid, Context).Contains(&Object))
-		{
-			return ThisGuid;
-		}
-	}
-	return FGuid();
+	public:
+		FMovieSceneRootEvaluationTemplateInstance Template;
+		virtual FMovieSceneRootEvaluationTemplateInstance& GetEvaluationTemplate() override { check(false); return Template; }
+		virtual void UpdateCameraCut(UObject* CameraObject, const EMovieSceneCameraCutParams& CameraCutParams) override {}
+		virtual void SetViewportSettings(const TMap<FViewportClient*, EMovieSceneViewportParams>& ViewportParamsMap) override {}
+		virtual void GetViewportSettings(TMap<FViewportClient*, EMovieSceneViewportParams>& ViewportParamsMap) const override {}
+		virtual EMovieScenePlayerStatus::Type GetPlaybackStatus() const { return EMovieScenePlayerStatus::Stopped; }
+		virtual void SetPlaybackStatus(EMovieScenePlayerStatus::Type InPlaybackStatus) override {}
+	} Player;
+
+	UMovieSceneSequence* ThisSequence = const_cast<UMovieSceneSequence*>(this);
+	Player.State.AssignSequence(MovieSceneSequenceID::Root, *ThisSequence, Player);
+
+	FGuid ExistingID = Player.FindObjectId(Object, MovieSceneSequenceID::Root);
+	return ExistingID;
 }
 
 FMovieSceneObjectBindingID UMovieSceneSequence::FindBindingByTag(FName InBindingName) const
