Sequencer: Fix final update happening when subsequence ledger wasn't empty yet. This fixes an issue that resulted in the camera cut instance being destroyed too early and not restoring the last valid camera.

diff --git a/Engine/Source/Runtime/MovieScene/Private/EntitySystem/MovieSceneEntitySystemRunner.cpp b/Engine/Source/Runtime/MovieScene/Private/EntitySystem/MovieSceneEntitySystemRunner.cpp
--- a/Engine/Source/Runtime/MovieScene/Private/EntitySystem/MovieSceneEntitySystemRunner.cpp
+++ b/Engine/Source/Runtime/MovieScene/Private/EntitySystem/MovieSceneEntitySystemRunner.cpp
@@ -234,10 +234,10 @@ bool FMovieSceneEntitySystemRunner::QueueFinalUpdateImpl(FInstanceHandle InInsta
 	// 2. we're not in the middle of an update loop 
 	// 3. the instance has no current updates
 	//
-	const bool bHasAnyEntities = !Instance.Ledger.IsEmpty();
+	const bool bCanFinishImmediately = Instance.CanFinishImmediately(Linker);
 	const ERunnerFlushState UnsafeDestroyMask = ERunnerFlushState::Everything & ~(ERunnerFlushState::PostEvaluation | ERunnerFlushState::End);
 	const bool bSafeToDestroyNow = !EnumHasAnyFlags(FlushState, UnsafeDestroyMask);
-	if (!bHasAnyEntities && bSafeToDestroyNow && !HasQueuedUpdates(InInstanceHandle))
+	if (bCanFinishImmediately && bSafeToDestroyNow && !HasQueuedUpdates(InInstanceHandle))
 	{
 		Instance.Finish(Linker);
 		Instance.PostEvaluation(Linker);
diff --git a/Engine/Source/Runtime/MovieScene/Private/EntitySystem/MovieSceneSequenceInstance.cpp b/Engine/Source/Runtime/MovieScene/Private/EntitySystem/MovieSceneSequenceInstance.cpp
--- a/Engine/Source/Runtime/MovieScene/Private/EntitySystem/MovieSceneSequenceInstance.cpp
+++ b/Engine/Source/Runtime/MovieScene/Private/EntitySystem/MovieSceneSequenceInstance.cpp
@@ -239,6 +239,18 @@ void FSequenceInstance::Update(UMovieSceneEntitySystemLinker* Linker, const FMov
 	SequenceUpdater->Update(Linker, RootInstanceHandle, GetPlayer(), InContext);
 }
 
+bool FSequenceInstance::CanFinishImmediately(UMovieSceneEntitySystemLinker* Linker) const
+{
+	if (SequenceUpdater)
+	{
+		check(RootInstanceHandle == InstanceHandle);
+
+		return SequenceUpdater->CanFinishImmediately(Linker, RootInstanceHandle);
+	}
+
+	return true;
+}
+
 void FSequenceInstance::Finish(UMovieSceneEntitySystemLinker* Linker)
 {
 	if (IsRootSequence() && !bHasEverUpdated)
diff --git a/Engine/Source/Runtime/MovieScene/Private/EntitySystem/MovieSceneSequenceUpdaters.cpp b/Engine/Source/Runtime/MovieScene/Private/EntitySystem/MovieSceneSequenceUpdaters.cpp
--- a/Engine/Source/Runtime/MovieScene/Private/EntitySystem/MovieSceneSequenceUpdaters.cpp
+++ b/Engine/Source/Runtime/MovieScene/Private/EntitySystem/MovieSceneSequenceUpdaters.cpp
@@ -38,6 +38,7 @@ struct FSequenceUpdater_Flat : ISequenceUpdater
 	virtual void DissectContext(UMovieSceneEntitySystemLinker* Linker, IMovieScenePlayer* InPlayer, const FMovieSceneContext& Context, TArray<TRange<FFrameTime>>& OutDissections) override;
 	virtual void Start(UMovieSceneEntitySystemLinker* Linker, FRootInstanceHandle InstanceHandle, IMovieScenePlayer* InPlayer, const FMovieSceneContext& InContext) override;
 	virtual void Update(UMovieSceneEntitySystemLinker* Linker, FRootInstanceHandle InstanceHandle, IMovieScenePlayer* InPlayer, const FMovieSceneContext& Context) override;
+	virtual bool CanFinishImmediately(UMovieSceneEntitySystemLinker* InLinker, FRootInstanceHandle RootInstanceHandle) const override;
 	virtual void Finish(UMovieSceneEntitySystemLinker* Linker, FRootInstanceHandle InstanceHandle, IMovieScenePlayer* InPlayer) override;
 	virtual void InvalidateCachedData(UMovieSceneEntitySystemLinker* Linker) override;
 	virtual void Destroy(UMovieSceneEntitySystemLinker* Linker) override;
@@ -65,6 +66,7 @@ struct FSequenceUpdater_Hierarchical : ISequenceUpdater
 	virtual void DissectContext(UMovieSceneEntitySystemLinker* Linker, IMovieScenePlayer* InPlayer, const FMovieSceneContext& Context, TArray<TRange<FFrameTime>>& OutDissections) override;
 	virtual void Start(UMovieSceneEntitySystemLinker* Linker, FRootInstanceHandle InstanceHandle, IMovieScenePlayer* InPlayer, const FMovieSceneContext& InContext) override;
 	virtual void Update(UMovieSceneEntitySystemLinker* Linker, FRootInstanceHandle InstanceHandle, IMovieScenePlayer* InPlayer, const FMovieSceneContext& Context) override;
+	virtual bool CanFinishImmediately(UMovieSceneEntitySystemLinker* InLinker, FRootInstanceHandle RootInstanceHandle) const override;
 	virtual void Finish(UMovieSceneEntitySystemLinker* Linker, FRootInstanceHandle InstanceHandle, IMovieScenePlayer* InPlayer) override;
 	virtual void InvalidateCachedData(UMovieSceneEntitySystemLinker* Linker) override;
 	virtual void Destroy(UMovieSceneEntitySystemLinker* Linker) override;
@@ -286,12 +288,17 @@ void FSequenceUpdater_Flat::Update(UMovieSceneEntitySystemLinker* Linker, FRootI
 	}
 }
 
+bool FSequenceUpdater_Flat::CanFinishImmediately(UMovieSceneEntitySystemLinker* InLinker, FRootInstanceHandle RootInstanceHandle) const
+{
+	const FSequenceInstance& SequenceInstance = InLinker->GetInstanceRegistry()->GetInstance(RootInstanceHandle);
+	return SequenceInstance.Ledger.IsEmpty();
+}
+
 void FSequenceUpdater_Flat::Finish(UMovieSceneEntitySystemLinker* Linker, FRootInstanceHandle InstanceHandle, IMovieScenePlayer* InPlayer)
 {
 	InvalidateCachedData(Linker);
 }
 
-
 void FSequenceUpdater_Flat::Destroy(UMovieSceneEntitySystemLinker* Linker)
 {
 }
@@ -680,6 +687,28 @@ void FSequenceUpdater_Hierarchical::Update(UMovieSceneEntitySystemLinker* Linker
 	}
 }
 
+bool FSequenceUpdater_Hierarchical::CanFinishImmediately(UMovieSceneEntitySystemLinker* InLinker, FRootInstanceHandle RootInstanceHandle) const
+{
+	FInstanceRegistry* InstanceRegistry = InLinker->GetInstanceRegistry();
+
+	const FSequenceInstance& RootInstance = InstanceRegistry->GetInstance(RootInstanceHandle);
+	if (!RootInstance.Ledger.IsEmpty())
+	{
+		return false;
+	}
+
+	for (TPair<FMovieSceneSequenceID, FSubInstanceData> Pair : SequenceInstances)
+	{
+		const FSequenceInstance& SubInstance = InstanceRegistry->GetInstance(Pair.Value.Handle);
+		if (!SubInstance.Ledger.IsEmpty())
+		{
+			return false;
+		}
+	}
+
+	return true;
+}
+
 void FSequenceUpdater_Hierarchical::Finish(UMovieSceneEntitySystemLinker* Linker, FRootInstanceHandle InstanceHandle, IMovieScenePlayer* InPlayer)
 {
 	FInstanceRegistry* InstanceRegistry = Linker->GetInstanceRegistry();
diff --git a/Engine/Source/Runtime/MovieScene/Public/EntitySystem/MovieSceneSequenceInstance.h b/Engine/Source/Runtime/MovieScene/Public/EntitySystem/MovieSceneSequenceInstance.h
--- a/Engine/Source/Runtime/MovieScene/Public/EntitySystem/MovieSceneSequenceInstance.h
+++ b/Engine/Source/Runtime/MovieScene/Public/EntitySystem/MovieSceneSequenceInstance.h
@@ -82,6 +82,15 @@ public:
 	void Update(UMovieSceneEntitySystemLinker* Linker, const FMovieSceneContext& InContext);
 
 
+	/**
+	 * Returns whether this instance can be finished immediately without any last update.
+	 *
+	 * @param Linker     The linker that owns this sequence instance
+	 * @return           Whether the instance can be finished immediately
+	 */
+	bool CanFinishImmediately(UMovieSceneEntitySystemLinker* Linker) const;
+
+
 	/**
 	 * Mark this instance as finished, causing all its entities to be unlinked and the instance to become inactive at the end of the next update.
 	 *
diff --git a/Engine/Source/Runtime/MovieScene/Public/EntitySystem/MovieSceneSequenceUpdaters.h b/Engine/Source/Runtime/MovieScene/Public/EntitySystem/MovieSceneSequenceUpdaters.h
--- a/Engine/Source/Runtime/MovieScene/Public/EntitySystem/MovieSceneSequenceUpdaters.h
+++ b/Engine/Source/Runtime/MovieScene/Public/EntitySystem/MovieSceneSequenceUpdaters.h
@@ -86,6 +86,16 @@ public:
 	virtual void Update(UMovieSceneEntitySystemLinker* InLinker, FRootInstanceHandle InInstanceHandle, IMovieScenePlayer* InPlayer, const FMovieSceneContext& InContext) = 0;
 
 
+	/**
+	 * Returns whether this instance can be finished immediately without any last update.
+	 *
+	 * @param Linker              The linker that owns this sequence instance
+	 * @param RootInstanceHandle  The handle to the root instance
+	 * @return                    Whether the instance can be finished immediately
+	 */
+	virtual bool CanFinishImmediately(UMovieSceneEntitySystemLinker* InLinker, FRootInstanceHandle RootInstanceHandle) const = 0;
+
+
 	/**
 	 * Called before evaluation when this updater's sequence is no longer required to be evaluated
 	 *
@@ -137,4 +147,4 @@ public:
 
 
 } // namespace MovieScene
-} // namespace UE
\ No newline at end of file
+} // namespace UE
