Sequencer: don't wait for a RPC stop event from the server when a replicated sequence is played only on the client

diff --git a/Engine/Source/Runtime/MovieScene/Private/MovieSceneSequencePlayer.cpp b/Engine/Source/Runtime/MovieScene/Private/MovieSceneSequencePlayer.cpp
--- a/Engine/Source/Runtime/MovieScene/Private/MovieSceneSequencePlayer.cpp
+++ b/Engine/Source/Runtime/MovieScene/Private/MovieSceneSequencePlayer.cpp
@@ -1102,12 +1102,15 @@ void UMovieSceneSequencePlayer::UpdateTimeCursorPosition_Internal(FFrameTime New
 			FMovieSceneEvaluationRange Range = UpdatePlayPosition(PlayPosition, NewPosition, Method);
 			UpdateMovieSceneInstance(Range, StatusOverride);
 
-			// If we are using replicated playback and are not an authoritative sequence player, we need to wait 
-			// for the server to tell us to finish.
+			// We have authority to finish playback if:
+			// 1. There's no playback replication (standalone sequence)
+			// 2. We are the server side of a replicated sequence
+			// 3. We are the client side of a replicated sequence, but playing is only happening on our side (i.e. the Play() method was
+			//    called only on the client, and the server sequence is stopped)
 			const bool bHasAuthorityToFinish = (
-					HasAuthority() || 
-					PlaybackClient == nullptr ||
-					!PlaybackClient->GetIsReplicatedPlayback());
+				(PlaybackClient == nullptr || !PlaybackClient->GetIsReplicatedPlayback()) ||
+				HasAuthority() ||
+				NetSyncProps.LastKnownStatus == EMovieScenePlayerStatus::Stopped);
 			if (bHasAuthorityToFinish)
 			{
 				FinishPlaybackInternal(NewPosition);
