Disabling the CachedChannelEvaluation logic since it was breaking animations playing in reverse on certain platforms.

diff --git a/Engine/Source/Runtime/MovieScene/Private/Channels/MovieSceneCurveChannelImpl.cpp b/Engine/Source/Runtime/MovieScene/Private/Channels/MovieSceneCurveChannelImpl.cpp
--- a/Engine/Source/Runtime/MovieScene/Private/Channels/MovieSceneCurveChannelImpl.cpp
+++ b/Engine/Source/Runtime/MovieScene/Private/Channels/MovieSceneCurveChannelImpl.cpp
@@ -36,7 +36,7 @@ static FAutoConsoleVariableRef CVarCachedChannelEvaluationParityThreshold(
 	TEXT("Threshold for testing Evaluate parity with cached/uncached routines."),
 	ECVF_Default);
 
-MOVIESCENE_API bool GEnableCachedChannelEvaluation = true;
+MOVIESCENE_API bool GEnableCachedChannelEvaluation = false;
 static FAutoConsoleVariableRef CVarEnableCachedChannelEvaluation(
 	TEXT("Sequencer.EnableCachedChannelEvaluation"),
 	GEnableCachedChannelEvaluation,
