Sequencer: Fixed reversed dissected evaluation ranges when playing backwards

diff --git a/Engine/Source/Runtime/MovieScene/Private/EntitySystem/MovieSceneEntitySystemRunner.cpp b/Engine/Source/Runtime/MovieScene/Private/EntitySystem/MovieSceneEntitySystemRunner.cpp
--- a/Engine/Source/Runtime/MovieScene/Private/EntitySystem/MovieSceneEntitySystemRunner.cpp
+++ b/Engine/Source/Runtime/MovieScene/Private/EntitySystem/MovieSceneEntitySystemRunner.cpp
@@ -9,6 +9,7 @@
 #include "Evaluation/MovieSceneEvaluationTemplateInstance.h"
 #include "IMovieScenePlayer.h"
 #include "MovieSceneSequence.h"
+#include "Algo/Reverse.h"
 #include "ProfilingDebugging/CountersTrace.h"
 
 DECLARE_CYCLE_STAT(TEXT("Runner Flush"), 				MovieSceneEval_RunnerFlush, 				STATGROUP_MovieSceneEval);
@@ -616,6 +617,11 @@ bool FMovieSceneEntitySystemRunner::GameThread_UpdateSequenceInstances(UMovieSce
 
 			if (Dissections.Num() != 0)
 			{
+				if (Request.Context.GetDirection() == EPlayDirection::Backwards)
+				{
+					Algo::Reverse(Dissections);
+				}
+
 				for (int32 Index = 0; Index < Dissections.Num() - 1; ++Index)
 				{
 					// Never finish or destroy sequence instances until the _last_ dissected update
