Sequencer:  Default Interpolation type will be based upon previous key set not the default setting. Default only used when setting first key.

diff --git a/Engine/Plugins/Animation/ControlRig/Source/ControlRigEditor/Private/Tools/ControlRigTweener.cpp b/Engine/Plugins/Animation/ControlRig/Source/ControlRigEditor/Private/Tools/ControlRigTweener.cpp
--- a/Engine/Plugins/Animation/ControlRig/Source/ControlRigEditor/Private/Tools/ControlRigTweener.cpp
+++ b/Engine/Plugins/Animation/ControlRig/Source/ControlRigEditor/Private/Tools/ControlRigTweener.cpp
@@ -549,13 +549,16 @@ bool FBasicBlendSlider::Blend(TWeakPtr<ISequencer>& InSequencer, const double Bl
 					double CurrentTime = TickResoultion.AsSeconds(FFrameTime(ObjectChannels.KeyBounds[Index].CurrentFrame));
 					double NewValue = DoBlend(PreviousTime, PreviousValue, CurrentTime, CurrentValue,
 						NextTime, NextValue, BlendValue);
+					using namespace UE::MovieScene;
 					if (ObjectChannels.KeyBounds[Index].FloatChannel)
 					{
-						AddKeyToChannel(ObjectChannels.KeyBounds[Index].FloatChannel, FrameTime.GetFrame(), (float)NewValue, Sequencer->GetKeyInterpolation());
+						EMovieSceneKeyInterpolation KeyInterpolation = GetInterpolationMode(ObjectChannels.KeyBounds[Index].FloatChannel, FrameTime.GetFrame(), Sequencer->GetKeyInterpolation());
+						AddKeyToChannel(ObjectChannels.KeyBounds[Index].FloatChannel, FrameTime.GetFrame(), (float)NewValue, KeyInterpolation);
 					}
 					else if (ObjectChannels.KeyBounds[Index].DoubleChannel)
 					{
-						AddKeyToChannel(ObjectChannels.KeyBounds[Index].DoubleChannel, FrameTime.GetFrame(), NewValue, Sequencer->GetKeyInterpolation());
+						EMovieSceneKeyInterpolation KeyInterpolation = GetInterpolationMode(ObjectChannels.KeyBounds[Index].DoubleChannel, FrameTime.GetFrame(), Sequencer->GetKeyInterpolation());
+						AddKeyToChannel(ObjectChannels.KeyBounds[Index].DoubleChannel, FrameTime.GetFrame(), NewValue, KeyInterpolation);
 					}
 					bDidBlend = true;
 				}
diff --git a/Engine/Source/Editor/MovieSceneTools/Private/Channels/BuiltInChannelEditors.cpp b/Engine/Source/Editor/MovieSceneTools/Private/Channels/BuiltInChannelEditors.cpp
--- a/Engine/Source/Editor/MovieSceneTools/Private/Channels/BuiltInChannelEditors.cpp
+++ b/Engine/Source/Editor/MovieSceneTools/Private/Channels/BuiltInChannelEditors.cpp
@@ -193,9 +193,9 @@ FKeyHandle AddOrUpdateKeyImpl(ChannelType* Channel, UMovieSceneSection* SectionT
 			}
 		}
 	}
-
 	using namespace UE::MovieScene;
-	return AddKeyToChannel(Channel, InTime, NewValue, Sequencer.GetKeyInterpolation());
+	EMovieSceneKeyInterpolation KeyInterpolation = GetInterpolationMode(Channel,InTime,Sequencer.GetKeyInterpolation());
+	return AddKeyToChannel(Channel, InTime, NewValue, KeyInterpolation);
 }
 
 FKeyHandle AddOrUpdateKey(FMovieSceneFloatChannel* Channel, UMovieSceneSection* SectionToKey, const TMovieSceneExternalValue<float>& ExternalValue, FFrameNumber InTime, ISequencer& Sequencer, const FGuid& InObjectBindingID, FTrackInstancePropertyBindings* PropertyBindings)
diff --git a/Engine/Source/Editor/MovieSceneTools/Private/CurveKeyEditors/SequencerKeyEditor.h b/Engine/Source/Editor/MovieSceneTools/Private/CurveKeyEditors/SequencerKeyEditor.h
--- a/Engine/Source/Editor/MovieSceneTools/Private/CurveKeyEditors/SequencerKeyEditor.h
+++ b/Engine/Source/Editor/MovieSceneTools/Private/CurveKeyEditors/SequencerKeyEditor.h
@@ -115,7 +115,7 @@ struct TSequencerKeyEditor
 		const FFrameNumber CurrentTime = Sequencer->GetLocalTime().Time.FloorToFrame();
 		const bool  bAutoSetTrackDefaults = Sequencer->GetAutoSetTrackDefaults();
 
-		EMovieSceneKeyInterpolation Interpolation = Sequencer->GetKeyInterpolation();
+		EMovieSceneKeyInterpolation Interpolation = GetInterpolationMode(Channel,CurrentTime,Sequencer->GetKeyInterpolation());
 
 		TArray<FKeyHandle> KeysAtCurrentTime;
 		Channel->GetKeys(TRange<FFrameNumber>(CurrentTime), nullptr, &KeysAtCurrentTime);
diff --git a/Engine/Source/Editor/MovieSceneTools/Private/TrackEditors/PropertyTrackEditors/EulerTransformPropertyTrackEditor.cpp b/Engine/Source/Editor/MovieSceneTools/Private/TrackEditors/PropertyTrackEditors/EulerTransformPropertyTrackEditor.cpp
--- a/Engine/Source/Editor/MovieSceneTools/Private/TrackEditors/PropertyTrackEditors/EulerTransformPropertyTrackEditor.cpp
+++ b/Engine/Source/Editor/MovieSceneTools/Private/TrackEditors/PropertyTrackEditors/EulerTransformPropertyTrackEditor.cpp
@@ -142,7 +142,8 @@ void FEulerTransformPropertyTrackEditor::ProcessKeyOperation(UObject* ObjectToKe
 					FMovieSceneDoubleChannel* Channel = static_cast<FMovieSceneDoubleChannel*>(Handle.Get());
 
 					double Value = RecomposedTransform[Handle.GetChannelIndex()];
-					AddKeyToChannel(Channel, KeyTime, Value, InSequencer.GetKeyInterpolation());
+					EMovieSceneKeyInterpolation Interpolation = GetInterpolationMode(Channel, KeyTime, InSequencer.GetKeyInterpolation());
+					AddKeyToChannel(Channel, KeyTime, Value, Interpolation);
 				}
 				else
 				{
diff --git a/Engine/Source/Editor/MovieSceneTools/Private/TrackEditors/TransformTrackEditor.cpp b/Engine/Source/Editor/MovieSceneTools/Private/TrackEditors/TransformTrackEditor.cpp
--- a/Engine/Source/Editor/MovieSceneTools/Private/TrackEditors/TransformTrackEditor.cpp
+++ b/Engine/Source/Editor/MovieSceneTools/Private/TrackEditors/TransformTrackEditor.cpp
@@ -1147,7 +1147,8 @@ void F3DTransformTrackEditor::ProcessKeyOperation(UObject* ObjectToKey, TArrayVi
 					FMovieSceneDoubleChannel* Channel = static_cast<FMovieSceneDoubleChannel*>(Handle.Get());
 
 					double Value = RecomposedTransform[Handle.GetChannelIndex()];
-					AddKeyToChannel(Channel, KeyTime, Value, InSequencer.GetKeyInterpolation());
+					EMovieSceneKeyInterpolation Interpolation = GetInterpolationMode(Channel,KeyTime,InSequencer.GetKeyInterpolation());
+					AddKeyToChannel(Channel, KeyTime, Value, Interpolation);
 				}
 				else
 				{
diff --git a/Engine/Source/Editor/MovieSceneTools/Public/KeyframeTrackEditor.h b/Engine/Source/Editor/MovieSceneTools/Public/KeyframeTrackEditor.h
--- a/Engine/Source/Editor/MovieSceneTools/Public/KeyframeTrackEditor.h
+++ b/Engine/Source/Editor/MovieSceneTools/Public/KeyframeTrackEditor.h
@@ -70,6 +70,7 @@ struct TAddKeyImpl : IImpl
 				{
 					if (Section->TryModify())
 					{
+						InterpolationMode = GetInterpolationMode(Channel, InTime, InterpolationMode);
 						AddKeyToChannel(Channel, InTime, ValueToSet, InterpolationMode);
 						bKeyCreated = true;
 					}
@@ -130,6 +131,7 @@ struct TAddKeyImpl<FMovieSceneBoolChannel, bool> : IImpl
 				{
 					if (Section->TryModify())
 					{
+						InterpolationMode = GetInterpolationMode(Channel, InTime, InterpolationMode);
 						AddKeyToChannel(Channel, InTime, ValueToSet, InterpolationMode);
 						bKeyCreated = true;
 					}
@@ -189,6 +191,7 @@ struct TAddKeyImpl<FMovieSceneDoubleChannel, double> : IImpl
 				{
 					if (Section->TryModify())
 					{
+						InterpolationMode = GetInterpolationMode(Channel, InTime, InterpolationMode);
 						AddKeyToChannel(Channel, InTime, ValueToSet, InterpolationMode);
 						bKeyCreated = true;
 					}
@@ -266,6 +269,7 @@ struct TAddKeyImpl<FMovieSceneFloatChannel, float> : IImpl
 				{
 					if (Section->TryModify())
 					{
+						InterpolationMode = GetInterpolationMode(Channel, InTime, InterpolationMode);
 						AddKeyToChannel(Channel, InTime, ValueToSet, InterpolationMode);
 						bKeyCreated = true;
 					}
@@ -343,6 +347,7 @@ struct TAddKeyImpl<FMovieSceneIntegerChannel, int32> : IImpl
 				{
 					if (Section->TryModify())
 					{
+						InterpolationMode = GetInterpolationMode(Channel, InTime, InterpolationMode);
 						AddKeyToChannel(Channel, InTime, ValueToSet, InterpolationMode);
 						bKeyCreated = true;
 					}
diff --git a/Engine/Source/Editor/Sequencer/Public/SequencerChannelTraits.h b/Engine/Source/Editor/Sequencer/Public/SequencerChannelTraits.h
--- a/Engine/Source/Editor/Sequencer/Public/SequencerChannelTraits.h
+++ b/Engine/Source/Editor/Sequencer/Public/SequencerChannelTraits.h
@@ -187,7 +187,8 @@ namespace Sequencer
 		ValueType ValueAtTime = InDefaultValue;
 		EvaluateChannel(InChannel, InTime, ValueAtTime);
 
-		return AddKeyToChannel(InChannel, InTime, ValueAtTime, InSequencer.GetKeyInterpolation());
+		EMovieSceneKeyInterpolation InterpolationMode = GetInterpolationMode(InChannel, InTime, InSequencer.GetKeyInterpolation());
+		return AddKeyToChannel(InChannel, InTime, ValueAtTime, InterpolationMode);
 	}
 	
 
@@ -226,7 +227,8 @@ namespace Sequencer
 					TOptional<ValueType> Value = InExternalValue.OnGetExternalValue(*Object, InPropertyBindings);
 					if (Value.IsSet())
 					{
-						return AddKeyToChannel(InChannel, InTime, Value.GetValue(), InSequencer.GetKeyInterpolation());
+						EMovieSceneKeyInterpolation InterpolationMode = GetInterpolationMode(InChannel, InTime, InSequencer.GetKeyInterpolation());
+						return AddKeyToChannel(InChannel, InTime, Value.GetValue(), InterpolationMode);
 					}
 				}
 			}
@@ -290,7 +292,8 @@ namespace Sequencer
 			ValueType ValueAtTime{};
 			EvaluateChannel(InChannel, InTime, ValueAtTime);
 		
-			Handle = AddKeyToChannel(InChannel, InTime, ValueAtTime, InSequencer.GetKeyInterpolation());
+			EMovieSceneKeyInterpolation InterpolationMode = GetInterpolationMode(InChannel, InTime, InSequencer.GetKeyInterpolation());
+			Handle = AddKeyToChannel(InChannel, InTime, ValueAtTime, InterpolationMode);
 		}
 
 		return Handle.GetValue();
diff --git a/Engine/Source/Editor/Sequencer/Public/SequencerSettings.h b/Engine/Source/Editor/Sequencer/Public/SequencerSettings.h
--- a/Engine/Source/Editor/Sequencer/Public/SequencerSettings.h
+++ b/Engine/Source/Editor/Sequencer/Public/SequencerSettings.h
@@ -113,9 +113,9 @@ public:
 	/** Sets which channels are keyed when a channel is keyed */
 	void SetKeyGroupMode(EKeyGroupMode);
 
-	/** Gets default key interpolation. */
+	/** Get the default Interpolation type for newly created keyframes if the channel does not already have keyframes */
 	EMovieSceneKeyInterpolation GetKeyInterpolation() const;
-	/** Sets default key interpolation */
+	/** Sets default key interpolation for creating new keys on empty channels */
 	void SetKeyInterpolation(EMovieSceneKeyInterpolation InKeyInterpolation);
 
 	/** Get initial spawn position. */
diff --git a/Engine/Source/Runtime/MovieScene/Private/Channels/MovieSceneCurveChannelImpl.cpp b/Engine/Source/Runtime/MovieScene/Private/Channels/MovieSceneCurveChannelImpl.cpp
--- a/Engine/Source/Runtime/MovieScene/Private/Channels/MovieSceneCurveChannelImpl.cpp
+++ b/Engine/Source/Runtime/MovieScene/Private/Channels/MovieSceneCurveChannelImpl.cpp
@@ -1095,6 +1095,52 @@ void TMovieSceneCurveChannelImpl<ChannelType>::Optimize(ChannelType* InChannel,
 	}
 }
 
+
+template<typename ChannelType>
+EMovieSceneKeyInterpolation TMovieSceneCurveChannelImpl<ChannelType>::GetInterpolationMode(ChannelType* InChannel, const FFrameNumber& InTime, EMovieSceneKeyInterpolation DefaultInterpolationMode)
+{
+	auto ChannelData = InChannel->GetData();
+	const TArrayView<const ChannelValueType> Values = ChannelData.GetValues();
+	const TArrayView<FFrameNumber> Times = ChannelData.GetTimes();
+	if (Times.Num() > 0)
+	{
+		//get previous key or first key if first
+		int32 Index = Algo::LowerBound(Times, InTime) - 1;
+		if (Index < 0)
+		{
+			Index = 0;
+		}
+		const ChannelValueType& Value = Values[Index]; //-V758
+		switch (Value.InterpMode.GetValue())
+		{
+		case RCIM_Linear:
+			return EMovieSceneKeyInterpolation::Linear;
+			break;
+		case RCIM_Constant:
+			return EMovieSceneKeyInterpolation::Constant;
+			break;
+
+		case RCIM_Cubic:
+
+			switch (Value.TangentMode.GetValue())
+			{
+			case RCTM_Auto:
+				return EMovieSceneKeyInterpolation::Auto;
+				break;
+			case RCTM_Break:
+				return EMovieSceneKeyInterpolation::Break;
+				break;
+			case RCTM_User:
+				return EMovieSceneKeyInterpolation::User;
+				break;
+			}
+			break;
+		}
+
+	}
+	return DefaultInterpolationMode;
+}
+
 template<typename ChannelType>
 FKeyHandle TMovieSceneCurveChannelImpl<ChannelType>::AddKeyToChannel(ChannelType* InChannel, FFrameNumber InFrameNumber, float InValue, EMovieSceneKeyInterpolation Interpolation)
 {
diff --git a/Engine/Source/Runtime/MovieScene/Private/Channels/MovieSceneCurveChannelImpl.h b/Engine/Source/Runtime/MovieScene/Private/Channels/MovieSceneCurveChannelImpl.h
--- a/Engine/Source/Runtime/MovieScene/Private/Channels/MovieSceneCurveChannelImpl.h
+++ b/Engine/Source/Runtime/MovieScene/Private/Channels/MovieSceneCurveChannelImpl.h
@@ -94,13 +94,22 @@ struct MOVIESCENE_API TMovieSceneCurveChannelImpl
 	/** Add a new key to a channel at a given time */
 	static FKeyHandle AddKeyToChannel(ChannelType* InChannel, FFrameNumber InFrameNumber, float InValue, EMovieSceneKeyInterpolation Interpolation);
 
+	/**
+	 * Get the interpolation mode to use at a specified time
+	 *
+	 * @param InChannel          The channel to find the interpolation mode
+	 * @param InTime             The time we are looking for an interpolation mode
+	 * @param DefaultInterpolationMode Current default interpolation mode, may be returned as the current mode
+	 * @return Interpolation mode to use at that frame
+	 */
+	static EMovieSceneKeyInterpolation GetInterpolationMode(ChannelType* InChannel, const FFrameNumber& InTime, EMovieSceneKeyInterpolation DefaultInterpolationMode);
+	
 	/** Dilate channel data.*/
 	static void Dilate(ChannelType* InChannel, FFrameNumber Origin, float DilationFactor);
 
 	/** Assigns a given value on a key */
 	static void AssignValue(ChannelType* InChannel, FKeyHandle InKeyHandle, typename ChannelType::CurveValueType InValue);
 
-
 	/** Serialization methods */
 
 	/** Serialize the channel */
@@ -198,3 +207,4 @@ void TMovieSceneCurveChannelImpl<ChannelType>::CopyChannel(const OtherChannelTyp
 #endif
 }
 
+
diff --git a/Engine/Source/Runtime/MovieScene/Private/Channels/MovieSceneDoubleChannel.cpp b/Engine/Source/Runtime/MovieScene/Private/Channels/MovieSceneDoubleChannel.cpp
--- a/Engine/Source/Runtime/MovieScene/Private/Channels/MovieSceneDoubleChannel.cpp
+++ b/Engine/Source/Runtime/MovieScene/Private/Channels/MovieSceneDoubleChannel.cpp
@@ -168,6 +168,11 @@ void FMovieSceneDoubleChannel::ClearDefault()
 	bHasDefaultValue = false;
 }
 
+EMovieSceneKeyInterpolation GetInterpolationMode(FMovieSceneDoubleChannel* InChannel, const FFrameNumber& InTime, EMovieSceneKeyInterpolation DefaultInterpolationMode)
+{
+	return TMovieSceneCurveChannelImpl<FMovieSceneDoubleChannel>::GetInterpolationMode(InChannel, InTime, DefaultInterpolationMode);
+}
+
 FKeyHandle AddKeyToChannel(FMovieSceneDoubleChannel* Channel, FFrameNumber InFrameNumber, double InValue, EMovieSceneKeyInterpolation Interpolation)
 {
 	return TMovieSceneCurveChannelImpl<FMovieSceneDoubleChannel>::AddKeyToChannel(Channel, InFrameNumber, InValue, Interpolation);
diff --git a/Engine/Source/Runtime/MovieScene/Private/Channels/MovieSceneFloatChannel.cpp b/Engine/Source/Runtime/MovieScene/Private/Channels/MovieSceneFloatChannel.cpp
--- a/Engine/Source/Runtime/MovieScene/Private/Channels/MovieSceneFloatChannel.cpp
+++ b/Engine/Source/Runtime/MovieScene/Private/Channels/MovieSceneFloatChannel.cpp
@@ -158,6 +158,11 @@ void FMovieSceneFloatChannel::ClearDefault()
 	bHasDefaultValue = false;
 }
 
+EMovieSceneKeyInterpolation GetInterpolationMode(FMovieSceneFloatChannel* InChannel, const FFrameNumber& InTime, EMovieSceneKeyInterpolation DefaultInterpolationMode)
+{
+	return TMovieSceneCurveChannelImpl<FMovieSceneFloatChannel>::GetInterpolationMode(InChannel, InTime, DefaultInterpolationMode);
+}
+
 FKeyHandle AddKeyToChannel(FMovieSceneFloatChannel* Channel, FFrameNumber InFrameNumber, float InValue, EMovieSceneKeyInterpolation Interpolation)
 {
 	return TMovieSceneCurveChannelImpl<FMovieSceneFloatChannel>::AddKeyToChannel(Channel, InFrameNumber, InValue, Interpolation);
diff --git a/Engine/Source/Runtime/MovieScene/Public/Channels/MovieSceneChannelTraits.h b/Engine/Source/Runtime/MovieScene/Public/Channels/MovieSceneChannelTraits.h
--- a/Engine/Source/Runtime/MovieScene/Public/Channels/MovieSceneChannelTraits.h
+++ b/Engine/Source/Runtime/MovieScene/Public/Channels/MovieSceneChannelTraits.h
@@ -126,6 +126,12 @@ namespace MovieScene
 		return false;
 	}
 
+	template<typename ChannelType>
+	EMovieSceneKeyInterpolation GetInterpolationMode(ChannelType* InChannel, const FFrameNumber& InTime, EMovieSceneKeyInterpolation DefaultInterpolationMode)
+	{
+		return DefaultInterpolationMode;
+	}
+
 	/**
 	 * Add a key to a channel, or update an existing key if one already exists at this time
 	 *
diff --git a/Engine/Source/Runtime/MovieScene/Public/Channels/MovieSceneDoubleChannel.h b/Engine/Source/Runtime/MovieScene/Public/Channels/MovieSceneDoubleChannel.h
--- a/Engine/Source/Runtime/MovieScene/Public/Channels/MovieSceneDoubleChannel.h
+++ b/Engine/Source/Runtime/MovieScene/Public/Channels/MovieSceneDoubleChannel.h
@@ -360,6 +360,11 @@ struct TMovieSceneChannelTraits<FMovieSceneDoubleChannel> : TMovieSceneChannelTr
 #endif
 };
 
+/**
+ * Overload for getting the interpolation mode for a channel at a specified time, it could be the previous key's mode.See UE::MovieScene::GetInterpolationMode for default implementation.
+ */
+MOVIESCENE_API EMovieSceneKeyInterpolation GetInterpolationMode(FMovieSceneDoubleChannel* InChannel, const FFrameNumber& InTime, EMovieSceneKeyInterpolation DefaultInterpolationMode);
+
 /**
  * Overload for adding a new key to a double channel at a given time. See UE::MovieScene::AddKeyToChannel for default implementation.
  */
diff --git a/Engine/Source/Runtime/MovieScene/Public/Channels/MovieSceneFloatChannel.h b/Engine/Source/Runtime/MovieScene/Public/Channels/MovieSceneFloatChannel.h
--- a/Engine/Source/Runtime/MovieScene/Public/Channels/MovieSceneFloatChannel.h
+++ b/Engine/Source/Runtime/MovieScene/Public/Channels/MovieSceneFloatChannel.h
@@ -360,12 +360,16 @@ struct TMovieSceneChannelTraits<FMovieSceneFloatChannel> : TMovieSceneChannelTra
 #endif
 };
 
+/**
+ * Overload for getting the interpolation mode for a channel at a specified time, it could be the previous key's mode.See UE::MovieScene::GetInterpolationMode for default implementation.
+ */
+MOVIESCENE_API EMovieSceneKeyInterpolation GetInterpolationMode(FMovieSceneFloatChannel* InChannel, const FFrameNumber& InTime, EMovieSceneKeyInterpolation DefaultInterpolationMode);
+
 /**
  * Overload for adding a new key to a float channel at a given time. See UE::MovieScene::AddKeyToChannel for default implementation.
  */
 MOVIESCENE_API FKeyHandle AddKeyToChannel(FMovieSceneFloatChannel* Channel, FFrameNumber InFrameNumber, float InValue, EMovieSceneKeyInterpolation Interpolation);
 
-
 /**
  * Overload for dilating float channel data. See UE::MovieScene::Dilate for default implementation.
  */
