Sequencer: Fixed runners becoming detached from their linkers when a change of playback context was detected

 - Also added code to ensure that InvalidateCachedData is called whenever RestoreGlobalState is called

diff --git a/Engine/Source/Editor/Sequencer/Private/LevelEditorSequencerIntegration.cpp b/Engine/Source/Editor/Sequencer/Private/LevelEditorSequencerIntegration.cpp
--- a/Engine/Source/Editor/Sequencer/Private/LevelEditorSequencerIntegration.cpp
+++ b/Engine/Source/Editor/Sequencer/Private/LevelEditorSequencerIntegration.cpp
@@ -1185,7 +1185,6 @@ void FLevelEditorSequencerIntegration::ResetToAnimatedState(UWorld* World)
 		{
 			if (Options.bRequiresLevelEvents)
 			{
-				In.InvalidateCachedData();
 				In.ForceEvaluate();
 
 				for (const TSharedPtr<ISequencerTrackEditor>& TrackEditor : In.GetTrackEditors())
diff --git a/Engine/Source/Runtime/MovieScene/Private/Evaluation/PreAnimatedState/MovieScenePreAnimatedStateExtension.cpp b/Engine/Source/Runtime/MovieScene/Private/Evaluation/PreAnimatedState/MovieScenePreAnimatedStateExtension.cpp
--- a/Engine/Source/Runtime/MovieScene/Private/Evaluation/PreAnimatedState/MovieScenePreAnimatedStateExtension.cpp
+++ b/Engine/Source/Runtime/MovieScene/Private/Evaluation/PreAnimatedState/MovieScenePreAnimatedStateExtension.cpp
@@ -387,6 +387,12 @@ void FPreAnimatedStateExtension::RestoreGlobalState(const FRestoreStateParams& P
 
 	GroupMetaData.Shrink();
 
+	// Invalidate cached data for any sequence instance that belongs to the terminal instance
+	if (Params.TerminalInstanceHandle.IsValid())
+	{
+		Linker->GetInstanceRegistry()->MutateInstance(Params.TerminalInstanceHandle).InvalidateCachedData(Linker);
+	}
+
 	bEntriesInvalidated = true;
 }
 
