Sequencer: Removed mutual entity initializers

  - These were actually never used, and are actually ambiguous when using mutual components as part of a mutation since it's hard to know whether we should run them for mutual components that existed before the mutation.

diff --git a/Engine/Source/Runtime/MovieScene/Private/EntitySystem/MovieSceneEntityFactory.cpp b/Engine/Source/Runtime/MovieScene/Private/EntitySystem/MovieSceneEntityFactory.cpp
--- a/Engine/Source/Runtime/MovieScene/Private/EntitySystem/MovieSceneEntityFactory.cpp
+++ b/Engine/Source/Runtime/MovieScene/Private/EntitySystem/MovieSceneEntityFactory.cpp
@@ -241,15 +241,6 @@ void FEntityFactories::DefineMutuallyInclusiveComponent(FComponentTypeID InCompo
 	Masks.AllMutualFirsts.Set(InComponentA);
 }
 
-void FEntityFactories::DefineMutuallyInclusiveComponent(TInlineValue<FMutualEntityInitializer>&& InInitializer)
-{
-	check(InInitializer.IsValid());
-
-	DefineChildComponent(InInitializer->GetComponentA(), InInitializer->GetComponentB());
-	// Note: after this line, InInitializer is reset
-	MutualInitializers.Add(MoveTemp(InInitializer));
-}
-
 void FEntityFactories::DefineComplexInclusiveComponents(const FComplexInclusivityFilter& InFilter, FComponentTypeID InComponent)
 {
 	FComponentMask ComponentsToInclude { InComponent };
@@ -375,17 +366,6 @@ int32 FEntityFactories::ComputeMutuallyInclusiveComponents(FComponentMask& Compo
 	return NumNewComponents;
 }
 
-void FEntityFactories::RunInitializers(const FComponentMask& Type, const FEntityRange& InEntityRange)
-{
-	for (TInlineValue<FMutualEntityInitializer>& MutualInit : MutualInitializers)
-	{
-		if (MutualInit->IsRelevant(Type))
-		{
-			MutualInit->Run(InEntityRange);
-		}
-	}
-}
-
 void FEntityFactories::RunInitializers(const FComponentMask& ParentType, const FComponentMask& ChildType, const FEntityAllocation* ParentAllocation, TArrayView<const int32> ParentAllocationOffsets, const FEntityRange& InChildEntityRange)
 {
 	// First off, run child initializers
@@ -396,9 +376,6 @@ void FEntityFactories::RunInitializers(const FComponentMask& ParentType, const F
 			ChildInit->Run(InChildEntityRange, ParentAllocation, ParentAllocationOffsets);
 		}
 	}
-
-	// Run generic initializers
-	RunInitializers(ChildType, InChildEntityRange);
 }
 
 }	// using namespace MovieScene
diff --git a/Engine/Source/Runtime/MovieScene/Private/EntitySystem/MovieSceneEntityManager.cpp b/Engine/Source/Runtime/MovieScene/Private/EntitySystem/MovieSceneEntityManager.cpp
--- a/Engine/Source/Runtime/MovieScene/Private/EntitySystem/MovieSceneEntityManager.cpp
+++ b/Engine/Source/Runtime/MovieScene/Private/EntitySystem/MovieSceneEntityManager.cpp
@@ -1885,28 +1885,6 @@ void FEntityManager::InitializeChildAllocation(const FComponentMask& ParentType,
 	}
 }
 
-void FEntityManager::InitializeMutualComponents(FMovieSceneEntityID EntityID)
-{
-	FEntityLocation Location = EntityLocations[EntityID.AsIndex()];
-	if (Location.IsValid())
-	{
-		return;
-	}
-
-	const FComponentMask& EntityType = EntityAllocationMasks[Location.GetAllocationIndex()];
-	FEntityAllocation* Allocation = EntityAllocations[Location.GetAllocationIndex()];
-
-	FEntityRange Range = { Allocation, Location.GetEntryIndexWithinAllocation(), 1 };
-
-	for (const TInlineValue<FMutualEntityInitializer>& MutualInit : ComponentRegistry->Factories.MutualInitializers)
-	{
-		if (MutualInit->IsRelevant(EntityType))
-		{
-			MutualInit->Run(Range);
-		}
-	}
-}
-
 void FEntityManager::AddMutualComponents()
 {
 	AddMutualComponents(FEntityComponentFilter());
@@ -1968,16 +1946,6 @@ void FEntityManager::AddMutualComponents(const FEntityComponentFilter& InFilter)
 			}
 		}
 
-		FEntityRange Range = { NewAllocation, 0, NewAllocation->Num() };
-		for (const TInlineValue<FMutualEntityInitializer>& MutualInit : ComponentRegistry->Factories.MutualInitializers)
-		{
-			// Only run mutual initializers for _new_ component types (ie, ones that were actually added)
-			if (NewComponents.Contains(MutualInit->GetComponentA()) && Pair.Value.Contains(MutualInit->GetComponentB()))
-			{
-				MutualInit->Run(Range);
-			}
-		}
-
 		EntityAllocationMasks[Pair.Key] = Pair.Value;
 		EntityAllocations[Pair.Key] = NewAllocation;
 
diff --git a/Engine/Source/Runtime/MovieScene/Public/EntitySystem/MovieSceneEntityBuilder.h b/Engine/Source/Runtime/MovieScene/Public/EntitySystem/MovieSceneEntityBuilder.h
--- a/Engine/Source/Runtime/MovieScene/Public/EntitySystem/MovieSceneEntityBuilder.h
+++ b/Engine/Source/Runtime/MovieScene/Public/EntitySystem/MovieSceneEntityBuilder.h
@@ -220,10 +220,6 @@ struct TEntityBuilderImpl<TIntegerSequence<int, Indices...>, T...> : IEntityBuil
 	{
 		if (Entity.Data.Allocation != nullptr)
 		{
-			if (bAddMutualComponents)
-			{
-				EntityManager->InitializeMutualComponents(Entity.EntityID);
-			}
 			VisitTupleElements([Entity](auto& In){ In.Apply(Entity.Data.Allocation, Entity.Data.ComponentOffset); }, this->Payload);
 		}
 	}
diff --git a/Engine/Source/Runtime/MovieScene/Public/EntitySystem/MovieSceneEntityFactory.h b/Engine/Source/Runtime/MovieScene/Public/EntitySystem/MovieSceneEntityFactory.h
--- a/Engine/Source/Runtime/MovieScene/Public/EntitySystem/MovieSceneEntityFactory.h
+++ b/Engine/Source/Runtime/MovieScene/Public/EntitySystem/MovieSceneEntityFactory.h
@@ -137,11 +137,6 @@ struct FEntityFactories
      */
 	MOVIESCENE_API void DefineMutuallyInclusiveComponent(FComponentTypeID InComponentA, FComponentTypeID InComponentB);
 
-	/**
-	 * Specifies a mutual inclusivity relationship. The helper method above is easier and preferrable.
-	 */
-	MOVIESCENE_API void DefineMutuallyInclusiveComponent(TInlineValue<FMutualEntityInitializer>&& InInitializer);
-
 	/**
 	 * Specifies that if an entity matches the given filter, the specified components should be created on it.
 	 * @note: include "EntitySystem/MovieSceneEntityFactoryTemplates.h" for definition
@@ -173,11 +168,9 @@ struct FEntityFactories
 	 */
 	MOVIESCENE_API int32 ComputeMutuallyInclusiveComponents(FComponentMask& ComponentMask);
 
-	MOVIESCENE_API void RunInitializers(const FComponentMask& Type, const FEntityRange& InEntityRange);
 	MOVIESCENE_API void RunInitializers(const FComponentMask& ParentType, const FComponentMask& ChildType, const FEntityAllocation* ParentAllocation, TArrayView<const int32> ParentAllocationOffsets, const FEntityRange& InChildEntityRange);
 
 	TArray<TInlineValue<FChildEntityInitializer>> ChildInitializers;
-	TArray<TInlineValue<FMutualEntityInitializer>> MutualInitializers;
 
 	TMultiMap<FComponentTypeID, FComponentTypeID> ParentToChildComponentTypes;
 	FMovieSceneEntitySystemDirectedGraph MutualInclusivityGraph;
diff --git a/Engine/Source/Runtime/MovieScene/Public/EntitySystem/MovieSceneEntityFactoryTemplates.h b/Engine/Source/Runtime/MovieScene/Public/EntitySystem/MovieSceneEntityFactoryTemplates.h
--- a/Engine/Source/Runtime/MovieScene/Public/EntitySystem/MovieSceneEntityFactoryTemplates.h
+++ b/Engine/Source/Runtime/MovieScene/Public/EntitySystem/MovieSceneEntityFactoryTemplates.h
@@ -191,31 +191,6 @@ private:
 
 
 
-template<typename ComponentTypeA, typename ComponentTypeB>
-struct TMutualEntityInitializer : FMutualEntityInitializer
-{
-	using CallbackType = void (*)(ComponentTypeA* ComponentsA, ComponentTypeB* ComponentsB, int32 Num);
-
-	explicit TMutualEntityInitializer(TComponentTypeID<ComponentTypeA> InComponentA, TComponentTypeID<ComponentTypeB> InComponentB, CallbackType InCallback)
-		: FMutualEntityInitializer(InComponentA, InComponentB)
-		, Callback(InCallback)
-	{}
-
-private:
-
-	virtual void Run(const FEntityRange& Range) override
-	{
-		TComponentLock<TWrite<ComponentTypeA>> A = Range.Allocation->WriteComponents(ComponentA.ReinterpretCast<ComponentTypeA>(), FEntityAllocationWriteContext::NewAllocation());
-		TComponentLock<TWrite<ComponentTypeB>> B = Range.Allocation->WriteComponents(ComponentB.ReinterpretCast<ComponentTypeB>(), FEntityAllocationWriteContext::NewAllocation());
-
-		Callback(&A[Range.ComponentStartOffset], &B[Range.ComponentStartOffset], Range.Num);
-	}
-
-	CallbackType Callback;
-};
-
-
-
 template<typename ComponentType>
 inline void FEntityFactories::DuplicateChildComponent(TComponentTypeID<ComponentType> InComponent)
 {
diff --git a/Engine/Source/Runtime/MovieScene/Public/EntitySystem/MovieSceneEntityFactoryTypes.h b/Engine/Source/Runtime/MovieScene/Public/EntitySystem/MovieSceneEntityFactoryTypes.h
--- a/Engine/Source/Runtime/MovieScene/Public/EntitySystem/MovieSceneEntityFactoryTypes.h
+++ b/Engine/Source/Runtime/MovieScene/Public/EntitySystem/MovieSceneEntityFactoryTypes.h
@@ -85,38 +85,6 @@ protected:
 };
 
 
-struct FMutualEntityInitializer
-{
-	virtual ~FMutualEntityInitializer(){}
-
-	bool IsRelevant(const FComponentMask& InType) const
-	{
-		return InType.Contains(ComponentA) && InType.Contains(ComponentB);
-	}
-
-	FComponentTypeID GetComponentA() const
-	{
-		return ComponentA;
-	}
-
-	FComponentTypeID GetComponentB() const
-	{
-		return ComponentB;
-	}
-
-	virtual void Run(const FEntityRange& Range) = 0;
-
-protected:
-
-	FComponentTypeID ComponentA, ComponentB;
-
-	explicit FMutualEntityInitializer(FComponentTypeID InComponentA, FComponentTypeID InComponentB)
-		: ComponentA(InComponentA), ComponentB(InComponentB)
-	{
-		check(ComponentA && ComponentB);
-	}
-};
-
 
 }	// using namespace MovieScene
 }	// using namespace UE
diff --git a/Engine/Source/Runtime/MovieScene/Public/EntitySystem/MovieSceneEntityManager.h b/Engine/Source/Runtime/MovieScene/Public/EntitySystem/MovieSceneEntityManager.h
--- a/Engine/Source/Runtime/MovieScene/Public/EntitySystem/MovieSceneEntityManager.h
+++ b/Engine/Source/Runtime/MovieScene/Public/EntitySystem/MovieSceneEntityManager.h
@@ -47,7 +47,6 @@ struct FChildEntityInitializer;
 struct FComponentRegistry;
 struct FEntityAllocationIteratorProxy;
 struct FFreeEntityOperation;
-struct FMutualEntityInitializer;
 struct IComponentTypeHandler;
 struct IMovieSceneEntityMutation;
 template <typename T> struct TReadOptional;
@@ -409,11 +408,6 @@ public:
 	 */
 	void InitializeChildAllocation(const FComponentMask& ParentType, const FComponentMask& ChildType, const FEntityAllocation* ParentAllocation, TArrayView<const int32> ParentAllocationOffsets, const FEntityRange& InChildEntityRange);
 
-	/**
-	 * Initialize a single entity using the mutual initializers
-	 */
-	void InitializeMutualComponents(FMovieSceneEntityID EntityID);
-
 	/**
 	 * Destroy a previously registered instanced child initializer using its index
 	 */
