Sequencer: Corrected erroneously recursive template function

diff --git a/Engine/Source/Runtime/MovieScene/Public/Evaluation/PreAnimatedState/MovieScenePreAnimatedStateStorage.h b/Engine/Source/Runtime/MovieScene/Public/Evaluation/PreAnimatedState/MovieScenePreAnimatedStateStorage.h
--- a/Engine/Source/Runtime/MovieScene/Public/Evaluation/PreAnimatedState/MovieScenePreAnimatedStateStorage.h
+++ b/Engine/Source/Runtime/MovieScene/Public/Evaluation/PreAnimatedState/MovieScenePreAnimatedStateStorage.h
@@ -683,7 +683,7 @@ public:
 	template<typename... ContributorTypes>
 	void BeginTrackingAndCachePreAnimatedValues(UMovieSceneEntitySystemLinker* Linker, TComponentTypeID<ContributorTypes>... InComponentTypes)
 	{
-		BeginTrackingAndCachePreAnimatedValues(Linker, TSavePreAnimatedStateParams<ContributorTypes...>(), InComponentTypes...);
+		BeginTrackingAndCachePreAnimatedValuesTask(Linker, TSavePreAnimatedStateParams<ContributorTypes...>(), InComponentTypes...);
 	}
 
 	/**
