Sequencer: cleanup pre-animated state APIs

There were some inconsistencies between the various pre-animated storage classes:

- The object storage was using static methods on the traits classes while the base storage was using instance methods on a nested traits object.
- It wasn't clear what methods were needed on a traits class. Now only the minimal methods are required, and optional codepaths can be enabled using some flags on the traits class.
- Some code was moved to the appropriate .cpp/.h files, so that there's a 1:1 match between them.
- Rename SavePreAnimatedState method to CachePreAnimatedValue so there's only one terminology.
- Rename TSavePreAnimatedStateParams to TPreAnimatedStateTaskParams since this parameter struct is also used for "begin tracking" tasks.
- Replaced the object storage's ValueInitializer with an OnCacheValue callback to be more consistent with the similar API on the base class.
- Added a few new ways to track entities and cache state.

diff --git a/Engine/Source/Runtime/MovieScene/Private/Evaluation/PreAnimatedState/MovieScenePreAnimatedObjectStorage.cpp b/Engine/Source/Runtime/MovieScene/Private/Evaluation/PreAnimatedState/MovieScenePreAnimatedObjectStorage.cpp
new file mode 100644
--- /dev/null
+++ b/Engine/Source/Runtime/MovieScene/Private/Evaluation/PreAnimatedState/MovieScenePreAnimatedObjectStorage.cpp
@@ -0,0 +1,25 @@
+
+// Copyright Epic Games, Inc. All Rights Reserved.
+
+#include "Evaluation/PreAnimatedState/MovieScenePreAnimatedObjectStorage.h"
+#include "Evaluation/PreAnimatedState/MovieScenePreAnimatedStateExtension.h"
+#include "Evaluation/PreAnimatedState/MovieScenePreAnimatedObjectGroupManager.h"
+
+namespace UE
+{
+namespace MovieScene
+{
+
+void FBoundObjectPreAnimatedStateTraits::Initialize(FPreAnimatedStorageID InStorageID, FPreAnimatedStateExtension* InParentExtension)
+{
+	ObjectGroupManager = InParentExtension->GetOrCreateGroupManager<FPreAnimatedObjectGroupManager>();
+}
+
+FPreAnimatedStorageGroupHandle FBoundObjectPreAnimatedStateTraits::MakeGroupImpl(UObject* BoundObject)
+{
+	return ObjectGroupManager->MakeGroupForKey(BoundObject);
+}
+
+} // namespace MovieScene
+} // namespace UE
+
diff --git a/Engine/Source/Runtime/MovieScene/Private/Evaluation/PreAnimatedState/MovieScenePreAnimatedStateStorage.cpp b/Engine/Source/Runtime/MovieScene/Private/Evaluation/PreAnimatedState/MovieScenePreAnimatedStateStorage.cpp
--- a/Engine/Source/Runtime/MovieScene/Private/Evaluation/PreAnimatedState/MovieScenePreAnimatedStateStorage.cpp
+++ b/Engine/Source/Runtime/MovieScene/Private/Evaluation/PreAnimatedState/MovieScenePreAnimatedStateStorage.cpp
@@ -18,21 +18,6 @@ FPreAnimatedTrackerParams::FPreAnimatedTrackerParams(FEntityAllocationProxy Item
 	bWantsRestoreState = Item.GetAllocationType().Contains(FBuiltInComponentTypes::Get()->Tags.RestoreState);
 }
 
-void FBoundObjectPreAnimatedStateTraits::Initialize(FPreAnimatedStorageID InStorageID, FPreAnimatedStateExtension* InParentExtension)
-{
-	ObjectGroupManager = InParentExtension->GetOrCreateGroupManager<FPreAnimatedObjectGroupManager>();
-}
-
-FPreAnimatedStorageGroupHandle FBoundObjectPreAnimatedStateTraits::MakeGroupImpl(UObject* BoundObject)
-{
-	return ObjectGroupManager->MakeGroupForKey(BoundObject);
-}
-
 } // namespace MovieScene
 } // namespace UE
 
-
-
-
-
-
diff --git a/Engine/Source/Runtime/MovieScene/Public/Evaluation/PreAnimatedState/MovieScenePreAnimatedObjectStorage.h b/Engine/Source/Runtime/MovieScene/Public/Evaluation/PreAnimatedState/MovieScenePreAnimatedObjectStorage.h
--- a/Engine/Source/Runtime/MovieScene/Public/Evaluation/PreAnimatedState/MovieScenePreAnimatedObjectStorage.h
+++ b/Engine/Source/Runtime/MovieScene/Public/Evaluation/PreAnimatedState/MovieScenePreAnimatedObjectStorage.h
@@ -16,15 +16,54 @@ namespace UE
 namespace MovieScene
 {
 
+/**
+ * Group state class that groups pre-animated storage together by bound object
+ *
+ * Inherit from this class by implementing the following members:
+ *
+ *	 using KeyType     = FObjectKey;
+ *	 using StorageType = <your storage type>;
+ *	 StorageType CachePreAnimatedValue(const FObjectKey& Object);
+ *	 void RestorePreAnimatedValue(const FObjectKey& Object, StorageType& InOutCachedValue, const FRestoreStateParams& Params);
+ */
+struct MOVIESCENE_API FBoundObjectPreAnimatedStateTraits : FPreAnimatedStateTraits
+{
+	enum { NeedsInitialize = true, SupportsGrouping = true, SupportsReplaceObject = true };
 
+	void Initialize(FPreAnimatedStorageID InStorageID, FPreAnimatedStateExtension* InParentExtension);
 
-// struct FPreAnimatedStateStorageObjectTraits
-// {
-// 	using KeyType     = FObjectKey;
-// 	using StorageType = IMovieScenePreAnimatedTokenPtr;
-// 	static void CachePreAnimatedValue(const FObjectKey& Object, StorageType& OutCachedValue);
-// 	static void RestorePreAnimatedValue(const FObjectKey& Object, StorageType& InOutCachedValue, const FRestoreStateParams& Params);
-// };
+	/* Defined as a template rather than a variadic function to prevent error C4840 */
+	template<typename... T>
+	FPreAnimatedStorageGroupHandle MakeGroup(UObject* BoundObject, T&&... Unused)
+	{
+		return MakeGroupImpl(BoundObject);
+	}
+	FPreAnimatedStorageGroupHandle MakeGroupImpl(UObject* BoundObject);
+
+	template<typename ...T>
+	void ReplaceObject(TTuple<FObjectKey, T...>& InOutKey, const FObjectKey& NewObject)
+	{
+		InOutKey.template Get<0>() = NewObject;
+	}
+	template<typename KeyType>
+	void ReplaceObject(KeyType& InOutKey, const FObjectKey& NewObject)
+	{
+		InOutKey.Object = NewObject;
+	}
+	template<typename ObjectType>
+	void ReplaceObject(TObjectKey<ObjectType>& InOutKey, const FObjectKey& NewObject)
+	{
+		if (ObjectType* CastResult = Cast<ObjectType>(NewObject.ResolveObjectPtr()))
+		{
+			InOutKey = CastResult;
+		}
+	}
+	void ReplaceObject(FObjectKey& InOutKey, const FObjectKey& NewObject)
+	{
+		InOutKey = NewObject;
+	}
+	TSharedPtr<FPreAnimatedObjectGroupManager> ObjectGroupManager;
+};
 
 template<typename ObjectTraits>
 struct TPreAnimatedStateStorage_ObjectTraits
@@ -34,6 +73,8 @@ struct TPreAnimatedStateStorage_ObjectTraits
 	using KeyType = typename ObjectTraits::KeyType;
 	using StorageType = typename ObjectTraits::StorageType;
 
+	static_assert(ObjectTraits::SupportsGrouping, "Pre-animated object state storage should support grouping by object");
+
 	TPreAnimatedStateStorage_ObjectTraits()
 	{}
 
@@ -121,9 +162,7 @@ public:
 		EPreAnimatedStorageRequirement StorageRequirement = this->ParentExtension->GetStorageRequirement(Entry);
 		if (!this->IsStorageRequirementSatisfied(StorageIndex, StorageRequirement))
 		{
-			StorageType NewValue;
-			ObjectTraits::CachePreAnimatedValue(BoundObject, NewValue);
-	
+			StorageType NewValue = this->Traits.CachePreAnimatedValue(BoundObject);
 			this->AssignPreAnimatedValue(StorageIndex, StorageRequirement, MoveTemp(NewValue));
 		}
 	
@@ -133,8 +172,8 @@ public:
 		}
 	}
 
-	template<typename ValueInitializer>
-	void CachePreAnimatedValue(const FCachePreAnimatedValueParams& Params, UObject* BoundObject, ValueInitializer&& InitCallback)
+	template<typename OnCacheValue /* StorageType(const KeyType&) */>
+	void CachePreAnimatedValue(const FCachePreAnimatedValueParams& Params, UObject* BoundObject, OnCacheValue&& CacheCallback)
 	{
 		FObjectKey Key{ BoundObject };
 
@@ -151,21 +190,10 @@ public:
 			return;
 		}
 
-		CachePreAnimatedValue(Params, Entry, BoundObject, Forward<ValueInitializer>(InitCallback));
-	}
-
-	template<typename ValueInitializer>
-	void CachePreAnimatedValue(const FCachePreAnimatedValueParams& Params, const FPreAnimatedStateEntry& Entry, UObject* BoundObject, ValueInitializer&& InitCallback)
-	{
-		FPreAnimatedStorageIndex StorageIndex = Entry.ValueHandle.StorageIndex;
-
 		EPreAnimatedStorageRequirement StorageRequirement = this->ParentExtension->GetStorageRequirement(Entry);
 		if (!this->IsStorageRequirementSatisfied(StorageIndex, StorageRequirement))
 		{
-			StorageType NewValue;
-			InitCallback(BoundObject, NewValue);
-			ObjectTraits::CachePreAnimatedValue(BoundObject, NewValue);
-
+			StorageType NewValue = CacheCallback(Key);
 			this->AssignPreAnimatedValue(StorageIndex, StorageRequirement, MoveTemp(NewValue));
 		}
 
@@ -176,16 +204,6 @@ public:
 	}
 };
 
-
-
-
-
-
 } // namespace MovieScene
 } // namespace UE
 
-
-
-
-
-
diff --git a/Engine/Source/Runtime/MovieScene/Public/Evaluation/PreAnimatedState/MovieScenePreAnimatedObjectTokenStorage.h b/Engine/Source/Runtime/MovieScene/Public/Evaluation/PreAnimatedState/MovieScenePreAnimatedObjectTokenStorage.h
--- a/Engine/Source/Runtime/MovieScene/Public/Evaluation/PreAnimatedState/MovieScenePreAnimatedObjectTokenStorage.h
+++ b/Engine/Source/Runtime/MovieScene/Public/Evaluation/PreAnimatedState/MovieScenePreAnimatedObjectTokenStorage.h
@@ -6,8 +6,8 @@
 #include "EntitySystem/MovieSceneSequenceInstanceHandle.h"
 #include "Evaluation/MovieSceneAnimTypeID.h"
 #include "Evaluation/PreAnimatedState/MovieScenePreAnimatedObjectGroupManager.h"
+#include "Evaluation/PreAnimatedState/MovieScenePreAnimatedObjectStorage.h"
 #include "Evaluation/PreAnimatedState/MovieScenePreAnimatedStateExtension.h"
-#include "Evaluation/PreAnimatedState/MovieScenePreAnimatedStateStorage.h"
 #include "Evaluation/PreAnimatedState/MovieScenePreAnimatedStateTypes.h"
 #include "Evaluation/PreAnimatedState/MovieScenePreAnimatedStorageID.h"
 #include "Evaluation/PreAnimatedState/MovieSceneRestoreStateParams.h"
diff --git a/Engine/Source/Runtime/MovieScene/Public/Evaluation/PreAnimatedState/MovieScenePreAnimatedPropertyStorage.h b/Engine/Source/Runtime/MovieScene/Public/Evaluation/PreAnimatedState/MovieScenePreAnimatedPropertyStorage.h
--- a/Engine/Source/Runtime/MovieScene/Public/Evaluation/PreAnimatedState/MovieScenePreAnimatedPropertyStorage.h
+++ b/Engine/Source/Runtime/MovieScene/Public/Evaluation/PreAnimatedState/MovieScenePreAnimatedPropertyStorage.h
@@ -14,7 +14,7 @@
 #include "EntitySystem/MovieScenePropertyBinding.h"
 #include "Evaluation/PreAnimatedState/MovieSceneRestoreStateParams.h"
 #include "Evaluation/PreAnimatedState/MovieScenePreAnimatedStorageID.inl"
-#include "Evaluation/PreAnimatedState/MovieScenePreAnimatedStateStorage.h"
+#include "Evaluation/PreAnimatedState/MovieScenePreAnimatedObjectStorage.h"
 #include "Evaluation/PreAnimatedState/MovieScenePreAnimatedEntityCaptureSource.h"
 
 
@@ -132,7 +132,10 @@ struct TPreAnimatedPropertyStorageImpl<PropertyTraits, TPropertyMetaData<MetaDat
 	: TPreAnimatedStateStorage<TPreAnimatedPropertyTraits<PropertyTraits, TIntegerSequence<int, MetaDataIndices...>, MetaDataTypes...>>
 	, IPreAnimatedObjectPropertyStorage
 {
-	using StorageType = typename TPreAnimatedPropertyTraits<PropertyTraits, TIntegerSequence<int, MetaDataIndices...>, MetaDataTypes...>::StorageType;
+	using StorageTraits = TPreAnimatedPropertyTraits<PropertyTraits, TIntegerSequence<int, MetaDataIndices...>, MetaDataTypes...>;
+	using StorageType = typename StorageTraits::StorageType;
+
+	static_assert(StorageTraits::SupportsGrouping, "Pre-animated storage for properties should support grouping by object");
 
 	TPreAnimatedPropertyStorageImpl(const FPropertyDefinition& InPropertyDefinition)
 		: MetaDataComponents(InPropertyDefinition.MetaDataTypes)
diff --git a/Engine/Source/Runtime/MovieScene/Public/Evaluation/PreAnimatedState/MovieScenePreAnimatedRootTokenStorage.h b/Engine/Source/Runtime/MovieScene/Public/Evaluation/PreAnimatedState/MovieScenePreAnimatedRootTokenStorage.h
--- a/Engine/Source/Runtime/MovieScene/Public/Evaluation/PreAnimatedState/MovieScenePreAnimatedRootTokenStorage.h
+++ b/Engine/Source/Runtime/MovieScene/Public/Evaluation/PreAnimatedState/MovieScenePreAnimatedRootTokenStorage.h
@@ -21,7 +21,7 @@ namespace MovieScene
 struct FPreAnimatedStateExtension;
 struct FRestoreStateParams;
 
-struct FPreAnimatedRootTokenTraits : FDefaultPreAnimatedStateTraits
+struct FPreAnimatedRootTokenTraits : FPreAnimatedStateTraits
 {
 	using KeyType     = FMovieSceneAnimTypeID;
 	using StorageType = IMovieScenePreAnimatedGlobalTokenPtr;
@@ -44,13 +44,5 @@ public:
 	void Initialize(FPreAnimatedStorageID InStorageID, FPreAnimatedStateExtension* ParentExtension) override;
 };
 
-
-
 } // namespace MovieScene
 } // namespace UE
-
-
-
-
-
-
diff --git a/Engine/Source/Runtime/MovieScene/Public/Evaluation/PreAnimatedState/MovieScenePreAnimatedStateStorage.h b/Engine/Source/Runtime/MovieScene/Public/Evaluation/PreAnimatedState/MovieScenePreAnimatedStateStorage.h
--- a/Engine/Source/Runtime/MovieScene/Public/Evaluation/PreAnimatedState/MovieScenePreAnimatedStateStorage.h
+++ b/Engine/Source/Runtime/MovieScene/Public/Evaluation/PreAnimatedState/MovieScenePreAnimatedStateStorage.h
@@ -46,7 +46,7 @@ struct FRootInstanceHandle;
  * to skip saving pre-animated state for some entities.
  */
 template<typename... InputTypes>
-struct TSavePreAnimatedStateParams
+struct TPreAnimatedStateTaskParams
 {
 	FEntityComponentFilter AdditionalFilter;
 
@@ -56,64 +56,39 @@ struct TSavePreAnimatedStateParams
 	}
 };
 
-struct FDefaultPreAnimatedStateTraits
-{
-	void Initialize(FPreAnimatedStorageID InStorageID, FPreAnimatedStateExtension* InParentExtension)
-	{}
-
-	/* Defined as a template rather than a variadic function to prevent error C4840 */
-	template<typename... T>
-	FPreAnimatedStorageGroupHandle MakeGroup(T&&... Unused)
-	{
-		return FPreAnimatedStorageGroupHandle();
-	}
-
-	template<typename KeyType>
-	void ReplaceObject(KeyType& InOutKey, const FObjectKey& NewObject)
-	{}
-};
-
 /**
- * Group state class that groups pre-animated storage together by bound object
+ * Base class for all pre-animated state traits.
+ *
+ * Inherit from this class to get the default flag values. The sub-class must/might implement
+ * the following members:
+ *
+ * (mandatory)
+ *   typedef or alias KeyType, must be constructible from (InputTypes...)
+ *   typedef or alias StorageType
+ *   void RestorePreAnimatedValue(const KeyType&, const StorageValue&, const FRestoreStateParams&);
+ *
+ * (optional, if using the ECS-wide tasks for tracking and caching state)
+ *   StorageType CachePreAnimatedValue(InputTypes...);
+ *
+ * (optional, if NeedsInitialize is true)
+ *   void Initialize(FPreAnimatedStorageID, FPreAnimatedStateExtension*);
+ *
+ * (optional, if SupportsGrouping is true)
+ *   FPreAnimatedStorageGroupHandle MakeGroup(InputTypes...);
+ *
+ * (optional, if SupportsReplaceObject is true)
+ *   void ReplaceObject(KeyType&, const FObjectKey&);
  */
-struct MOVIESCENE_API FBoundObjectPreAnimatedStateTraits
+struct MOVIESCENE_API FPreAnimatedStateTraits
 {
-	void Initialize(FPreAnimatedStorageID InStorageID, FPreAnimatedStateExtension* InParentExtension);
-
-	/* Defined as a template rather than a variadic function to prevent error C4840 */
-	template<typename... T>
-	FPreAnimatedStorageGroupHandle MakeGroup(UObject* BoundObject, T&&... Unused)
+	enum
 	{
-		return MakeGroupImpl(BoundObject);
-	}
-	FPreAnimatedStorageGroupHandle MakeGroupImpl(UObject* BoundObject);
-
-	template<typename ...T>
-	void ReplaceObject(TTuple<FObjectKey, T...>& InOutKey, const FObjectKey& NewObject)
-	{
-		InOutKey.template Get<0>() = NewObject;
-	}
-	template<typename KeyType>
-	void ReplaceObject(KeyType& InOutKey, const FObjectKey& NewObject)
-	{
-		InOutKey.Object = NewObject;
-	}
-	template<typename ObjectType>
-	void ReplaceObject(TObjectKey<ObjectType>& InOutKey, const FObjectKey& NewObject)
-	{
-		if (ObjectType* CastResult = Cast<ObjectType>(NewObject.ResolveObjectPtr()))
-		{
-			InOutKey = CastResult;
-		}
-	}
-	void ReplaceObject(FObjectKey& InOutKey, const FObjectKey& NewObject)
-	{
-		InOutKey = NewObject;
-	}
-	TSharedPtr<FPreAnimatedObjectGroupManager> ObjectGroupManager;
+		NeedsInitialize = false,
+		SupportsGrouping = false,
+		SupportsReplaceObject = false,
+	};
 };
 
-
 /**
  * Storage container for a specific type of pre-animated state as defined by the specified traits.
  *
@@ -187,7 +162,10 @@ public:
 		ParentExtension = InParentExtension;
 		StorageID = InStorageID;
 
-		this->Traits.Initialize(InStorageID, InParentExtension);
+		if constexpr (StorageTraits::NeedsInitialize)
+		{
+			this->Traits.Initialize(InStorageID, InParentExtension);
+		}
 	}
 
 	/** Retrieve the ID that uniquely identifies this storage container */
@@ -381,31 +359,6 @@ public:
 		}
 	}
 
-	/**
-	 * Save pre-aniamted state for the specified group and key using a callback.
-	 * Callback will only be invoked if state has not already been saved.
-	 */
-	template<typename OnCacheValue /* StorageType(const KeyType&) */>
-	void SavePreAnimatedState(FPreAnimatedStorageGroupHandle GroupHandle, const KeyType& InKey, OnCacheValue&& CacheCallback)
-	{
-		// Find the storage index for the specific key we're animating
-		FPreAnimatedStorageIndex StorageIndex = GetOrCreateStorageIndex(InKey);
-
-		FPreAnimatedStateEntry Entry{ GroupHandle, FPreAnimatedStateCachedValueHandle{ StorageID, StorageIndex } };
-
-		// Begin tracking the entry from its capture source
-		ParentExtension->AddSourceMetaData(Entry);
-
-		EPreAnimatedStorageRequirement Requirement = ParentExtension->HasActiveCaptureSource()
-			? EPreAnimatedStorageRequirement::Transient
-			: EPreAnimatedStorageRequirement::Persistent;
-
-		if (!IsStorageRequirementSatisfied(Entry.ValueHandle.StorageIndex, Requirement))
-		{
-			AssignPreAnimatedValue(StorageIndex, Requirement, CacheCallback(InKey));
-		}
-	}
-
 	/**
 	 * Attempt to find a storage index for the specified key, creating a new one if it doesn't exist
 	 */
@@ -540,16 +493,18 @@ public:
 
 	/**
 	 * Look at any entity with the specified component types, and set up new associations with storage indices for those entities
+	 * The provided component values are put together to make up the storage key
 	 * WARNING: Does not cache actual pre-animated values
 	 */
 	template<typename... ContributorTypes>
 	void BeginTrackingEntities(UMovieSceneEntitySystemLinker* Linker, TComponentTypeID<ContributorTypes>... InComponentTypes)
 	{
-		BeginTrackingEntitiesTask(Linker, TSavePreAnimatedStateParams<ContributorTypes...>(), InComponentTypes...);
+		BeginTrackingEntitiesTask(Linker, TPreAnimatedStateTaskParams<ContributorTypes...>(), InComponentTypes...);
 	}
 
 	/**
 	 * Look at any entity with the specified component types, and set up new associations with storage indices for those entities
+	 * The provided component values are put together to make up the storage key
 	 * WARNING: Does not cache actual pre-animated values
 	 */
 	template<typename TaskType, typename... ContributorTypes>
@@ -577,8 +532,13 @@ public:
 
 				KeyType Key{ Inputs[Index]... };
 
-				FPreAnimatedStorageGroupHandle GroupHandle  = this->Traits.MakeGroup(Inputs[Index]...);
-				FPreAnimatedStorageIndex       StorageIndex = this->GetOrCreateStorageIndex(Key);
+				FPreAnimatedStorageGroupHandle GroupHandle;
+				if constexpr(StorageTraits::SupportsGrouping)
+				{
+					GroupHandle = this->Traits.MakeGroup(Inputs[Index]...);
+				}
+
+				FPreAnimatedStorageIndex StorageIndex = this->GetOrCreateStorageIndex(Key);
 
 				FPreAnimatedStateEntry Entry{ GroupHandle, FPreAnimatedStateCachedValueHandle{ this->StorageID, StorageIndex } };
 				EntityMetaData->BeginTrackingEntity(Entry, EntityIDs[Index], RootInstanceHandles[Index], bWantsRestore);
@@ -598,6 +558,7 @@ public:
 
 	/**
 	 * Set up a new associations with a storage index for the given entity
+	 * The provided component values are put together to make up the storage key
 	 * WARNING: Does not cache actual pre-animated values
 	 */
 	template<typename... ContributorTypes>
@@ -608,7 +569,7 @@ public:
 			return;
 		}
 
-		TSavePreAnimatedStateParams<ContributorTypes...> Params;
+		TPreAnimatedStateTaskParams<ContributorTypes...> Params;
 
 		if (!Params.ShouldSavePreAnimatedState(InComponents...))
 		{
@@ -617,8 +578,13 @@ public:
 
 		KeyType Key{ InComponents... };
 
-		FPreAnimatedStorageGroupHandle GroupHandle  = this->Traits.MakeGroup(InComponents...);
-		FPreAnimatedStorageIndex       StorageIndex = this->GetOrCreateStorageIndex(Key);
+		FPreAnimatedStorageGroupHandle GroupHandle;
+		if constexpr(StorageTraits::SupportsGrouping)
+		{
+			GroupHandle = this->Traits.MakeGroup(InComponents...);
+		}
+
+		FPreAnimatedStorageIndex StorageIndex = this->GetOrCreateStorageIndex(Key);
 
 		FPreAnimatedStateEntry Entry{ GroupHandle, FPreAnimatedStateCachedValueHandle{ this->StorageID, StorageIndex } };
 		FPreAnimatedEntityCaptureSource* EntityMetaData = this->ParentExtension->GetOrCreateEntityMetaData();
@@ -628,21 +594,21 @@ public:
 	/**
 	 * Cache pre-animated values for entities with the specified component types
 	 */
-	template<typename... CacheValueInputs>
-	void CachePreAnimatedValues(UMovieSceneEntitySystemLinker* Linker, TComponentTypeID<CacheValueInputs>... InComponentTypes)
+	template<typename... ContributorTypes>
+	void CachePreAnimatedValues(UMovieSceneEntitySystemLinker* Linker, TComponentTypeID<ContributorTypes>... InComponentTypes)
 	{
-		CachePreAnimatedValuesTask(Linker, TSavePreAnimatedStateParams<CacheValueInputs...>(), InComponentTypes...);
+		CachePreAnimatedValuesTask(Linker, TPreAnimatedStateTaskParams<ContributorTypes...>(), InComponentTypes...);
 	}
 
 	/**
 	 * Cache pre-animated values for entities with the specified component types
 	 */
-	template<typename TaskType, typename... CacheValueInputs>
-	void CachePreAnimatedValuesTask(UMovieSceneEntitySystemLinker* Linker, const TaskType& InParams, TComponentTypeID<CacheValueInputs>... InComponentTypes)
+	template<typename TaskType, typename... ContributorTypes>
+	void CachePreAnimatedValuesTask(UMovieSceneEntitySystemLinker* Linker, const TaskType& InParams, TComponentTypeID<ContributorTypes>... InComponentTypes)
 	{
 		FBuiltInComponentTypes* BuiltInComponents = FBuiltInComponentTypes::Get();
 
-		auto VisitAllocation = [this, &InParams](FEntityAllocationIteratorItem Item, TRead<CacheValueInputs>... Values)
+		auto VisitAllocation = [this, &InParams](FEntityAllocationIteratorItem Item, TRead<ContributorTypes>... Values)
 		{
 			const int32 Num = Item.GetAllocation()->Num();
 			for (int32 Index = 0; Index < Num; ++Index)
@@ -654,17 +620,30 @@ public:
 
 				KeyType Key{ Values[Index]... };
 
-				FPreAnimatedStorageGroupHandle GroupHandle  = this->Traits.MakeGroup(Values[Index]...);
-				FPreAnimatedStorageIndex       StorageIndex = this->GetOrCreateStorageIndex(Key);
+				FPreAnimatedStorageGroupHandle GroupHandle;
+				if constexpr(StorageTraits::SupportsGrouping)
+				{
+					GroupHandle = this->Traits.MakeGroup(Values[Index]...);
+				}
+		
+				FPreAnimatedStorageIndex StorageIndex = this->GetOrCreateStorageIndex(Key);
 
 				FPreAnimatedStateEntry Entry{ GroupHandle, FPreAnimatedStateCachedValueHandle{ this->StorageID, StorageIndex } };
 
-				this->ParentExtension->EnsureMetaData(Entry);
+				if (this->ParentExtension->IsCapturingGlobalState())
+				{
+					this->ParentExtension->EnsureMetaData(Entry);
+				}
+				else if (!this->ParentExtension->MetaDataExists(Entry))
+				{
+					continue;
+				}
 
 				EPreAnimatedStorageRequirement StorageRequirement = this->ParentExtension->GetStorageRequirement(Entry);
 				if (!this->IsStorageRequirementSatisfied(StorageIndex, StorageRequirement))
 				{
-					this->AssignPreAnimatedValue(StorageIndex, StorageRequirement, this->Traits.CachePreAnimatedValue(Values[Index]...));
+					StorageType NewValue = this->Traits.CachePreAnimatedValue(Values[Index]...);
+					this->AssignPreAnimatedValue(StorageIndex, StorageRequirement, MoveTemp(NewValue));
 				}
 			}
 		};
@@ -676,6 +655,95 @@ public:
 		.Iterate_PerAllocation(&Linker->EntityManager, VisitAllocation);
 	}
 
+	/**
+	 * Save pre-animated state for the specified values
+	 * Requires that the traits class implements CachePreAnimatedValue(ContributorTypes...)
+	 */
+	template<typename... ContributorTypes>
+	void CachePreAnimatedValue(ContributorTypes... Values)
+	{
+		KeyType Key{ Values... };
+
+		FPreAnimatedStorageGroupHandle GroupHandle;
+		if constexpr(StorageTraits::SupportsGrouping)
+		{
+			GroupHandle = this->Traits.MakeGroup(Values...);
+		}
+
+		FPreAnimatedStorageIndex StorageIndex = this->GetOrCreateStorageIndex(Key);
+
+		FPreAnimatedStateEntry Entry{ GroupHandle, FPreAnimatedStateCachedValueHandle{ this->StorageID, StorageIndex } };
+
+		if (this->ParentExtension->IsCapturingGlobalState())
+		{
+			this->ParentExtension->EnsureMetaData(Entry);
+		}
+		else if (!this->ParentExtension->MetaDataExists(Entry))
+		{
+			return;
+		}
+
+		// Begin tracking the entry from its capture source
+		ParentExtension->AddSourceMetaData(Entry);
+
+		EPreAnimatedStorageRequirement Requirement = ParentExtension->HasActiveCaptureSource()
+			? EPreAnimatedStorageRequirement::Transient
+			: EPreAnimatedStorageRequirement::Persistent;
+
+		if (!IsStorageRequirementSatisfied(Entry.ValueHandle.StorageIndex, Requirement))
+		{
+			StorageType NewValue = this->Traits.CachePreAnimatedValue(Values...);
+			this->AssignPreAnimatedValue(StorageIndex, Requirement, MoveTemp(NewValue));
+		}
+	}
+
+	/**
+	 * Save pre-animated state for the specified group and key using a callback.
+	 * Callback will only be invoked if state has not already been saved.
+	 */
+	template<typename OnCacheValue /* StorageType(const KeyType&) */>
+	void CachePreAnimatedValue(const KeyType& InKey, OnCacheValue&& CacheCallback)
+	{
+		static_assert(StorageTraits::SupportsGrouping == false, "Grouped pre-animated state requires passing a group handle");
+		CachePreAnimatedValue(FPreAnimatedStorageGroupHandle(), InKey, Forward<OnCacheValue>(CacheCallback));
+	}
+
+	/**
+	 * Save pre-animated state for the specified group and key using a callback.
+	 * Callback will only be invoked if state has not already been saved.
+	 */
+	template<typename OnCacheValue /* StorageType(const KeyType&) */>
+	void CachePreAnimatedValue(FPreAnimatedStorageGroupHandle GroupHandle, const KeyType& InKey, OnCacheValue&& CacheCallback)
+	{
+		ensureMsgf(GroupHandle.IsValid() || !StorageTraits::SupportsGrouping, TEXT("The group handle must be valid for pre-animated state that supports grouping, and invalid if not"));
+
+		// Find the storage index for the specific key we're animating
+		FPreAnimatedStorageIndex StorageIndex = GetOrCreateStorageIndex(InKey);
+
+		FPreAnimatedStateEntry Entry{ GroupHandle, FPreAnimatedStateCachedValueHandle{ this->StorageID, StorageIndex } };
+
+		if (this->ParentExtension->IsCapturingGlobalState())
+		{
+			this->ParentExtension->EnsureMetaData(Entry);
+		}
+		else if (!this->ParentExtension->MetaDataExists(Entry))
+		{
+			return;
+		}
+
+		// Begin tracking the entry from its capture source
+		ParentExtension->AddSourceMetaData(Entry);
+
+		EPreAnimatedStorageRequirement Requirement = ParentExtension->HasActiveCaptureSource()
+			? EPreAnimatedStorageRequirement::Transient
+			: EPreAnimatedStorageRequirement::Persistent;
+
+		if (!IsStorageRequirementSatisfied(Entry.ValueHandle.StorageIndex, Requirement))
+		{
+			StorageType NewValue = CacheCallback(InKey);
+			AssignPreAnimatedValue(StorageIndex, Requirement, MoveTemp(NewValue));
+		}
+	}
 
 	/**
 	 * Look at any entity with the specified component types, and set up new associations with storage indices for those entities, whilst also caching pre-animated values at the same time.
@@ -683,7 +751,7 @@ public:
 	template<typename... ContributorTypes>
 	void BeginTrackingAndCachePreAnimatedValues(UMovieSceneEntitySystemLinker* Linker, TComponentTypeID<ContributorTypes>... InComponentTypes)
 	{
-		BeginTrackingAndCachePreAnimatedValuesTask(Linker, TSavePreAnimatedStateParams<ContributorTypes...>(), InComponentTypes...);
+		BeginTrackingAndCachePreAnimatedValuesTask(Linker, TPreAnimatedStateTaskParams<ContributorTypes...>(), InComponentTypes...);
 	}
 
 	/**
@@ -714,8 +782,13 @@ public:
 
 				KeyType Key{ Inputs[Index]... };
 
-				FPreAnimatedStorageGroupHandle GroupHandle  = this->Traits.MakeGroup(Inputs[Index]...);
-				FPreAnimatedStorageIndex       StorageIndex = this->GetOrCreateStorageIndex(Key);
+				FPreAnimatedStorageGroupHandle GroupHandle;
+				if constexpr(StorageTraits::SupportsGrouping)
+				{
+					GroupHandle = this->Traits.MakeGroup(Inputs[Index]...);
+				}
+
+				FPreAnimatedStorageIndex StorageIndex = this->GetOrCreateStorageIndex(Key);
 
 				FPreAnimatedStateEntry Entry{ GroupHandle, FPreAnimatedStateCachedValueHandle{ this->StorageID, StorageIndex } };
 				EntityMetaData->BeginTrackingEntity(Entry, EntityIDs[Index], RootInstanceHandles[Index], bWantsRestore);
@@ -723,7 +796,8 @@ public:
 				EPreAnimatedStorageRequirement StorageRequirement = this->ParentExtension->GetStorageRequirement(Entry);
 				if (!this->IsStorageRequirementSatisfied(StorageIndex, StorageRequirement))
 				{
-					this->AssignPreAnimatedValue(StorageIndex, StorageRequirement, this->Traits.CachePreAnimatedValue(Inputs[Index]...));
+					StorageType NewValue = this->Traits.CachePreAnimatedValue(Inputs[Index]...);
+					this->AssignPreAnimatedValue(StorageIndex, StorageRequirement, MoveTemp(NewValue));
 				}
 
 			}
@@ -743,7 +817,10 @@ public:
 	void OnObjectReplaced(FPreAnimatedStorageIndex StorageIndex, const FObjectKey& OldObject, const FObjectKey& NewObject) override
 	{
 		KeyType ExistingKey = this->GetKey(StorageIndex);
-		this->Traits.ReplaceObject(ExistingKey, NewObject);
+		if constexpr(StorageTraits::SupportsReplaceObject)
+		{
+			this->Traits.ReplaceObject(ExistingKey, NewObject);
+		}
 		this->ReplaceKey(StorageIndex, ExistingKey);
 	}
 
@@ -805,7 +882,32 @@ protected:
 };
 
 
-} // namespace MovieScene
-} // namespace UE
+/**
+ * Simple traits class that works with the simple storage below.
+ * The only thing that is needed is a `void RestoreState(KeyType, Params)` method on the state
+ * class.
+ */
+template<typename InKeyType, typename InStorageType>
+struct TSimplePreAnimatedStateTraits : FPreAnimatedStateTraits
+{
+	using KeyType = InKeyType;
+	using StorageType = InStorageType;
+
+	void RestorePreAnimatedValue(const KeyType& Key, StorageType& SavedValue, const FRestoreStateParams& Params)
+	{
+		SavedValue.RestoreState(Key, Params);
+	}
+};
 
 
+/**
+ * Simple pre-animated state storage that doesn't support any grouping, and lets people simply implement
+ * a `void RestoreState(KeyType, Params)` method on their state class.
+ */
+template<typename KeyType, typename StorageType>
+struct TSimplePreAnimatedStateStorage : TPreAnimatedStateStorage<TSimplePreAnimatedStateTraits<KeyType, StorageType>>
+{
+};
+
+} // namespace MovieScene
+} // namespace UE
diff --git a/Engine/Source/Runtime/MovieSceneTracks/Private/Systems/MovieSceneComponentAttachmentSystem.cpp b/Engine/Source/Runtime/MovieSceneTracks/Private/Systems/MovieSceneComponentAttachmentSystem.cpp
--- a/Engine/Source/Runtime/MovieSceneTracks/Private/Systems/MovieSceneComponentAttachmentSystem.cpp
+++ b/Engine/Source/Runtime/MovieSceneTracks/Private/Systems/MovieSceneComponentAttachmentSystem.cpp
@@ -110,12 +110,14 @@ struct FComponentAttachmentPreAnimatedTraits : FBoundObjectPreAnimatedStateTrait
 	using KeyType     = FObjectKey;
 	using StorageType = FPreAnimAttachment;
 
-	static void CachePreAnimatedValue(UObject* InObject, FPreAnimAttachment& OutCachedAttachment)
+	static FPreAnimAttachment CachePreAnimatedValue(UObject* InObject)
 	{
 		USceneComponent* SceneComponent = CastChecked<USceneComponent>(InObject);
 
+		FPreAnimAttachment OutCachedAttachment;
 		OutCachedAttachment.OldAttachParent = SceneComponent->GetAttachParent();
 		OutCachedAttachment.OldAttachSocket = SceneComponent->GetAttachSocketName();
+		return OutCachedAttachment;
 	}
 	static void RestorePreAnimatedValue(const FObjectKey& InKey, FPreAnimAttachment& InOutCachedAttachment, const FRestoreStateParams& Params)
 	{
@@ -126,8 +128,7 @@ struct FComponentAttachmentPreAnimatedTraits : FBoundObjectPreAnimatedStateTrait
 	}
 };
 
-struct FPreAnimatedComponentAttachmentStorage
-	: TPreAnimatedStateStorage_ObjectTraits<FComponentAttachmentPreAnimatedTraits>
+struct FPreAnimatedComponentAttachmentStorage : TPreAnimatedStateStorage_ObjectTraits<FComponentAttachmentPreAnimatedTraits>
 {
 	static TAutoRegisterPreAnimatedStorageID<FPreAnimatedComponentAttachmentStorage> StorageID;
 };
@@ -177,6 +178,7 @@ UMovieSceneComponentAttachmentSystem::UMovieSceneComponentAttachmentSystem(const
 
 		DefineImplicitPrerequisite(UMovieSceneCachePreAnimatedStateSystem::StaticClass(), GetClass());
 		DefineImplicitPrerequisite(UMovieSceneComponentMobilitySystem::StaticClass(), GetClass());
+		DefineImplicitPrerequisite(UMovieSceneComponentAttachmentInvalidatorSystem::StaticClass(), GetClass());
 		DefineImplicitPrerequisite(GetClass(), UMovieSceneRestorePreAnimatedStateSystem::StaticClass());
 
 		DefineComponentConsumer(GetClass(), FBuiltInComponentTypes::Get()->BoundObject);
@@ -191,7 +193,6 @@ void UMovieSceneComponentAttachmentSystem::OnLink()
 
 	UMovieSceneComponentAttachmentInvalidatorSystem* AttachmentInvalidator = Linker->LinkSystem<UMovieSceneComponentAttachmentInvalidatorSystem>();
 	Linker->SystemGraph.AddReference(this, AttachmentInvalidator);
-	Linker->SystemGraph.AddPrerequisite(AttachmentInvalidator, this);
 
 	AttachmentTracker.Initialize(this);
 }
diff --git a/Engine/Source/Runtime/MovieSceneTracks/Private/Systems/MovieSceneComponentMobilitySystem.cpp b/Engine/Source/Runtime/MovieSceneTracks/Private/Systems/MovieSceneComponentMobilitySystem.cpp
--- a/Engine/Source/Runtime/MovieSceneTracks/Private/Systems/MovieSceneComponentMobilitySystem.cpp
+++ b/Engine/Source/Runtime/MovieSceneTracks/Private/Systems/MovieSceneComponentMobilitySystem.cpp
@@ -100,12 +100,14 @@ void FPreAnimatedMobilityTraits::RestorePreAnimatedValue(const FObjectKey& InKey
 	}
 }
 
-void FPreAnimatedMobilityTraits::CachePreAnimatedValue(UObject* InObject, EComponentMobility::Type& OutMobility)
+EComponentMobility::Type FPreAnimatedMobilityTraits::CachePreAnimatedValue(UObject* InObject)
 {
+	EComponentMobility::Type OutMobility = EComponentMobility::Static;
 	if (USceneComponent* SceneComponent = Cast<USceneComponent>(InObject))
 	{
 		OutMobility = SceneComponent->Mobility;
 	}
+	return OutMobility;
 }
 
 } // namespace MovieScene
diff --git a/Engine/Source/Runtime/MovieSceneTracks/Private/Systems/MovieSceneMaterialParameterCollectionSystem.cpp b/Engine/Source/Runtime/MovieSceneTracks/Private/Systems/MovieSceneMaterialParameterCollectionSystem.cpp
--- a/Engine/Source/Runtime/MovieSceneTracks/Private/Systems/MovieSceneMaterialParameterCollectionSystem.cpp
+++ b/Engine/Source/Runtime/MovieSceneTracks/Private/Systems/MovieSceneMaterialParameterCollectionSystem.cpp
@@ -277,7 +277,7 @@ void UMovieSceneMaterialParameterCollectionSystem::SavePreAnimatedState(const FP
 	FBuiltInComponentTypes*          BuiltInComponents = FBuiltInComponentTypes::Get();
 	FMovieSceneTracksComponentTypes* TracksComponents  = FMovieSceneTracksComponentTypes::Get();
 
-	TSavePreAnimatedStateParams<UObject*, FName> Params;
+	TPreAnimatedStateTaskParams<UObject*, FName> Params;
 
 	Params.AdditionalFilter.All({ TracksComponents->MPC });
 	ScalarParameterStorage->BeginTrackingEntitiesTask(Linker, Params, TracksComponents->BoundMaterial, TracksComponents->ScalarParameterName);
diff --git a/Engine/Source/Runtime/MovieSceneTracks/Private/Systems/MovieSceneSkeletalAnimationSystem.cpp b/Engine/Source/Runtime/MovieSceneTracks/Private/Systems/MovieSceneSkeletalAnimationSystem.cpp
--- a/Engine/Source/Runtime/MovieSceneTracks/Private/Systems/MovieSceneSkeletalAnimationSystem.cpp
+++ b/Engine/Source/Runtime/MovieSceneTracks/Private/Systems/MovieSceneSkeletalAnimationSystem.cpp
@@ -73,8 +73,9 @@ struct FPreAnimatedSkeletalAnimationTraits : FBoundObjectPreAnimatedStateTraits
 	using KeyType     = FObjectKey;
 	using StorageType = FPreAnimatedSkeletalAnimationState;
 
-	static void CachePreAnimatedValue(const KeyType& Object, StorageType& OutCachedValue)
+	FPreAnimatedSkeletalAnimationState CachePreAnimatedValue(const KeyType& Object)
 	{
+		FPreAnimatedSkeletalAnimationState OutCachedValue;
 		USkeletalMeshComponent* Component = Cast<USkeletalMeshComponent>(Object.ResolveObjectPtr());
 		if (ensure(Component))
 		{
@@ -82,9 +83,10 @@ struct FPreAnimatedSkeletalAnimationTraits : FBoundObjectPreAnimatedStateTraits
 			OutCachedValue.CachedAnimInstance.Reset(Component->AnimScriptInstance);
 			OutCachedValue.SkeletalMeshRestoreState.SaveState(Component);
 		}
+		return OutCachedValue;
 	}
 
-	static void RestorePreAnimatedValue(const KeyType& Object, StorageType& InOutCachedValue, const FRestoreStateParams& Params)
+	void RestorePreAnimatedValue(const KeyType& Object, StorageType& InOutCachedValue, const FRestoreStateParams& Params)
 	{
 		USkeletalMeshComponent* Component = Cast<USkeletalMeshComponent>(Object.ResolveObjectPtr());
 		if (!Component)
@@ -169,12 +171,13 @@ struct FPreAnimatedSkeletalAnimationMontageTraits : FBoundObjectPreAnimatedState
 	using KeyType     = FObjectKey;
 	using StorageType = FPreAnimatedSkeletalAnimationMontageState;
 
-	static void CachePreAnimatedValue(const KeyType& Object, StorageType& OutCachedValue)
+	FPreAnimatedSkeletalAnimationMontageState CachePreAnimatedValue(const FObjectKey& Object)
 	{
-		// Should already be initialized.
+		// Should be unused, as we always cache state with captured values.
+		return FPreAnimatedSkeletalAnimationMontageState();
 	}
 
-	static void RestorePreAnimatedValue(const FObjectKey& Object, FPreAnimatedSkeletalAnimationMontageState& InOutCachedValue, const FRestoreStateParams& Params)
+	void RestorePreAnimatedValue(const FObjectKey& Object, FPreAnimatedSkeletalAnimationMontageState& InOutCachedValue, const FRestoreStateParams& Params)
 	{
 		UAnimInstance* AnimInstance = InOutCachedValue.WeakInstance.Get();
 		if (AnimInstance)
@@ -204,12 +207,13 @@ struct FPreAnimatedSkeletalAnimationAnimInstanceTraits : FBoundObjectPreAnimated
 	using KeyType     = FObjectKey;
 	using StorageType = bool; // We actually don't need any state, so this will be a dummy value
 
-	static void CachePreAnimatedValue(const KeyType& Object, StorageType& Unused)
+	bool CachePreAnimatedValue(const FObjectKey& Object)
 	{
 		// Nothing to do, we just need the object pointer to restore state.
+		return true;
 	}
 
-	static void RestorePreAnimatedValue(const FObjectKey& Object, bool& Unused, const FRestoreStateParams& Params)
+	void RestorePreAnimatedValue(const FObjectKey& Object, bool& Unused, const FRestoreStateParams& Params)
 	{
 		if (UObject* ObjectPtr = Object.ResolveObjectPtr())
 		{
@@ -756,9 +760,11 @@ private:
 				DataContainer.MontageInstanceId = InstanceId;
 
 				PreAnimatedMontageStorage->BeginTrackingEntity(Params.EntityID, Params.bWantsRestoreState, Params.RootInstanceHandle, Montage);
-				PreAnimatedMontageStorage->CachePreAnimatedValue(FCachePreAnimatedValueParams(), Montage, [=](UObject* Unused, FPreAnimatedSkeletalAnimationMontageState& OutState) {
+				PreAnimatedMontageStorage->CachePreAnimatedValue(FCachePreAnimatedValueParams(), Montage, [=](const FObjectKey& Unused) {
+						FPreAnimatedSkeletalAnimationMontageState OutState;
 						OutState.WeakInstance = AnimInst;
 						OutState.MontageInstanceId = InstanceId;
+						return OutState;
 					});
 
 				// Make sure it's playing if the sequence is.
@@ -848,9 +854,11 @@ private:
 				DataContainer.MontageInstanceId = InstanceId;
 
 				PreAnimatedMontageStorage->BeginTrackingEntity(Params.EntityID, Params.bWantsRestoreState, Params.RootInstanceHandle, Montage);
-				PreAnimatedMontageStorage->CachePreAnimatedValue(FCachePreAnimatedValueParams(), Montage, [=](UObject* Unused, FPreAnimatedSkeletalAnimationMontageState& OutState) {
+				PreAnimatedMontageStorage->CachePreAnimatedValue(FCachePreAnimatedValueParams(), Montage, [=](const FObjectKey& Unused) {
+						FPreAnimatedSkeletalAnimationMontageState OutState;
 						OutState.WeakInstance = AnimInst;
 						OutState.MontageInstanceId = InstanceId;
+						return OutState;
 					});
 
 				FAnimMontageInstance* Instance = AnimInst->GetMontageInstanceForID(InstanceId);
diff --git a/Engine/Source/Runtime/MovieSceneTracks/Private/TrackInstances/MovieSceneCVarTrackInstance.cpp b/Engine/Source/Runtime/MovieSceneTracks/Private/TrackInstances/MovieSceneCVarTrackInstance.cpp
--- a/Engine/Source/Runtime/MovieSceneTracks/Private/TrackInstances/MovieSceneCVarTrackInstance.cpp
+++ b/Engine/Source/Runtime/MovieSceneTracks/Private/TrackInstances/MovieSceneCVarTrackInstance.cpp
@@ -21,12 +21,12 @@ namespace UE
 namespace MovieScene
 {
 
-struct FPreAnimatedCVarTraits : FDefaultPreAnimatedStateTraits
+struct FPreAnimatedCVarTraits : FPreAnimatedStateTraits
 {
 	using KeyType     = FString;
 	using StorageType = FString;
 
-	static void RestorePreAnimatedValue(const FString& InKey, const FString& CachedValue, const FRestoreStateParams& Params)
+	void RestorePreAnimatedValue(const FString& InKey, const FString& CachedValue, const FRestoreStateParams& Params)
 	{
 		static IConsoleManager& ConsoleManager = IConsoleManager::Get();
 		IConsoleVariable* ConsoleVariable = ConsoleManager.FindConsoleVariable(*InKey);
@@ -90,7 +90,7 @@ void UMovieSceneCVarTrackInstance::OnInputAdded(const FMovieSceneTrackInstanceIn
 
 		for (const TPair<FString, FString>& Pair : CommandsAndValues)
 		{
-			CVarStorage->SavePreAnimatedState(UngroupedHandle, Pair.Key, FPreAnimatedCVarTraits::CachePreAnimatedValue);
+			CVarStorage->CachePreAnimatedValue(UngroupedHandle, Pair.Key, FPreAnimatedCVarTraits::CachePreAnimatedValue);
 		}
 	}
 
diff --git a/Engine/Source/Runtime/MovieSceneTracks/Public/Systems/MovieSceneComponentMobilitySystem.h b/Engine/Source/Runtime/MovieSceneTracks/Public/Systems/MovieSceneComponentMobilitySystem.h
--- a/Engine/Source/Runtime/MovieSceneTracks/Public/Systems/MovieSceneComponentMobilitySystem.h
+++ b/Engine/Source/Runtime/MovieSceneTracks/Public/Systems/MovieSceneComponentMobilitySystem.h
@@ -54,8 +54,8 @@ struct FPreAnimatedMobilityTraits : FBoundObjectPreAnimatedStateTraits
 	using KeyType     = FObjectKey;
 	using StorageType = EComponentMobility::Type;
 
-	static void RestorePreAnimatedValue(const FObjectKey& InKey, EComponentMobility::Type Mobility, const FRestoreStateParams& Params);
-	static void CachePreAnimatedValue(UObject* InObject, EComponentMobility::Type& OutMobility);
+	void RestorePreAnimatedValue(const FObjectKey& InKey, EComponentMobility::Type Mobility, const FRestoreStateParams& Params);
+	EComponentMobility::Type CachePreAnimatedValue(UObject* InObject);
 };
 
 struct FPreAnimatedComponentMobilityStorage : TPreAnimatedStateStorage_ObjectTraits<FPreAnimatedMobilityTraits>
diff --git a/Engine/Source/Runtime/MovieSceneTracks/Public/Systems/MovieSceneMaterialSystem.h b/Engine/Source/Runtime/MovieSceneTracks/Public/Systems/MovieSceneMaterialSystem.h
--- a/Engine/Source/Runtime/MovieSceneTracks/Public/Systems/MovieSceneMaterialSystem.h
+++ b/Engine/Source/Runtime/MovieSceneTracks/Public/Systems/MovieSceneMaterialSystem.h
@@ -12,6 +12,7 @@
 #include "EntitySystem/BuiltInComponentTypes.h"
 #include "EntitySystem/MovieSceneEntitySystemRunner.h"
 #include "EntitySystem/MovieSceneComponentTypeInfo.h"
+#include "Evaluation/PreAnimatedState/MovieScenePreAnimatedObjectStorage.h"
 #include "MovieSceneTracksComponentTypes.h"
 
 #include "MovieSceneMaterialSystem.generated.h"
@@ -363,13 +364,10 @@ void TMovieSceneMaterialSystem<AccessorType, RequiredComponents...>::SavePreAnim
 {
 	using namespace UE::MovieScene;
 
-	FBuiltInComponentTypes*          BuiltInComponents = FBuiltInComponentTypes::Get();
-	FMovieSceneTracksComponentTypes* TracksComponents  = FMovieSceneTracksComponentTypes::Get();
-
 	// If we have material results to apply save those as well
 	if (Linker->EntityManager.Contains(MaterialSwitcherFilter))
 	{
-		TSavePreAnimatedStateParams<RequiredComponents...> Params;
+		TPreAnimatedStateTaskParams<RequiredComponents...> Params;
 		Params.AdditionalFilter = MaterialSwitcherFilter;
 
 		MaterialSwitcherStorage->BeginTrackingAndCachePreAnimatedValuesTask(Linker, Params, InRequiredComponents...);
@@ -378,11 +376,11 @@ void TMovieSceneMaterialSystem<AccessorType, RequiredComponents...>::SavePreAnim
 	// If we have bound materials to resolve save the current material
 	if (Linker->EntityManager.Contains(MaterialParameterFilter))
 	{
-		TSavePreAnimatedStateParams<RequiredComponents...> Params;
+		TPreAnimatedStateTaskParams<RequiredComponents...> Params;
 		Params.AdditionalFilter = MaterialParameterFilter;
 
 		MaterialParameterStorage->BeginTrackingAndCachePreAnimatedValuesTask(Linker, Params, InRequiredComponents...);
 	}
 }
 
-} // namespace UE::MovieScene
\ No newline at end of file
+} // namespace UE::MovieScene
