Sequencer: move system dependency tracking to the global dependency graph

This fixes a problem with systems not being ordered properly for execution when they don't have any direct dependencies, but have indirect ones. They would at first be (correctly) sorted in some arbitrary order, but when the "in between" systems are created, the dependency trail would force them to be re-sorted correctly. This trips the assert that detects when systems are inserted before the currently executing system, which causes other bugs because that re-sorted system would miss a frame.

diff --git a/Engine/Source/Runtime/MovieScene/Private/EntitySystem/MovieSceneEntitySystem.cpp b/Engine/Source/Runtime/MovieScene/Private/EntitySystem/MovieSceneEntitySystem.cpp
--- a/Engine/Source/Runtime/MovieScene/Private/EntitySystem/MovieSceneEntitySystem.cpp
+++ b/Engine/Source/Runtime/MovieScene/Private/EntitySystem/MovieSceneEntitySystem.cpp
@@ -32,6 +32,10 @@ struct FSystemDependencyGraph
 
 		GraphIDsByClass.Add(Class->GetFName(), NewGraphID);
 
+		FlowGraph.AllocateNode(NewGraphID);
+
+		++SerialNumber;
+
 		return NewGraphID;
 	}
 
@@ -53,92 +57,351 @@ struct FSystemDependencyGraph
 
 		GraphIDsByComponent.Add(InComponentType, NewGraphID);
 
+		FlowGraph.AllocateNode(NewGraphID);
+
+		++SerialNumber;
+
 		return NewGraphID;
 	}
 
-	void SetupDependencies(UMovieSceneEntitySystem* ThisSystem, UMovieSceneEntitySystemLinker* Linker)
+	void MakeRelationship(uint16 UpstreamGraphID, uint16 DownstreamGraphID)
+	{
+		ImplicitSubsequents.MakeEdge(UpstreamGraphID, DownstreamGraphID);
+		ImplicitPrerequisites.MakeEdge(DownstreamGraphID, UpstreamGraphID);
+		++SerialNumber;
+	}
+
+	UClass* ClassFromGraphID(uint16 GraphID) const
+	{
+		check(Nodes.IsValidIndex(GraphID));
+		return Nodes[GraphID].Class.Get();
+	}
+
+	uint16 NumGraphIDs() const
+	{
+		return Nodes.Num();
+	}
+
+	void SortByFlowOrder(TArray<uint16>& InOutNodeIDs)
+	{
+		struct FSortableID
+		{
+			uint16 NodeID;
+			uint16 FlowOrder;
+		};
+
+		UpdateCache();
+
+		TArray<FSortableID, TInlineAllocator<32>> SortableNodeIDs;
+		for (uint16 NodeID : InOutNodeIDs)
+		{
+			SortableNodeIDs.Add(FSortableID{ NodeID, ReverseLookupFlowOrderNodes[NodeID] });
+		}
+		Algo::SortBy(SortableNodeIDs, &FSortableID::FlowOrder);
+
+		InOutNodeIDs.Reset();
+		for (const FSortableID& SortableNodeID : SortableNodeIDs)
+		{
+			InOutNodeIDs.Add(SortableNodeID.NodeID);
+		}
+	}
+
+	void GetSubsequents(uint16 FromNodeID, TArray<uint16>& OutSubsequentNodeIDs)
 	{
 		using FDirectionalEdge = FMovieSceneEntitySystemDirectedGraph::FDirectionalEdge;
 
-		const uint16 FromNode = ThisSystem->GetGlobalDependencyGraphID();
+		UpdateCache();
 
-		// Set up prerequisites
-		for (FDirectionalEdge Edge : ImplicitPrerequisites.GetEdgesFrom(FromNode))
+		for (FDirectionalEdge Edge : FlowGraph.GetEdgesFrom(FromNodeID))
 		{
-			FNode Node = Nodes[Edge.ToNode];
+			OutSubsequentNodeIDs.Add(Edge.ToNode);
+		}
+	}
 
-			// Follow edges from components
-			if (Node.WriteComponentType)
+	void DebugPrint(bool bUpdateCache = true)
+	{
+		const TCHAR FormatString[] =
+			TEXT("----------------------------------------------------------------------------------\n")
+			TEXT("%s\n")
+			TEXT("----------------------------------------------------------------------------------\n");
+
+		if (bUpdateCache)
+		{
+			UpdateCache();
+		}
+
+		GLog->Log(TEXT("Printing debug graph for Entity System Graph (in standard graphviz syntax):"));
+		GLog->Log(FString::Printf(FormatString, *ToString()));
+	}
+
+	FString ToString() const
+	{
+		using FDirectionalEdge = FMovieSceneEntitySystemDirectedGraph::FDirectionalEdge;
+
+		FComponentRegistry* ComponentRegistry = UMovieSceneEntitySystemLinker::GetComponents();
+
+		FString String;
+		String += TEXT("\ndigraph FMovieSceneEntitySystemGraph {\n");
+		String += TEXT("\tnode [shape=record,height=.1];\n");
+
+		FString FlowStrings[] = 
+		{
+			TEXT("\tsubgraph cluster_flow_0 { label=\"Spawn\"; color=\"#0e868c\";\n"),
+			TEXT("\tsubgraph cluster_flow_1 { label=\"Instantiation\"; color=\"#96c74c\";\n"),
+			TEXT("\tsubgraph cluster_flow_2 { label=\"Evaluation\"; color=\"#6dc74c\";\n"),
+			TEXT("\tsubgraph cluster_flow_3 { label=\"Finalization\"; color=\"#aa42f5\";\n"),
+		};
+
+		TMap<uint16, ESystemPhase> SystemPhasesConnectedToWriteComponents;
+		for (FDirectionalEdge Edge : FlowGraph.GetEdges())
+		{
+			if (Nodes[Edge.FromNode].WriteComponentType)
 			{
-				for (FDirectionalEdge ComponentEdge : ImplicitPrerequisites.GetEdgesFrom(Edge.ToNode))
+				if (UClass* SystemClass = Nodes[Edge.ToNode].Class.Get())
 				{
-					// Components shouldn't be connected to other components
-					SetupPrereq(Nodes[ComponentEdge.ToNode].Class.Get(), ThisSystem, Linker);
+					UMovieSceneEntitySystem* System = SystemClass->GetDefaultObject<UMovieSceneEntitySystem>();
+					SystemPhasesConnectedToWriteComponents.FindOrAdd(Edge.FromNode) |= System->GetPhase();
 				}
 			}
-			else
+			else if (Nodes[Edge.ToNode].WriteComponentType)
 			{
-				SetupPrereq(Node.Class.Get(), ThisSystem, Linker);
+				if (UClass* SystemClass = Nodes[Edge.FromNode].Class.Get())
+				{
+					UMovieSceneEntitySystem* System = SystemClass->GetDefaultObject<UMovieSceneEntitySystem>();
+					SystemPhasesConnectedToWriteComponents.FindOrAdd(Edge.ToNode) |= System->GetPhase();
+				}
 			}
 		}
 
-		// Set up subsequents
-		for (FDirectionalEdge Edge : ImplicitSubsequents.GetEdgesFrom(FromNode))
+		for (int32 NodeID = 0; NodeID < Nodes.Num(); ++NodeID)
 		{
-			FNode Node = Nodes[Edge.ToNode];
-
-			// Follow edges from components
-			if (Node.WriteComponentType)
+			const FNode& Node(Nodes[NodeID]);
+			if (UClass* SystemClass = Node.Class.Get())
 			{
-				for (FDirectionalEdge ComponentEdge : ImplicitSubsequents.GetEdgesFrom(Edge.ToNode))
+				UMovieSceneEntitySystem* DefaultSystem = SystemClass->GetDefaultObject<UMovieSceneEntitySystem>();
+
+				ESystemPhase SystemPhase = DefaultSystem->GetPhase();
+
+				if (EnumHasAnyFlags(SystemPhase, ESystemPhase::Spawn))
 				{
-					// Components shouldn't be connected to other components
-					SetupSubsequent(ThisSystem, Nodes[ComponentEdge.ToNode].Class.Get(), Linker);
+					FlowStrings[0] += FString::Printf(TEXT("\t\tflow_node%d_0[label=\"%s\"];\n"), NodeID, *SystemClass->GetName());
+				}
+				if (EnumHasAnyFlags(SystemPhase, ESystemPhase::Instantiation))
+				{
+					FlowStrings[1] += FString::Printf(TEXT("\t\tflow_node%d_1[label=\"%s\"];\n"), NodeID, *SystemClass->GetName());
+				}
+				if (EnumHasAnyFlags(SystemPhase, ESystemPhase::Evaluation))
+				{
+					FlowStrings[2] += FString::Printf(TEXT("\t\tflow_node%d_2[label=\"%s\"];\n"), NodeID, *SystemClass->GetName());
+				}
+				if (EnumHasAnyFlags(SystemPhase, ESystemPhase::Finalization))
+				{
+					FlowStrings[3] += FString::Printf(TEXT("\t\tflow_node%d_3[label=\"%s\"];\n"), NodeID, *SystemClass->GetName());
 				}
 			}
-			else
+			else if(Node.WriteComponentType)
 			{
-				SetupSubsequent(ThisSystem, Node.Class.Get(), Linker);
+				const FComponentTypeInfo& WriteComponentInfo = ComponentRegistry->GetComponentTypeChecked(Node.WriteComponentType);
+
+				ESystemPhase* ConnectedSystemPhase = SystemPhasesConnectedToWriteComponents.Find(NodeID);
+				if (ConnectedSystemPhase)
+				{
+					if (EnumHasAnyFlags(*ConnectedSystemPhase, ESystemPhase::Spawn))
+					{
+#if UE_MOVIESCENE_ENTITY_DEBUG
+						FlowStrings[0] += FString::Printf(TEXT("\t\tflow_node%d_0[label=\"%s\"];\n"), NodeID, *WriteComponentInfo.DebugInfo->DebugName);
+#else
+						FlowStrings[0] += FString::Printf(TEXT("\t\tflow_node%d_0[label=\"Component ID=%s\"];\n"), NodeID, Node.WriteComponentType.BitIndex());
+#endif
+					}
+					if (EnumHasAnyFlags(*ConnectedSystemPhase, ESystemPhase::Instantiation))
+					{
+#if UE_MOVIESCENE_ENTITY_DEBUG
+						FlowStrings[1] += FString::Printf(TEXT("\t\tflow_node%d_1[label=\"%s\"];\n"), NodeID, *WriteComponentInfo.DebugInfo->DebugName);
+#else
+						FlowStrings[1] += FString::Printf(TEXT("\t\tflow_node%d_1[label=\"Component ID=%s\"];\n"), NodeID, Node.WriteComponentType.BitIndex());
+#endif
+					}
+					if (EnumHasAnyFlags(*ConnectedSystemPhase, ESystemPhase::Evaluation))
+					{
+#if UE_MOVIESCENE_ENTITY_DEBUG
+						FlowStrings[2] += FString::Printf(TEXT("\t\tflow_node%d_2[label=\"%s\"];\n"), NodeID, *WriteComponentInfo.DebugInfo->DebugName);
+#else
+						FlowStrings[2] += FString::Printf(TEXT("\t\tflow_node%d_2[label=\"Component ID=%s\"];\n"), NodeID, Node.WriteComponentType.BitIndex());
+#endif
+					}
+					if (EnumHasAnyFlags(*ConnectedSystemPhase, ESystemPhase::Finalization))
+					{
+#if UE_MOVIESCENE_ENTITY_DEBUG
+						FlowStrings[3] += FString::Printf(TEXT("\t\tflow_node%d_3[label=\"%s\"];\n"), NodeID, *WriteComponentInfo.DebugInfo->DebugName);
+#else
+						FlowStrings[3] += FString::Printf(TEXT("\t\tflow_node%d_3[label=\"Component ID=%s\"];\n"), NodeID, Node.WriteComponentType.BitIndex());
+#endif
+					}
+				}
 			}
 		}
-	}
 
-	void SetupPrereq(UClass* PrereqType, UMovieSceneEntitySystem* ThisSystem, UMovieSceneEntitySystemLinker* Linker)
-	{
-		if (PrereqType)
+		for (FString& FlowString : FlowStrings)
+		{
+			String += FlowString;
+			String += TEXT("\t}\n");
+		}
+
 		{
-			if (UMovieSceneEntitySystem* TargetSystem = Linker->FindSystem(PrereqType))
+			FMovieSceneEntitySystemDirectedGraph::FDiscoverCyclicEdges CyclicEdges(&FlowGraph);
+			CyclicEdges.Search();
+
+			TArrayView<const FDirectionalEdge> FlowEdges = FlowGraph.GetEdges();
+			for (int32 EdgeIndex = 0; EdgeIndex < FlowEdges.Num(); ++EdgeIndex)
 			{
-				Linker->SystemGraph.AddPrerequisite(TargetSystem, ThisSystem);
+				FDirectionalEdge Edge = FlowEdges[EdgeIndex];
+				const bool bIsCyclic = CyclicEdges.IsCyclic(EdgeIndex);
+				const FString EdgeColor = bIsCyclic ? TEXT("#FF0000") : TEXT("#39ad3b");
+
+				ESystemPhase FromPhase = ESystemPhase::None;
+				if (UClass* FromClass = Nodes[Edge.FromNode].Class.Get())
+				{
+					FromPhase = FromClass->GetDefaultObject<UMovieSceneEntitySystem>()->GetPhase();
+				}
+
+				ESystemPhase ToPhase = ESystemPhase::None;
+				if (UClass* ToClass = Nodes[Edge.ToNode].Class.Get())
+				{
+					ToPhase = ToClass->GetDefaultObject<UMovieSceneEntitySystem>()->GetPhase();
+				}
+
+				if ((FromPhase == ESystemPhase::None || EnumHasAnyFlags(FromPhase, ESystemPhase::Spawn)) && 
+						(ToPhase == ESystemPhase::None || EnumHasAnyFlags(ToPhase, ESystemPhase::Spawn)))
+				{
+					String += FString::Printf(TEXT("\tflow_node%d_0 -> flow_node%d_0 [color=\"%s\"];\n"), (int32)Edge.FromNode, (int32)Edge.ToNode, *EdgeColor);
+				}
+				if ((FromPhase == ESystemPhase::None || EnumHasAnyFlags(FromPhase, ESystemPhase::Instantiation)) && 
+						(ToPhase == ESystemPhase::None || EnumHasAnyFlags(ToPhase, ESystemPhase::Instantiation)))
+				{
+					String += FString::Printf(TEXT("\tflow_node%d_1 -> flow_node%d_1 [color=\"%s\"];\n"), (int32)Edge.FromNode, (int32)Edge.ToNode, *EdgeColor);
+				}
+				if ((FromPhase == ESystemPhase::None || EnumHasAnyFlags(FromPhase, ESystemPhase::Evaluation)) && 
+						(ToPhase == ESystemPhase::None || EnumHasAnyFlags(ToPhase, ESystemPhase::Evaluation)))
+				{
+					String += FString::Printf(TEXT("\tflow_node%d_2 -> flow_node%d_2 [color=\"%s\"];\n"), (int32)Edge.FromNode, (int32)Edge.ToNode, *EdgeColor);
+				}
+				if ((FromPhase == ESystemPhase::None || EnumHasAnyFlags(FromPhase, ESystemPhase::Finalization)) && 
+						(ToPhase == ESystemPhase::None || EnumHasAnyFlags(ToPhase, ESystemPhase::Finalization)))
+				{
+					String += FString::Printf(TEXT("\tflow_node%d_3 -> flow_node%d_3 [color=\"%s\"];\n"), (int32)Edge.FromNode, (int32)Edge.ToNode, *EdgeColor);
+				}
 			}
 		}
+
+		String += TEXT("}");
+		return String;
 	}
-	void SetupSubsequent(UMovieSceneEntitySystem* ThisSystem, UClass* SubsequentType, UMovieSceneEntitySystemLinker* Linker)
+
+private:
+
+	void UpdateCache()
 	{
-		if (SubsequentType)
+		if (PreviousSerialNumber == SerialNumber)
 		{
-			if (UMovieSceneEntitySystem* TargetSystem = Linker->FindSystem(SubsequentType))
-			{
-				Linker->SystemGraph.AddPrerequisite(ThisSystem, TargetSystem);
-			}
+			return;
+		}
+
+		BuildFlowGraph();
+
+		FMovieSceneEntitySystemDirectedGraph::FDepthFirstSearch DepthFirstSearch(&FlowGraph);
+
+		TBitArray<> EdgeNodes = FlowGraph.FindEdgeUpstreamNodes();
+		for (TConstSetBitIterator<> EdgeNodeIt(EdgeNodes); EdgeNodeIt; ++EdgeNodeIt)
+		{
+			const uint16 NodeID = static_cast<uint16>(EdgeNodeIt.GetIndex());
+			check(Nodes.IsValidIndex(NodeID));
+
+			DepthFirstSearch.Search(NodeID);
 		}
+
+		Algo::Reverse(DepthFirstSearch.PostNodes);
+
+		// FlowOrderNodes has the NodeIDs sorted by dependencies.
+		FlowOrderNodes = DepthFirstSearch.PostNodes;
+
+		// Build a reverse lookup array for the flow order, i.e. an array that gives the flow order index
+		// of a system given that system's NodeID.
+		ReverseLookupFlowOrderNodes.SetNum(Nodes.Num());
+		for (int32 Index = 0; Index < FlowOrderNodes.Num(); ++Index)
+		{
+			ReverseLookupFlowOrderNodes[FlowOrderNodes[Index]] = Index;
+		}
+
+		PreviousSerialNumber = SerialNumber;
 	}
 
-	void MakeRelationship(uint16 UpstreamGraphID, uint16 DownstreamGraphID)
+	void BuildFlowGraph()
 	{
-		ImplicitSubsequents.MakeEdge(UpstreamGraphID, DownstreamGraphID);
-		ImplicitPrerequisites.MakeEdge(DownstreamGraphID, UpstreamGraphID);
+		FlowGraph.DestroyAllEdges();
+
+		for (int32 NodeIndex = 0, NodeCount = Nodes.Num(); NodeIndex < NodeCount; ++NodeIndex)
+		{
+			SetupFlowDependencies(NodeIndex);
+		}
+
+		FlowGraph.CleanUpDanglingEdges();
+
+		checkf(!FlowGraph.IsCyclic(), TEXT("Cycle detected in system flow graph!"));
 	}
 
-	UClass* ClassFromGraphID(uint16 GraphID) const
+	void SetupFlowDependencies(int32 NodeIndex)
 	{
-		check(Nodes.IsValidIndex(GraphID));
-		return Nodes[GraphID].Class.Get();
+		using FDirectionalEdge = FMovieSceneEntitySystemDirectedGraph::FDirectionalEdge;
+
+		const uint16 CurrentNodeID(static_cast<uint16>(NodeIndex));
+
+		// Set up prerequisites
+		for (FDirectionalEdge Edge : ImplicitPrerequisites.GetEdgesFrom(CurrentNodeID))
+		{
+			const FNode& Node = Nodes[Edge.ToNode];
+
+			// Follow edges from components
+			if (Node.WriteComponentType)
+			{
+				for (FDirectionalEdge ComponentEdge : ImplicitPrerequisites.GetEdgesFrom(Edge.ToNode))
+				{
+					// Components shouldn't be connected to other components
+					SetupFlowDependency(ComponentEdge.ToNode, CurrentNodeID);
+				}
+			}
+			else
+			{
+				SetupFlowDependency(Edge.ToNode, CurrentNodeID);
+			}
+		}
+
+		// Set up subsequents
+		for (FDirectionalEdge Edge : ImplicitSubsequents.GetEdgesFrom(CurrentNodeID))
+		{
+			const FNode& Node = Nodes[Edge.ToNode];
+
+			// Follow edges from components
+			if (Node.WriteComponentType)
+			{
+				for (FDirectionalEdge ComponentEdge : ImplicitSubsequents.GetEdgesFrom(Edge.ToNode))
+				{
+					// Components shouldn't be connected to other components
+					SetupFlowDependency(CurrentNodeID, ComponentEdge.ToNode);
+				}
+			}
+			else
+			{
+				SetupFlowDependency(CurrentNodeID, Edge.ToNode);
+			}
+		}
 	}
 
-	uint16 NumGraphIDs() const
+	void SetupFlowDependency(uint16 UpstreamID, uint16 DownstreamID)
 	{
-		return Nodes.Num();
+		check(UpstreamID != TNumericLimits<uint16>::Max() && DownstreamID != TNumericLimits<uint16>::Max());
+
+		FlowGraph.MakeEdge(UpstreamID, DownstreamID);
 	}
 
 private:
@@ -162,8 +425,16 @@ private:
 	};
 
 	TArray<FNode> Nodes;
+
 	FMovieSceneEntitySystemDirectedGraph ImplicitPrerequisites;
 	FMovieSceneEntitySystemDirectedGraph ImplicitSubsequents;
+
+	FMovieSceneEntitySystemDirectedGraph FlowGraph;
+	TArray<uint16> FlowOrderNodes;
+	TArray<uint16> ReverseLookupFlowOrderNodes;
+
+	uint32 SerialNumber = 0;
+	uint32 PreviousSerialNumber = 0;
 };
 
 FSystemDependencyGraph GlobalDependencyGraph;
@@ -328,6 +599,27 @@ UE::MovieScene::EEntitySystemCategory UMovieSceneEntitySystem::RegisterCustomSys
 	return Result;
 }
 
+void UMovieSceneEntitySystem::SortByFlowOrder(TArray<uint16>& InOutGlobalNodeIDs)
+{
+	using namespace UE::MovieScene;
+
+	GlobalDependencyGraph.SortByFlowOrder(InOutGlobalNodeIDs);
+}
+
+void UMovieSceneEntitySystem::GetSubsequentSystems(uint16 FromGlobalNodeID, TArray<uint16>& OutSubsequentGlobalNodeIDs)
+{
+	using namespace UE::MovieScene;
+
+	GlobalDependencyGraph.GetSubsequents(FromGlobalNodeID, OutSubsequentGlobalNodeIDs);
+}
+
+void UMovieSceneEntitySystem::DebugPrintGlobalDependencyGraph(bool bUpdateCache)
+{
+	using namespace UE::MovieScene;
+
+	GlobalDependencyGraph.DebugPrint(bUpdateCache);
+}
+
 bool UMovieSceneEntitySystem::IsRelevant(UMovieSceneEntitySystemLinker* InLinker) const
 {
 	if (RelevantComponent && InLinker->EntityManager.ContainsComponent(RelevantComponent))
@@ -429,7 +721,6 @@ void UMovieSceneEntitySystem::Link(UMovieSceneEntitySystemLinker* InLinker)
 	Linker = InLinker;
 	OnLink();
 
-	GlobalDependencyGraph.SetupDependencies(this, InLinker);
 	Linker->SystemLinked(this);
 }
 
diff --git a/Engine/Source/Runtime/MovieScene/Private/EntitySystem/MovieSceneEntitySystemDirectedGraph.cpp b/Engine/Source/Runtime/MovieScene/Private/EntitySystem/MovieSceneEntitySystemDirectedGraph.cpp
new file mode 100644
--- /dev/null
+++ b/Engine/Source/Runtime/MovieScene/Private/EntitySystem/MovieSceneEntitySystemDirectedGraph.cpp
@@ -0,0 +1,344 @@
+// Copyright Epic Games, Inc. All Rights Reserved.
+
+#include "EntitySystem/MovieSceneEntitySystemDirectedGraph.h"
+
+//#include UE_INLINE_GENERATED_CPP_BY_NAME(MovieSceneEntitySystemDirectedGraph)
+
+FMovieSceneEntitySystemDirectedGraph::FDepthFirstSearch::FDepthFirstSearch(const FMovieSceneEntitySystemDirectedGraph* InGraph)
+	: Visited(false, InGraph->Nodes.Num())
+	, IsVisiting(false, InGraph->Nodes.Num())
+{
+	Graph = InGraph;
+	PostNodes.Reserve(InGraph->Nodes.CountSetBits());
+	check(!Graph->bHasDanglingEdges);
+}
+
+void FMovieSceneEntitySystemDirectedGraph::FDepthFirstSearch::Search(uint16 InNodeID)
+{
+	IsVisiting[InNodeID] = true;
+
+	for (FMovieSceneEntitySystemDirectedGraph::FDirectionalEdge Edge : Graph->GetEdgesFrom(InNodeID))
+	{
+		if (Visited[Edge.ToNode] == false)
+		{
+			if (!ensureMsgf(IsVisiting[Edge.ToNode] == false, TEXT("Cycle found in graph.")))
+			{
+				return;
+			}
+			Visited[Edge.ToNode] = true;
+			Search(Edge.ToNode);
+		}
+	}
+
+	PostNodes.Add(InNodeID);
+
+	IsVisiting[InNodeID] = false;
+}
+
+FMovieSceneEntitySystemDirectedGraph::FBreadthFirstSearch::FBreadthFirstSearch(const FMovieSceneEntitySystemDirectedGraph* InGraph)
+	: Visited(false, InGraph->Nodes.Num())
+	, StackIndex(0)
+{
+	Graph = InGraph;
+	Nodes.Reserve(InGraph->Nodes.CountSetBits());
+	check(!InGraph->bHasDanglingEdges);
+}
+
+void FMovieSceneEntitySystemDirectedGraph::FBreadthFirstSearch::Search(uint16 InNodeID)
+{
+	if (Visited[InNodeID] == true)
+	{
+		return;
+	}
+
+	Nodes.Reset();
+	StackIndex = 0;
+
+	Visited[InNodeID] = true;
+	Nodes.Add(InNodeID);
+
+	while (StackIndex < Nodes.Num())
+	{
+		const int32 StackEnd = Nodes.Num();
+		for ( ; StackIndex < StackEnd; ++StackIndex)
+		{
+			const uint16 NodeID = Nodes[StackIndex];
+
+			// Visit all nodes this points to
+			for (FMovieSceneEntitySystemDirectedGraph::FDirectionalEdge Edge : Graph->GetEdgesFrom(NodeID))
+			{
+				if (Visited[Edge.ToNode] == false)
+				{
+					Visited[Edge.ToNode] = true;
+					Nodes.Add(Edge.ToNode);
+				}
+			}
+		}
+	}
+}
+
+FMovieSceneEntitySystemDirectedGraph::FDiscoverCyclicEdges::FDiscoverCyclicEdges(const FMovieSceneEntitySystemDirectedGraph* InGraph)
+	: CyclicEdges(false, InGraph->SortedEdges.Num())
+	, VisitedEdges(false, InGraph->SortedEdges.Num())
+{
+	Graph = InGraph;
+	check(!Graph->bHasDanglingEdges);
+}
+
+void FMovieSceneEntitySystemDirectedGraph::FDiscoverCyclicEdges::Search()
+{
+	for (uint16 EdgeIndex = 0; EdgeIndex < Graph->SortedEdges.Num(); ++EdgeIndex)
+	{
+		if (VisitedEdges[EdgeIndex] == false)
+		{
+			EdgeChain.Reset();
+			EdgeChain.Add(EdgeIndex);
+
+			SearchFrom(Graph->SortedEdges[EdgeIndex].FromNode);
+
+			VisitedEdges[EdgeIndex] = true;
+		}
+	}
+}
+
+void FMovieSceneEntitySystemDirectedGraph::FDiscoverCyclicEdges::SearchFrom(uint16 NodeID)
+{
+	TBitArray<> VisitedNodes(false, Graph->Nodes.Num());
+	VisitedNodes[NodeID] = true;
+	DiscoverCycles(NodeID, VisitedNodes);
+}
+
+void FMovieSceneEntitySystemDirectedGraph::FDiscoverCyclicEdges::DiscoverCycles(uint16 NodeID, TBitArray<>& VisitedNodes)
+{
+	// Iterate all edges from this node
+	for (int32 SubsequentEdge = Graph->FindEdgeStart(NodeID); SubsequentEdge < Graph->SortedEdges.Num() && Graph->SortedEdges[SubsequentEdge].FromNode == NodeID; ++SubsequentEdge)
+	{
+		if (VisitedEdges[SubsequentEdge] == true)
+		{
+			continue;
+		}
+
+		VisitedEdges[SubsequentEdge] = true;
+		EdgeChain.Add(SubsequentEdge);
+
+		const uint16 SubsequentNode = Graph->SortedEdges[SubsequentEdge].ToNode;
+		if (VisitedNodes[SubsequentNode] == true)
+		{
+			TagCyclicChain(SubsequentNode);
+		}
+		else
+		{
+			VisitedNodes[SubsequentNode] = true;
+			DiscoverCycles(SubsequentNode, VisitedNodes);
+			VisitedNodes[SubsequentNode] = false;
+		}
+
+		EdgeChain.Pop();
+	}
+}
+
+void FMovieSceneEntitySystemDirectedGraph::FDiscoverCyclicEdges::TagCyclicChain(uint16 CyclicNodeID)
+{
+	// Found a cycle
+	for (int32 EdgeChainIndex = EdgeChain.Num() - 1; EdgeChainIndex >= 0; --EdgeChainIndex)
+	{
+		const uint16 UpstreamEdgeIndex = EdgeChain[EdgeChainIndex];
+		CyclicEdges.PadToNum(UpstreamEdgeIndex + 1, false);
+		CyclicEdges[UpstreamEdgeIndex] = true;
+
+		if (Graph->SortedEdges[UpstreamEdgeIndex].FromNode == CyclicNodeID)
+		{
+			return;
+		}
+	}
+}
+
+void FMovieSceneEntitySystemDirectedGraph::AllocateNode(uint16 NodeID)
+{
+	CleanUpDanglingEdges();
+	Nodes.PadToNum(NodeID + 1, false);
+	Nodes[NodeID] = true;
+}
+
+bool FMovieSceneEntitySystemDirectedGraph::IsNodeAllocated(uint16 NodeID) const
+{
+	return Nodes.IsValidIndex(NodeID) && Nodes[NodeID] == true;
+}
+
+void FMovieSceneEntitySystemDirectedGraph::RemoveNode(uint16 NodeID)
+{
+	check(NodeID != TNumericLimits<uint16>::Max() && IsNodeAllocated(NodeID));
+
+	// Remove the node from the graph
+	Nodes[NodeID] = false;
+
+	bHasDanglingEdges = true;
+}
+
+void FMovieSceneEntitySystemDirectedGraph::CleanUpDanglingEdges()
+{
+	if (!bHasDanglingEdges)
+	{
+		return;
+	}
+
+	bHasDanglingEdges = false;
+	for (int32 Index = 0; Index < SortedEdges.Num(); )
+	{
+		FDirectionalEdge Edge = SortedEdges[Index];
+		if (!IsNodeAllocated(Edge.ToNode) || !IsNodeAllocated(Edge.FromNode))
+		{
+			SortedEdges.RemoveAt(Index, 1, false);
+		}
+		else
+		{
+			++Index;
+		}
+	}
+}
+
+bool FMovieSceneEntitySystemDirectedGraph::IsCyclic() const
+{
+	TBitArray<> Visited(false, Nodes.Num());
+
+	for (FDirectionalEdge Edge : SortedEdges)
+	{
+		if (Visited[Edge.ToNode] == true)
+		{
+			continue;
+		}
+
+		TBitArray<> Visiting(false, Nodes.Num());
+		if (IsCyclicImpl(Edge.ToNode, Visiting))
+		{
+			return true;
+		}
+
+		Visited.CombineWithBitwiseOR(Visiting, EBitwiseOperatorFlags::MaxSize);
+	}
+
+	return false;
+}
+
+bool FMovieSceneEntitySystemDirectedGraph::IsCyclicImpl(uint16 NodeID, TBitArray<>& Visiting) const
+{
+	if (Visiting[NodeID] == true)
+	{
+		return true;
+	}
+
+	Visiting[NodeID] = true;
+
+	for (FDirectionalEdge Edge : GetEdgesFrom(NodeID))
+	{
+		if (IsCyclicImpl(Edge.ToNode, Visiting))
+		{
+			return true;
+		}
+	}
+
+	Visiting[NodeID] = false;
+	return false;
+}
+
+void FMovieSceneEntitySystemDirectedGraph::MakeEdge(uint16 FromNode, uint16 ToNode)
+{
+	FDirectionalEdge NewEdge(FromNode, ToNode);
+
+	const int32 InsertIndex = FindEdgeIndex(NewEdge);
+	if (!SortedEdges.IsValidIndex(InsertIndex) || SortedEdges[InsertIndex] != NewEdge)
+	{
+		SortedEdges.Insert(NewEdge, InsertIndex);
+	}
+}
+
+void FMovieSceneEntitySystemDirectedGraph::DestroyEdge(uint16 FromNode, uint16 ToNode)
+{
+	FDirectionalEdge Edge(FromNode, ToNode);
+
+	const int32 RemoveIndex = Algo::BinarySearch(SortedEdges, Edge);
+	if (RemoveIndex != INDEX_NONE)
+	{
+		SortedEdges.RemoveAt(RemoveIndex, 1, false);
+	}
+}
+
+void FMovieSceneEntitySystemDirectedGraph::DestroyAllEdges()
+{
+	SortedEdges.Reset();
+	bHasDanglingEdges = false;
+}
+
+int32 FMovieSceneEntitySystemDirectedGraph::FindEdgeStart(uint16 FromNode) const
+{
+	return Algo::LowerBoundBy(SortedEdges, FromNode, &FDirectionalEdge::FromNode);
+}
+
+TArrayView<const FMovieSceneEntitySystemDirectedGraph::FDirectionalEdge> FMovieSceneEntitySystemDirectedGraph::GetEdges() const
+{
+	check(!bHasDanglingEdges);
+	return SortedEdges;
+}
+
+bool FMovieSceneEntitySystemDirectedGraph::HasEdgeFrom(uint16 InNode) const
+{
+	const int32 ExpectedIndex = FindEdgeStart(InNode);
+	return SortedEdges.IsValidIndex(ExpectedIndex) && SortedEdges[ExpectedIndex].FromNode == InNode;
+}
+
+bool FMovieSceneEntitySystemDirectedGraph::HasEdgeTo(uint16 InNode) const
+{
+	check(InNode != TNumericLimits<uint16>::Max());
+	return Algo::FindBy(SortedEdges, InNode, &FDirectionalEdge::ToNode) != nullptr;
+}
+
+TArrayView<const FMovieSceneEntitySystemDirectedGraph::FDirectionalEdge> FMovieSceneEntitySystemDirectedGraph::GetEdgesFrom(uint16 InNodeID) const
+{
+	check(!bHasDanglingEdges);
+
+	const int32 EdgeIndex = FindEdgeStart(InNodeID);
+
+	int32 Num = 0;
+	while (EdgeIndex + Num < SortedEdges.Num() && SortedEdges[EdgeIndex + Num].FromNode == InNodeID)
+	{
+		++Num;
+	}
+
+	if (Num > 0)
+	{
+		return MakeArrayView(SortedEdges.GetData() + EdgeIndex, Num);
+	}
+	return TArrayView<const FDirectionalEdge>();
+}
+
+TBitArray<> FMovieSceneEntitySystemDirectedGraph::FindEdgeUpstreamNodes() const
+{
+	check(!bHasDanglingEdges);
+
+	TBitArray<> EdgeNodes(true, Nodes.Num());
+
+	// Unmark nodes that have edges pointing towards them
+	for (uint16 EdgeIndex = 0; EdgeIndex < SortedEdges.Num(); ++EdgeIndex)
+	{
+		const uint16 ToNode = SortedEdges[EdgeIndex].ToNode;
+		EdgeNodes[ToNode] = false;
+	}
+
+	// Mask with nodes that are actually allocated
+	return TBitArray<>::BitwiseAND(EdgeNodes, Nodes, EBitwiseOperatorFlags::MaxSize);
+}
+
+int32 FMovieSceneEntitySystemDirectedGraph::FindEdgeIndex(const FDirectionalEdge& Edge) const
+{
+	check(!bHasDanglingEdges);
+	return Algo::LowerBound(SortedEdges, Edge);
+}
+
+bool FMovieSceneEntitySystemDirectedGraph::EdgeExists(const FDirectionalEdge& Edge) const
+{
+	check(!bHasDanglingEdges);
+
+	const int32 EdgeIndex = FindEdgeIndex(Edge);
+	return EdgeIndex < SortedEdges.Num() && SortedEdges[EdgeIndex] == Edge;
+}
+
diff --git a/Engine/Source/Runtime/MovieScene/Private/EntitySystem/MovieSceneEntitySystemGraphs.cpp b/Engine/Source/Runtime/MovieScene/Private/EntitySystem/MovieSceneEntitySystemGraphs.cpp
--- a/Engine/Source/Runtime/MovieScene/Private/EntitySystem/MovieSceneEntitySystemGraphs.cpp
+++ b/Engine/Source/Runtime/MovieScene/Private/EntitySystem/MovieSceneEntitySystemGraphs.cpp
@@ -13,156 +13,6 @@
 
 #include UE_INLINE_GENERATED_CPP_BY_NAME(MovieSceneEntitySystemGraphs)
 
-
-FMovieSceneEntitySystemDirectedGraph::FDepthFirstSearch::FDepthFirstSearch(const FMovieSceneEntitySystemDirectedGraph* InGraph)
-	: Visited(false, InGraph->Nodes.Num())
-	, IsVisiting(false, InGraph->Nodes.Num())
-{
-	Graph = InGraph;
-	PostNodes.Reserve(InGraph->Nodes.CountSetBits());
-	check(!Graph->bHasDanglingEdges);
-}
-
-void FMovieSceneEntitySystemDirectedGraph::FDepthFirstSearch::Search(uint16 InNodeID)
-{
-	IsVisiting[InNodeID] = true;
-
-	for (FMovieSceneEntitySystemDirectedGraph::FDirectionalEdge Edge : Graph->GetEdgesFrom(InNodeID))
-	{
-		if (Visited[Edge.ToNode] == false)
-		{
-			if (!ensureMsgf(IsVisiting[Edge.ToNode] == false, TEXT("Cycle found in graph.")))
-			{
-				return;
-			}
-			Visited[Edge.ToNode] = true;
-			Search(Edge.ToNode);
-		}
-	}
-
-	PostNodes.Add(InNodeID);
-
-	IsVisiting[InNodeID] = false;
-}
-
-FMovieSceneEntitySystemDirectedGraph::FBreadthFirstSearch::FBreadthFirstSearch(const FMovieSceneEntitySystemDirectedGraph* InGraph)
-	: Visited(false, InGraph->Nodes.Num())
-	, StackIndex(0)
-{
-	Graph = InGraph;
-	Nodes.Reserve(InGraph->Nodes.CountSetBits());
-	check(!InGraph->bHasDanglingEdges);
-}
-
-void FMovieSceneEntitySystemDirectedGraph::FBreadthFirstSearch::Search(uint16 InNodeID)
-{
-	if (Visited[InNodeID] == true)
-	{
-		return;
-	}
-
-	Nodes.Reset();
-	StackIndex = 0;
-
-	Visited[InNodeID] = true;
-	Nodes.Add(InNodeID);
-
-	while (StackIndex < Nodes.Num())
-	{
-		const int32 StackEnd = Nodes.Num();
-		for ( ; StackIndex < StackEnd; ++StackIndex)
-		{
-			const uint16 NodeID = Nodes[StackIndex];
-
-			// Visit all nodes this points to
-			for (FMovieSceneEntitySystemDirectedGraph::FDirectionalEdge Edge : Graph->GetEdgesFrom(NodeID))
-			{
-				if (Visited[Edge.ToNode] == false)
-				{
-					Visited[Edge.ToNode] = true;
-					Nodes.Add(Edge.ToNode);
-				}
-			}
-		}
-	}
-}
-
-FMovieSceneEntitySystemDirectedGraph::FDiscoverCyclicEdges::FDiscoverCyclicEdges(const FMovieSceneEntitySystemDirectedGraph* InGraph)
-	: CyclicEdges(false, InGraph->SortedEdges.Num())
-	, VisitedEdges(false, InGraph->SortedEdges.Num())
-{
-	Graph = InGraph;
-	check(!Graph->bHasDanglingEdges);
-}
-
-void FMovieSceneEntitySystemDirectedGraph::FDiscoverCyclicEdges::Search()
-{
-	for (uint16 EdgeIndex = 0; EdgeIndex < Graph->SortedEdges.Num(); ++EdgeIndex)
-	{
-		if (VisitedEdges[EdgeIndex] == false)
-		{
-			EdgeChain.Reset();
-			EdgeChain.Add(EdgeIndex);
-
-			SearchFrom(Graph->SortedEdges[EdgeIndex].FromNode);
-
-			VisitedEdges[EdgeIndex] = true;
-		}
-	}
-}
-
-void FMovieSceneEntitySystemDirectedGraph::FDiscoverCyclicEdges::SearchFrom(uint16 NodeID)
-{
-	TBitArray<> VisitedNodes(false, Graph->Nodes.Num());
-	VisitedNodes[NodeID] = true;
-	DiscoverCycles(NodeID, VisitedNodes);
-}
-
-void FMovieSceneEntitySystemDirectedGraph::FDiscoverCyclicEdges::DiscoverCycles(uint16 NodeID, TBitArray<>& VisitedNodes)
-{
-	// Iterate all edges from this node
-	for (int32 SubsequentEdge = Graph->FindEdgeStart(NodeID); SubsequentEdge < Graph->SortedEdges.Num() && Graph->SortedEdges[SubsequentEdge].FromNode == NodeID; ++SubsequentEdge)
-	{
-		if (VisitedEdges[SubsequentEdge] == true)
-		{
-			continue;
-		}
-
-		VisitedEdges[SubsequentEdge] = true;
-		EdgeChain.Add(SubsequentEdge);
-
-		const uint16 SubsequentNode = Graph->SortedEdges[SubsequentEdge].ToNode;
-		if (VisitedNodes[SubsequentNode] == true)
-		{
-			TagCyclicChain(SubsequentNode);
-		}
-		else
-		{
-			VisitedNodes[SubsequentNode] = true;
-			DiscoverCycles(SubsequentNode, VisitedNodes);
-			VisitedNodes[SubsequentNode] = false;
-		}
-
-		EdgeChain.Pop();
-	}
-}
-
-void FMovieSceneEntitySystemDirectedGraph::FDiscoverCyclicEdges::TagCyclicChain(uint16 CyclicNodeID)
-{
-	// Found a cycle
-	for (int32 EdgeChainIndex = EdgeChain.Num() - 1; EdgeChainIndex >= 0; --EdgeChainIndex)
-	{
-		const uint16 UpstreamEdgeIndex = EdgeChain[EdgeChainIndex];
-		CyclicEdges.PadToNum(UpstreamEdgeIndex + 1, false);
-		CyclicEdges[UpstreamEdgeIndex] = true;
-
-		if (Graph->SortedEdges[UpstreamEdgeIndex].FromNode == CyclicNodeID)
-		{
-			return;
-		}
-	}
-}
-
 void FMovieSceneEntitySystemGraphNodes::AddStructReferencedObjects(FReferenceCollector& Collector) const
 {
 	for (FMovieSceneEntitySystemGraphNode& Node : const_cast<TSparseArray<FMovieSceneEntitySystemGraphNode>&>(Array))
@@ -171,190 +21,6 @@ void FMovieSceneEntitySystemGraphNodes::AddStructReferencedObjects(FReferenceCol
 	}
 }
 
-void FMovieSceneEntitySystemDirectedGraph::AllocateNode(uint16 NodeID)
-{
-	CleanUpDanglingEdges();
-	Nodes.PadToNum(NodeID + 1, false);
-	Nodes[NodeID] = true;
-}
-
-bool FMovieSceneEntitySystemDirectedGraph::IsNodeAllocated(uint16 NodeID) const
-{
-	return Nodes.IsValidIndex(NodeID) && Nodes[NodeID] == true;
-}
-
-void FMovieSceneEntitySystemDirectedGraph::RemoveNode(uint16 NodeID)
-{
-	check(NodeID != TNumericLimits<uint16>::Max() && IsNodeAllocated(NodeID));
-
-	// Remove the node from the graph
-	Nodes[NodeID] = false;
-
-	bHasDanglingEdges = true;
-}
-
-void FMovieSceneEntitySystemDirectedGraph::CleanUpDanglingEdges()
-{
-	if (!bHasDanglingEdges)
-	{
-		return;
-	}
-
-	bHasDanglingEdges = false;
-	for (int32 Index = 0; Index < SortedEdges.Num(); )
-	{
-		FDirectionalEdge Edge = SortedEdges[Index];
-		if (!IsNodeAllocated(Edge.ToNode) || !IsNodeAllocated(Edge.FromNode))
-		{
-			SortedEdges.RemoveAt(Index, 1, false);
-		}
-		else
-		{
-			++Index;
-		}
-	}
-}
-
-bool FMovieSceneEntitySystemDirectedGraph::IsCyclic() const
-{
-	TBitArray<> Visited(false, Nodes.Num());
-
-	for (FDirectionalEdge Edge : SortedEdges)
-	{
-		if (Visited[Edge.ToNode] == true)
-		{
-			continue;
-		}
-
-		TBitArray<> Visiting(false, Nodes.Num());
-		if (IsCyclicImpl(Edge.ToNode, Visiting))
-		{
-			return true;
-		}
-
-		Visited.CombineWithBitwiseOR(Visiting, EBitwiseOperatorFlags::MaxSize);
-	}
-
-	return false;
-}
-
-bool FMovieSceneEntitySystemDirectedGraph::IsCyclicImpl(uint16 NodeID, TBitArray<>& Visiting) const
-{
-	if (Visiting[NodeID] == true)
-	{
-		return true;
-	}
-
-	Visiting[NodeID] = true;
-
-	for (FDirectionalEdge Edge : GetEdgesFrom(NodeID))
-	{
-		if (IsCyclicImpl(Edge.ToNode, Visiting))
-		{
-			return true;
-		}
-	}
-
-	Visiting[NodeID] = false;
-	return false;
-}
-
-void FMovieSceneEntitySystemDirectedGraph::MakeEdge(uint16 FromNode, uint16 ToNode)
-{
-	FDirectionalEdge NewEdge(FromNode, ToNode);
-
-	const int32 InsertIndex = FindEdgeIndex(NewEdge);
-	if (!SortedEdges.IsValidIndex(InsertIndex) || SortedEdges[InsertIndex] != NewEdge)
-	{
-		SortedEdges.Insert(NewEdge, InsertIndex);
-	}
-}
-
-void FMovieSceneEntitySystemDirectedGraph::DestroyEdge(uint16 FromNode, uint16 ToNode)
-{
-	FDirectionalEdge Edge(FromNode, ToNode);
-
-	const int32 RemoveIndex = Algo::BinarySearch(SortedEdges, Edge);
-	if (RemoveIndex != INDEX_NONE)
-	{
-		SortedEdges.RemoveAt(RemoveIndex, 1, false);
-	}
-}
-
-int32 FMovieSceneEntitySystemDirectedGraph::FindEdgeStart(uint16 FromNode) const
-{
-	return Algo::LowerBoundBy(SortedEdges, FromNode, &FDirectionalEdge::FromNode);
-}
-
-TArrayView<const FMovieSceneEntitySystemDirectedGraph::FDirectionalEdge> FMovieSceneEntitySystemDirectedGraph::GetEdges() const
-{
-	check(!bHasDanglingEdges);
-	return SortedEdges;
-}
-
-bool FMovieSceneEntitySystemDirectedGraph::HasEdgeFrom(uint16 InNode) const
-{
-	const int32 ExpectedIndex = FindEdgeStart(InNode);
-	return SortedEdges.IsValidIndex(ExpectedIndex) && SortedEdges[ExpectedIndex].FromNode == InNode;
-}
-
-bool FMovieSceneEntitySystemDirectedGraph::HasEdgeTo(uint16 InNode) const
-{
-	check(InNode != TNumericLimits<uint16>::Max());
-	return Algo::FindBy(SortedEdges, InNode, &FDirectionalEdge::ToNode) != nullptr;
-}
-
-TArrayView<const FMovieSceneEntitySystemDirectedGraph::FDirectionalEdge> FMovieSceneEntitySystemDirectedGraph::GetEdgesFrom(uint16 InNodeID) const
-{
-	check(!bHasDanglingEdges);
-
-	const int32 EdgeIndex = FindEdgeStart(InNodeID);
-
-	int32 Num = 0;
-	while (EdgeIndex + Num < SortedEdges.Num() && SortedEdges[EdgeIndex + Num].FromNode == InNodeID)
-	{
-		++Num;
-	}
-
-	if (Num > 0)
-	{
-		return MakeArrayView(SortedEdges.GetData() + EdgeIndex, Num);
-	}
-	return TArrayView<const FDirectionalEdge>();
-}
-
-TBitArray<> FMovieSceneEntitySystemDirectedGraph::FindEdgeUpstreamNodes() const
-{
-	check(!bHasDanglingEdges);
-
-	TBitArray<> EdgeNodes(true, Nodes.Num());
-
-	// Unmark nodes that have edges pointing towards them
-	for (uint16 EdgeIndex = 0; EdgeIndex < SortedEdges.Num(); ++EdgeIndex)
-	{
-		const uint16 ToNode = SortedEdges[EdgeIndex].ToNode;
-		EdgeNodes[ToNode] = false;
-	}
-
-	// Mask with nodes that are actually allocated
-	return TBitArray<>::BitwiseAND(EdgeNodes, Nodes, EBitwiseOperatorFlags::MaxSize);
-}
-
-int32 FMovieSceneEntitySystemDirectedGraph::FindEdgeIndex(const FDirectionalEdge& Edge) const
-{
-	check(!bHasDanglingEdges);
-	return Algo::LowerBound(SortedEdges, Edge);
-}
-
-bool FMovieSceneEntitySystemDirectedGraph::EdgeExists(const FDirectionalEdge& Edge) const
-{
-	check(!bHasDanglingEdges);
-
-	const int32 EdgeIndex = FindEdgeIndex(Edge);
-	return EdgeIndex < SortedEdges.Num() && SortedEdges[EdgeIndex] == Edge;
-}
-
-
 void FMovieSceneEntitySystemGraph::AddSystem(UMovieSceneEntitySystem* InSystem)
 {
 	checkf(ReentrancyGuard == 0, TEXT("Attempting to add a system to the graph recursively."));
@@ -366,30 +32,12 @@ void FMovieSceneEntitySystemGraph::AddSystem(UMovieSceneEntitySystem* InSystem)
 
 	const uint16 NewNodeID = static_cast<uint16>(NewIndex);
 
-	FlowGraph.AllocateNode(NewNodeID);
 	ReferenceGraph.AllocateNode(NewNodeID);
 
 	InSystem->SetGraphID(NewNodeID);
 
-	++SerialNumber;
-}
-
-int32 FMovieSceneEntitySystemGraph::NumSubsequents(UMovieSceneEntitySystem* InSystem) const
-{
-	const uint16 GraphID = InSystem->GetGraphID();
-	check(GraphID != TNumericLimits<uint16>::Max());
-
-	return FlowGraph.GetEdgesFrom(GraphID).Num();
-}
-
-void FMovieSceneEntitySystemGraph::AddPrerequisite(UMovieSceneEntitySystem* Upstream, UMovieSceneEntitySystem* Downstream)
-{
-	const uint16 UpstreamID   = Upstream->GetGraphID();
-	const uint16 DownstreamID = Downstream->GetGraphID();
-
-	check(UpstreamID != TNumericLimits<uint16>::Max() && DownstreamID != TNumericLimits<uint16>::Max());
-
-	FlowGraph.MakeEdge(UpstreamID, DownstreamID);
+	checkf(!GlobalToLocalNodeIDs.Contains(InSystem->GetGlobalDependencyGraphID()), TEXT("Got more than one instance of a given system class"));
+	GlobalToLocalNodeIDs.Add(InSystem->GetGlobalDependencyGraphID(), NewNodeID);
 
 	++SerialNumber;
 }
@@ -423,17 +71,17 @@ void FMovieSceneEntitySystemGraph::RemoveSystem(UMovieSceneEntitySystem* InSyste
 	const uint16 NodeID = InSystem->GetGraphID();
 	check(NodeID != TNumericLimits<uint16>::Max());
 
-	FlowGraph.RemoveNode(NodeID);
 	ReferenceGraph.RemoveNode(NodeID);
 
 	Nodes.Array.RemoveAt(NodeID);
 
 	InSystem->SetGraphID(TNumericLimits<uint16>::Max());
 
+	GlobalToLocalNodeIDs.Remove(InSystem->GetGlobalDependencyGraphID());
+
 	++SerialNumber;
 	--ReentrancyGuard;
 
-	FlowGraph.CleanUpDanglingEdges();
 	ReferenceGraph.CleanUpDanglingEdges();
 }
 
@@ -445,7 +93,6 @@ int32 FMovieSceneEntitySystemGraph::RemoveIrrelevantSystems(UMovieSceneEntitySys
 	check(PreviousSerialNumber == SerialNumber);
 
 	int32 NumRemoved = 0;
-	bool bHasUnreferecedIntermediateSystems = false;
 
 	FMovieSceneEntitySystemDirectedGraph::FBreadthFirstSearch Search(&ReferenceGraph);
 
@@ -468,14 +115,16 @@ int32 FMovieSceneEntitySystemGraph::RemoveIrrelevantSystems(UMovieSceneEntitySys
 			{
 				const uint16 NodeID = Index;
 
-				FlowGraph.RemoveNode(NodeID);
 				ReferenceGraph.RemoveNode(NodeID);
 
 				UMovieSceneEntitySystem* System = Nodes.Array[NodeID].System;
 				Nodes.Array.RemoveAt(NodeID);
 
-				// Remove this system from the graph to ensure we are not re-entrant
+				// Remove this system from the graph to ensure we are not re-entrant when calling Unlink() on it
 				System->SetGraphID(TNumericLimits<uint16>::Max());
+
+				GlobalToLocalNodeIDs.Remove(System->GetGlobalDependencyGraphID());
+
 				System->Unlink();
 				++NumRemoved;
 			}
@@ -486,7 +135,6 @@ int32 FMovieSceneEntitySystemGraph::RemoveIrrelevantSystems(UMovieSceneEntitySys
 	{
 		++SerialNumber;
 
-		FlowGraph.CleanUpDanglingEdges();
 		ReferenceGraph.CleanUpDanglingEdges();
 	}
 
@@ -503,15 +151,8 @@ void FMovieSceneEntitySystemGraph::UpdateCache()
 		return;
 	}
 
-	FlowGraph.CleanUpDanglingEdges();
 	ReferenceGraph.CleanUpDanglingEdges();
 
-	checkf(!FlowGraph.IsCyclic(), TEXT("Cycle detected in system flow graph.\n")
-		TEXT("----------------------------------------------------------------------------------\n")
-		TEXT("%s\n")
-		TEXT("----------------------------------------------------------------------------------\n"),
-		*ToString());
-
 	checkf(!ReferenceGraph.IsCyclic(), TEXT("Cycle detected in system reference graph.\n")
 		TEXT("----------------------------------------------------------------------------------\n")
 		TEXT("%s\n")
@@ -523,20 +164,20 @@ void FMovieSceneEntitySystemGraph::UpdateCache()
 	EvaluationPhase.Empty();
 	FinalizationPhase.Empty();
 
-	FMovieSceneEntitySystemDirectedGraph::FDepthFirstSearch DepthFirstSearch(&FlowGraph);
-
-	TBitArray<> EdgeNodes = FlowGraph.FindEdgeUpstreamNodes();
-	for (TConstSetBitIterator<> EdgeNodeIt(EdgeNodes); EdgeNodeIt; ++EdgeNodeIt)
+	TArray<uint16> SortedGlobalNodeIDs;
+	for (const FMovieSceneEntitySystemGraphNode& Node : Nodes.Array)
 	{
-		const uint16 NodeID = static_cast<uint16>(EdgeNodeIt.GetIndex());
-		check(Nodes.Array.IsAllocated(NodeID));
-
-		DepthFirstSearch.Search(NodeID);
+		SortedGlobalNodeIDs.Add(Node.System->GetGlobalDependencyGraphID());
 	}
+	UMovieSceneEntitySystem::SortByFlowOrder(SortedGlobalNodeIDs);
 
-	Algo::Reverse(DepthFirstSearch.PostNodes);
+	TArray<uint16> SortedNodeIDs;
+	for (uint16 GlobalNodeID : SortedGlobalNodeIDs)
+	{
+		SortedNodeIDs.Add(GlobalToLocalNodeIDs[GlobalNodeID]);
+	}
 
-	for (uint16 NodeID : DepthFirstSearch.PostNodes)
+	for (uint16 NodeID : SortedNodeIDs)
 	{
 		ESystemPhase SystemPhase = Nodes.Array[NodeID].System->GetPhase();
 
@@ -627,39 +268,6 @@ FString FMovieSceneEntitySystemGraph::ToString() const
 	String += ReferenceGraphString;
 	String += TEXT("\t}\n");
 
-	{
-		FMovieSceneEntitySystemDirectedGraph::FDiscoverCyclicEdges CyclicEdges(&FlowGraph);
-		CyclicEdges.Search();
-
-		TArrayView<const FDirectionalEdge> FlowEdges = FlowGraph.GetEdges();
-		for (int32 EdgeIndex = 0; EdgeIndex < FlowEdges.Num(); ++EdgeIndex)
-		{
-			FDirectionalEdge Edge = FlowEdges[EdgeIndex];
-			const bool bIsCyclic = CyclicEdges.IsCyclic(EdgeIndex);
-
-			ESystemPhase FromPhase = Nodes.Array[Edge.FromNode].System->GetPhase();
-			ESystemPhase ToPhase   = Nodes.Array[Edge.ToNode].System->GetPhase();
-
-			if (EnumHasAnyFlags(FromPhase, ESystemPhase::Spawn) && EnumHasAnyFlags(ToPhase, ESystemPhase::Spawn))
-			{
-				String += FString::Printf(TEXT("\tflow_node%d_0 -> flow_node%d_0 [color=\"%s\"];\n"), (int32)Edge.FromNode, (int32)Edge.ToNode, bIsCyclic ? TEXT("#FF0000") : TEXT("#39ad3b"));
-			}
-			if (EnumHasAnyFlags(FromPhase, ESystemPhase::Instantiation) && EnumHasAnyFlags(ToPhase, ESystemPhase::Instantiation))
-			{
-				String += FString::Printf(TEXT("\tflow_node%d_1 -> flow_node%d_1 [color=\"%s\"];\n"), (int32)Edge.FromNode, (int32)Edge.ToNode, bIsCyclic ? TEXT("#FF0000") : TEXT("#39ad3b"));
-			}
-			if (EnumHasAnyFlags(FromPhase, ESystemPhase::Evaluation) && EnumHasAnyFlags(ToPhase, ESystemPhase::Evaluation))
-			{
-				String += FString::Printf(TEXT("\tflow_node%d_2 -> flow_node%d_2 [color=\"%s\"];\n"), (int32)Edge.FromNode, (int32)Edge.ToNode, bIsCyclic ? TEXT("#FF0000") : TEXT("#39ad3b"));
-			}
-			if (EnumHasAnyFlags(FromPhase, ESystemPhase::Finalization) && EnumHasAnyFlags(ToPhase, ESystemPhase::Finalization))
-			{
-				String += FString::Printf(TEXT("\tflow_node%d_3 -> flow_node%d_3 [color=\"%s\"];\n"), (int32)Edge.FromNode, (int32)Edge.ToNode, bIsCyclic ? TEXT("#FF0000") : TEXT("#39ad3b"));
-			}
-		}
-	}
-
-
 	{
 		FMovieSceneEntitySystemDirectedGraph::FDiscoverCyclicEdges CyclicEdges(&ReferenceGraph);
 		CyclicEdges.Search();
@@ -811,14 +419,20 @@ void FMovieSceneEntitySystemGraph::ExecutePhase(const ArrayType& SortedEntries,
 		{
 			SCOPE_CYCLE_COUNTER(MovieSceneEval_SystemDependencyCost)
 
-			for (FDirectionalEdge Edge : FlowGraph.GetEdgesFrom(NodeID))
+			TArray<uint16> ToGlobalNodeIDs;
+			UMovieSceneEntitySystem::GetSubsequentSystems(System->GetGlobalDependencyGraphID(), ToGlobalNodeIDs);
+			for (uint16 ToGlobalNodeID : ToGlobalNodeIDs)
 			{
-				FMovieSceneEntitySystemGraphNode& ToNode = Nodes.Array[Edge.ToNode];
-				if (!ToNode.Prerequisites)
+				uint16* ToNodeID = GlobalToLocalNodeIDs.Find(ToGlobalNodeID);
+				if (ToNodeID)
 				{
-					ToNode.Prerequisites = MakeShared<FSystemTaskPrerequisites>();
+					FMovieSceneEntitySystemGraphNode& ToNode = Nodes.Array[*ToNodeID];
+					if (!ToNode.Prerequisites)
+					{
+						ToNode.Prerequisites = MakeShared<FSystemTaskPrerequisites>();
+					}
+					ToNode.Prerequisites->Consume(*DownstreamTasks.Subsequents);
 				}
-				ToNode.Prerequisites->Consume(*DownstreamTasks.Subsequents);
 			}
 		}
 	}
@@ -847,3 +461,4 @@ uint16 FMovieSceneEntitySystemGraph::GetGraphID(const UMovieSceneEntitySystem* I
 {
 	return InSystem->GetGraphID();
 }
+
diff --git a/Engine/Source/Runtime/MovieScene/Public/EntitySystem/MovieSceneEntitySystem.h b/Engine/Source/Runtime/MovieScene/Public/EntitySystem/MovieSceneEntitySystem.h
--- a/Engine/Source/Runtime/MovieScene/Public/EntitySystem/MovieSceneEntitySystem.h
+++ b/Engine/Source/Runtime/MovieScene/Public/EntitySystem/MovieSceneEntitySystem.h
@@ -123,6 +123,21 @@ public:
 	 */
 	static UE::MovieScene::EEntitySystemCategory RegisterCustomSystemCategory();
 
+	/**
+	 * Sort the given systems by their flow order, suitable for execution
+	 */
+	static void SortByFlowOrder(TArray<uint16>& InOutGlobalNodeIDs);
+
+	/**
+	 * Get the global IDs of all subsequent systems of the given system
+	 */
+	static void GetSubsequentSystems(uint16 FromGlobalNodeID, TArray<uint16>& OutSubsequentGlobalNodeIDs);
+
+	/**
+	 * Prints a graphviz markup for the global system dependency graph
+	 */
+	static void DebugPrintGlobalDependencyGraph(bool bUpdateCache = true);
+
 public:
 
 	/** Returns system categories */
diff --git a/Engine/Source/Runtime/MovieScene/Public/EntitySystem/MovieSceneEntitySystemDirectedGraph.h b/Engine/Source/Runtime/MovieScene/Public/EntitySystem/MovieSceneEntitySystemDirectedGraph.h
--- a/Engine/Source/Runtime/MovieScene/Public/EntitySystem/MovieSceneEntitySystemDirectedGraph.h
+++ b/Engine/Source/Runtime/MovieScene/Public/EntitySystem/MovieSceneEntitySystemDirectedGraph.h
@@ -138,6 +138,8 @@ public:
 
 	void DestroyEdge(uint16 FromNode, uint16 ToNode);
 
+	void DestroyAllEdges();
+
 	TBitArray<> FindEdgeUpstreamNodes() const;
 
 	TArrayView<const FDirectionalEdge> GetEdges() const;
@@ -169,4 +171,3 @@ private:
 	bool bHasDanglingEdges;
 };
 
-
diff --git a/Engine/Source/Runtime/MovieScene/Public/EntitySystem/MovieSceneEntitySystemGraphs.h b/Engine/Source/Runtime/MovieScene/Public/EntitySystem/MovieSceneEntitySystemGraphs.h
--- a/Engine/Source/Runtime/MovieScene/Public/EntitySystem/MovieSceneEntitySystemGraphs.h
+++ b/Engine/Source/Runtime/MovieScene/Public/EntitySystem/MovieSceneEntitySystemGraphs.h
@@ -87,8 +87,6 @@ struct MOVIESCENE_API FMovieSceneEntitySystemGraph
 
 	GENERATED_BODY()
 
-	void AddPrerequisite(UMovieSceneEntitySystem* Upstream, UMovieSceneEntitySystem* Downstream);
-
 	void AddReference(UMovieSceneEntitySystem* FromReference, UMovieSceneEntitySystem* ToReference);
 
 	void RemoveReference(UMovieSceneEntitySystem* FromReference, UMovieSceneEntitySystem* ToReference);
@@ -147,8 +145,6 @@ struct MOVIESCENE_API FMovieSceneEntitySystemGraph
 
 	void AddSystem(UMovieSceneEntitySystem* InSystem);
 
-	int32 NumSubsequents(UMovieSceneEntitySystem* InSystem) const;
-
 	void RemoveSystem(UMovieSceneEntitySystem* InSystem);
 
 	int32 RemoveIrrelevantSystems(UMovieSceneEntitySystemLinker* Linker);
@@ -194,10 +190,11 @@ private:
 	UPROPERTY()
 	FMovieSceneEntitySystemGraphNodes Nodes;
 
-	FMovieSceneEntitySystemDirectedGraph FlowGraph;
+	TMap<uint16, uint16> GlobalToLocalNodeIDs;
+
 	FMovieSceneEntitySystemDirectedGraph ReferenceGraph;
 
 	uint32 SerialNumber = 0;
 	uint32 PreviousSerialNumber = 0;
 	uint32 ReentrancyGuard = 0;
-};
\ No newline at end of file
+};
