Sequencer: support "additive-from-base" blend type with byte, integer, and perlin noise channels

- New component types are added for byte and integer base values.
- The base value evaluator system now also computes base values for byte and integer channels that are using the additive-from-base blend type. This lets us keep all core channel types inside one system, which makes it easy to optimize or manage the main use-case of the overall edge-case that is additive-from-base channels.
- This system also now computes the base eval time in seconds, for those who need it, like the perlin noise channels.
- The perlin noise channel evaluator systems now also run during instantiation in order to compute base values for the super rare case of noise channels that are additive-from-base. They depend on the base value evaluator system in order to have a correctly computed base eval time in seconds.
- Note that the component factories are now creating the additive-from-base components based on the result components, and not the channel components. This is because we don't care if, say, a double result was created by a double channel or a double perlin noise. In either case, if it has the additive-from-blend tag and a base value eval time, we need to create the other appropriate components. We just trust that the appropriate evaluator channel will compute the values later.

diff --git a/Engine/Source/Runtime/MovieScene/Private/EntitySystem/BuiltInComponentTypes.cpp b/Engine/Source/Runtime/MovieScene/Private/EntitySystem/BuiltInComponentTypes.cpp
--- a/Engine/Source/Runtime/MovieScene/Private/EntitySystem/BuiltInComponentTypes.cpp
+++ b/Engine/Source/Runtime/MovieScene/Private/EntitySystem/BuiltInComponentTypes.cpp
@@ -104,6 +104,7 @@ FBuiltInComponentTypes::FBuiltInComponentTypes()
 	ComponentRegistry->NewComponentType(&DoubleResult[8],       TEXT("Double Result 8"));
 	ComponentRegistry->NewComponentType(&ObjectResult,          TEXT("Object Result"));
 
+	ComponentRegistry->NewComponentType(&BaseByte,			    TEXT("Base Byte"));
 	ComponentRegistry->NewComponentType(&BaseInteger,			TEXT("Base Integer"));
 	ComponentRegistry->NewComponentType(&BaseDouble[0],         TEXT("Base Double 0"));
 	ComponentRegistry->NewComponentType(&BaseDouble[1],         TEXT("Base Double 1"));
@@ -116,6 +117,7 @@ FBuiltInComponentTypes::FBuiltInComponentTypes()
 	ComponentRegistry->NewComponentType(&BaseDouble[8],         TEXT("Base Double 8"));
 
 	ComponentRegistry->NewComponentType(&BaseValueEvalTime,     TEXT("Base Value Eval Time"));
+	ComponentRegistry->NewComponentType(&BaseValueEvalSeconds,  TEXT("Base Value Eval Seconds"));
 
 	ComponentRegistry->NewComponentType(&WeightResult,          TEXT("Weight Result"));
 	ComponentRegistry->NewComponentType(&WeightAndEasingResult, TEXT("Weight/Easing Result"));
@@ -174,6 +176,7 @@ FBuiltInComponentTypes::FBuiltInComponentTypes()
 	ComponentRegistry->Factories.DuplicateChildComponent(EvalTime);
 	ComponentRegistry->Factories.DuplicateChildComponent(EvalSeconds);
 	ComponentRegistry->Factories.DuplicateChildComponent(BaseValueEvalTime);
+	ComponentRegistry->Factories.DuplicateChildComponent(BaseValueEvalSeconds);
 
 	ComponentRegistry->Factories.DuplicateChildComponent(SequenceID);
 	ComponentRegistry->Factories.DuplicateChildComponent(InstanceHandle);
@@ -200,6 +203,10 @@ FBuiltInComponentTypes::FBuiltInComponentTypes()
 		ComponentRegistry->Factories.DuplicateChildComponent(ByteResult);
 		ComponentRegistry->Factories.DefineMutuallyInclusiveComponent(ByteChannel, ByteResult);
 		ComponentRegistry->Factories.DefineMutuallyInclusiveComponent(ByteChannel, EvalTime);
+
+		ComponentRegistry->Factories.DefineComplexInclusiveComponents(
+				FComplexInclusivityFilter::All({ ByteResult, BaseValueEvalTime, Tags.AdditiveFromBaseBlend }),
+				BaseByte);
 	}
 	
 	// Integer channel relationships
@@ -210,7 +217,7 @@ FBuiltInComponentTypes::FBuiltInComponentTypes()
 		ComponentRegistry->Factories.DefineMutuallyInclusiveComponent(IntegerChannel, EvalTime);
 
 		ComponentRegistry->Factories.DefineComplexInclusiveComponents(
-				FComplexInclusivityFilter::All({ IntegerChannel, BaseValueEvalTime, Tags.AdditiveFromBaseBlend }),
+				FComplexInclusivityFilter::All({ IntegerResult, BaseValueEvalTime, Tags.AdditiveFromBaseBlend }),
 				BaseInteger);
 	}
 
@@ -231,14 +238,6 @@ FBuiltInComponentTypes::FBuiltInComponentTypes()
 			ComponentRegistry->Factories.DefineMutuallyInclusiveComponent(FloatChannel[Index], EvalTime);
 			ComponentRegistry->Factories.DefineMutuallyInclusiveComponent(FloatChannel[Index], CachedInterpolation[Index]);
 		}
-
-		// Create base float components for float channels that are meant to be additive from base.
-		for (int32 Index = 0; Index < UE_ARRAY_COUNT(BaseDouble); ++Index)
-		{
-			ComponentRegistry->Factories.DefineComplexInclusiveComponents(
-					FComplexInclusivityFilter::All({ FloatChannel[Index], BaseValueEvalTime, Tags.AdditiveFromBaseBlend }),
-					BaseDouble[Index]);
-		}
 	}
 
 	// Double channel relationships
@@ -259,21 +258,28 @@ FBuiltInComponentTypes::FBuiltInComponentTypes()
 			ComponentRegistry->Factories.DefineMutuallyInclusiveComponent(DoubleChannel[Index], EvalTime);
 			ComponentRegistry->Factories.DefineMutuallyInclusiveComponent(DoubleChannel[Index], CachedInterpolation[Index]);
 		}
-
-		// Create base double components for double channels that are meant to be additive from base.
-		for (int32 Index = 0; Index < UE_ARRAY_COUNT(BaseDouble); ++Index)
-		{
-			ComponentRegistry->Factories.DefineComplexInclusiveComponents(
-					FComplexInclusivityFilter::All({ DoubleChannel[Index], BaseValueEvalTime, Tags.AdditiveFromBaseBlend }),
-					BaseDouble[Index]);
-		}
 	}
 
 	{
+		// Associate result and base values.
 		for (int32 Index = 0; Index < UE_ARRAY_COUNT(DoubleResult); ++Index)
 		{
 			ResultToBase.Add(DoubleResult[Index], BaseDouble[Index]);
 		}
+
+		// Create base double components for any channels that are meant to be additive from base.
+		for (int32 Index = 0; Index < UE_ARRAY_COUNT(BaseDouble); ++Index)
+		{
+			ComponentRegistry->Factories.DefineComplexInclusiveComponents(
+					FComplexInclusivityFilter::All({ DoubleResult[Index], BaseValueEvalTime, Tags.AdditiveFromBaseBlend }),
+					BaseDouble[Index]);
+		}
+
+		// If normal evaluation requires the time in seconds, the base value evaluation probably also requires
+		// a base eval time in seconds.
+		ComponentRegistry->Factories.DefineComplexInclusiveComponents(
+				FComplexInclusivityFilter::All({ EvalSeconds, BaseValueEvalTime }),
+				BaseValueEvalSeconds);
 	}
 
 	// Easing component relationships
diff --git a/Engine/Source/Runtime/MovieScene/Public/EntitySystem/BuiltInComponentTypes.h b/Engine/Source/Runtime/MovieScene/Public/EntitySystem/BuiltInComponentTypes.h
--- a/Engine/Source/Runtime/MovieScene/Public/EntitySystem/BuiltInComponentTypes.h
+++ b/Engine/Source/Runtime/MovieScene/Public/EntitySystem/BuiltInComponentTypes.h
@@ -284,6 +284,9 @@ public:
 	// A byte representing the output of a byte or enum track or channel
 	TComponentTypeID<uint8> ByteResult;
 
+	// A byte representing the base value for the byte channel for the purposes of "additive from base" blending.
+	TComponentTypeID<uint8> BaseByte;
+
 	// An FMovieSceneIntegerChannel
 	TComponentTypeID<FSourceIntegerChannel> IntegerChannel;
 
@@ -318,8 +321,9 @@ public:
 	// A double representing the base value for the double channel at index N, for the purposes of "additive from base" blending.
 	TComponentTypeID<double> BaseDouble[9];
 
-	// The time at which to evaluate a base value, such as BaseFloat[] or BaseDouble[].
+	// The time (in frames or in seconds) at which to evaluate a base value, such as BaseFloat[] or BaseDouble[].
 	TComponentTypeID<FFrameTime> BaseValueEvalTime;
+	TComponentTypeID<double> BaseValueEvalSeconds;
 
 	// A float representing the evaluated output of a weight channel
 	TComponentTypeID<double> WeightResult;
diff --git a/Engine/Source/Runtime/MovieSceneTracks/Private/Systems/DoublePerlinNoiseChannelEvaluatorSystem.cpp b/Engine/Source/Runtime/MovieSceneTracks/Private/Systems/DoublePerlinNoiseChannelEvaluatorSystem.cpp
--- a/Engine/Source/Runtime/MovieSceneTracks/Private/Systems/DoublePerlinNoiseChannelEvaluatorSystem.cpp
+++ b/Engine/Source/Runtime/MovieSceneTracks/Private/Systems/DoublePerlinNoiseChannelEvaluatorSystem.cpp
@@ -2,9 +2,11 @@
 
 #include "Systems/DoublePerlinNoiseChannelEvaluatorSystem.h"
 #include "Systems/MovieScenePiecewiseDoubleBlenderSystem.h"
+#include "Systems/MovieSceneBaseValueEvaluatorSystem.h"
 #include "EntitySystem/BuiltInComponentTypes.h"
 #include "EntitySystem/MovieSceneEntitySystemTask.h"
 #include "EntitySystem/MovieSceneEntitySystemLinker.h"
+#include "EntitySystem/MovieSceneEntitySystemRunner.h"
 #include "EntitySystem/MovieSceneEvalTimeSystem.h"
 #include "Channels/MovieSceneDoublePerlinNoiseChannel.h"
 #include "MovieSceneTracksComponentTypes.h"
@@ -34,6 +36,7 @@ UDoublePerlinNoiseChannelEvaluatorSystem::UDoublePerlinNoiseChannelEvaluatorSyst
 	using namespace UE::MovieScene;
 
 	SystemCategories = EEntitySystemCategory::ChannelEvaluators;
+	Phase = ESystemPhase::Instantiation | ESystemPhase::Evaluation;
 
 	const FBuiltInComponentTypes* BuiltInComponents = FBuiltInComponentTypes::Get();
 	const FMovieSceneTracksComponentTypes* TrackComponents = FMovieSceneTracksComponentTypes::Get();
@@ -49,6 +52,7 @@ UDoublePerlinNoiseChannelEvaluatorSystem::UDoublePerlinNoiseChannelEvaluatorSyst
 		}
 
 		DefineImplicitPrerequisite(UMovieSceneEvalTimeSystem::StaticClass(), GetClass());
+		DefineImplicitPrerequisite(UMovieSceneBaseValueEvaluatorSystem::StaticClass(), GetClass());
 		DefineImplicitPrerequisite(GetClass(), UMovieScenePiecewiseDoubleBlenderSystem::StaticClass());
 	}
 }
@@ -60,15 +64,37 @@ void UDoublePerlinNoiseChannelEvaluatorSystem::OnRun(FSystemTaskPrerequisites& I
 	const FBuiltInComponentTypes* BuiltInComponents = FBuiltInComponentTypes::Get();
 	const FMovieSceneTracksComponentTypes* TrackComponents = FMovieSceneTracksComponentTypes::Get();
 
-	for (int32 i = 0; i < UE_ARRAY_COUNT(BuiltInComponents->DoubleResult); ++i)
+	FMovieSceneEntitySystemRunner* Runner = Linker->GetActiveRunner();
+	if (!Runner)
 	{
-		FEntityTaskBuilder()
-			.Read(BuiltInComponents->EvalSeconds)
-			.Read(TrackComponents->DoublePerlinNoiseChannel)
-			.Write(BuiltInComponents->DoubleResult[i])
-			.FilterNone({ BuiltInComponents->Tags.Ignored })
-			.SetStat(GET_STATID(MovieSceneEval_EvaluateDoublePerlinNoiseChannelTask))
-			.Dispatch_PerEntity<FEvaluateDoublePerlinNoiseChannels>(&Linker->EntityManager, InPrerequisites, &Subsequents);
+		return;
+	}
+
+	if (Runner->GetCurrentPhase() == ESystemPhase::Instantiation)
+	{
+		for (int32 i = 0; i < UE_ARRAY_COUNT(BuiltInComponents->DoubleResult); ++i)
+		{
+			FEntityTaskBuilder()
+				.Read(BuiltInComponents->BaseValueEvalSeconds)
+				.Read(TrackComponents->DoublePerlinNoiseChannel)
+				.Write(BuiltInComponents->BaseDouble[i])
+				.FilterAll({ BuiltInComponents->Tags.NeedsLink })
+				.FilterNone({ BuiltInComponents->Tags.Ignored })
+				.RunInline_PerEntity(&Linker->EntityManager, FEvaluateDoublePerlinNoiseChannels());
+		}
+	}
+	else if (Runner->GetCurrentPhase() == ESystemPhase::Evaluation)
+	{
+		for (int32 i = 0; i < UE_ARRAY_COUNT(BuiltInComponents->DoubleResult); ++i)
+		{
+			FEntityTaskBuilder()
+				.Read(BuiltInComponents->EvalSeconds)
+				.Read(TrackComponents->DoublePerlinNoiseChannel)
+				.Write(BuiltInComponents->DoubleResult[i])
+				.FilterNone({ BuiltInComponents->Tags.Ignored })
+				.SetStat(GET_STATID(MovieSceneEval_EvaluateDoublePerlinNoiseChannelTask))
+				.Dispatch_PerEntity<FEvaluateDoublePerlinNoiseChannels>(&Linker->EntityManager, InPrerequisites, &Subsequents);
+		}
 	}
 }
 
diff --git a/Engine/Source/Runtime/MovieSceneTracks/Private/Systems/FloatPerlinNoiseChannelEvaluatorSystem.cpp b/Engine/Source/Runtime/MovieSceneTracks/Private/Systems/FloatPerlinNoiseChannelEvaluatorSystem.cpp
--- a/Engine/Source/Runtime/MovieSceneTracks/Private/Systems/FloatPerlinNoiseChannelEvaluatorSystem.cpp
+++ b/Engine/Source/Runtime/MovieSceneTracks/Private/Systems/FloatPerlinNoiseChannelEvaluatorSystem.cpp
@@ -2,9 +2,11 @@
 
 #include "Systems/FloatPerlinNoiseChannelEvaluatorSystem.h"
 #include "Systems/MovieScenePiecewiseDoubleBlenderSystem.h"
+#include "Systems/MovieSceneBaseValueEvaluatorSystem.h"
 #include "EntitySystem/BuiltInComponentTypes.h"
 #include "EntitySystem/MovieSceneEntitySystemTask.h"
 #include "EntitySystem/MovieSceneEntitySystemLinker.h"
+#include "EntitySystem/MovieSceneEntitySystemRunner.h"
 #include "EntitySystem/MovieSceneEvalTimeSystem.h"
 #include "Channels/MovieSceneFloatPerlinNoiseChannel.h"
 #include "MovieSceneTracksComponentTypes.h"
@@ -35,6 +37,7 @@ UFloatPerlinNoiseChannelEvaluatorSystem::UFloatPerlinNoiseChannelEvaluatorSystem
 	using namespace UE::MovieScene;
 
 	SystemCategories = EEntitySystemCategory::ChannelEvaluators;
+	Phase = ESystemPhase::Instantiation | ESystemPhase::Evaluation;
 
 	const FBuiltInComponentTypes* BuiltInComponents = FBuiltInComponentTypes::Get();
 	const FMovieSceneTracksComponentTypes* TrackComponents = FMovieSceneTracksComponentTypes::Get();
@@ -50,6 +53,7 @@ UFloatPerlinNoiseChannelEvaluatorSystem::UFloatPerlinNoiseChannelEvaluatorSystem
 		}
 
 		DefineImplicitPrerequisite(UMovieSceneEvalTimeSystem::StaticClass(), GetClass());
+		DefineImplicitPrerequisite(UMovieSceneBaseValueEvaluatorSystem::StaticClass(), GetClass());
 		DefineImplicitPrerequisite(GetClass(), UMovieScenePiecewiseDoubleBlenderSystem::StaticClass());
 	}
 }
@@ -61,15 +65,37 @@ void UFloatPerlinNoiseChannelEvaluatorSystem::OnRun(FSystemTaskPrerequisites& In
 	const FBuiltInComponentTypes* BuiltInComponents = FBuiltInComponentTypes::Get();
 	const FMovieSceneTracksComponentTypes* TrackComponents = FMovieSceneTracksComponentTypes::Get();
 
-	for (int32 i = 0; i < UE_ARRAY_COUNT(BuiltInComponents->DoubleResult); ++i)
+	FMovieSceneEntitySystemRunner* Runner = Linker->GetActiveRunner();
+	if (!Runner)
 	{
-		FEntityTaskBuilder()
-			.Read(BuiltInComponents->EvalSeconds)
-			.Read(TrackComponents->FloatPerlinNoiseChannel)
-			.Write(BuiltInComponents->DoubleResult[i])
-			.FilterNone({ BuiltInComponents->Tags.Ignored })
-			.SetStat(GET_STATID(MovieSceneEval_EvaluateFloatPerlinNoiseChannelTask))
-			.Dispatch_PerEntity<FEvaluateFloatPerlinNoiseChannels>(&Linker->EntityManager, InPrerequisites, &Subsequents);
+		return;
+	}
+
+	if (Runner->GetCurrentPhase() == ESystemPhase::Instantiation)
+	{
+		for (int32 i = 0; i < UE_ARRAY_COUNT(BuiltInComponents->DoubleResult); ++i)
+		{
+			FEntityTaskBuilder()
+				.Read(BuiltInComponents->BaseValueEvalSeconds)
+				.Read(TrackComponents->FloatPerlinNoiseChannel)
+				.Write(BuiltInComponents->BaseDouble[i])
+				.FilterAll({ BuiltInComponents->Tags.NeedsLink })
+				.FilterNone({ BuiltInComponents->Tags.Ignored })
+				.RunInline_PerEntity(&Linker->EntityManager, FEvaluateFloatPerlinNoiseChannels());
+		}
+	}
+	else if (Runner->GetCurrentPhase() == ESystemPhase::Evaluation)
+	{
+		for (int32 i = 0; i < UE_ARRAY_COUNT(BuiltInComponents->DoubleResult); ++i)
+		{
+			FEntityTaskBuilder()
+				.Read(BuiltInComponents->EvalSeconds)
+				.Read(TrackComponents->FloatPerlinNoiseChannel)
+				.Write(BuiltInComponents->DoubleResult[i])
+				.FilterNone({ BuiltInComponents->Tags.Ignored })
+				.SetStat(GET_STATID(MovieSceneEval_EvaluateFloatPerlinNoiseChannelTask))
+				.Dispatch_PerEntity<FEvaluateFloatPerlinNoiseChannels>(&Linker->EntityManager, InPrerequisites, &Subsequents);
+		}
 	}
 }
 
diff --git a/Engine/Source/Runtime/MovieSceneTracks/Private/Systems/MovieSceneBaseValueEvaluatorSystem.cpp b/Engine/Source/Runtime/MovieSceneTracks/Private/Systems/MovieSceneBaseValueEvaluatorSystem.cpp
--- a/Engine/Source/Runtime/MovieSceneTracks/Private/Systems/MovieSceneBaseValueEvaluatorSystem.cpp
+++ b/Engine/Source/Runtime/MovieSceneTracks/Private/Systems/MovieSceneBaseValueEvaluatorSystem.cpp
@@ -7,6 +7,8 @@
 #include "EntitySystem/MovieSceneEntitySystemTask.h"
 #include "EntitySystem/MovieSceneEntitySystemLinker.h"
 #include "EntitySystem/MovieSceneComponentRegistry.h"
+#include "Channels/MovieSceneByteChannel.h"
+#include "Channels/MovieSceneIntegerChannel.h"
 #include "Channels/MovieSceneDoubleChannel.h"
 #include "Channels/MovieSceneFloatChannel.h"
 
@@ -17,6 +19,44 @@ namespace UE
 namespace MovieScene
 {
 
+struct FAssignBaseValueEvalSeconds
+{
+	FInstanceRegistry* InstanceRegistry;
+
+	explicit FAssignBaseValueEvalSeconds(FInstanceRegistry* InInstanceRegistry)
+		: InstanceRegistry(InInstanceRegistry)
+	{}
+
+	void ForEachEntity(UE::MovieScene::FInstanceHandle InstanceHandle, const FFrameTime& BaseValueEvalTime, double& BaseValueEvalSeconds)
+	{
+		const FSequenceInstance& Instance = InstanceRegistry->GetInstance(InstanceHandle);
+		const FMovieSceneContext& Context = Instance.GetContext();
+		BaseValueEvalSeconds = Context.GetFrameRate().AsSeconds(BaseValueEvalTime);
+	}
+};
+
+struct FEvaluateBaseByteValues
+{
+	void ForEachEntity(FSourceByteChannel ByteChannel, FFrameTime FrameTime, uint8& OutResult)
+	{
+		if (!ByteChannel.Source->Evaluate(FrameTime, OutResult))
+		{
+			OutResult = MIN_uint8;
+		}
+	}
+};
+
+struct FEvaluateBaseIntegerValues
+{
+	void ForEachEntity(FSourceIntegerChannel IntegerChannel, FFrameTime FrameTime, int32& OutResult)
+	{
+		if (!IntegerChannel.Source->Evaluate(FrameTime, OutResult))
+		{
+			OutResult = MIN_int32;
+		}
+	}
+};
+
 struct FEvaluateBaseFloatValues
 {
 	void ForEachEntity(FSourceFloatChannel FloatChannel, FFrameTime FrameTime, double& OutResult)
@@ -65,6 +105,33 @@ void UMovieSceneBaseValueEvaluatorSystem::OnRun(FSystemTaskPrerequisites& InPrer
 
 	FBuiltInComponentTypes* BuiltInComponents = FBuiltInComponentTypes::Get();
 
+	// First, compute the base eval time in seconds, for those systems and tasks that need it.
+
+	FAssignBaseValueEvalSeconds AssignBaseValueEvalSecondsTask(Linker->GetInstanceRegistry());
+	FEntityTaskBuilder()
+	.Read(BuiltInComponents->InstanceHandle)
+	.Read(BuiltInComponents->BaseValueEvalTime)
+	.Write(BuiltInComponents->BaseValueEvalSeconds)
+	.RunInline_PerEntity(&Linker->EntityManager, AssignBaseValueEvalSecondsTask);
+
+	// Second, evaluate the base values for all the core channel types.
+
+	FEntityTaskBuilder()
+	.Read(BuiltInComponents->ByteChannel)
+	.Read(BuiltInComponents->BaseValueEvalTime)
+	.Write(BuiltInComponents->BaseByte)
+	.FilterAll({ BuiltInComponents->Tags.NeedsLink })
+	.FilterNone({ BuiltInComponents->Tags.Ignored })
+	.RunInline_PerEntity(&Linker->EntityManager, FEvaluateBaseByteValues());
+
+	FEntityTaskBuilder()
+	.Read(BuiltInComponents->IntegerChannel)
+	.Read(BuiltInComponents->BaseValueEvalTime)
+	.Write(BuiltInComponents->BaseInteger)
+	.FilterAll({ BuiltInComponents->Tags.NeedsLink })
+	.FilterNone({ BuiltInComponents->Tags.Ignored })
+	.RunInline_PerEntity(&Linker->EntityManager, FEvaluateBaseIntegerValues());
+
 	static_assert(
 			UE_ARRAY_COUNT(BuiltInComponents->BaseDouble) == UE_ARRAY_COUNT(BuiltInComponents->FloatChannel),
 			"There should be a matching number of float channels and float base values.");
