Sequencer: Fixed DesiredThread and StatIds not being propogated to the final task when being defined before component accessors

diff --git a/Engine/Source/Runtime/MovieScene/Public/EntitySystem/MovieSceneEntitySystemTask.h b/Engine/Source/Runtime/MovieScene/Public/EntitySystem/MovieSceneEntitySystemTask.h
--- a/Engine/Source/Runtime/MovieScene/Public/EntitySystem/MovieSceneEntitySystemTask.h
+++ b/Engine/Source/Runtime/MovieScene/Public/EntitySystem/MovieSceneEntitySystemTask.h
@@ -45,6 +45,18 @@ template<typename> struct TUnstructuredTask;
 template<typename TaskImpl, int32 NumComponents, bool AutoExpandAccessors>
 struct TEntityTaskCaller;
 
+struct FCommonEntityTaskParams
+{
+	/** A stat ID for the task */
+	TStatId StatId;
+
+	/** The thread that this task wants to run on */
+	ENamedThreads::Type DesiredThread = ENamedThreads::AnyHiPriThreadHiPriTask;
+
+	/** Useful for debugging to break the debugger when this task is run */
+	bool bBreakOnRun = false;
+};
+
 /** Default traits specialized for each user TaskImplInstance */
 template<typename T>
 struct TDefaultEntityTaskTraits
@@ -202,13 +214,9 @@ struct TEntityTaskComponents : TEntityTaskComponentsImpl<TMakeIntegerSequence<in
 		return TFilteredEntityTask< T... >(*this, InFilter);
 	}
 
-	/**
-	 * Assign the current thread for task dispatch to ensure that it is issued on the correct thread.
-	 * This should only be required for tasks dispatched outside of the main linker execution, or tasks dispatched for the global entity manager
-	 */
+	UE_DEPRECATED(5.2, "This function is not required.")
 	TEntityTaskComponents< T... >& SetCurrentThread(ENamedThreads::Type InCurrentThread)
 	{
-		CurrentThread = InCurrentThread;
 		return *this;
 	}
 
@@ -217,7 +225,7 @@ struct TEntityTaskComponents : TEntityTaskComponentsImpl<TMakeIntegerSequence<in
 	 */
 	TEntityTaskComponents< T... >& SetDesiredThread(ENamedThreads::Type InDesiredThread)
 	{
-		DesiredThread = InDesiredThread;
+		this->CommonParams.DesiredThread = InDesiredThread;
 		return *this;
 	}
 
@@ -226,7 +234,7 @@ struct TEntityTaskComponents : TEntityTaskComponentsImpl<TMakeIntegerSequence<in
 	 */
 	TEntityTaskComponents< T... >& SetStat(TStatId InStatId)
 	{
-		StatId = InStatId;
+		this->CommonParams.StatId = InStatId;
 		return *this;
 	}
 
@@ -259,11 +267,11 @@ struct TEntityTaskComponents : TEntityTaskComponentsImpl<TMakeIntegerSequence<in
 			FGraphEventArray GatheredPrereqs;
 			this->PopulatePrerequisites(Prerequisites, &GatheredPrereqs);
 
-			ENamedThreads::Type ThisThread = CurrentThread == ENamedThreads::AnyThread ? EntityManager->GetDispatchThread() : CurrentThread;
+			ENamedThreads::Type ThisThread = EntityManager->GetDispatchThread();
 			checkSlow(ThisThread != ENamedThreads::AnyThread);
 
 			FGraphEventRef NewTask = TGraphTask< TUnstructuredTask<TaskImpl> >::CreateTask(GatheredPrereqs.Num() != 0 ? &GatheredPrereqs : nullptr, ThisThread)
-				.ConstructAndDispatchWhenReady( DesiredThread, StatId, bBreakOnRun, Forward<TaskConstructionArgs>(InArgs)... );
+				.ConstructAndDispatchWhenReady( this->CommonParams, Forward<TaskConstructionArgs>(InArgs)... );
 
 			if (Subsequents)
 			{
@@ -330,11 +338,11 @@ struct TEntityTaskComponents : TEntityTaskComponentsImpl<TMakeIntegerSequence<in
 			FGraphEventArray GatheredPrereqs;
 			this->PopulatePrerequisites(Prerequisites, &GatheredPrereqs);
 
-			ENamedThreads::Type ThisThread = CurrentThread == ENamedThreads::AnyThread ? EntityManager->GetDispatchThread() : CurrentThread;
+			ENamedThreads::Type ThisThread = EntityManager->GetDispatchThread();
 			checkSlow(ThisThread != ENamedThreads::AnyThread);
 
 			FGraphEventRef NewTask = TGraphTask< TEntityAllocationTask<TaskImpl, T...> >::CreateTask(GatheredPrereqs.Num() != 0 ? &GatheredPrereqs : nullptr, ThisThread)
-				.ConstructAndDispatchWhenReady( EntityManager, *this, DesiredThread, StatId, bBreakOnRun, Forward<TaskConstructionArgs>(InArgs)... );
+				.ConstructAndDispatchWhenReady( EntityManager, *this, Forward<TaskConstructionArgs>(InArgs)... );
 
 			if (Subsequents)
 			{
@@ -408,11 +416,11 @@ struct TEntityTaskComponents : TEntityTaskComponentsImpl<TMakeIntegerSequence<in
 			FGraphEventArray GatheredPrereqs;
 			this->PopulatePrerequisites(Prerequisites, &GatheredPrereqs);
 
-			ENamedThreads::Type ThisThread = CurrentThread == ENamedThreads::AnyThread ? EntityManager->GetDispatchThread() : CurrentThread;
+			ENamedThreads::Type ThisThread = EntityManager->GetDispatchThread();
 			checkSlow(ThisThread != ENamedThreads::AnyThread);
 
 			FGraphEventRef NewTask = TGraphTask< TEntityTask<TaskImpl, T...> >::CreateTask(GatheredPrereqs.Num() != 0 ? &GatheredPrereqs : nullptr, ThisThread)
-				.ConstructAndDispatchWhenReady( EntityManager, *this, DesiredThread, StatId, bBreakOnRun, Forward<TaskConstructionArgs>(InArgs)... );
+				.ConstructAndDispatchWhenReady( EntityManager, *this, Forward<TaskConstructionArgs>(InArgs)... );
 
 			if (Subsequents)
 			{
@@ -435,29 +443,14 @@ struct TEntityTaskComponents : TEntityTaskComponentsImpl<TMakeIntegerSequence<in
 
 public:
 
-	/** Useful for debugging to break the debugger when this task is run */
-	bool bBreakOnRun;
-	/** The current thread that is being used to dispatch from. Only necessary when EntityManager->GetDispatchThread is not available (ie when tasks are being run outside of system linkers) */
-	ENamedThreads::Type CurrentThread;
-	/** The thread that this task wants to run on */
-	ENamedThreads::Type DesiredThread;
-	/** A stat ID for the task */
-	TStatId StatId;
-
 	TEntityTaskComponents()
-		: bBreakOnRun(false)
-		, CurrentThread(ENamedThreads::AnyThread)
-		, DesiredThread(ENamedThreads::AnyHiPriThreadHiPriTask)
 	{
 		static_assert(sizeof...(T) == 0, "Default construction is only supported for TEntityTaskComponents<>");
 	}
 
 	template<typename... ConstructionTypes>
-	explicit TEntityTaskComponents(ConstructionTypes&&... InTypes)
-		: Super(Forward<ConstructionTypes>(InTypes)...)
-		, bBreakOnRun(false)
-		, CurrentThread(ENamedThreads::AnyThread)
-		, DesiredThread(ENamedThreads::AnyHiPriThreadHiPriTask)
+	explicit TEntityTaskComponents(const FCommonEntityTaskParams& InCommonParams, ConstructionTypes&&... InTypes)
+		: Super(InCommonParams, Forward<ConstructionTypes>(InTypes)...)
 	{}
 };
 
@@ -470,7 +463,7 @@ struct TEntityTaskComponentsImpl<TIntegerSequence<int, Indices...>, T...>
 	 */
 	TEntityTaskComponents< T..., FEntityIDAccess > ReadEntityIDs() const
 	{
-		return TEntityTaskComponents< T..., FEntityIDAccess >( Accessors.template Get<Indices>()..., FEntityIDAccess{} );
+		return TEntityTaskComponents< T..., FEntityIDAccess >( CommonParams, Accessors.template Get<Indices>()..., FEntityIDAccess{} );
 	}
 
 	/**
@@ -482,7 +475,7 @@ struct TEntityTaskComponentsImpl<TIntegerSequence<int, Indices...>, T...>
 	template<typename U>
 	TEntityTaskComponents<T..., TReadAccess<U>> Read(TComponentTypeID<U> ComponentType) const
 	{
-		return TEntityTaskComponents<T..., TReadAccess<U> >(Accessors.template Get<Indices>()..., TReadAccess<U>(ComponentType));
+		return TEntityTaskComponents<T..., TReadAccess<U> >(CommonParams, Accessors.template Get<Indices>()..., TReadAccess<U>(ComponentType));
 	}
 
 	/**
@@ -494,7 +487,7 @@ struct TEntityTaskComponentsImpl<TIntegerSequence<int, Indices...>, T...>
 	template<typename... U>
 	TEntityTaskComponents<T..., TReadOneOfAccessor<U...>> ReadOneOf(TComponentTypeID<U>... ComponentTypes) const
 	{
-		return TEntityTaskComponents<T..., TReadOneOfAccessor<U...> >( Accessors.template Get<Indices>()..., TReadOneOfAccessor<U...>(ComponentTypes...) );
+		return TEntityTaskComponents<T..., TReadOneOfAccessor<U...> >( CommonParams, Accessors.template Get<Indices>()..., TReadOneOfAccessor<U...>(ComponentTypes...) );
 	}
 
 	/**
@@ -505,7 +498,7 @@ struct TEntityTaskComponentsImpl<TIntegerSequence<int, Indices...>, T...>
 	template<typename... U>
 	TEntityTaskComponents<T..., TReadOneOrMoreOfAccessor<U...>> ReadOneOrMoreOf(TComponentTypeID<U>... ComponentTypes) const
 	{
-		return TEntityTaskComponents<T..., TReadOneOrMoreOfAccessor<U...> >(Accessors.template Get<Indices>()..., TReadOneOrMoreOfAccessor<U...>(ComponentTypes...));
+		return TEntityTaskComponents<T..., TReadOneOrMoreOfAccessor<U...> >(CommonParams, Accessors.template Get<Indices>()..., TReadOneOrMoreOfAccessor<U...>(ComponentTypes...));
 	}
 
 	/**
@@ -516,7 +509,7 @@ struct TEntityTaskComponentsImpl<TIntegerSequence<int, Indices...>, T...>
 	template<typename... U>
 	TEntityTaskComponents<T..., TReadAccess<U>...> ReadAllOf(TComponentTypeID<U>... ComponentTypes) const
 	{
-		return TEntityTaskComponents<T..., TReadAccess<U>... >( Accessors.template Get<Indices>()..., TReadAccess<U>(ComponentTypes)... );
+		return TEntityTaskComponents<T..., TReadAccess<U>... >( CommonParams, Accessors.template Get<Indices>()..., TReadAccess<U>(ComponentTypes)... );
 	}
 
 	/**
@@ -527,7 +520,7 @@ struct TEntityTaskComponentsImpl<TIntegerSequence<int, Indices...>, T...>
 	template<typename... U>
 	TEntityTaskComponents<T..., TOptionalReadAccess<U>...> ReadAnyOf(TComponentTypeID<U>... ComponentTypes) const
 	{
-		return TEntityTaskComponents<T..., TOptionalReadAccess<U>... >( Accessors.template Get<Indices>()..., TOptionalReadAccess<U>(ComponentTypes)... );
+		return TEntityTaskComponents<T..., TOptionalReadAccess<U>... >( CommonParams, Accessors.template Get<Indices>()..., TOptionalReadAccess<U>(ComponentTypes)... );
 	}
 
 	/**
@@ -538,7 +531,7 @@ struct TEntityTaskComponentsImpl<TIntegerSequence<int, Indices...>, T...>
 	 */
 	TEntityTaskComponents<T..., FErasedReadAccess> ReadErased(FComponentTypeID ComponentType) const
 	{
-		return TEntityTaskComponents<T..., FErasedReadAccess >(Accessors.template Get<Indices>()..., FErasedReadAccess(ComponentType));
+		return TEntityTaskComponents<T..., FErasedReadAccess >(CommonParams, Accessors.template Get<Indices>()..., FErasedReadAccess(ComponentType));
 	}
 
 	/**
@@ -548,7 +541,7 @@ struct TEntityTaskComponentsImpl<TIntegerSequence<int, Indices...>, T...>
 	template<typename U>
 	TEntityTaskComponents<T..., TOptionalReadAccess<U>> ReadOptional(TComponentTypeID<U> ComponentType) const
 	{
-		return TEntityTaskComponents<T..., TOptionalReadAccess<U> >(Accessors.template Get<Indices>()..., TOptionalReadAccess<U>(ComponentType));
+		return TEntityTaskComponents<T..., TOptionalReadAccess<U> >(CommonParams, Accessors.template Get<Indices>()..., TOptionalReadAccess<U>(ComponentType));
 	}
 
 	/**
@@ -560,7 +553,7 @@ struct TEntityTaskComponentsImpl<TIntegerSequence<int, Indices...>, T...>
 	template<typename U>
 	TEntityTaskComponents<T..., TWriteAccess<U>> Write(TComponentTypeID<U> ComponentType) const
 	{
-		return TEntityTaskComponents<T..., TWriteAccess<U> >(Accessors.template Get<Indices>()..., TWriteAccess<U>(ComponentType));
+		return TEntityTaskComponents<T..., TWriteAccess<U> >(CommonParams, Accessors.template Get<Indices>()..., TWriteAccess<U>(ComponentType));
 	}
 
 	/**
@@ -571,7 +564,7 @@ struct TEntityTaskComponentsImpl<TIntegerSequence<int, Indices...>, T...>
 	 */
 	TEntityTaskComponents<T..., FErasedWriteAccess> WriteErased(FComponentTypeID ComponentType) const
 	{
-		return TEntityTaskComponents<T..., FErasedWriteAccess >(Accessors.template Get<Indices>()..., FErasedWriteAccess(ComponentType));
+		return TEntityTaskComponents<T..., FErasedWriteAccess >(CommonParams, Accessors.template Get<Indices>()..., FErasedWriteAccess(ComponentType));
 	}
 
 
@@ -582,7 +575,7 @@ struct TEntityTaskComponentsImpl<TIntegerSequence<int, Indices...>, T...>
 	template<typename U>
 	TEntityTaskComponents<T..., TOptionalWriteAccess<U>> WriteOptional(TComponentTypeID<U> ComponentType) const
 	{
-		return TEntityTaskComponents<T..., TOptionalWriteAccess<U> >(Accessors.template Get<Indices>()..., TOptionalWriteAccess<U>(ComponentType));
+		return TEntityTaskComponents<T..., TOptionalWriteAccess<U> >(CommonParams, Accessors.template Get<Indices>()..., TOptionalWriteAccess<U>(ComponentType));
 	}
 
 	bool HasBeenWrittenToSince(uint32 InSystemVersion)
@@ -810,9 +803,13 @@ public:
 
 protected:
 
+	TEntityTaskComponentsImpl()
+	{}
+
 	template<typename... ConstructionTypes>
-	explicit TEntityTaskComponentsImpl(ConstructionTypes&&... InTypes)
+	explicit TEntityTaskComponentsImpl(const FCommonEntityTaskParams& InCommonParams, ConstructionTypes&&... InTypes)
 		: Accessors{ Forward<ConstructionTypes>(InTypes)... }
+		, CommonParams(InCommonParams)
 	{}
 
 protected:
@@ -820,6 +817,10 @@ protected:
 	template<typename...> friend struct TEntityTaskComponentsImpl;
 
 	TTuple<T...> Accessors;
+
+public:
+
+	FCommonEntityTaskParams CommonParams;
 };
 
 
@@ -837,20 +838,12 @@ struct TFilteredEntityTask
 {
 	TFilteredEntityTask(const TEntityTaskComponents<T...>& InComponents)
 		: Components(InComponents)
-		, bBreakOnRun(InComponents.bBreakOnRun)
-		, CurrentThread(InComponents.CurrentThread)
-		, DesiredThread(InComponents.DesiredThread)
-		, StatId(InComponents.StatId)
 	{
 		Components.PopulateFilter(&Filter);
 	}
 	TFilteredEntityTask(const TEntityTaskComponents<T...>& InComponents, const FEntityComponentFilter& InFilter)
 		: Components(InComponents)
 		, Filter(InFilter)
-		, bBreakOnRun(InComponents.bBreakOnRun)
-		, CurrentThread(InComponents.CurrentThread)
-		, DesiredThread(InComponents.DesiredThread)
-		, StatId(InComponents.StatId)
 	{
 		Components.PopulateFilter(&Filter);
 	}
@@ -937,13 +930,9 @@ struct TFilteredEntityTask
 		return *this;
 	}
 
-	/**
-	 * Assign the current thread for task dispatch to ensure that it is issued on the correct thread.
-	 * This should only be required for tasks dispatched outside of the main linker execution, or tasks dispatched for the global entity manager
-	 */
+	UE_DEPRECATED(5.2, "This function is not required.")
 	TFilteredEntityTask< T... >& SetCurrentThread(ENamedThreads::Type InCurrentThread)
 	{
-		CurrentThread = InCurrentThread;
 		return *this;
 	}
 
@@ -952,7 +941,7 @@ struct TFilteredEntityTask
 	 */
 	TFilteredEntityTask< T... >& SetDesiredThread(ENamedThreads::Type InDesiredThread)
 	{
-		DesiredThread = InDesiredThread;
+		Components.CommonParams.DesiredThread = InDesiredThread;
 		return *this;
 	}
 
@@ -961,10 +950,34 @@ struct TFilteredEntityTask
 	 */
 	TFilteredEntityTask< T... >& SetStat(TStatId InStatId)
 	{
-		StatId = InStatId;
+		Components.CommonParams.StatId = InStatId;
 		return *this;
 	}
 
+	/**
+	 * Get the desired thread for this task to run on
+	 */
+	ENamedThreads::Type GetDesiredThread() const
+	{
+		return Components.CommonParams.DesiredThread;
+	}
+
+	/**
+	 * Return this task's stat id
+	 */
+	TStatId GetStatId() const
+	{
+		return Components.CommonParams.StatId;
+	}
+
+	/**
+	 * Check whether we should break the debugger when this task is run
+	 */
+	bool ShouldBreakOnRun() const
+	{
+		return Components.CommonParams.bBreakOnRun;
+	}
+
 	/**
 	 * Access the pre-populated filter that should be used for iterating relevant entities for this task
 	 */
@@ -1035,11 +1048,11 @@ struct TFilteredEntityTask
 			FGraphEventArray GatheredPrereqs;
 			Components.PopulatePrerequisites(Prerequisites, &GatheredPrereqs);
 
-			ENamedThreads::Type ThisThread = CurrentThread == ENamedThreads::AnyThread ? EntityManager->GetDispatchThread() : CurrentThread;
+			ENamedThreads::Type ThisThread = EntityManager->GetDispatchThread();
 			checkSlow(ThisThread != ENamedThreads::AnyThread);
 
 			FGraphEventRef NewTask = TGraphTask< TEntityAllocationTask<TaskImpl, T...> >::CreateTask(GatheredPrereqs.Num() != 0 ? &GatheredPrereqs : nullptr, ThisThread)
-				.ConstructAndDispatchWhenReady( EntityManager, *this, DesiredThread, StatId, bBreakOnRun, Forward<TaskConstructionArgs>(InArgs)... );
+				.ConstructAndDispatchWhenReady( EntityManager, *this, Forward<TaskConstructionArgs>(InArgs)... );
 
 			if (Subsequents)
 			{
@@ -1113,11 +1126,11 @@ struct TFilteredEntityTask
 			FGraphEventArray GatheredPrereqs;
 			Components.PopulatePrerequisites(Prerequisites, &GatheredPrereqs);
 
-			ENamedThreads::Type ThisThread = CurrentThread == ENamedThreads::AnyThread ? EntityManager->GetDispatchThread() : CurrentThread;
+			ENamedThreads::Type ThisThread = EntityManager->GetDispatchThread();
 			checkSlow(ThisThread != ENamedThreads::AnyThread);
 
 			FGraphEventRef NewTask = TGraphTask< TEntityTask<TaskImpl, T...> >::CreateTask(GatheredPrereqs.Num() != 0 ? &GatheredPrereqs : nullptr, ThisThread)
-				.ConstructAndDispatchWhenReady( EntityManager, *this, DesiredThread, StatId, bBreakOnRun, Forward<TaskConstructionArgs>(InArgs)... );
+				.ConstructAndDispatchWhenReady( EntityManager, *this, Forward<TaskConstructionArgs>(InArgs)... );
 
 			if (Subsequents)
 			{
@@ -1166,14 +1179,6 @@ private:
 	TEntityTaskComponents<T...> Components;
 
 	FEntityComponentFilter Filter;
-
-	bool bBreakOnRun;
-
-	ENamedThreads::Type CurrentThread;
-
-	ENamedThreads::Type DesiredThread;
-
-	TStatId StatId;
 };
 
 template<typename TaskImpl, typename... ComponentTypes>
@@ -1191,6 +1196,16 @@ struct TEntityTaskBase
 		, WriteContext(*InEntityManager)
 	{}
 
+	TStatId GetStatId() const
+	{
+		return FilteredTask.GetStatId();
+	}
+
+	ENamedThreads::Type GetDesiredThread() const
+	{
+		return FilteredTask.GetDesiredThread();
+	}
+
 	void Run(TaskImpl& TaskImplInstance)
 	{
 		UE_LOG(LogMovieSceneECS, VeryVerbose, TEXT("Running entity task the following components: %s"), *FilteredTask.GetComponents().ToString(EntityManager));
@@ -1210,7 +1225,7 @@ struct TEntityTaskBase
 		PostTask(&TaskImplInstance);
 	}
 
-private:
+protected:
 
 	static void PreTask(void*, ...){}
 	template <typename T> static void PreTask(T* InTask, decltype(&T::PreTask)* = 0)
@@ -1237,33 +1252,17 @@ struct TEntityTask : TEntityTaskBase<TaskImpl, ComponentTypes...>
 {
 
 	template<typename... ArgTypes>
-	explicit TEntityTask(FEntityManager* InEntityManager, const TEntityTaskComponents<ComponentTypes...>& InComponents, ENamedThreads::Type InDesiredThread, TStatId InStatId, bool bInBreakOnRun, ArgTypes&&... InArgs)
+	explicit TEntityTask(FEntityManager* InEntityManager, const TEntityTaskComponents<ComponentTypes...>& InComponents, ArgTypes&&... InArgs)
 		: TEntityTaskBase<TaskImpl, ComponentTypes...>(InEntityManager, InComponents)
 		, TaskImplInstance{ Forward<ArgTypes>(InArgs)... }
-		, DesiredThread(InDesiredThread)
-		, StatId(InStatId)
-		, bBreakOnRun(bInBreakOnRun)
 	{}
 
 	template<typename... ArgTypes>
-	explicit TEntityTask(FEntityManager* InEntityManager, const TFilteredEntityTask<ComponentTypes...>& InFilteredTask, ENamedThreads::Type InDesiredThread, TStatId InStatId, bool bInBreakOnRun, ArgTypes&&... InArgs)
+	explicit TEntityTask(FEntityManager* InEntityManager, const TFilteredEntityTask<ComponentTypes...>& InFilteredTask, ArgTypes&&... InArgs)
 		: TEntityTaskBase<TaskImpl, ComponentTypes...>(InEntityManager, InFilteredTask)
 		, TaskImplInstance{ Forward<ArgTypes>(InArgs)... }
-		, DesiredThread(InDesiredThread)
-		, StatId(InStatId)
-		, bBreakOnRun(bInBreakOnRun)
 	{}
 
-	TStatId GetStatId() const
-	{
-		return StatId;
-	}
-
-	ENamedThreads::Type GetDesiredThread()
-	{
-		return DesiredThread;
-	}
-
 	static ESubsequentsMode::Type GetSubsequentsMode()
 	{
 		return ESubsequentsMode::TrackSubsequents;
@@ -1271,14 +1270,14 @@ struct TEntityTask : TEntityTaskBase<TaskImpl, ComponentTypes...>
 
 	void DoTask(ENamedThreads::Type CurrentThread, const FGraphEventRef& CompletionGraphEvent)
 	{
-		if (bBreakOnRun)
+		if (this->FilteredTask.ShouldBreakOnRun())
 		{
 			UE_DEBUG_BREAK();
 		}
 
-		if ((DesiredThread & ENamedThreads::AnyThread) == 0)
+		if ((this->FilteredTask.GetDesiredThread() & ENamedThreads::AnyThread) == 0)
 		{
-			checkf(CurrentThread == DesiredThread, TEXT("MovieScene evaluation task is not being run on its desired thread"));
+			checkf(CurrentThread == this->FilteredTask.GetDesiredThread(), TEXT("MovieScene evaluation task is not being run on its desired thread"));
 		}
 
 		this->Run(TaskImplInstance);
@@ -1287,12 +1286,6 @@ struct TEntityTask : TEntityTaskBase<TaskImpl, ComponentTypes...>
 private:
 
 	TaskImpl TaskImplInstance;
-
-	ENamedThreads::Type DesiredThread;
-
-	TStatId StatId;
-
-	bool bBreakOnRun;
 };
 
 template<typename TaskImpl, typename... ComponentTypes>
@@ -1310,6 +1303,16 @@ struct TEntityAllocationTaskBase
 		, WriteContext(*InEntityManager)
 	{}
 
+	TStatId GetStatId() const
+	{
+		return ComponentFilter.GetStatId();
+	}
+
+	ENamedThreads::Type GetDesiredThread() const
+	{
+		return ComponentFilter.GetDesiredThread();
+	}
+
 	void Run(TaskImpl& TaskImplInstance)
 	{
 		UE_LOG(LogMovieSceneECS, VeryVerbose, TEXT("Running entity task the following components: %s"), *ComponentFilter.GetComponents().ToString(EntityManager));
@@ -1356,33 +1359,17 @@ template<typename TaskImpl, typename... ComponentTypes>
 struct TEntityAllocationTask : TEntityAllocationTaskBase<TaskImpl, ComponentTypes...>
 {
 	template<typename... ArgTypes>
-	explicit TEntityAllocationTask(FEntityManager* InEntityManager, const TEntityTaskComponents<ComponentTypes...>& InComponents, ENamedThreads::Type InDesiredThread, TStatId InStatId, bool bInBreakOnRun, ArgTypes&&... InArgs)
+	explicit TEntityAllocationTask(FEntityManager* InEntityManager, const TEntityTaskComponents<ComponentTypes...>& InComponents, ArgTypes&&... InArgs)
 		: TEntityAllocationTaskBase<TaskImpl, ComponentTypes...>(InEntityManager, InComponents)
 		, TaskImplInstance{ Forward<ArgTypes>(InArgs)... }
-		, DesiredThread(InDesiredThread)
-		, StatId(InStatId)
-		, bBreakOnRun(bInBreakOnRun)
 	{}
 
 	template<typename... ArgTypes>
-	explicit TEntityAllocationTask(FEntityManager* InEntityManager, const TFilteredEntityTask<ComponentTypes...>& InComponentFilter, ENamedThreads::Type InDesiredThread, TStatId InStatId, bool bInBreakOnRun, ArgTypes&&... InArgs)
+	explicit TEntityAllocationTask(FEntityManager* InEntityManager, const TFilteredEntityTask<ComponentTypes...>& InComponentFilter, ArgTypes&&... InArgs)
 		: TEntityAllocationTaskBase<TaskImpl, ComponentTypes...>(InEntityManager, InComponentFilter)
 		, TaskImplInstance{ Forward<ArgTypes>(InArgs)... }
-		, DesiredThread(InDesiredThread)
-		, StatId(InStatId)
-		, bBreakOnRun(bInBreakOnRun)
 	{}
 
-	TStatId GetStatId() const
-	{
-		return StatId;
-	}
-
-	ENamedThreads::Type GetDesiredThread()
-	{
-		return DesiredThread;
-	}
-
 	static ESubsequentsMode::Type GetSubsequentsMode()
 	{
 		return ESubsequentsMode::TrackSubsequents;
@@ -1390,14 +1377,14 @@ struct TEntityAllocationTask : TEntityAllocationTaskBase<TaskImpl, ComponentType
 
 	void DoTask(ENamedThreads::Type CurrentThread, const FGraphEventRef& CompletionGraphEvent)
 	{
-		if (bBreakOnRun)
+		if (this->ComponentFilter.ShouldBreakOnRun())
 		{
 			UE_DEBUG_BREAK();
 		}
 
-		if ((DesiredThread & ENamedThreads::AnyThread) == 0)
+		if ((this->ComponentFilter.GetDesiredThread() & ENamedThreads::AnyThread) == 0)
 		{
-			checkf(CurrentThread == DesiredThread, TEXT("MovieScene evaluation task is not being run on its desired thread"));
+			checkf(CurrentThread == this->ComponentFilter.GetDesiredThread(), TEXT("MovieScene evaluation task is not being run on its desired thread"));
 		}
 
 		this->Run(TaskImplInstance);
@@ -1406,33 +1393,25 @@ struct TEntityAllocationTask : TEntityAllocationTaskBase<TaskImpl, ComponentType
 private:
 
 	TaskImpl TaskImplInstance;
-
-	ENamedThreads::Type DesiredThread;
-
-	TStatId StatId;
-
-	bool bBreakOnRun;
 };
 
 template<typename TaskImpl>
 struct TUnstructuredTask
 {
 	template<typename... ArgTypes>
-	explicit TUnstructuredTask(ENamedThreads::Type InDesiredThread, TStatId InStatId, bool bInBreakOnRun, ArgTypes&&... InArgs)
+	explicit TUnstructuredTask(const FCommonEntityTaskParams& InCommonParams, ArgTypes&&... InArgs)
 		: TaskImplInstance{ Forward<ArgTypes>(InArgs)... }
-		, DesiredThread(InDesiredThread)
-		, StatId(InStatId)
-		, bBreakOnRun(bInBreakOnRun)
+		, CommonParams(InCommonParams)
 	{}
 
 	TStatId GetStatId() const
 	{
-		return StatId;
+		return CommonParams.StatId;
 	}
 
-	ENamedThreads::Type GetDesiredThread()
+	ENamedThreads::Type GetDesiredThread() const
 	{
-		return DesiredThread;
+		return CommonParams.DesiredThread;
 	}
 
 	static ESubsequentsMode::Type GetSubsequentsMode()
@@ -1442,14 +1421,14 @@ struct TUnstructuredTask
 
 	void DoTask(ENamedThreads::Type CurrentThread, const FGraphEventRef& CompletionGraphEvent)
 	{
-		if (bBreakOnRun)
+		if (CommonParams.bBreakOnRun)
 		{
 			UE_DEBUG_BREAK();
 		}
 
-		if ((DesiredThread & ENamedThreads::AnyThread) == 0)
+		if ((CommonParams.DesiredThread & ENamedThreads::AnyThread) == 0)
 		{
-			checkf(CurrentThread == DesiredThread, TEXT("MovieScene evaluation task is not being run on its desired thread"));
+			checkf(CurrentThread == CommonParams.DesiredThread, TEXT("MovieScene evaluation task is not being run on its desired thread"));
 		}
 
 		UE_LOG(LogMovieSceneECS, VeryVerbose, TEXT("Running unstructured task"));
@@ -1461,11 +1440,7 @@ private:
 
 	TaskImpl TaskImplInstance;
 
-	ENamedThreads::Type DesiredThread;
-
-	TStatId StatId;
-
-	bool bBreakOnRun;
+	FCommonEntityTaskParams CommonParams;
 };
 
 template<typename...> struct TEntityTaskCaller_AutoExpansion;
