Fix indeterminism in cooked build due to uninitialized struct padding being included in the serialized data

#rnx

diff --git a/Engine/Source/Runtime/MovieScene/Public/Channels/MovieSceneBoolChannel.h b/Engine/Source/Runtime/MovieScene/Public/Channels/MovieSceneBoolChannel.h
--- a/Engine/Source/Runtime/MovieScene/Public/Channels/MovieSceneBoolChannel.h
+++ b/Engine/Source/Runtime/MovieScene/Public/Channels/MovieSceneBoolChannel.h
@@ -32,7 +32,7 @@ struct MOVIESCENE_API FMovieSceneBoolChannel : public FMovieSceneChannel
 	GENERATED_BODY()
 
 	FMovieSceneBoolChannel()
-		: DefaultValue(), bHasDefaultValue(false)
+		: DefaultValue(false), bHasDefaultValue(false)
 	{}
 
 	/**
diff --git a/Engine/Source/Runtime/MovieScene/Public/Channels/MovieSceneByteChannel.h b/Engine/Source/Runtime/MovieScene/Public/Channels/MovieSceneByteChannel.h
--- a/Engine/Source/Runtime/MovieScene/Public/Channels/MovieSceneByteChannel.h
+++ b/Engine/Source/Runtime/MovieScene/Public/Channels/MovieSceneByteChannel.h
@@ -33,7 +33,7 @@ struct MOVIESCENE_API FMovieSceneByteChannel : public FMovieSceneChannel
 	GENERATED_BODY()
 
 	FMovieSceneByteChannel()
-		: DefaultValue()
+		: DefaultValue(0)
 		, bHasDefaultValue(false)
 		, Enum(nullptr)
 		, KeyHandles()
diff --git a/Engine/Source/Runtime/MovieScene/Public/Channels/MovieSceneCurveChannelCommon.h b/Engine/Source/Runtime/MovieScene/Public/Channels/MovieSceneCurveChannelCommon.h
--- a/Engine/Source/Runtime/MovieScene/Public/Channels/MovieSceneCurveChannelCommon.h
+++ b/Engine/Source/Runtime/MovieScene/Public/Channels/MovieSceneCurveChannelCommon.h
@@ -64,8 +64,17 @@ struct FMovieSceneTangentData
 	UPROPERTY(EditAnywhere, Category = "Key")
 	TEnumAsByte<ERichCurveTangentWeightMode> TangentWeightMode;
 
+	// This is required because TMovieSceneCurveChannelImpl<ChannelType>::Serialize dumps us as a byte array so we need padding to be initialized to avoid indeterminism in the cooked build
+	uint8 UnserializedPaddingBytes[3] = {0};
 };
 
+// Make sure padding size matches our expectation
+static_assert(sizeof(FMovieSceneTangentData) == 
+	sizeof(FMovieSceneTangentData::ArriveTangent) + sizeof(FMovieSceneTangentData::LeaveTangent) + sizeof(FMovieSceneTangentData::ArriveTangentWeight) + 
+	sizeof(FMovieSceneTangentData::LeaveTangentWeight) + sizeof(FMovieSceneTangentData::TangentWeightMode) + sizeof(FMovieSceneTangentData::UnserializedPaddingBytes),
+	"Adjust padding size to avoid cooked build indeterminism with uninitialized padded data");
+
+
 template<>
 struct TIsPODType<FMovieSceneTangentData>
 {
diff --git a/Engine/Source/Runtime/MovieScene/Public/Channels/MovieSceneDoubleChannel.h b/Engine/Source/Runtime/MovieScene/Public/Channels/MovieSceneDoubleChannel.h
--- a/Engine/Source/Runtime/MovieScene/Public/Channels/MovieSceneDoubleChannel.h
+++ b/Engine/Source/Runtime/MovieScene/Public/Channels/MovieSceneDoubleChannel.h
@@ -82,8 +82,14 @@ struct FMovieSceneDoubleValue
 	 */
 	UPROPERTY()
 	uint8 PaddingByte;
-};
 
+	// This is required because TMovieSceneCurveChannelImpl<ChannelType>::Serialize dumps us as a byte array so we need padding to be initialized to avoid indeterminism in the cooked build
+	uint8 UnserializedPaddingBytes[1] = { 0 };
+};
+static_assert(
+	sizeof(FMovieSceneDoubleValue) == 
+	sizeof(FMovieSceneDoubleValue::Value) + sizeof(FMovieSceneDoubleValue::Tangent) + sizeof(FMovieSceneDoubleValue::InterpMode) + sizeof(FMovieSceneDoubleValue::TangentMode) + sizeof(FMovieSceneDoubleValue::PaddingByte) + sizeof(FMovieSceneDoubleValue::UnserializedPaddingBytes),
+	"Adjust padding size to avoid cooked build indeterminism with uninitialized padded data");
 
 template<>
 struct TIsPODType<FMovieSceneDoubleValue>
@@ -116,7 +122,7 @@ struct MOVIESCENE_API FMovieSceneDoubleChannel : public FMovieSceneChannel
 	FMovieSceneDoubleChannel() 
 		: PreInfinityExtrap(RCCE_Constant)
 		, PostInfinityExtrap(RCCE_Constant)
-		, DefaultValue()
+		, DefaultValue(0.0)
 		, bHasDefaultValue(false)
 #if WITH_EDITORONLY_DATA
 		, bShowCurve(false)
diff --git a/Engine/Source/Runtime/MovieScene/Public/Channels/MovieSceneFloatChannel.h b/Engine/Source/Runtime/MovieScene/Public/Channels/MovieSceneFloatChannel.h
--- a/Engine/Source/Runtime/MovieScene/Public/Channels/MovieSceneFloatChannel.h
+++ b/Engine/Source/Runtime/MovieScene/Public/Channels/MovieSceneFloatChannel.h
@@ -81,8 +81,14 @@ struct FMovieSceneFloatValue
 	 */
 	UPROPERTY()
 	uint8 PaddingByte;
-};
 
+	// This is required because TMovieSceneCurveChannelImpl<ChannelType>::Serialize dumps us as a byte array so we need padding to be initialized to avoid indeterminism in the cooked build
+	uint8 UnserializedPaddingBytes[1] = {0};
+};
+static_assert(
+	sizeof(FMovieSceneFloatValue) ==
+	sizeof(FMovieSceneFloatValue::Value) + sizeof(FMovieSceneFloatValue::Tangent) + sizeof(FMovieSceneFloatValue::InterpMode) + sizeof(FMovieSceneFloatValue::TangentMode) + sizeof(FMovieSceneFloatValue::PaddingByte) + sizeof(FMovieSceneFloatValue::UnserializedPaddingBytes),
+	"Adjust padding size to avoid cooked build indeterminism with uninitialized padded data");
 
 template<>
 struct TIsPODType<FMovieSceneFloatValue>
@@ -115,7 +121,7 @@ struct MOVIESCENE_API FMovieSceneFloatChannel : public FMovieSceneChannel
 	FMovieSceneFloatChannel() 
 		: PreInfinityExtrap(RCCE_Constant)
 		, PostInfinityExtrap(RCCE_Constant)
-		, DefaultValue()
+		, DefaultValue(0.f)
 		, bHasDefaultValue(false)
 #if WITH_EDITORONLY_DATA
 		, bShowCurve(false)
diff --git a/Engine/Source/Runtime/MovieScene/Public/Channels/MovieSceneIntegerChannel.h b/Engine/Source/Runtime/MovieScene/Public/Channels/MovieSceneIntegerChannel.h
--- a/Engine/Source/Runtime/MovieScene/Public/Channels/MovieSceneIntegerChannel.h
+++ b/Engine/Source/Runtime/MovieScene/Public/Channels/MovieSceneIntegerChannel.h
@@ -33,7 +33,7 @@ struct MOVIESCENE_API FMovieSceneIntegerChannel : public FMovieSceneChannel
 	GENERATED_BODY()
 
 	FMovieSceneIntegerChannel()
-		: DefaultValue(), bHasDefaultValue(false)
+		: DefaultValue(0), bHasDefaultValue(false)
 	{}
 
 	/**
