Removed TIsSame usage in MovieScene (part of a larger change to deprecate TIsSame).

diff --git a/Engine/Source/Runtime/MovieScene/Private/Channels/MovieSceneCurveChannelImpl.cpp b/Engine/Source/Runtime/MovieScene/Private/Channels/MovieSceneCurveChannelImpl.cpp
--- a/Engine/Source/Runtime/MovieScene/Private/Channels/MovieSceneCurveChannelImpl.cpp
+++ b/Engine/Source/Runtime/MovieScene/Private/Channels/MovieSceneCurveChannelImpl.cpp
@@ -1445,7 +1445,7 @@ bool TMovieSceneCurveChannelImpl<ChannelType>::SerializeChannelValue(ChannelValu
 		return false;
 	}
 
-	if constexpr(TIsSame<CurveValueType, double>::Value)
+	if constexpr(std::is_same_v<CurveValueType, double>)
 	{
 		if(Ar.UEVer() >= EUnrealEngineObjectUE5Version::LARGE_WORLD_COORDINATES)
 		{
diff --git a/Engine/Source/Runtime/MovieScene/Public/Channels/MovieSceneChannelEditorDataEntry.h b/Engine/Source/Runtime/MovieScene/Public/Channels/MovieSceneChannelEditorDataEntry.h
--- a/Engine/Source/Runtime/MovieScene/Public/Channels/MovieSceneChannelEditorDataEntry.h
+++ b/Engine/Source/Runtime/MovieScene/Public/Channels/MovieSceneChannelEditorDataEntry.h
@@ -56,7 +56,7 @@ protected:
 	template<typename ChannelType>
 	void AddMetaData(const FMovieSceneChannelMetaData& MetaData)
 	{
-		static_assert(TIsSame<typename TMovieSceneChannelTraits<ChannelType>::ExtendedEditorDataType, void>::Value, "Must supply extended editor data according to the channel's traits.");
+		static_assert(std::is_same_v<typename TMovieSceneChannelTraits<ChannelType>::ExtendedEditorDataType, void>, "Must supply extended editor data according to the channel's traits.");
 
 		// Add the editor meta-data
 		MetaDataArray.Add(MetaData);
@@ -93,7 +93,7 @@ protected:
 	template<typename ChannelType>
 	TArrayView<const typename TMovieSceneChannelTraits<ChannelType>::ExtendedEditorDataType> GetAllExtendedEditorData() const
 	{
-		static_assert(!TIsSame<typename TMovieSceneChannelTraits<ChannelType>::ExtendedEditorDataType, void>::Value, "This channel type does not define any extended editor data.");
+		static_assert(!std::is_same_v<typename TMovieSceneChannelTraits<ChannelType>::ExtendedEditorDataType, void>, "This channel type does not define any extended editor data.");
 
 		const auto& TypedImpl = static_cast<const TMovieSceneExtendedEditorDataArray<ChannelType>&>(ExtendedEditorDataArray.GetValue());
 		return TypedImpl.Data;
diff --git a/Engine/Source/Runtime/MovieScene/Public/Channels/MovieSceneChannelProxy.h b/Engine/Source/Runtime/MovieScene/Public/Channels/MovieSceneChannelProxy.h
--- a/Engine/Source/Runtime/MovieScene/Public/Channels/MovieSceneChannelProxy.h
+++ b/Engine/Source/Runtime/MovieScene/Public/Channels/MovieSceneChannelProxy.h
@@ -96,8 +96,8 @@ struct FMovieSceneChannelProxyData
 	template<typename ChannelType>
 	void Add(ChannelType& InChannel, const FMovieSceneChannelMetaData& InMetaData)
 	{
-		static_assert(TIsSame<typename TMovieSceneChannelTraits<ChannelType>::ExtendedEditorDataType, void>::Value, "Must supply extended editor data according to the channel's traits.");
-		static_assert(!TIsSame<ChannelType, FMovieSceneChannel>::Value, "Cannot add channels by their base FMovieSceneChannel type.");
+		static_assert(std::is_same_v<typename TMovieSceneChannelTraits<ChannelType>::ExtendedEditorDataType, void>, "Must supply extended editor data according to the channel's traits.");
+		static_assert(!std::is_same_v<ChannelType, FMovieSceneChannel>, "Cannot add channels by their base FMovieSceneChannel type.");
 
 		// Add the channel
 		const int32 ChannelTypeIndex = AddInternal(InChannel);
@@ -114,8 +114,8 @@ struct FMovieSceneChannelProxyData
 	template<typename ChannelType>
 	FMovieSceneChannelHandle AddWithDefaultEditorData(ChannelType& InChannel, const FMovieSceneChannelMetaData& InMetaData)
 	{
-		static_assert(!TIsSame<typename TMovieSceneChannelTraits<ChannelType>::ExtendedEditorDataType, void>::Value, "This method is for channels with typed editor data. You *must* call SetExtendedEditorData afterwards.");
-		static_assert(!TIsSame<ChannelType, FMovieSceneChannel>::Value, "Cannot add channels by their base FMovieSceneChannel type.");
+		static_assert(!std::is_same_v<typename TMovieSceneChannelTraits<ChannelType>::ExtendedEditorDataType, void>, "This method is for channels with typed editor data. You *must* call SetExtendedEditorData afterwards.");
+		static_assert(!std::is_same_v<ChannelType, FMovieSceneChannel>, "Cannot add channels by their base FMovieSceneChannel type.");
 
 		// Add the channel
 		const int32 ChannelTypeIndex = AddInternal(InChannel);
@@ -136,8 +136,8 @@ struct FMovieSceneChannelProxyData
 	template<typename ChannelType, typename ExtendedEditorDataType>
 	void Add(ChannelType& InChannel, const FMovieSceneChannelMetaData& InMetaData, ExtendedEditorDataType&& InExtendedEditorData)
 	{
-		static_assert(!TIsSame<typename TMovieSceneChannelTraits<ChannelType>::ExtendedEditorDataType, void>::Value, "Must supply typed editor data according to the channel's traits. Define TMovieSceneChannelTraits<ChannelType>::ExtendedEditorDataType to use this function.");
-		static_assert(!TIsSame<ChannelType, FMovieSceneChannel>::Value, "Cannot add channels by their base FMovieSceneChannel type.");
+		static_assert(!std::is_same_v<typename TMovieSceneChannelTraits<ChannelType>::ExtendedEditorDataType, void>, "Must supply typed editor data according to the channel's traits. Define TMovieSceneChannelTraits<ChannelType>::ExtendedEditorDataType to use this function.");
+		static_assert(!std::is_same_v<ChannelType, FMovieSceneChannel>, "Cannot add channels by their base FMovieSceneChannel type.");
 
 		// Add the channel
 		const int32 ChannelTypeIndex = AddInternal(InChannel);
@@ -448,7 +448,7 @@ ChannelType* FMovieSceneChannelProxy::GetChannel(int32 ChannelIndex) const
 template<typename ChannelType>
 FMovieSceneChannelProxy::FMovieSceneChannelProxy(ChannelType& InChannel)
 {
-	static_assert(!TIsSame<ChannelType, FMovieSceneChannel>::Value, "Cannot add channels by their base FMovieSceneChannel type..");
+	static_assert(!std::is_same_v<ChannelType, FMovieSceneChannel>, "Cannot add channels by their base FMovieSceneChannel type..");
 
 	const FName ChannelTypeName = ChannelType::StaticStruct()->GetFName();
 	Entries.Add(FMovieSceneChannelEntry(ChannelTypeName, InChannel));
@@ -466,8 +466,8 @@ FMovieSceneChannelProxy::FMovieSceneChannelProxy(ChannelType& InChannel)
 template<typename ChannelType>
 FMovieSceneChannelProxy::FMovieSceneChannelProxy(ChannelType& InChannel, const FMovieSceneChannelMetaData& InMetaData)
 {
-	static_assert(!TIsSame<ChannelType, FMovieSceneChannel>::Value, "Cannot add channels by their base FMovieSceneChannel type..");
-	static_assert(TIsSame<typename TMovieSceneChannelTraits<ChannelType>::ExtendedEditorDataType, void>::Value, "Must supply typed editor data according to the channel's traits.");
+	static_assert(!std::is_same_v<ChannelType, FMovieSceneChannel>, "Cannot add channels by their base FMovieSceneChannel type..");
+	static_assert(std::is_same_v<typename TMovieSceneChannelTraits<ChannelType>::ExtendedEditorDataType, void>, "Must supply typed editor data according to the channel's traits.");
 
 	const FName ChannelTypeName = ChannelType::StaticStruct()->GetFName();
 	Entries.Add(FMovieSceneChannelEntry(ChannelTypeName, InChannel));
@@ -483,8 +483,8 @@ FMovieSceneChannelProxy::FMovieSceneChannelProxy(ChannelType& InChannel, const F
 template<typename ChannelType, typename ExtendedEditorDataType>
 FMovieSceneChannelProxy::FMovieSceneChannelProxy(ChannelType& InChannel, const FMovieSceneChannelMetaData& InMetaData, ExtendedEditorDataType&& InExtendedEditorDataType)
 {
-	static_assert(!TIsSame<ChannelType, FMovieSceneChannel>::Value, "Cannot add channels by their base FMovieSceneChannel type..");
-	static_assert(!TIsSame<typename TMovieSceneChannelTraits<ChannelType>::ExtendedEditorDataType, void>::Value, "Must supply typed editor data according to the channel's traits. Define TMovieSceneChannelTraits<ChannelType>::ExtendedEditorDataType to use this function.");
+	static_assert(!std::is_same_v<ChannelType, FMovieSceneChannel>, "Cannot add channels by their base FMovieSceneChannel type..");
+	static_assert(!std::is_same_v<typename TMovieSceneChannelTraits<ChannelType>::ExtendedEditorDataType, void>, "Must supply typed editor data according to the channel's traits. Define TMovieSceneChannelTraits<ChannelType>::ExtendedEditorDataType to use this function.");
 
 	const FName ChannelTypeName = ChannelType::StaticStruct()->GetFName();
 	Entries.Add(FMovieSceneChannelEntry(ChannelTypeName, InChannel));
diff --git a/Engine/Source/Runtime/MovieScene/Public/EntitySystem/MovieSceneComponentTypeInfo.h b/Engine/Source/Runtime/MovieScene/Public/EntitySystem/MovieSceneComponentTypeInfo.h
--- a/Engine/Source/Runtime/MovieScene/Public/EntitySystem/MovieSceneComponentTypeInfo.h
+++ b/Engine/Source/Runtime/MovieScene/Public/EntitySystem/MovieSceneComponentTypeInfo.h
@@ -32,7 +32,7 @@ typename TEnableIf<TPointerIsConvertibleFromTo<T, UObject>::Value>::Type AddRefe
 }
 
 template<typename T>
-typename TEnableIf< TIsSame<decltype(T::StaticStruct()), decltype(T::StaticStruct())>::Value >::Type AddReferencedObjectForComponent(FReferenceCollector* ReferenceCollector, T* Component)
+typename TEnableIf< std::is_same_v<decltype(T::StaticStruct()), decltype(T::StaticStruct())> >::Type AddReferencedObjectForComponent(FReferenceCollector* ReferenceCollector, T* Component)
 {
 	for (TPropertyValueIterator<const FObjectProperty> It(T::StaticStruct(), Component); It; ++It)
 	{
@@ -43,7 +43,7 @@ typename TEnableIf< TIsSame<decltype(T::StaticStruct()), decltype(T::StaticStruc
 template<typename T>
 struct THasAddReferencedObjectForComponent
 {
-	static constexpr bool Value = !TIsSame< FNotImplemented*, decltype( AddReferencedObjectForComponent((FReferenceCollector*)0, (T*)0) ) >::Value;
+	static constexpr bool Value = !std::is_same_v< FNotImplemented*, decltype( AddReferencedObjectForComponent((FReferenceCollector*)0, (T*)0) ) >;
 };
 
 /**
diff --git a/Engine/Source/Runtime/MovieScene/Public/EntitySystem/MovieSceneEntitySystemTask.h b/Engine/Source/Runtime/MovieScene/Public/EntitySystem/MovieSceneEntitySystemTask.h
--- a/Engine/Source/Runtime/MovieScene/Public/EntitySystem/MovieSceneEntitySystemTask.h
+++ b/Engine/Source/Runtime/MovieScene/Public/EntitySystem/MovieSceneEntitySystemTask.h
@@ -1481,7 +1481,7 @@ struct TEntityTaskCaller<TaskImpl, NumComponents, false>
 {
 	FORCEINLINE static void ForEachEntityImpl(...)
 	{
-		static_assert(TNot<TIsSame<TaskImpl, TaskImpl>>::Value, "non-expanded entity iteration is not supported");
+		static_assert(!std::is_same_v<TaskImpl, TaskImpl>, "non-expanded entity iteration is not supported");
 	}
 
 	template<typename... AccessorTypes>
diff --git a/Engine/Source/Runtime/MovieScene/Public/EntitySystem/MovieScenePropertyComponentHandler.h b/Engine/Source/Runtime/MovieScene/Public/EntitySystem/MovieScenePropertyComponentHandler.h
--- a/Engine/Source/Runtime/MovieScene/Public/EntitySystem/MovieScenePropertyComponentHandler.h
+++ b/Engine/Source/Runtime/MovieScene/Public/EntitySystem/MovieScenePropertyComponentHandler.h
@@ -559,9 +559,9 @@ struct TPropertyDefinitionBuilder
 
 		Definition->CompositeSize = 1;
 
-		static_assert(!TIsSame<typename PropertyTraits::StorageType, float>::Value, "Please use double-precision composites");
+		static_assert(!std::is_same_v<typename PropertyTraits::StorageType, float>, "Please use double-precision composites");
 
-		if (TIsSame<typename PropertyTraits::StorageType, double>::Value)
+		if constexpr (std::is_same_v<typename PropertyTraits::StorageType, double>)
 		{
 			Definition->DoubleCompositeMask = 1;
 		}
@@ -639,9 +639,9 @@ struct TCompositePropertyDefinitionBuilder
 		FPropertyCompositeDefinition NewChannel = { InComponent, static_cast<uint16>(CompositeOffset) };
 		Registry->CompositeDefinitions.Add(NewChannel);
 
-		static_assert(!TIsSame<T, float>::Value, "Please use double-precision composites");
+		static_assert(!std::is_same_v<T, float>, "Please use double-precision composites");
 
-		if (TIsSame<T, double>::Value)
+		if constexpr (std::is_same_v<T, double>)
 		{
 			Definition->DoubleCompositeMask |= 1 << Definition->CompositeSize;
 		}
diff --git a/Engine/Source/Runtime/MovieScene/Public/EntitySystem/MovieScenePropertySystemTypes.h b/Engine/Source/Runtime/MovieScene/Public/EntitySystem/MovieScenePropertySystemTypes.h
--- a/Engine/Source/Runtime/MovieScene/Public/EntitySystem/MovieScenePropertySystemTypes.h
+++ b/Engine/Source/Runtime/MovieScene/Public/EntitySystem/MovieScenePropertySystemTypes.h
@@ -93,7 +93,7 @@ namespace Private
 			using StorageType = bool;
 		};
 		static_assert(PropertyTraitsHaveCustomAccessorStorageType<TestNormal> == false, "Normal has no custom storage type");
-		static_assert(TIsSame<TCustomPropertyAccessorStorageType<TestNormal>::Value, bool>::Value, "Normal has bool storage type");
+		static_assert(std::is_same_v<TCustomPropertyAccessorStorageType<TestNormal>::Value, bool>, "Normal has bool storage type");
 
 		struct TestCustom
 		{
@@ -101,7 +101,7 @@ namespace Private
 			using CustomAccessorStorageType = int;
 		};
 		static_assert(PropertyTraitsHaveCustomAccessorStorageType<TestCustom> == true, "Custom does have custom storage type");
-		static_assert(TIsSame<TCustomPropertyAccessorStorageType<TestCustom>::Value, int>::Value, "Custom has int storage type");
+		static_assert(std::is_same_v<TCustomPropertyAccessorStorageType<TestCustom>::Value, int>, "Custom has int storage type");
 	}
 }
 
diff --git a/Engine/Source/Runtime/MovieScene/Public/Evaluation/Blending/BlendableToken.h b/Engine/Source/Runtime/MovieScene/Public/Evaluation/Blending/BlendableToken.h
--- a/Engine/Source/Runtime/MovieScene/Public/Evaluation/Blending/BlendableToken.h
+++ b/Engine/Source/Runtime/MovieScene/Public/Evaluation/Blending/BlendableToken.h
@@ -80,7 +80,7 @@ namespace MovieScene
 	void BlendValue(WorkingDataType& OutBlend, InType InValue, float Weight, EMovieSceneBlendType BlendType, TMovieSceneInitialValueStore<SourceDataType>& InitialValueStore)
 	{
 		// Always assert on instantiation
-		static_assert(TIsSame<WorkingDataType, void>::Value, "BlendValue must be implemented for the specified types in order to blend them with Sequencer.");
+		static_assert(std::is_same_v<WorkingDataType, void>, "BlendValue must be implemented for the specified types in order to blend them with Sequencer.");
 	}
 } // namespace MovieScene
 } // namespace UE
diff --git a/Engine/Source/Runtime/MovieScene/Public/Evaluation/Blending/BlendableTokenStack.h b/Engine/Source/Runtime/MovieScene/Public/Evaluation/Blending/BlendableTokenStack.h
--- a/Engine/Source/Runtime/MovieScene/Public/Evaluation/Blending/BlendableTokenStack.h
+++ b/Engine/Source/Runtime/MovieScene/Public/Evaluation/Blending/BlendableTokenStack.h
@@ -25,7 +25,7 @@ template<typename T> struct TMovieSceneBlendingActuator;
 template<typename T> FMovieSceneAnimTypeID GetBlendingDataType()
 {
 	// Always assert on instantiation
-	static_assert(TIsSame<T, void>::Value, "GetBlendingDataType must be specialized for a type in order to use it with an accumulator.");
+	static_assert(std::is_same_v<T, void>, "GetBlendingDataType must be specialized for a type in order to use it with an accumulator.");
 	return FMovieSceneAnimTypeID::Unique();
 }
 
diff --git a/Engine/Source/Runtime/MovieScene/Public/Evaluation/Blending/MovieSceneMultiChannelBlending.h b/Engine/Source/Runtime/MovieScene/Public/Evaluation/Blending/MovieSceneMultiChannelBlending.h
--- a/Engine/Source/Runtime/MovieScene/Public/Evaluation/Blending/MovieSceneMultiChannelBlending.h
+++ b/Engine/Source/Runtime/MovieScene/Public/Evaluation/Blending/MovieSceneMultiChannelBlending.h
@@ -151,7 +151,7 @@ namespace MovieScene
 	template<typename T, typename SourceData, uint8 N>
 	void MultiChannelFromData(SourceData InSourceData, TMultiChannelValue<T, N>& OutChannelData)
 	{
-		static_assert(TIsSame<T, void>::Value, "MultiChannelFromData must be implemented to blend SourceData with multi-channel data.");
+		static_assert(std::is_same_v<T, void>, "MultiChannelFromData must be implemented to blend SourceData with multi-channel data.");
 	}
 
 	/**
@@ -161,7 +161,7 @@ namespace MovieScene
 	template<typename T, typename SourceData, uint8 N>
 	void ResolveChannelsToData(const TMultiChannelValue<T, N>& OutChannelData, SourceData& OutData)
 	{
-		static_assert(TIsSame<T, void>::Value, "ResolveChannelsToData must be implemented to blend SourceData with multi-channel data.");
+		static_assert(std::is_same_v<T, void>, "ResolveChannelsToData must be implemented to blend SourceData with multi-channel data.");
 	}
 
 	/** Working data type used to blend multi-channel values */
diff --git a/Engine/Source/Runtime/MovieScene/Public/Evaluation/MovieSceneEvalTemplate.h b/Engine/Source/Runtime/MovieScene/Public/Evaluation/MovieSceneEvalTemplate.h
--- a/Engine/Source/Runtime/MovieScene/Public/Evaluation/MovieSceneEvalTemplate.h
+++ b/Engine/Source/Runtime/MovieScene/Public/Evaluation/MovieSceneEvalTemplate.h
@@ -225,7 +225,7 @@ struct FMovieSceneEvalTemplatePtr
 	FMovieSceneEvalTemplatePtr(T&& In)
 		: TInlineValue(Forward<T>(In))
 	{
-		static_assert(!TIsSame<typename TDecay<T>::Type, FMovieSceneEvalTemplate>::Value, "Direct usage of FMovieSceneEvalTemplate is prohibited.");
+		static_assert(!std::is_same_v<typename TDecay<T>::Type, FMovieSceneEvalTemplate>, "Direct usage of FMovieSceneEvalTemplate is prohibited.");
 		
 #if WITH_EDITOR
 		checkf(T::StaticStruct() == &In.GetScriptStruct() && T::StaticStruct() != FMovieSceneEvalTemplate::StaticStruct(), TEXT("%s does not correctly override GetScriptStructImpl. Template will not serialize correctly."), *T::StaticStruct()->GetName());
diff --git a/Engine/Source/Runtime/MovieScene/Public/Evaluation/MovieSceneSequenceInstanceData.h b/Engine/Source/Runtime/MovieScene/Public/Evaluation/MovieSceneSequenceInstanceData.h
--- a/Engine/Source/Runtime/MovieScene/Public/Evaluation/MovieSceneSequenceInstanceData.h
+++ b/Engine/Source/Runtime/MovieScene/Public/Evaluation/MovieSceneSequenceInstanceData.h
@@ -79,7 +79,7 @@ struct MOVIESCENE_API FMovieSceneSequenceInstanceDataPtr
 	{
 		typedef typename TDecay<T>::Type ClientType;
 
-		static_assert(!TIsSame<ClientType, FMovieSceneSequenceInstanceData>::Value, "Direct usage of FMovieSceneSequenceInstanceData is prohibited.");
+		static_assert(!std::is_same_v<ClientType, FMovieSceneSequenceInstanceData>, "Direct usage of FMovieSceneSequenceInstanceData is prohibited.");
 
 #if WITH_EDITOR && DO_CHECK
 		const UStruct* ClientStruct = ClientType::StaticStruct();
diff --git a/Engine/Source/Runtime/MovieScene/Public/Evaluation/MovieSceneTrackImplementation.h b/Engine/Source/Runtime/MovieScene/Public/Evaluation/MovieSceneTrackImplementation.h
--- a/Engine/Source/Runtime/MovieScene/Public/Evaluation/MovieSceneTrackImplementation.h
+++ b/Engine/Source/Runtime/MovieScene/Public/Evaluation/MovieSceneTrackImplementation.h
@@ -141,7 +141,7 @@ struct FMovieSceneTrackImplementationPtr
 	FMovieSceneTrackImplementationPtr(T&& In)
 		: TInlineValue(Forward<T>(In))
 	{
-		static_assert(!TIsSame<typename TDecay<T>::Type, FMovieSceneTrackImplementation>::Value, "Direct usage of FMovieSceneTrackImplementation is prohibited.");
+		static_assert(!std::is_same_v<typename TDecay<T>::Type, FMovieSceneTrackImplementation>, "Direct usage of FMovieSceneTrackImplementation is prohibited.");
 
 #if WITH_EDITOR
 		checkf(T::StaticStruct() == &In.GetScriptStruct() && T::StaticStruct() != FMovieSceneTrackImplementation::StaticStruct(), TEXT("%s type does not correctly override GetScriptStructImpl. Track will not serialize correctly."), *T::StaticStruct()->GetName());
diff --git a/Engine/Source/Runtime/MovieScene/Public/MovieSceneClipboard.h b/Engine/Source/Runtime/MovieScene/Public/MovieSceneClipboard.h
--- a/Engine/Source/Runtime/MovieScene/Public/MovieSceneClipboard.h
+++ b/Engine/Source/Runtime/MovieScene/Public/MovieSceneClipboard.h
@@ -30,7 +30,7 @@ namespace MovieSceneClipboard
 	template<typename T> static FName GetKeyTypeName()
 	{
 		// Use the template type here to stop the compiler generating the code unless it's actually used
-		static_assert(!TIsSame<T, T>::Value, "This function must be specialized to use with the specified type");
+		static_assert(!std::is_same_v<T, T>, "This function must be specialized to use with the specified type");
 		return NAME_None;
 	}
 }
