Sequencer: Added more component debug types

diff --git a/Engine/Source/Runtime/MovieScene/Public/EntitySystem/MovieSceneBlenderSystemTypes.h b/Engine/Source/Runtime/MovieScene/Public/EntitySystem/MovieSceneBlenderSystemTypes.h
--- a/Engine/Source/Runtime/MovieScene/Public/EntitySystem/MovieSceneBlenderSystemTypes.h
+++ b/Engine/Source/Runtime/MovieScene/Public/EntitySystem/MovieSceneBlenderSystemTypes.h
@@ -8,6 +8,7 @@
 #include "Templates/Casts.h"
 #include "Templates/SubclassOf.h"
 #include "Templates/TypeHash.h"
+#include "EntitySystem/MovieSceneComponentDebug.h"
 
 class UMovieSceneBlenderSystem;
 class UMovieSceneEntitySystemLinker;
@@ -89,3 +90,9 @@ struct FMovieSceneBlendChannelID
 	uint16 ChannelID;
 };
 
+#if UE_MOVIESCENE_ENTITY_DEBUG
+namespace UE::MovieScene
+{
+	template<> struct TComponentDebugType<FMovieSceneBlendChannelID> { static const EComponentDebugType Type = EComponentDebugType::BlendChannelID; };
+}
+#endif
diff --git a/Engine/Source/Runtime/MovieScene/Public/EntitySystem/MovieSceneComponentDebug.h b/Engine/Source/Runtime/MovieScene/Public/EntitySystem/MovieSceneComponentDebug.h
--- a/Engine/Source/Runtime/MovieScene/Public/EntitySystem/MovieSceneComponentDebug.h
+++ b/Engine/Source/Runtime/MovieScene/Public/EntitySystem/MovieSceneComponentDebug.h
@@ -4,6 +4,8 @@
 
 #include "EntitySystem/MovieSceneEntitySystemTypes.h"
 #include "Math/Vector4.h"
+#include "Misc/FrameTime.h"
+#include "UObject/ObjectKey.h"
 
 #if UE_MOVIESCENE_ENTITY_DEBUG
 
@@ -20,6 +22,7 @@ enum class EComponentDebugType
 	Unknown,
 	Bool,
 	Uint8,
+	int16,
 	Uint16,
 	Int32,
 	Float,
@@ -27,11 +30,15 @@ enum class EComponentDebugType
 	Vector2,
 	Vector3,
 	Vector4,
+	FrameTime,
+	ObjectKey,
 	Object,
 	Property,
 	InstanceHandle,
+	SequenceID,
 	EntityID,
 	Name,
+	BlendChannelID,
 };
 
 /**
@@ -44,19 +51,22 @@ struct FComponentTypeDebugInfo
 	EComponentDebugType Type = EComponentDebugType::Unknown;
 };
 
-template<typename T> struct TComponentDebugType                      { static const EComponentDebugType Type = EComponentDebugType::Unknown;  };
-template<>           struct TComponentDebugType<bool>                { static const EComponentDebugType Type = EComponentDebugType::Bool;     };
-template<>           struct TComponentDebugType<uint8>               { static const EComponentDebugType Type = EComponentDebugType::Uint8;    };
-template<>           struct TComponentDebugType<uint16>              { static const EComponentDebugType Type = EComponentDebugType::Uint16;   };
-template<>           struct TComponentDebugType<int32>               { static const EComponentDebugType Type = EComponentDebugType::Int32;    };
-template<>           struct TComponentDebugType<float>               { static const EComponentDebugType Type = EComponentDebugType::Float;    };
-template<>           struct TComponentDebugType<double>              { static const EComponentDebugType Type = EComponentDebugType::Double;    };
-template<>           struct TComponentDebugType<FVector2D>           { static const EComponentDebugType Type = EComponentDebugType::Vector2;  };
-template<>           struct TComponentDebugType<FVector>             { static const EComponentDebugType Type = EComponentDebugType::Vector3;  };
-template<>           struct TComponentDebugType<FVector4>            { static const EComponentDebugType Type = EComponentDebugType::Vector4;  };
-template<>           struct TComponentDebugType<UObject*>            { static const EComponentDebugType Type = EComponentDebugType::Object;   };
-template<>           struct TComponentDebugType<FMovieSceneEntityID> { static const EComponentDebugType Type = EComponentDebugType::EntityID; };
-template<>           struct TComponentDebugType<FName>               { static const EComponentDebugType Type = EComponentDebugType::Name;     };
+template<typename T> struct TComponentDebugType                            { static const EComponentDebugType Type = EComponentDebugType::Unknown;           };
+template<>           struct TComponentDebugType<bool>                      { static const EComponentDebugType Type = EComponentDebugType::Bool;              };
+template<>           struct TComponentDebugType<uint8>                     { static const EComponentDebugType Type = EComponentDebugType::Uint8;             };
+template<>           struct TComponentDebugType<int16>                     { static const EComponentDebugType Type = EComponentDebugType::int16;             };
+template<>           struct TComponentDebugType<uint16>                    { static const EComponentDebugType Type = EComponentDebugType::Uint16;            };
+template<>           struct TComponentDebugType<int32>                     { static const EComponentDebugType Type = EComponentDebugType::Int32;             };
+template<>           struct TComponentDebugType<float>                     { static const EComponentDebugType Type = EComponentDebugType::Float;             };
+template<>           struct TComponentDebugType<double>                    { static const EComponentDebugType Type = EComponentDebugType::Double;            };
+template<>           struct TComponentDebugType<FVector2D>                 { static const EComponentDebugType Type = EComponentDebugType::Vector2;           };
+template<>           struct TComponentDebugType<FVector>                   { static const EComponentDebugType Type = EComponentDebugType::Vector3;           };
+template<>           struct TComponentDebugType<FVector4>                  { static const EComponentDebugType Type = EComponentDebugType::Vector4;           };
+template<>           struct TComponentDebugType<FFrameTime>                { static const EComponentDebugType Type = EComponentDebugType::FrameTime;         };
+template<>           struct TComponentDebugType<FObjectKey>                { static const EComponentDebugType Type = EComponentDebugType::ObjectKey;         };
+template<>           struct TComponentDebugType<UObject*>                  { static const EComponentDebugType Type = EComponentDebugType::Object;            };
+template<>           struct TComponentDebugType<FMovieSceneEntityID>       { static const EComponentDebugType Type = EComponentDebugType::EntityID;          };
+template<>           struct TComponentDebugType<FName>                     { static const EComponentDebugType Type = EComponentDebugType::Name;              };
 
 } // namespace MovieScene
 } // namespace UE
diff --git a/Engine/Source/Runtime/MovieScene/Public/EntitySystem/MovieSceneSequenceInstanceHandle.h b/Engine/Source/Runtime/MovieScene/Public/EntitySystem/MovieSceneSequenceInstanceHandle.h
--- a/Engine/Source/Runtime/MovieScene/Public/EntitySystem/MovieSceneSequenceInstanceHandle.h
+++ b/Engine/Source/Runtime/MovieScene/Public/EntitySystem/MovieSceneSequenceInstanceHandle.h
@@ -65,7 +65,8 @@ struct FRootInstanceHandle : FInstanceHandle
 };
 
 #if UE_MOVIESCENE_ENTITY_DEBUG
-template<> struct TComponentDebugType<FInstanceHandle> { static const EComponentDebugType Type = EComponentDebugType::InstanceHandle; };
+template<> struct TComponentDebugType<FInstanceHandle>     { static const EComponentDebugType Type = EComponentDebugType::InstanceHandle; };
+template<> struct TComponentDebugType<FRootInstanceHandle> { static const EComponentDebugType Type = EComponentDebugType::InstanceHandle; };
 #endif
 
 
diff --git a/Engine/Source/Runtime/MovieScene/Public/MovieSceneSequenceID.h b/Engine/Source/Runtime/MovieScene/Public/MovieSceneSequenceID.h
--- a/Engine/Source/Runtime/MovieScene/Public/MovieSceneSequenceID.h
+++ b/Engine/Source/Runtime/MovieScene/Public/MovieSceneSequenceID.h
@@ -5,6 +5,7 @@
 #include "CoreMinimal.h"
 #include "UObject/ObjectMacros.h"
 #include "UObject/Class.h"
+#include "EntitySystem/MovieSceneComponentDebug.h"
 #include "MovieSceneSequenceID.generated.h"
 
 USTRUCT()
@@ -88,6 +89,13 @@ struct TStructOpsTypeTraits<FMovieSceneSequenceID> : public TStructOpsTypeTraits
 	};
 };
 
+#if UE_MOVIESCENE_ENTITY_DEBUG
+namespace UE::MovieScene
+{
+	template<> struct TComponentDebugType<FMovieSceneSequenceID> { static const EComponentDebugType Type = EComponentDebugType::SequenceID; };
+}
+#endif
+
 typedef TCallTraits<FMovieSceneSequenceID>::ParamType FMovieSceneSequenceIDRef;
 
 namespace MovieSceneSequenceID
