Sequencer: Fixed erroneous float <> double conversion and comparison and re-enabled cached channel evaluation by default

diff --git a/Engine/Source/Runtime/MovieScene/Private/Channels/MovieSceneCurveChannelImpl.cpp b/Engine/Source/Runtime/MovieScene/Private/Channels/MovieSceneCurveChannelImpl.cpp
--- a/Engine/Source/Runtime/MovieScene/Private/Channels/MovieSceneCurveChannelImpl.cpp
+++ b/Engine/Source/Runtime/MovieScene/Private/Channels/MovieSceneCurveChannelImpl.cpp
@@ -36,7 +36,7 @@ static FAutoConsoleVariableRef CVarCachedChannelEvaluationParityThreshold(
 	TEXT("Threshold for testing Evaluate parity with cached/uncached routines."),
 	ECVF_Default);
 
-MOVIESCENE_API bool GEnableCachedChannelEvaluation = false;
+MOVIESCENE_API bool GEnableCachedChannelEvaluation = true;
 static FAutoConsoleVariableRef CVarEnableCachedChannelEvaluation(
 	TEXT("Sequencer.EnableCachedChannelEvaluation"),
 	GEnableCachedChannelEvaluation,
diff --git a/Engine/Source/Runtime/MovieScene/Private/Channels/MovieSceneInterpolation.cpp b/Engine/Source/Runtime/MovieScene/Private/Channels/MovieSceneInterpolation.cpp
--- a/Engine/Source/Runtime/MovieScene/Private/Channels/MovieSceneInterpolation.cpp
+++ b/Engine/Source/Runtime/MovieScene/Private/Channels/MovieSceneInterpolation.cpp
@@ -208,7 +208,7 @@ double FWeightedCubicInterpolation::Evaluate(FFrameTime InTime) const
 	Coeff[0] = Coeff[0] - Interp;
 
 	const int32 NumResults = UE::Curves::SolveCubic(Coeff, Results);
-	float NewInterp = Interp;
+	double NewInterp = Interp;
 	if (NumResults == 1)
 	{
 		NewInterp = Results[0];
@@ -218,9 +218,9 @@ double FWeightedCubicInterpolation::Evaluate(FFrameTime InTime) const
 		NewInterp = TNumericLimits<double>::Lowest(); //just need to be out of range
 		for (double Result : Results)
 		{
-			if ((Result >= 0.0f) && (Result <= 1.0f))
+			if ((Result >= 0.0) && (Result <= 1.0))
 			{
-				if (NewInterp < 0.0f || Result > NewInterp)
+				if (NewInterp < 0.0 || Result > NewInterp)
 				{
 					NewInterp = Result;
 				}
@@ -229,7 +229,7 @@ double FWeightedCubicInterpolation::Evaluate(FFrameTime InTime) const
 
 		if (NewInterp == TNumericLimits<double>::Lowest())
 		{
-			NewInterp = 0.f;
+			NewInterp = 0.0;
 		}
 	}
 
