Sequencer: Copy unique objects so that they don't get copy/pasted multiple times when gathering descendants

diff --git a/Engine/Source/Editor/Sequencer/Private/Sequencer.cpp b/Engine/Source/Editor/Sequencer/Private/Sequencer.cpp
--- a/Engine/Source/Editor/Sequencer/Private/Sequencer.cpp
+++ b/Engine/Source/Editor/Sequencer/Private/Sequencer.cpp
@@ -7869,7 +7869,7 @@ void FSequencer::CopySelectedObjects(TArray<TSharedPtr<UE::Sequencer::FObjectBin
 	using namespace UE::Sequencer;
 
 	// Gather guids for the object nodes and any child object nodes
-	TArray<FMovieSceneBindingProxy> Bindings;
+	TSet<FMovieSceneBindingProxy> Bindings;
 	for (TSharedPtr<FObjectBindingModel> ObjectNode : ObjectNodes)
 	{
 		Bindings.Add(FMovieSceneBindingProxy(ObjectNode->GetObjectGuid(), GetFocusedMovieSceneSequence()));
@@ -7887,7 +7887,8 @@ void FSequencer::CopySelectedObjects(TArray<TSharedPtr<UE::Sequencer::FObjectBin
 		}
 	}
 
-	FSequencerUtilities::CopyBindings(AsShared(), Bindings, Folders, ExportedText);
+	TArray<FMovieSceneBindingProxy> BindingsArray = Bindings.Array();
+	FSequencerUtilities::CopyBindings(AsShared(), BindingsArray, Folders, ExportedText);
 		
 	// Make sure to clear the clipboard for the keys
 	GClipboardStack.Empty();
diff --git a/Engine/Source/Runtime/MovieScene/Public/MovieSceneBindingProxy.h b/Engine/Source/Runtime/MovieScene/Public/MovieSceneBindingProxy.h
--- a/Engine/Source/Runtime/MovieScene/Public/MovieSceneBindingProxy.h
+++ b/Engine/Source/Runtime/MovieScene/Public/MovieSceneBindingProxy.h
@@ -6,6 +6,7 @@
 #include "Misc/Guid.h"
 #include "UObject/ObjectMacros.h"
 #include "UObject/ObjectPtr.h"
+#include "Templates/TypeHash.h"
 
 #include "MovieSceneBindingProxy.generated.h"
 
@@ -31,6 +32,21 @@ struct MOVIESCENE_API FMovieSceneBindingProxy
 		, Sequence(InSequence)
 	{}
 
+	FORCEINLINE friend bool operator==(const FMovieSceneBindingProxy &LHS, const FMovieSceneBindingProxy &RHS)
+	{
+		return LHS.BindingID == RHS.BindingID && LHS.Sequence == RHS.Sequence;
+	}
+
+	FORCEINLINE friend bool operator!=(const FMovieSceneBindingProxy &LHS, const FMovieSceneBindingProxy &RHS)
+	{
+		return LHS.BindingID != RHS.BindingID || LHS.Sequence != RHS.Sequence;
+	}
+
+	FORCEINLINE friend uint32 GetTypeHash(const FMovieSceneBindingProxy& In)
+	{
+		return HashCombine(GetTypeHash(In.BindingID), GetTypeHash(In.Sequence));
+	}
+
 	UMovieScene* GetMovieScene() const;
 
 	UPROPERTY(BlueprintReadOnly, Category=Binding)
