Sequencer: Added a way of overriding Keep or Restore State behavior at the sub-sequence leve

This is useful for sitations where a Sequence with Keep State tracks is included as a sub-sequence, but the outer sequence wants it to restore state instead.

diff --git a/Engine/Source/Runtime/MovieScene/Private/Compilation/MovieSceneCompiledDataManager.cpp b/Engine/Source/Runtime/MovieScene/Private/Compilation/MovieSceneCompiledDataManager.cpp
--- a/Engine/Source/Runtime/MovieScene/Private/Compilation/MovieSceneCompiledDataManager.cpp
+++ b/Engine/Source/Runtime/MovieScene/Private/Compilation/MovieSceneCompiledDataManager.cpp
@@ -160,6 +160,8 @@ struct FGatherParameters
 
 		SubParams.LocalClampRange           = SubData.RootToSequenceTransform.TransformRangeUnwarped(SubParams.RootClampRange);
 
+		SubParams.AccumulatedFlags          = UE::MovieScene::AccumulateChildSubSectionFlags(AccumulatedFlags, SubData.AccumulatedFlags);
+
 		return SubParams;
 	}
 
@@ -195,6 +197,9 @@ struct FGatherParameters
 	/** Current accumulated hierarchical bias */
 	int16 HierarchicalBias;
 
+	/** Current accumulated sub-section flags */
+	EMovieSceneSubSectionFlags AccumulatedFlags;
+
 	EMovieSceneServerClientMask NetworkMask;
 };
 
@@ -1458,6 +1463,7 @@ bool UMovieSceneCompiledDataManager::GenerateSubSequenceData(UMovieSceneSubTrack
 		NewSubData.PlayRange               = TRange<FFrameNumber>::Intersection(InnerClampRange, NewSubData.PlayRange.Value);
 		NewSubData.RootToSequenceTransform = NewSubData.RootToSequenceTransform * Params.RootToSequenceTransform;
 		NewSubData.HierarchicalBias        = Params.HierarchicalBias + NewSubData.HierarchicalBias;
+		NewSubData.AccumulatedFlags        = Params.AccumulatedFlags;
 
 		// Add the sub data to the root hierarchy
 		InOutHierarchy->Add(NewSubData, InnerSequenceID, ParentSequenceID);
diff --git a/Engine/Source/Runtime/MovieScene/Private/EntitySystem/MovieSceneSequenceUpdaters.cpp b/Engine/Source/Runtime/MovieScene/Private/EntitySystem/MovieSceneSequenceUpdaters.cpp
--- a/Engine/Source/Runtime/MovieScene/Private/EntitySystem/MovieSceneSequenceUpdaters.cpp
+++ b/Engine/Source/Runtime/MovieScene/Private/EntitySystem/MovieSceneSequenceUpdaters.cpp
@@ -619,6 +619,7 @@ void FSequenceUpdater_Hierarchical::Update(UMovieSceneEntitySystemLinker* Linker
 				Params.RootInstanceHandle = InstanceHandle;
 				Params.DefaultCompletionMode = SubSequence->DefaultCompletionMode;
 				Params.HierarchicalBias = SubData->HierarchicalBias;
+				Params.SubSectionFlags = SubData->AccumulatedFlags;
 				Params.bPreRoll  = bIsPreRoll;
 				Params.bPostRoll = bIsPostRoll;
 				Params.bDynamicWeighting = bDynamicWeighting.Get(false); // Always inherit dynamic weighting flags
diff --git a/Engine/Source/Runtime/MovieScene/Private/Evaluation/MovieSceneSequenceHierarchy.cpp b/Engine/Source/Runtime/MovieScene/Private/Evaluation/MovieSceneSequenceHierarchy.cpp
--- a/Engine/Source/Runtime/MovieScene/Private/Evaluation/MovieSceneSequenceHierarchy.cpp
+++ b/Engine/Source/Runtime/MovieScene/Private/Evaluation/MovieSceneSequenceHierarchy.cpp
@@ -26,6 +26,7 @@ FMovieSceneSubSequenceData::FMovieSceneSubSequenceData(const UMovieSceneSubSecti
 	, ParentFirstLoopStartFrameOffset(InSubSection.Parameters.FirstLoopStartFrameOffset)
 	, bCanLoop(InSubSection.Parameters.bCanLoop)
 	, HierarchicalBias(InSubSection.Parameters.HierarchicalBias)
+	, AccumulatedFlags(InSubSection.Parameters.Flags)
 #if WITH_EDITORONLY_DATA
 	, SectionPath(*InSubSection.GetPathNameInMovieScene())
 #endif
diff --git a/Engine/Source/Runtime/MovieScene/Private/MovieSceneSection.cpp b/Engine/Source/Runtime/MovieScene/Private/MovieSceneSection.cpp
--- a/Engine/Source/Runtime/MovieScene/Private/MovieSceneSection.cpp
+++ b/Engine/Source/Runtime/MovieScene/Private/MovieSceneSection.cpp
@@ -315,8 +315,17 @@ void UMovieSceneSection::BuildDefaultComponents(UMovieSceneEntitySystemLinker* E
 
 	const bool bHasEasing = (Easing.GetEaseInDuration() > 0 || Easing.GetEaseOutDuration() > 0);
 
-	const bool bShouldRestoreState = (EvalOptions.CompletionMode == EMovieSceneCompletionMode::RestoreState) ||
-		( EvalOptions.CompletionMode == EMovieSceneCompletionMode::ProjectDefault && Params.Sequence.DefaultCompletionMode == EMovieSceneCompletionMode::RestoreState);
+	// Should restore state if we're not forcing keep state and any one of the following:
+	// - We're forcing restore state
+	// - This section is set to restore state
+	// - This section is set to the default, and the default is restore state
+	const bool bForceKeepState     = EnumHasAnyFlags(Params.Sequence.SubSectionFlags, EMovieSceneSubSectionFlags::OverrideKeepState);
+	const bool bShouldRestoreState = bForceKeepState == false &&
+		(
+			EnumHasAnyFlags(Params.Sequence.SubSectionFlags, EMovieSceneSubSectionFlags::OverrideRestoreState) ||
+			(EvalOptions.CompletionMode == EMovieSceneCompletionMode::RestoreState) ||
+			(EvalOptions.CompletionMode == EMovieSceneCompletionMode::ProjectDefault && Params.Sequence.DefaultCompletionMode == EMovieSceneCompletionMode::RestoreState)
+		);
 
 	TComponentTypeID<FEasingComponentData> EasingComponentID = Components->Easing;
 	FComponentTypeID RestoreStateTag = Components->Tags.RestoreState;
@@ -342,17 +351,18 @@ void UMovieSceneSection::BuildDefaultComponents(UMovieSceneEntitySystemLinker* E
 
 	OutImportedEntity->AddBuilder(
 		FEntityBuilder()
-		.AddConditional(Components->BlenderType,                BlenderSystemClass, BlenderSystemClass.Get() != nullptr)
-		.AddConditional(Components->Easing,                     FEasingComponentData{ decltype(FEasingComponentData::Section)(this) }, bHasEasing)
-		.AddConditional(Components->HierarchicalBias,           Params.Sequence.HierarchicalBias, Params.Sequence.HierarchicalBias != 0)
-		.AddConditional(Components->Interrogation.InputKey,     Params.InterrogationKey, Params.InterrogationKey.IsValid())
-		.AddConditional(Components->Interrogation.Instance,     Params.InterrogationInstance, Params.InterrogationInstance.IsValid())
-		.AddConditional(Components->EvalTime,                   Params.EntityMetaData ? Params.EntityMetaData->ForcedTime : 0, bHasForcedTime)
-		.AddTagConditional(Components->Tags.RestoreState,       bShouldRestoreState)
-		.AddTagConditional(Components->Tags.FixedTime,          bHasForcedTime)
-		.AddTagConditional(Components->Tags.SectionPreRoll,     bHasSectionPreRoll)
-		.AddTagConditional(Components->Tags.PreRoll,            bHasSequencePreRoll || bHasSectionPreRoll)
-		.AddTagConditional(BlendTag,                            BlendTag != FComponentTypeID::Invalid())
+		.AddConditional(Components->BlenderType,                    BlenderSystemClass, BlenderSystemClass.Get() != nullptr)
+		.AddConditional(Components->Easing,                         FEasingComponentData{ decltype(FEasingComponentData::Section)(this) }, bHasEasing)
+		.AddConditional(Components->HierarchicalBias,               Params.Sequence.HierarchicalBias, Params.Sequence.HierarchicalBias != 0)
+		.AddConditional(Components->Interrogation.InputKey,         Params.InterrogationKey, Params.InterrogationKey.IsValid())
+		.AddConditional(Components->Interrogation.Instance,         Params.InterrogationInstance, Params.InterrogationInstance.IsValid())
+		.AddConditional(Components->EvalTime,                       Params.EntityMetaData ? Params.EntityMetaData->ForcedTime : 0, bHasForcedTime)
+		.AddTagConditional(Components->Tags.RestoreState,           bShouldRestoreState)
+		.AddTagConditional(Components->Tags.IgnoreHierarchicalBias, EnumHasAnyFlags(Params.Sequence.SubSectionFlags, EMovieSceneSubSectionFlags::IgnoreHierarchicalBias))
+		.AddTagConditional(Components->Tags.FixedTime,              bHasForcedTime)
+		.AddTagConditional(Components->Tags.SectionPreRoll,         bHasSectionPreRoll)
+		.AddTagConditional(Components->Tags.PreRoll,                bHasSequencePreRoll || bHasSectionPreRoll)
+		.AddTagConditional(BlendTag,                                BlendTag != FComponentTypeID::Invalid())
 	);
 
 	if (BlendTag == Components->Tags.AdditiveFromBaseBlend)
diff --git a/Engine/Source/Runtime/MovieScene/Public/EntitySystem/IMovieSceneEntityProvider.h b/Engine/Source/Runtime/MovieScene/Public/EntitySystem/IMovieSceneEntityProvider.h
--- a/Engine/Source/Runtime/MovieScene/Public/EntitySystem/IMovieSceneEntityProvider.h
+++ b/Engine/Source/Runtime/MovieScene/Public/EntitySystem/IMovieSceneEntityProvider.h
@@ -4,6 +4,7 @@
 
 #include "Containers/Array.h"
 #include "Containers/ContainerAllocationPolicies.h"
+#include "Evaluation/MovieSceneSectionParameters.h"
 #include "EntitySystem/MovieSceneEntityBuilder.h"
 #include "EntitySystem/MovieSceneEntityIDs.h"
 #include "EntitySystem/MovieSceneEntitySystemTypes.h"
@@ -68,6 +69,7 @@ struct FEntityImportSequenceParams
 		: HierarchicalBias(0)
 		, SequenceID(MovieSceneSequenceID::Root)
 		, DefaultCompletionMode(EMovieSceneCompletionMode::KeepState)
+		, SubSectionFlags(EMovieSceneSubSectionFlags::None)
 		, bPreRoll(false)
 		, bPostRoll(false)
 		, bDynamicWeighting(false)
@@ -80,6 +82,7 @@ struct FEntityImportSequenceParams
 	FRootInstanceHandle RootInstanceHandle;
 
 	EMovieSceneCompletionMode DefaultCompletionMode;
+	EMovieSceneSubSectionFlags SubSectionFlags;
 
 	bool bPreRoll : 1;
 	bool bPostRoll : 1;
diff --git a/Engine/Source/Runtime/MovieScene/Public/Evaluation/MovieSceneSectionParameters.h b/Engine/Source/Runtime/MovieScene/Public/Evaluation/MovieSceneSectionParameters.h
--- a/Engine/Source/Runtime/MovieScene/Public/Evaluation/MovieSceneSectionParameters.h
+++ b/Engine/Source/Runtime/MovieScene/Public/Evaluation/MovieSceneSectionParameters.h
@@ -4,8 +4,50 @@
 
 #include "CoreMinimal.h"
 #include "UObject/ObjectMacros.h"
+#include "Misc/EnumClassFlags.h"
 #include "MovieSceneSectionParameters.generated.h"
 
+/**
+ * Flag structure that can be applied to any sub-section allowing control over various
+ * behaviors for the nested sub-sequence.
+ */
+UENUM(BlueprintType)
+enum class EMovieSceneSubSectionFlags
+{
+	None = 0,
+
+	/*~ Begin mutually exclusive */
+	/** When set, everything within the sub-section (including further sub-sections) should be keep-state. Mutually exclusive with OverrideRestoreState. */
+	OverrideKeepState = 1 << 0,
+	/** When set, everything within the sub-section (including further sub-sections) should be restore-state. Mutually exclusive with OverrideKeepState. */
+	OverrideRestoreState = 1 << 1,
+	/*~ End mutually exclusive */
+
+	/** Everything inside this sub-sequence should ignore hierarchical bias and always be relevant */
+	IgnoreHierarchicalBias = 1 << 2,
+
+	AnyRestoreStateOverride = OverrideKeepState | OverrideRestoreState,
+};
+ENUM_CLASS_FLAGS(EMovieSceneSubSectionFlags)
+
+namespace UE::MovieScene
+{
+	/** Accumulate parent and chld sub-section flags ensuring that flags are inherited correctly. */
+	inline EMovieSceneSubSectionFlags AccumulateChildSubSectionFlags(EMovieSceneSubSectionFlags ParentFlags, EMovieSceneSubSectionFlags ChildFlags)
+	{
+		if (EnumHasAnyFlags(ParentFlags, EMovieSceneSubSectionFlags::AnyRestoreStateOverride))
+		{
+			// If the parent has any uninheritable flags based on the parent, ensure the child has the parent's flags
+			return (ChildFlags & ~EMovieSceneSubSectionFlags::AnyRestoreStateOverride) | ParentFlags;
+		}
+		else
+		{
+			return ChildFlags | ParentFlags;
+		}
+	}
+
+} // namespace UE::MovieScene
+
 USTRUCT(BlueprintType)
 struct FMovieSceneSectionParameters
 {
@@ -48,6 +90,10 @@ public:
 	UPROPERTY(config, BlueprintReadWrite, EditAnywhere, Category="Sequence")
 	int32 HierarchicalBias;
 
+	/** Sub-section flags defining how to deal with this sub-sequence */
+	UPROPERTY(config, BlueprintReadWrite, EditAnywhere, Category="Sequence")
+	EMovieSceneSubSectionFlags Flags;
+
 	UPROPERTY()
 	float StartOffset_DEPRECATED;
 	UPROPERTY()
diff --git a/Engine/Source/Runtime/MovieScene/Public/Evaluation/MovieSceneSequenceHierarchy.h b/Engine/Source/Runtime/MovieScene/Public/Evaluation/MovieSceneSequenceHierarchy.h
--- a/Engine/Source/Runtime/MovieScene/Public/Evaluation/MovieSceneSequenceHierarchy.h
+++ b/Engine/Source/Runtime/MovieScene/Public/Evaluation/MovieSceneSequenceHierarchy.h
@@ -146,6 +146,10 @@ struct FMovieSceneSubSequenceData
 	UPROPERTY()
 	int16 HierarchicalBias;
 
+	/** Flags accumulated from parent->child for each sub-section that led to the inclusion of this sub-sequence */
+	UPROPERTY()
+	EMovieSceneSubSectionFlags AccumulatedFlags;
+
 	/** Instance data that should be used for any tracks contained immediately within this sub sequence */
 	UPROPERTY()
 	FMovieSceneSequenceInstanceDataPtr InstanceData;
