Sequencer: Keep State overrides now work with legacy track templates

diff --git a/Engine/Source/Runtime/MovieScene/Private/Compilation/MovieSceneCompiledDataManager.cpp b/Engine/Source/Runtime/MovieScene/Private/Compilation/MovieSceneCompiledDataManager.cpp
--- a/Engine/Source/Runtime/MovieScene/Private/Compilation/MovieSceneCompiledDataManager.cpp
+++ b/Engine/Source/Runtime/MovieScene/Private/Compilation/MovieSceneCompiledDataManager.cpp
@@ -1343,6 +1343,15 @@ void UMovieSceneCompiledDataManager::GatherTrack(const FMovieSceneBinding* Objec
 
 				ESectionEvaluationFlags Flags = Params.Flags == ESectionEvaluationFlags::None ? Entry.Flags : Params.Flags;
 
+				if (EnumHasAnyFlags(Params.AccumulatedFlags, EMovieSceneSubSectionFlags::OverrideRestoreState))
+				{
+					Flags |= ESectionEvaluationFlags::ForceRestoreState;
+				}
+				else if (EnumHasAnyFlags(Params.AccumulatedFlags, EMovieSceneSubSectionFlags::OverrideKeepState))
+				{
+					Flags |= ESectionEvaluationFlags::ForceKeepState;
+				}
+
 				CompileData.ChildPriority = Entry.LegacySortOrder;
 				CompileData.Child         = FMovieSceneFieldEntry_ChildTemplate((uint16)ChildTemplateIndex, Flags, Entry.ForcedTime);
 				CompileData.bRequiresInit = EvaluationTrack->GetChildTemplate(ChildTemplateIndex).RequiresInitialization();
diff --git a/Engine/Source/Runtime/MovieScene/Private/Evaluation/MovieSceneEvaluationTrack.cpp b/Engine/Source/Runtime/MovieScene/Private/Evaluation/MovieSceneEvaluationTrack.cpp
--- a/Engine/Source/Runtime/MovieScene/Private/Evaluation/MovieSceneEvaluationTrack.cpp
+++ b/Engine/Source/Runtime/MovieScene/Private/Evaluation/MovieSceneEvaluationTrack.cpp
@@ -198,8 +198,13 @@ void FMovieSceneEvaluationTrack::EvaluateStatic(TArrayView<const FMovieSceneFiel
 
 		Context.ApplySectionPrePostRoll(bIsPreRoll, bIsPostRoll);
 
+		EMovieSceneCompletionMode CompletionMode =
+			EnumHasAnyFlags(FieldEntry.Flags, ESectionEvaluationFlags::ForceRestoreState) ? EMovieSceneCompletionMode::RestoreState :
+			EnumHasAnyFlags(FieldEntry.Flags, ESectionEvaluationFlags::ForceKeepState)    ? EMovieSceneCompletionMode::KeepState :
+			Template->GetCompletionMode();
+
 		PersistentData.DeriveSectionKey(FieldEntry.ChildIndex);
-		ExecutionTokens.SetCurrentScope(FMovieSceneEvaluationScope(PersistentData.GetSectionKey(), Template->GetCompletionMode()));
+		ExecutionTokens.SetCurrentScope(FMovieSceneEvaluationScope(PersistentData.GetSectionKey(), CompletionMode));
 		ExecutionTokens.SetContext(Context);
 
 		Template->Evaluate(Operand, Context, PersistentData, ExecutionTokens);
@@ -226,8 +231,13 @@ void FMovieSceneEvaluationTrack::EvaluateSwept(TArrayView<const FMovieSceneField
 		{
 			const FMovieSceneEvalTemplate& Template = *ChildTemplates[ChildIndex];
 
+			EMovieSceneCompletionMode CompletionMode =
+				EnumHasAnyFlags(Entry.Flags, ESectionEvaluationFlags::ForceRestoreState) ? EMovieSceneCompletionMode::RestoreState :
+				EnumHasAnyFlags(Entry.Flags, ESectionEvaluationFlags::ForceKeepState)    ? EMovieSceneCompletionMode::KeepState :
+				Template.GetCompletionMode();
+
 			PersistentData.DeriveSectionKey(ChildIndex);
-			ExecutionTokens.SetCurrentScope(FMovieSceneEvaluationScope(PersistentData.GetSectionKey(), Template.GetCompletionMode()));
+			ExecutionTokens.SetCurrentScope(FMovieSceneEvaluationScope(PersistentData.GetSectionKey(), CompletionMode));
 			ExecutionTokens.SetContext(Context);
 
 			Template.EvaluateSwept(
diff --git a/Engine/Source/Runtime/MovieScene/Public/Evaluation/MovieSceneSegment.h b/Engine/Source/Runtime/MovieScene/Public/Evaluation/MovieSceneSegment.h
--- a/Engine/Source/Runtime/MovieScene/Public/Evaluation/MovieSceneSegment.h
+++ b/Engine/Source/Runtime/MovieScene/Public/Evaluation/MovieSceneSegment.h
@@ -23,6 +23,9 @@ enum class ESectionEvaluationFlags : uint8
 	PreRoll		= 0x01,
 	/** Segment resides inside the 'post-roll' time for the section */
 	PostRoll	= 0x02,
+
+	ForceKeepState = 0x04,
+	ForceRestoreState = 0x08,
 };
 ENUM_CLASS_FLAGS(ESectionEvaluationFlags);
 
