Fixed uninitialized variable

diff --git a/Engine/Source/Runtime/MovieScene/Private/Compilation/MovieSceneCompiledDataManager.cpp b/Engine/Source/Runtime/MovieScene/Private/Compilation/MovieSceneCompiledDataManager.cpp
--- a/Engine/Source/Runtime/MovieScene/Private/Compilation/MovieSceneCompiledDataManager.cpp
+++ b/Engine/Source/Runtime/MovieScene/Private/Compilation/MovieSceneCompiledDataManager.cpp
@@ -142,6 +142,7 @@ struct FGatherParameters
 		, LocalClampRange(RootClampRange)
 		, Flags(ESectionEvaluationFlags::None)
 		, HierarchicalBias(0)
+		, AccumulatedFlags(EMovieSceneSubSectionFlags::None)
 	{}
 
 	FGatherParameters CreateForSubData(const FMovieSceneSubSequenceData& SubData, FMovieSceneSequenceID InSubSequenceID) const
