Sequencer: Fixed accumulated sub section flags getting accumulated the wrong way round

diff --git a/Engine/Source/Runtime/MovieScene/Private/Compilation/MovieSceneCompiledDataManager.cpp b/Engine/Source/Runtime/MovieScene/Private/Compilation/MovieSceneCompiledDataManager.cpp
--- a/Engine/Source/Runtime/MovieScene/Private/Compilation/MovieSceneCompiledDataManager.cpp
+++ b/Engine/Source/Runtime/MovieScene/Private/Compilation/MovieSceneCompiledDataManager.cpp
@@ -156,13 +156,12 @@ struct FGatherParameters
 
 		SubParams.RootToSequenceTransform   = SubData.RootToSequenceTransform;
 		SubParams.HierarchicalBias          = SubData.HierarchicalBias;
+		SubParams.AccumulatedFlags          = SubData.AccumulatedFlags;
 		SubParams.SequenceID                = InSubSequenceID;
 		SubParams.RootToSequenceWarpCounter = WarpCounter;
 
 		SubParams.LocalClampRange           = SubData.RootToSequenceTransform.TransformRangeUnwarped(SubParams.RootClampRange);
 
-		SubParams.AccumulatedFlags          = UE::MovieScene::AccumulateChildSubSectionFlags(AccumulatedFlags, SubData.AccumulatedFlags);
-
 		return SubParams;
 	}
 
@@ -1473,7 +1472,7 @@ bool UMovieSceneCompiledDataManager::GenerateSubSequenceData(UMovieSceneSubTrack
 		NewSubData.PlayRange               = TRange<FFrameNumber>::Intersection(InnerClampRange, NewSubData.PlayRange.Value);
 		NewSubData.RootToSequenceTransform = NewSubData.RootToSequenceTransform * Params.RootToSequenceTransform;
 		NewSubData.HierarchicalBias        = Params.HierarchicalBias + NewSubData.HierarchicalBias;
-		NewSubData.AccumulatedFlags        = Params.AccumulatedFlags;
+		NewSubData.AccumulatedFlags        = UE::MovieScene::AccumulateChildSubSectionFlags(Params.AccumulatedFlags, NewSubData.AccumulatedFlags);
 
 		// Add the sub data to the root hierarchy
 		InOutHierarchy->Add(NewSubData, InnerSequenceID, ParentSequenceID);
