Sequencer: Expanded API for managing custom time controllers on UMovieSceneSequencePlayer

diff --git a/Engine/Source/Runtime/MovieScene/Private/MovieSceneSequencePlayer.cpp b/Engine/Source/Runtime/MovieScene/Private/MovieSceneSequencePlayer.cpp
--- a/Engine/Source/Runtime/MovieScene/Private/MovieSceneSequencePlayer.cpp
+++ b/Engine/Source/Runtime/MovieScene/Private/MovieSceneSequencePlayer.cpp
@@ -1308,15 +1308,25 @@ void UMovieSceneSequencePlayer::SetPlaybackClient(TScriptInterface<IMovieScenePl
 	PlaybackClient = InPlaybackClient;
 }
 
+TSharedPtr<FMovieSceneTimeController> UMovieSceneSequencePlayer::GetTimeController() const
+{
+	return TimeController;
+}
+
 void UMovieSceneSequencePlayer::SetTimeController(TSharedPtr<FMovieSceneTimeController> InTimeController)
 {
-	TimeController = InTimeController;
+	SetTimeControllerDirectly(InTimeController);
 	if (TimeController.IsValid())
 	{
 		TimeController->Reset(GetCurrentTime());
 	}
 }
 
+void UMovieSceneSequencePlayer::SetTimeControllerDirectly(TSharedPtr<FMovieSceneTimeController> InTimeController)
+{
+	TimeController = InTimeController;
+}
+
 void UMovieSceneSequencePlayer::SetIgnorePlaybackReplication(bool bState)
 {
 	bIgnorePlaybackReplication = bState;
diff --git a/Engine/Source/Runtime/MovieScene/Public/MovieSceneSequencePlayer.h b/Engine/Source/Runtime/MovieScene/Public/MovieSceneSequencePlayer.h
--- a/Engine/Source/Runtime/MovieScene/Public/MovieSceneSequencePlayer.h
+++ b/Engine/Source/Runtime/MovieScene/Public/MovieSceneSequencePlayer.h
@@ -486,10 +486,23 @@ public:
 	 */
 	void SetPlaybackClient(TScriptInterface<IMovieScenePlaybackClient> InPlaybackClient);
 
+	/**
+	 * Retrieve the currently assigned time controller
+	 */
+	TSharedPtr<FMovieSceneTimeController> GetTimeController() const;
+
 	/**
 	 * Assign a time controller for this sequence player allowing custom time management implementations.
+	 * Will reset the supplied time controller to the current time.
 	 */
 	void SetTimeController(TSharedPtr<FMovieSceneTimeController> InTimeController);
+	
+
+	/**
+	 * Assign a time controller for this sequence player allowing custom time management implementations.
+	 * Will not reset the supplied time controller in any way, so the sequence will receive its time directly from the controller.
+	 */
+	void SetTimeControllerDirectly(TSharedPtr<FMovieSceneTimeController> InTimeController);
 
 	/**
 	 * Sets whether to listen or ignore playback replication events.
