Applying changes lost during previous move.

diff --git a/Engine/Source/Runtime/MovieScene/Public/MovieSceneMarkedFrame.h b/Engine/Source/Runtime/MovieScene/Public/MovieSceneMarkedFrame.h
--- a/Engine/Source/Runtime/MovieScene/Public/MovieSceneMarkedFrame.h
+++ b/Engine/Source/Runtime/MovieScene/Public/MovieSceneMarkedFrame.h
@@ -16,6 +16,7 @@ struct FMovieSceneMarkedFrame
 	FMovieSceneMarkedFrame()
 		: Label(FString())
 #if WITH_EDITORONLY_DATA
+		, Comment(FString())
 		, Color(0.f, 1.f, 1.f, 0.4f)
 #endif
 		, bIsDeterminismFence(false)
@@ -25,6 +26,7 @@ struct FMovieSceneMarkedFrame
 		: FrameNumber(InFrameNumber)
 		, Label(FString())
 #if WITH_EDITORONLY_DATA
+		, Comment(FString())
 		, Color(0.f, 1.f, 1.f, 0.4f)
 #endif
 		, bIsDeterminismFence(false)
@@ -37,10 +39,13 @@ struct FMovieSceneMarkedFrame
 	FString Label;
 
 #if WITH_EDITORONLY_DATA
+	UPROPERTY(EditAnywhere, Category = "Marked Frame")
+	FString Comment;
+
 	UPROPERTY(EditAnywhere, Category = "Marked Frame")
 	FLinearColor Color;
 #endif
 
-	UPROPERTY(EditAnywhere, BlueprintReadWrite, Category = "Marked Frame", DisplayName = "Is Determinism Fence?", meta = (Tooltip = "When checked, treat this mark as a fence for evaluation purposes. Fences cannot be crossed in a single evaluation, and force the evaluation to be split into 2 separate parts."))
+	UPROPERTY(EditAnywhere, BlueprintReadWrite, Category = "Marked Frame", DisplayName="Is Determinism Fence?", meta=(Tooltip="When checked, treat this mark as a fence for evaluation purposes. Fences cannot be crossed in a single evaluation, and force the evaluation to be split into 2 separate parts."))
 	bool bIsDeterminismFence;
 };
