Sequencer: Initialize EMovieSceneSubSectionFlags in FMovieSceneSectionParameters and FMovieSceneSubSequenceData

diff --git a/Engine/Source/Runtime/MovieScene/Private/Evaluation/MovieSceneSequenceHierarchy.cpp b/Engine/Source/Runtime/MovieScene/Private/Evaluation/MovieSceneSequenceHierarchy.cpp
--- a/Engine/Source/Runtime/MovieScene/Private/Evaluation/MovieSceneSequenceHierarchy.cpp
+++ b/Engine/Source/Runtime/MovieScene/Private/Evaluation/MovieSceneSequenceHierarchy.cpp
@@ -15,6 +15,7 @@
 FMovieSceneSubSequenceData::FMovieSceneSubSequenceData()
 	: Sequence(nullptr)
 	, HierarchicalBias(0)
+	, AccumulatedFlags(EMovieSceneSubSectionFlags::None)
 {}
 
 FMovieSceneSubSequenceData::FMovieSceneSubSequenceData(const UMovieSceneSubSection& InSubSection)
diff --git a/Engine/Source/Runtime/MovieScene/Public/Evaluation/MovieSceneSectionParameters.h b/Engine/Source/Runtime/MovieScene/Public/Evaluation/MovieSceneSectionParameters.h
--- a/Engine/Source/Runtime/MovieScene/Public/Evaluation/MovieSceneSectionParameters.h
+++ b/Engine/Source/Runtime/MovieScene/Public/Evaluation/MovieSceneSectionParameters.h
@@ -60,6 +60,7 @@ struct FMovieSceneSectionParameters
 		, FirstLoopStartFrameOffset(0)
 		, TimeScale(1.0f)
 		, HierarchicalBias(100)
+		, Flags(EMovieSceneSubSectionFlags::None)
 		, StartOffset_DEPRECATED(0.f)
 		, PrerollTime_DEPRECATED(0.0f)
 		, PostrollTime_DEPRECATED(0.0f)
