Sequencer: Overhauled component inclusivity behavior / initialization
This work is being done in order to re-introduce the concept of mutual component initializers in a more robust and manageable way. Mutally inclusive component rules define the presence (and now value) of a component based on the presence of another (or multiple). One example would be that DoubleResult[0] must be present alongside DoubleChannel[0].

Normally a mutual included component's presence is all that matters, but in some cases you want to initialize the value of that component to something based off the predicate component. This is a mutual component initializer.

Mutual component dependencies were previously stored as a digraph where edges denote inclusion. Additionally we supported 'complex' inclusions which were stored as an any/all component filter alongside the component mask to include. The logic for computing mutually inclusive components was previously implemented as a stabilizing loop algorithm that kept walking the graph and processing complex rules until it stopped adding components. This logic made it relatively hard to inject the notion of 'initializers' in a robust and performant way that correctly runs initializers if they were relevant to a particular operation (ie, we should not run an initializer if the component already existed).

The new logic stores simple and complex rules in separate maps, but has a generation step that generates a command buffer from the implied dependency graph. The resulting command buffer is ordered sequentially such that all rules that are predecated other mutually included components will be processed after their predicate has a chance to be included. This is decided based on the length of the predicate 'chain' in the graph. Cyclic dependencies are now an explicit error.

At runtime this command buffer is probably only generated once, and is relatively fast to execute (measured avg 600ns down from 12us on my Threadripper). It could be further optimized by removing commands that are known not to be relevant (ie, the leaf predicate component doesn't exist anywhere, so can never be active), but that is a separate task of diminishing returns.

Use of mutual component initializers are now mandatory when calling ComputeMutuallyInclusiveComponents. An option has been added to FEntityManager::MutateXYZ functions to specify whether you want to also add mutuals as part of the mutation.

Moved FMovieSceneEntitySystemDirectedGraph to UE::MovieScene::FDirectedGraph

diff --git a/Engine/Plugins/MovieScene/TemplateSequence/Source/TemplateSequence/Private/CameraAnimationSequenceSubsystem.cpp b/Engine/Plugins/MovieScene/TemplateSequence/Source/TemplateSequence/Private/CameraAnimationSequenceSubsystem.cpp
--- a/Engine/Plugins/MovieScene/TemplateSequence/Source/TemplateSequence/Private/CameraAnimationSequenceSubsystem.cpp
+++ b/Engine/Plugins/MovieScene/TemplateSequence/Source/TemplateSequence/Private/CameraAnimationSequenceSubsystem.cpp
@@ -31,8 +31,6 @@ struct FCameraAnimationInstantiationMutation : IMovieSceneEntityMutation
 	virtual void CreateMutation(FEntityManager* EntityManager, FComponentMask* InOutEntityComponentTypes) const
 	{
 		// Create all output components and bind objects/components.
-		FComponentRegistry* ComponentRegistry = EntityManager->GetComponents();
-		ComponentRegistry->Factories.ComputeMutuallyInclusiveComponents(*InOutEntityComponentTypes);
 		InOutEntityComponentTypes->Set(BuiltInComponents->BoundObject);
 	}
 	virtual void InitializeAllocation(FEntityAllocation* Allocation, const FComponentMask& AllocationType) const
@@ -152,7 +150,7 @@ void UCameraAnimationBoundObjectInstantiator::OnRun(FSystemTaskPrerequisites& In
 				BuiltInComponents->SceneComponentBinding })
 		.All({ BuiltInComponents->InstanceHandle, BuiltInComponents->Tags.NeedsLink })
 		.None({ BuiltInComponents->Tags.NeedsUnlink });
-	Linker->EntityManager.MutateAll(Filter, Mutation);
+	Linker->EntityManager.MutateAll(Filter, Mutation, EMutuallyInclusiveComponentType::All);
 }
 
 UCameraAnimationEntitySystemLinker::UCameraAnimationEntitySystemLinker(const FObjectInitializer& ObjInit)
diff --git a/Engine/Source/Runtime/MovieScene/Private/EntitySystem/IMovieSceneEntityProvider.cpp b/Engine/Source/Runtime/MovieScene/Private/EntitySystem/IMovieSceneEntityProvider.cpp
--- a/Engine/Source/Runtime/MovieScene/Private/EntitySystem/IMovieSceneEntityProvider.cpp
+++ b/Engine/Source/Runtime/MovieScene/Private/EntitySystem/IMovieSceneEntityProvider.cpp
@@ -35,13 +35,21 @@ FMovieSceneEntityID FImportedEntity::Manufacture(const FEntityImportParams& Para
 	.AddConditional(Components->RootInstanceHandle, Params.Sequence.RootInstanceHandle, Params.Sequence.RootInstanceHandle.IsValid())
 	.AddConditional(Components->InstanceHandle, Params.Sequence.InstanceHandle, Params.Sequence.InstanceHandle.IsValid());
 
+	FEntityAllocationWriteContext WriteContext(*EntityManager);
+
+	bool bAddMutualComponents = false;
+
 	FComponentMask NewMask;
-	BaseBuilder.GenerateType(EntityManager, NewMask);
+	BaseBuilder.GenerateType(EntityManager, NewMask, bAddMutualComponents);
 	for (TInlineValue<IEntityBuilder>& Builder : Builders)
 	{
-		Builder->GenerateType(EntityManager, NewMask);
+		Builder->GenerateType(EntityManager, NewMask, bAddMutualComponents);
 	}
 
+	FMutualComponentInitializers MutualInitializers;
+	EMutuallyInclusiveComponentType MutualTypes = bAddMutualComponents ? EMutuallyInclusiveComponentType::All : EMutuallyInclusiveComponentType::Mandatory;
+	EntityManager->GetComponents()->Factories.ComputeMutuallyInclusiveComponents(MutualTypes, NewMask, MutualInitializers);
+
 	FEntityInfo NewEntity = EntityManager->AllocateEntity(NewMask);
 
 	BaseBuilder.Initialize(EntityManager, NewEntity);
@@ -50,6 +58,10 @@ FMovieSceneEntityID FImportedEntity::Manufacture(const FEntityImportParams& Para
 		Builder->Initialize(EntityManager, NewEntity);
 	}
 
+	// Run mutual initializers after the builder has actually constructed the entity
+	// otherwise mutual components would be reading garbage
+	MutualInitializers.Execute(NewEntity.Data.AsRange(), WriteContext);
+
 	return NewEntity.EntityID;
 }
 
diff --git a/Engine/Source/Runtime/MovieScene/Private/EntitySystem/MovieSceneEntityFactory.cpp b/Engine/Source/Runtime/MovieScene/Private/EntitySystem/MovieSceneEntityFactory.cpp
--- a/Engine/Source/Runtime/MovieScene/Private/EntitySystem/MovieSceneEntityFactory.cpp
+++ b/Engine/Source/Runtime/MovieScene/Private/EntitySystem/MovieSceneEntityFactory.cpp
@@ -7,6 +7,7 @@
 #include "EntitySystem/BuiltInComponentTypes.h"
 #include "EntitySystem/MovieSceneEntitySystemTask.h"
 #include "EntitySystem/MovieSceneEntityFactoryTemplates.h"
+#include "MovieSceneFwd.h"
 
 namespace UE
 {
@@ -32,11 +33,13 @@ void FChildEntityFactory::Apply(UMovieSceneEntitySystemLinker* Linker, FEntityAl
 	const FComponentMask ParentType = ParentAllocationProxy.GetAllocationType();
 
 	FComponentMask DerivedEntityType;
+	FMutualComponentInitializers MutualInitializers;
+
 	{
 		GenerateDerivedType(DerivedEntityType);
 
 		Linker->EntityManager.GetComponents()->Factories.ComputeChildComponents(ParentType, DerivedEntityType);
-		Linker->EntityManager.GetComponents()->Factories.ComputeMutuallyInclusiveComponents(DerivedEntityType);
+		Linker->EntityManager.GetComponents()->Factories.ComputeMutuallyInclusiveComponents(EMutuallyInclusiveComponentType::All, DerivedEntityType, MutualInitializers);
 	}
 
 	const bool bHasAnyType = DerivedEntityType.Find(true) != INDEX_NONE;
@@ -49,7 +52,7 @@ void FChildEntityFactory::Apply(UMovieSceneEntitySystemLinker* Linker, FEntityAl
 
 	int32 CurrentParentOffset = 0;
 	const FEntityAllocation* ParentAllocation = ParentAllocationProxy.GetAllocation();
-	FEntityAllocationWriteContext WriteContext = FEntityAllocationWriteContext::NewAllocation();
+	FEntityAllocationWriteContext WriteContext(Linker->EntityManager);
 
 	// We attempt to allocate all the linker entities contiguously in memory for efficient initialization,
 	// but we may reach capacity constraints within allocations so we may have to run the factories more than once
@@ -79,6 +82,7 @@ void FChildEntityFactory::Apply(UMovieSceneEntitySystemLinker* Linker, FEntityAl
 		// Initialize the bound objects before we call child initializers
 		InitializeAllocation(Linker, ParentType, DerivedEntityType, ParentAllocation, CurrentEntityOffsets, ChildRange);
 
+		MutualInitializers.Execute(ChildRange, WriteContext);
 		Linker->EntityManager.InitializeChildAllocation(ParentType, DerivedEntityType, ParentAllocation, CurrentEntityOffsets, ChildRange);
 
 		CurrentParentOffset += NumAdded;
@@ -233,26 +237,35 @@ void FEntityFactories::DefineChildComponent(TInlineValue<FChildEntityInitializer
 	ChildInitializers.Add(MoveTemp(InInitializer));
 }
 
-void FEntityFactories::DefineMutuallyInclusiveComponent(FComponentTypeID InComponentA, FComponentTypeID InComponentB)
+void FEntityFactories::DefineMutuallyInclusiveComponents(FComponentTypeID InComponentA, std::initializer_list<FComponentTypeID> InMutualComponents)
+{
+	MutualInclusivityGraph.DefineMutualInclusionRule(InComponentA, InMutualComponents);
+}
+
+void FEntityFactories::DefineMutuallyInclusiveComponents(FComponentTypeID InComponentA, std::initializer_list<FComponentTypeID> InMutualComponents, FMutuallyInclusiveComponentParams&& Params)
 {
-	MutualInclusivityGraph.AllocateNode(InComponentA.BitIndex());
-	MutualInclusivityGraph.AllocateNode(InComponentB.BitIndex());
-	MutualInclusivityGraph.MakeEdge(InComponentA.BitIndex(), InComponentB.BitIndex());
-	Masks.AllMutualFirsts.Set(InComponentA);
+	MutualInclusivityGraph.DefineMutualInclusionRule(InComponentA, InMutualComponents, MoveTemp(Params));
 }
 
 void FEntityFactories::DefineComplexInclusiveComponents(const FComplexInclusivityFilter& InFilter, FComponentTypeID InComponent)
 {
-	FComponentMask ComponentsToInclude { InComponent };
-	FComplexInclusivity NewComplexInclusivity { InFilter, ComponentsToInclude };
-	DefineComplexInclusiveComponents(NewComplexInclusivity);
+	MutualInclusivityGraph.DefineComplexInclusionRule(InFilter, { InComponent });
 }
 
+void FEntityFactories::DefineComplexInclusiveComponents(const FComplexInclusivityFilter& InFilter, std::initializer_list<FComponentTypeID> InComponents, FMutuallyInclusiveComponentParams&& Params)
+{
+	MutualInclusivityGraph.DefineComplexInclusionRule(InFilter, InComponents, MoveTemp(Params));
+}
+
+PRAGMA_DISABLE_DEPRECATION_WARNINGS
 void FEntityFactories::DefineComplexInclusiveComponents(const FComplexInclusivity& InInclusivity)
 {
-	ComplexInclusivity.Add(InInclusivity);
-	Masks.AllComplexFirsts.CombineWithBitwiseOR(InInclusivity.Filter.Mask, EBitwiseOperatorFlags::MaxSize);
+	for (FComponentMaskIterator It(InInclusivity.ComponentsToInclude.Iterate()); It; ++It)
+	{
+		MutualInclusivityGraph.DefineComplexInclusionRule(InInclusivity.Filter, { FComponentTypeID::FromBitIndex(It.GetIndex()) });
+	}
 }
+PRAGMA_ENABLE_DEPRECATION_WARNINGS
 
 int32 FEntityFactories::ComputeChildComponents(const FComponentMask& ParentComponentMask, FComponentMask& ChildComponentMask)
 {
@@ -285,85 +298,9 @@ int32 FEntityFactories::ComputeChildComponents(const FComponentMask& ParentCompo
 	return NumNewComponents;
 }
 
-int32 FEntityFactories::ComputeMutuallyInclusiveComponents(FComponentMask& ComponentMask)
+int32 FEntityFactories::ComputeMutuallyInclusiveComponents(EMutuallyInclusiveComponentType MutualTypes, FComponentMask& ComponentMask, FMutualComponentInitializers& OutInitializers)
 {
-	int32 NumNewComponents = 0;
-
-	// We have two things that can add components: filtered includes and mutual includes.
-	//
-	// Since a mutual include might add a component that will make a filter pass, and a passing filter
-	// might add a component that has a mutual include, we need to loop over both until the whole
-	// thing "stabilizes".
-	//
-	// To avoid always having to loop one extra time (with the last loop not doing anything), we check
-	// if the previous loop added anything that can potentially make an additional loop useful. It won't
-	// prevent doing a loop for nothing, but it will prevent it *most* of the time.
-	//
-	while (true)
-	{
-		int32 NumNewComponentsThisTime = 0;
-		FComponentMask NewComponentsFromMutuals;
-
-		// Complex includes.
-		for (const FComplexInclusivity& Inclusivity : ComplexInclusivity)
-		{
-			if (Inclusivity.Filter.Match(ComponentMask))
-			{
-				// Only count the components that we are truly adding. Some of the components in ComponentsToInclude
-				// could already be present in our mask, and wouldn't count as "new" here.
-				const FComponentMask Added = FComponentMask::BitwiseAND(
-						Inclusivity.ComponentsToInclude, FComponentMask::BitwiseNOT(ComponentMask),
-						EBitwiseOperatorFlags::MaxSize);
-				NumNewComponentsThisTime += Added.NumComponents();
-
-				ComponentMask.CombineWithBitwiseOR(Inclusivity.ComponentsToInclude, EBitwiseOperatorFlags::MaxSize);
-			}
-		}
-
-		// Mutual includes.
-		FMovieSceneEntitySystemDirectedGraph::FBreadthFirstSearch BFS(&MutualInclusivityGraph);
-
-		for (FComponentMaskIterator It = ComponentMask.Iterate(); It; ++It)
-		{
-			const uint16 NodeID = static_cast<uint16>(It.GetIndex());
-			if (MutualInclusivityGraph.IsNodeAllocated(NodeID))
-			{
-				BFS.Search(NodeID);
-			}
-		}
-
-		// Ideally would do a bitwise OR here
-		for (TConstSetBitIterator<> It(BFS.GetVisited()); It; ++It)
-		{
-			FComponentTypeID ComponentType = FComponentTypeID::FromBitIndex(It.GetIndex());
-			if (!ComponentMask.Contains(ComponentType))
-			{
-				NewComponentsFromMutuals.Set(ComponentType);
-				++NumNewComponentsThisTime;
-
-				ComponentMask.Set(ComponentType);
-			}
-		}
-
-		// Accumulate our count of new components.
-		NumNewComponents += NumNewComponentsThisTime;
-
-		// We don't need to do another loop if:
-		//
-		// 1. We didn't add anything this loop... 
-		//   OR
-		// 2. We added something in the "mutuals" part that we know doesn't match
-		//    any complex filter.
-		if (
-				(NumNewComponentsThisTime == 0) ||
-				(!NewComponentsFromMutuals.ContainsAny(Masks.AllComplexFirsts))
-			)
-		{
-			break;
-		}
-	}
-
-	return NumNewComponents;
+	return MutualInclusivityGraph.ComputeMutuallyInclusiveComponents(MutualTypes, ComponentMask, ComponentMask, OutInitializers);
 }
 
 void FEntityFactories::RunInitializers(const FComponentMask& ParentType, const FComponentMask& ChildType, const FEntityAllocation* ParentAllocation, TArrayView<const int32> ParentAllocationOffsets, const FEntityRange& InChildEntityRange)
diff --git a/Engine/Source/Runtime/MovieScene/Private/EntitySystem/MovieSceneEntityManager.cpp b/Engine/Source/Runtime/MovieScene/Private/EntitySystem/MovieSceneEntityManager.cpp
--- a/Engine/Source/Runtime/MovieScene/Private/EntitySystem/MovieSceneEntityManager.cpp
+++ b/Engine/Source/Runtime/MovieScene/Private/EntitySystem/MovieSceneEntityManager.cpp
@@ -1622,56 +1622,52 @@ void FEntityManager::CombineComponents(FMovieSceneEntityID DestinationEntityID,
 	OnStructureChanged();
 }
 
-int32 FEntityManager::MutateAll(const FEntityComponentFilter& Filter, const IMovieSceneEntityMutation& Mutation)
+int32 FEntityManager::MutateAll(const FEntityComponentFilter& Filter, const IMovieSceneEntityMutation& Mutation, EMutuallyInclusiveComponentType MutualTypes)
 {
 	CheckCanChangeStructure();
 
-	TMap<int32, FComponentMask> AllocationMutations;
+	FMutualComponentInitializers MutualInitializers;
+	FEntityAllocationWriteContext WriteContext(*this);
+
+	int32 TotalNumMutations = 0;
 
 	for (int32 AllocationIndex = 0; AllocationIndex < EntityAllocationMasks.GetMaxIndex(); ++AllocationIndex)
 	{
-		if (EntityAllocationMasks.IsAllocated(AllocationIndex) && Filter.Match(EntityAllocationMasks[AllocationIndex]))
+		if (!EntityAllocationMasks.IsAllocated(AllocationIndex) || !Filter.Match(EntityAllocationMasks[AllocationIndex]))
 		{
-			FComponentMask NewMutation = EntityAllocationMasks[AllocationIndex];
-			Mutation.CreateMutation(this, &NewMutation);
-
-			if (!NewMutation.CompareSetBits(EntityAllocationMasks[AllocationIndex]))
-			{
-				AllocationMutations.Add(AllocationIndex, MoveTemp(NewMutation));
-			}
-			else
-			{
-				FEntityAllocation* SourceAllocation = EntityAllocations[AllocationIndex];
-				Mutation.InitializeUnmodifiedAllocation(SourceAllocation, NewMutation);
-			}
+			continue;
 		}
-	}
 
-	if (AllocationMutations.Num() == 0)
-	{
-		return 0;
-	}
+		// Process the mutation
+		FComponentMask NewAllocationType = EntityAllocationMasks[AllocationIndex];
+		Mutation.CreateMutation(this, &NewAllocationType);
 
-	int32 TotalNumMutations = 0;
+		// Add mutual components
+		MutualInitializers.Reset();
+		ComponentRegistry->Factories.ComputeMutuallyInclusiveComponents(MutualTypes, NewAllocationType, MutualInitializers);
 
-	for (TTuple<int32, FComponentMask>& Pair : AllocationMutations)
-	{
-		int32 AllocationIndex = Pair.Key;
 		FEntityAllocation* SourceAllocation = EntityAllocations[AllocationIndex];
 
-		// When adding a component to an entire allocation we just reallocate within the same allocation entry to avoid having to fix up 
-		// Specific entity entry indices
-		FEntityAllocation* NewAllocation = MigrateAllocation(AllocationIndex, Pair.Value);
+		// If the type hasn't changed at all, we can just run the unmodified initializer
+		if (NewAllocationType.CompareSetBits(EntityAllocationMasks[AllocationIndex]))
+		{
+			Mutation.InitializeUnmodifiedAllocation(SourceAllocation, NewAllocationType);
+			continue;
+		}
+
+		// The type has changed so we need to migrate the allocation - we just reallocate within
+		// the same allocation index to avoid having to fix up specific entity entry indices
+		FEntityAllocation* NewAllocation = MigrateAllocation(AllocationIndex, NewAllocationType);
 		TotalNumMutations += NewAllocation->Num();
 
-		FComponentMask OldAllocationType = EntityAllocationMasks[Pair.Key];
-		EntityAllocationMasks[Pair.Key] = Pair.Value;
-		EntityAllocations[Pair.Key] = NewAllocation;
+		FComponentMask OldAllocationType = EntityAllocationMasks[AllocationIndex];
+		EntityAllocationMasks[AllocationIndex] = NewAllocationType;
+		EntityAllocations[AllocationIndex] = NewAllocation;
 
 		FEntityAllocationMutexGuard LockGuard(NewAllocation, EComponentHeaderLockMode::LockFree);
 
 		// Default construct all the new components in the allocation, and then allow the mutation to further initialize this data if needed
-		for (FComponentMaskIterator Component = Pair.Value.Iterate(); Component; ++Component)
+		for (FComponentMaskIterator Component = NewAllocationType.Iterate(); Component; ++Component)
 		{
 			FComponentTypeID ComponentTypeID = FComponentTypeID::FromBitIndex(Component.GetIndex());
 			if (OldAllocationType.Contains(ComponentTypeID))
@@ -1688,24 +1684,39 @@ int32 FEntityManager::MutateAll(const FEntityComponentFilter& Filter, const IMov
 				ComponentTypeInfo.ConstructItems(Components, NewAllocation->Num());
 			}
 		}
-		Mutation.InitializeAllocation(NewAllocation, Pair.Value);
 
+		// Run custom initializers
+		MutualInitializers.Execute(FEntityRange{ NewAllocation, 0, NewAllocation->Num() }, WriteContext);
+
+		// Run mutation initializer
+		Mutation.InitializeAllocation(NewAllocation, NewAllocationType);
+
+		// Destroy the old data
 		DestroyAllocation(SourceAllocation);
 	}
 
-	CheckInvariants();
-
-	OnStructureChanged();
+	if (TotalNumMutations != 0)
+	{
+		CheckInvariants();
+		OnStructureChanged();
+	}
 
 	return TotalNumMutations;
 }
 
-int32 FEntityManager::MutateConditional(const FEntityComponentFilter& Filter, const IMovieSceneConditionalEntityMutation& Mutation)
+int32 FEntityManager::MutateConditional(const FEntityComponentFilter& Filter, const IMovieSceneConditionalEntityMutation& Mutation, EMutuallyInclusiveComponentType MutualTypes)
 {
 	CheckCanChangeStructure();
 
-	TMap<int32, TTuple<FComponentMask, TBitArray<>>> AllocationMutations;
+	struct FConditionalMutationState
+	{
+		FComponentMask NewType;
+		TBitArray<> MarkedEntities;
+		FMutualComponentInitializers MutualInitializers;
+	};
+	TMap<int32, FConditionalMutationState> AllocationMutations;
 
+	FMutualComponentInitializers MutualInitializersScratch;
 	for (int32 AllocationIndex = 0; AllocationIndex < EntityAllocationMasks.GetMaxIndex(); ++AllocationIndex)
 	{
 		if (EntityAllocationMasks.IsAllocated(AllocationIndex) && Filter.Match(EntityAllocationMasks[AllocationIndex]))
@@ -1729,9 +1740,17 @@ int32 FEntityManager::MutateConditional(const FEntityComponentFilter& Filter, co
 			FComponentMask NewMutation = EntityAllocationMasks[AllocationIndex];
 			Mutation.CreateMutation(this, &NewMutation);
 
+			// Add mutual components
+			MutualInitializersScratch.Reset();
+			ComponentRegistry->Factories.ComputeMutuallyInclusiveComponents(MutualTypes, NewMutation, MutualInitializersScratch);
+
 			if (!NewMutation.CompareSetBits(EntityAllocationMasks[AllocationIndex]))
 			{
-				AllocationMutations.Add(AllocationIndex, MakeTuple(MoveTemp(NewMutation), MoveTemp(MarkedEntities)));
+				AllocationMutations.Add(AllocationIndex, FConditionalMutationState{
+					MoveTemp(NewMutation),
+					MoveTemp(MarkedEntities),
+					MoveTemp(MutualInitializersScratch)
+				});
 			}
 		}
 	}
@@ -1743,24 +1762,28 @@ int32 FEntityManager::MutateConditional(const FEntityComponentFilter& Filter, co
 
 	int32 TotalNumMutations = 0;
 
-	for (TTuple<int32, TTuple<FComponentMask, TBitArray<>>>& Pair : AllocationMutations)
+	FMutualComponentInitializers MutualInitializers;
+	FEntityAllocationWriteContext WriteContext(*this);
+
+	for (TTuple<int32, FConditionalMutationState>& Pair : AllocationMutations)
 	{
 		const int32 SourceAllocationIndex = Pair.Key;
 		FEntityAllocation* SourceAllocation = EntityAllocations[SourceAllocationIndex];
 
-		const int32 NumEntities = Pair.Value.Value.CountSetBits();
+		const int32 NumEntities = Pair.Value.MarkedEntities.CountSetBits();
 		const bool bAllEntitiesMarked = NumEntities == SourceAllocation->Num();
 
 		TotalNumMutations += NumEntities;
 
+		const FComponentMask& DesiredType = Pair.Value.NewType;
 		if (bAllEntitiesMarked)
 		{
 			// When adding a component to an entire allocation we just reallocate within the same allocation entry to avoid having to fix up 
 			// Specific entity entry indices
-			FEntityAllocation* NewAllocation = MigrateAllocation(SourceAllocationIndex, Pair.Value.Key);
+			FEntityAllocation* NewAllocation = MigrateAllocation(SourceAllocationIndex, DesiredType);
 
 			// Default construct all the new components in the allocation, and then allow the mutation to further initialize this data if needed
-			for (FComponentMaskIterator Component = Pair.Value.Key.Iterate(); Component; ++Component)
+			for (FComponentMaskIterator Component = DesiredType.Iterate(); Component; ++Component)
 			{
 				FComponentTypeID ComponentTypeID = FComponentTypeID::FromBitIndex(Component.GetIndex());
 				if (EntityAllocationMasks[SourceAllocationIndex].Contains(ComponentTypeID))
@@ -1778,19 +1801,22 @@ int32 FEntityManager::MutateConditional(const FEntityComponentFilter& Filter, co
 				}
 			}
 
-			EntityAllocationMasks[SourceAllocationIndex] = Pair.Value.Key;
+			EntityAllocationMasks[SourceAllocationIndex] = DesiredType;
 			EntityAllocations[SourceAllocationIndex] = NewAllocation;
 
 			FEntityRange Range { NewAllocation, 0, NewAllocation->Num() };
 
-			Mutation.InitializeEntities(Range, Pair.Value.Key);
+			// Run custom initializers
+			Pair.Value.MutualInitializers.Execute(Range, WriteContext);
+
+			Mutation.InitializeEntities(Range, DesiredType);
 
 			DestroyAllocation(SourceAllocation);
 		}
 		else
 		{
 			int32 DesiredSlack = NumEntities;
-			const int32 NewAllocationIndex = CreateAllocationWithSlack(Pair.Value.Key, &DesiredSlack);
+			const int32 NewAllocationIndex = CreateAllocationWithSlack(DesiredType, &DesiredSlack);
 
 			const FEntityAllocation* DestAllocation = EntityAllocations[NewAllocationIndex];
 			FEntityRange Range { DestAllocation, DestAllocation->Num(), NumEntities };
@@ -1799,7 +1825,7 @@ int32 FEntityManager::MutateConditional(const FEntityComponentFilter& Filter, co
 
 			// Migrate entities to the new allocation - care is taken to iterate the allocation backwards
 			// since entities can shift around in the allocation as they are migrated
-			for (TBitArray<>::FConstReverseIterator It(Pair.Value.Value); It; ++It)
+			for (TBitArray<>::FConstReverseIterator It(Pair.Value.MarkedEntities); It; ++It)
 			{
 				if (It.GetValue())
 				{
@@ -1811,7 +1837,11 @@ int32 FEntityManager::MutateConditional(const FEntityComponentFilter& Filter, co
 			}
 
 			check(Range.ComponentStartOffset + Range.Num == Range.Allocation->Num());
-			Mutation.InitializeEntities(Range, Pair.Value.Key);
+
+			// Run custom initializers
+			Pair.Value.MutualInitializers.Execute(Range, WriteContext);
+
+			Mutation.InitializeEntities(Range, DesiredType);
 		}
 	}
 
@@ -1893,66 +1923,14 @@ void FEntityManager::AddMutualComponents(const FEntityComponentFilter& InFilter)
 {
 	CheckCanChangeStructure();
 
-	TMap<int32, FComponentMask> AllocationMutations;
-
-	for (int32 AllocationIndex = 0; AllocationIndex < EntityAllocations.GetMaxIndex(); ++AllocationIndex)
-	{
-		if (!EntityAllocations.IsValidIndex(AllocationIndex))
-		{
-			continue;
-		}
-
-		FComponentMask NewAllocationType = EntityAllocationMasks[AllocationIndex];
-		if (InFilter.IsValid() && !InFilter.Match(NewAllocationType))
-		{
-			continue;
-		}
-
-		const int32 NumNewComponents = ComponentRegistry->Factories.ComputeMutuallyInclusiveComponents(NewAllocationType);
-		if (NumNewComponents != 0)
-		{
-			AllocationMutations.Add(AllocationIndex, NewAllocationType);
-		}
-	}
-
-	if (AllocationMutations.Num() == 0)
+	struct FBenignMutation : IMovieSceneEntityMutation
 	{
-		return;
-	}
-
-	for (TTuple<int32, FComponentMask>& Pair : AllocationMutations)
-	{
-		int32 AllocationIndex = Pair.Key;
-		FEntityAllocation* SourceAllocation = EntityAllocations[AllocationIndex];
-
-		// When adding a component to an entire allocation we just reallocate within the same allocation entry to avoid having to fix up 
-		// Specific entity entry indices
-		FEntityAllocation* NewAllocation = MigrateAllocation(AllocationIndex, Pair.Value);
-
-		FComponentMask NewComponents = FComponentMask::BitwiseXOR(Pair.Value, EntityAllocationMasks[AllocationIndex], EBitwiseOperatorFlags::MaxSize);
-
-		for (FComponentMaskIterator Component = NewComponents.Iterate(); Component; ++Component)
-		{
-			FComponentTypeID ComponentTypeID = FComponentTypeID::FromBitIndex(Component.GetIndex());
-
-			const FComponentHeader& ComponentHeader = NewAllocation->GetComponentHeaderChecked(ComponentTypeID);
-			if (!ComponentHeader.IsTag())
-			{
-				const FComponentTypeInfo& ComponentTypeInfo = ComponentRegistry->GetComponentTypeChecked(ComponentTypeID);
-
-				void* Components = ComponentHeader.GetValuePtr(0);
-				ComponentTypeInfo.ConstructItems(Components, NewAllocation->Num());
-			}
-		}
-
-		EntityAllocationMasks[Pair.Key] = Pair.Value;
-		EntityAllocations[Pair.Key] = NewAllocation;
-
-		DestroyAllocation(SourceAllocation);
-	}
+		void CreateMutation(FEntityManager* EntityManager, FComponentMask* InOutEntityComponentTypes) const override
+		{}
+	};
 
-	CheckInvariants();
-	OnStructureChanged();
+	FBenignMutation BenignMutation;
+	MutateAll(InFilter, BenignMutation, EMutuallyInclusiveComponentType::All);
 }
 
 FEntityAllocation* FEntityManager::MigrateAllocation(int32 AllocationIndex, const FComponentMask& NewComponentMask)
diff --git a/Engine/Source/Runtime/MovieScene/Private/EntitySystem/MovieSceneEntitySystem.cpp b/Engine/Source/Runtime/MovieScene/Private/EntitySystem/MovieSceneEntitySystem.cpp
--- a/Engine/Source/Runtime/MovieScene/Private/EntitySystem/MovieSceneEntitySystem.cpp
+++ b/Engine/Source/Runtime/MovieScene/Private/EntitySystem/MovieSceneEntitySystem.cpp
@@ -108,7 +108,7 @@ struct FSystemDependencyGraph
 
 	void GetSubsequents(uint16 FromNodeID, TArray<uint16>& OutSubsequentNodeIDs)
 	{
-		using FDirectionalEdge = FMovieSceneEntitySystemDirectedGraph::FDirectionalEdge;
+		using FDirectionalEdge = FDirectedGraph::FDirectionalEdge;
 
 		UpdateCache();
 
@@ -136,7 +136,7 @@ struct FSystemDependencyGraph
 
 	FString ToString() const
 	{
-		using FDirectionalEdge = FMovieSceneEntitySystemDirectedGraph::FDirectionalEdge;
+		using FDirectionalEdge = FDirectedGraph::FDirectionalEdge;
 
 		FComponentRegistry* ComponentRegistry = UMovieSceneEntitySystemLinker::GetComponents();
 
@@ -249,7 +249,7 @@ struct FSystemDependencyGraph
 		}
 
 		{
-			FMovieSceneEntitySystemDirectedGraph::FDiscoverCyclicEdges CyclicEdges(&FlowGraph);
+			FDirectedGraph::FDiscoverCyclicEdges CyclicEdges(&FlowGraph);
 			CyclicEdges.Search();
 
 			TArrayView<const FDirectionalEdge> FlowEdges = FlowGraph.GetEdges();
@@ -309,7 +309,7 @@ private:
 
 		BuildFlowGraph();
 
-		FMovieSceneEntitySystemDirectedGraph::FDepthFirstSearch DepthFirstSearch(&FlowGraph);
+		FDirectedGraph::FDepthFirstSearch DepthFirstSearch(&FlowGraph);
 
 		TBitArray<> EdgeNodes = FlowGraph.FindEdgeUpstreamNodes();
 		for (TConstSetBitIterator<> EdgeNodeIt(EdgeNodes); EdgeNodeIt; ++EdgeNodeIt)
@@ -352,7 +352,7 @@ private:
 
 	void SetupFlowDependencies(int32 NodeIndex)
 	{
-		using FDirectionalEdge = FMovieSceneEntitySystemDirectedGraph::FDirectionalEdge;
+		using FDirectionalEdge = FDirectedGraph::FDirectionalEdge;
 
 		const uint16 CurrentNodeID(static_cast<uint16>(NodeIndex));
 
@@ -426,10 +426,10 @@ private:
 
 	TArray<FNode> Nodes;
 
-	FMovieSceneEntitySystemDirectedGraph ImplicitPrerequisites;
-	FMovieSceneEntitySystemDirectedGraph ImplicitSubsequents;
+	FDirectedGraph ImplicitPrerequisites;
+	FDirectedGraph ImplicitSubsequents;
 
-	FMovieSceneEntitySystemDirectedGraph FlowGraph;
+	FDirectedGraph FlowGraph;
 	TArray<uint16> FlowOrderNodes;
 	TArray<uint16> ReverseLookupFlowOrderNodes;
 
diff --git a/Engine/Source/Runtime/MovieScene/Private/EntitySystem/MovieSceneEntitySystemDirectedGraph.cpp b/Engine/Source/Runtime/MovieScene/Private/EntitySystem/MovieSceneEntitySystemDirectedGraph.cpp
--- a/Engine/Source/Runtime/MovieScene/Private/EntitySystem/MovieSceneEntitySystemDirectedGraph.cpp
+++ b/Engine/Source/Runtime/MovieScene/Private/EntitySystem/MovieSceneEntitySystemDirectedGraph.cpp
@@ -1,10 +1,15 @@
 // Copyright Epic Games, Inc. All Rights Reserved.
 
 #include "EntitySystem/MovieSceneEntitySystemDirectedGraph.h"
+#include "Containers/UnrealString.h"
+#include "Misc/StringBuilder.h"
 
 //#include UE_INLINE_GENERATED_CPP_BY_NAME(MovieSceneEntitySystemDirectedGraph)
 
-FMovieSceneEntitySystemDirectedGraph::FDepthFirstSearch::FDepthFirstSearch(const FMovieSceneEntitySystemDirectedGraph* InGraph)
+namespace UE::MovieScene
+{
+
+FDirectedGraph::FDepthFirstSearch::FDepthFirstSearch(const FDirectedGraph* InGraph)
 	: Visited(false, InGraph->Nodes.Num())
 	, IsVisiting(false, InGraph->Nodes.Num())
 {
@@ -13,11 +18,11 @@ FMovieSceneEntitySystemDirectedGraph::FDepthFirstSearch::FDepthFirstSearch(const
 	check(!Graph->bHasDanglingEdges);
 }
 
-void FMovieSceneEntitySystemDirectedGraph::FDepthFirstSearch::Search(uint16 InNodeID)
+void FDirectedGraph::FDepthFirstSearch::Search(uint16 InNodeID)
 {
 	IsVisiting[InNodeID] = true;
 
-	for (FMovieSceneEntitySystemDirectedGraph::FDirectionalEdge Edge : Graph->GetEdgesFrom(InNodeID))
+	for (FDirectedGraph::FDirectionalEdge Edge : Graph->GetEdgesFrom(InNodeID))
 	{
 		if (Visited[Edge.ToNode] == false)
 		{
@@ -35,7 +40,7 @@ void FMovieSceneEntitySystemDirectedGraph::FDepthFirstSearch::Search(uint16 InNo
 	IsVisiting[InNodeID] = false;
 }
 
-FMovieSceneEntitySystemDirectedGraph::FBreadthFirstSearch::FBreadthFirstSearch(const FMovieSceneEntitySystemDirectedGraph* InGraph)
+FDirectedGraph::FBreadthFirstSearch::FBreadthFirstSearch(const FDirectedGraph* InGraph)
 	: Visited(false, InGraph->Nodes.Num())
 	, StackIndex(0)
 {
@@ -44,7 +49,7 @@ FMovieSceneEntitySystemDirectedGraph::FBreadthFirstSearch::FBreadthFirstSearch(c
 	check(!InGraph->bHasDanglingEdges);
 }
 
-void FMovieSceneEntitySystemDirectedGraph::FBreadthFirstSearch::Search(uint16 InNodeID)
+void FDirectedGraph::FBreadthFirstSearch::Search(uint16 InNodeID)
 {
 	if (Visited[InNodeID] == true)
 	{
@@ -65,7 +70,7 @@ void FMovieSceneEntitySystemDirectedGraph::FBreadthFirstSearch::Search(uint16 In
 			const uint16 NodeID = Nodes[StackIndex];
 
 			// Visit all nodes this points to
-			for (FMovieSceneEntitySystemDirectedGraph::FDirectionalEdge Edge : Graph->GetEdgesFrom(NodeID))
+			for (FDirectedGraph::FDirectionalEdge Edge : Graph->GetEdgesFrom(NodeID))
 			{
 				if (Visited[Edge.ToNode] == false)
 				{
@@ -77,7 +82,7 @@ void FMovieSceneEntitySystemDirectedGraph::FBreadthFirstSearch::Search(uint16 In
 	}
 }
 
-FMovieSceneEntitySystemDirectedGraph::FDiscoverCyclicEdges::FDiscoverCyclicEdges(const FMovieSceneEntitySystemDirectedGraph* InGraph)
+FDirectedGraph::FDiscoverCyclicEdges::FDiscoverCyclicEdges(const FDirectedGraph* InGraph)
 	: CyclicEdges(false, InGraph->SortedEdges.Num())
 	, VisitedEdges(false, InGraph->SortedEdges.Num())
 {
@@ -85,7 +90,7 @@ FMovieSceneEntitySystemDirectedGraph::FDiscoverCyclicEdges::FDiscoverCyclicEdges
 	check(!Graph->bHasDanglingEdges);
 }
 
-void FMovieSceneEntitySystemDirectedGraph::FDiscoverCyclicEdges::Search()
+void FDirectedGraph::FDiscoverCyclicEdges::Search()
 {
 	for (uint16 EdgeIndex = 0; EdgeIndex < Graph->SortedEdges.Num(); ++EdgeIndex)
 	{
@@ -101,14 +106,14 @@ void FMovieSceneEntitySystemDirectedGraph::FDiscoverCyclicEdges::Search()
 	}
 }
 
-void FMovieSceneEntitySystemDirectedGraph::FDiscoverCyclicEdges::SearchFrom(uint16 NodeID)
+void FDirectedGraph::FDiscoverCyclicEdges::SearchFrom(uint16 NodeID)
 {
 	TBitArray<> VisitedNodes(false, Graph->Nodes.Num());
 	VisitedNodes[NodeID] = true;
 	DiscoverCycles(NodeID, VisitedNodes);
 }
 
-void FMovieSceneEntitySystemDirectedGraph::FDiscoverCyclicEdges::DiscoverCycles(uint16 NodeID, TBitArray<>& VisitedNodes)
+void FDirectedGraph::FDiscoverCyclicEdges::DiscoverCycles(uint16 NodeID, TBitArray<>& VisitedNodes)
 {
 	// Iterate all edges from this node
 	for (int32 SubsequentEdge = Graph->FindEdgeStart(NodeID); SubsequentEdge < Graph->SortedEdges.Num() && Graph->SortedEdges[SubsequentEdge].FromNode == NodeID; ++SubsequentEdge)
@@ -137,7 +142,7 @@ void FMovieSceneEntitySystemDirectedGraph::FDiscoverCyclicEdges::DiscoverCycles(
 	}
 }
 
-void FMovieSceneEntitySystemDirectedGraph::FDiscoverCyclicEdges::TagCyclicChain(uint16 CyclicNodeID)
+void FDirectedGraph::FDiscoverCyclicEdges::TagCyclicChain(uint16 CyclicNodeID)
 {
 	// Found a cycle
 	for (int32 EdgeChainIndex = EdgeChain.Num() - 1; EdgeChainIndex >= 0; --EdgeChainIndex)
@@ -153,19 +158,19 @@ void FMovieSceneEntitySystemDirectedGraph::FDiscoverCyclicEdges::TagCyclicChain(
 	}
 }
 
-void FMovieSceneEntitySystemDirectedGraph::AllocateNode(uint16 NodeID)
+void FDirectedGraph::AllocateNode(uint16 NodeID)
 {
 	CleanUpDanglingEdges();
 	Nodes.PadToNum(NodeID + 1, false);
 	Nodes[NodeID] = true;
 }
 
-bool FMovieSceneEntitySystemDirectedGraph::IsNodeAllocated(uint16 NodeID) const
+bool FDirectedGraph::IsNodeAllocated(uint16 NodeID) const
 {
 	return Nodes.IsValidIndex(NodeID) && Nodes[NodeID] == true;
 }
 
-void FMovieSceneEntitySystemDirectedGraph::RemoveNode(uint16 NodeID)
+void FDirectedGraph::RemoveNode(uint16 NodeID)
 {
 	check(NodeID != TNumericLimits<uint16>::Max() && IsNodeAllocated(NodeID));
 
@@ -175,7 +180,7 @@ void FMovieSceneEntitySystemDirectedGraph::RemoveNode(uint16 NodeID)
 	bHasDanglingEdges = true;
 }
 
-void FMovieSceneEntitySystemDirectedGraph::CleanUpDanglingEdges()
+void FDirectedGraph::CleanUpDanglingEdges()
 {
 	if (!bHasDanglingEdges)
 	{
@@ -197,7 +202,7 @@ void FMovieSceneEntitySystemDirectedGraph::CleanUpDanglingEdges()
 	}
 }
 
-bool FMovieSceneEntitySystemDirectedGraph::IsCyclic() const
+bool FDirectedGraph::IsCyclic() const
 {
 	TBitArray<> Visited(false, Nodes.Num());
 
@@ -220,7 +225,7 @@ bool FMovieSceneEntitySystemDirectedGraph::IsCyclic() const
 	return false;
 }
 
-bool FMovieSceneEntitySystemDirectedGraph::IsCyclicImpl(uint16 NodeID, TBitArray<>& Visiting) const
+bool FDirectedGraph::IsCyclicImpl(uint16 NodeID, TBitArray<>& Visiting) const
 {
 	if (Visiting[NodeID] == true)
 	{
@@ -241,7 +246,7 @@ bool FMovieSceneEntitySystemDirectedGraph::IsCyclicImpl(uint16 NodeID, TBitArray
 	return false;
 }
 
-void FMovieSceneEntitySystemDirectedGraph::MakeEdge(uint16 FromNode, uint16 ToNode)
+void FDirectedGraph::MakeEdge(uint16 FromNode, uint16 ToNode)
 {
 	FDirectionalEdge NewEdge(FromNode, ToNode);
 
@@ -252,7 +257,7 @@ void FMovieSceneEntitySystemDirectedGraph::MakeEdge(uint16 FromNode, uint16 ToNo
 	}
 }
 
-void FMovieSceneEntitySystemDirectedGraph::DestroyEdge(uint16 FromNode, uint16 ToNode)
+void FDirectedGraph::DestroyEdge(uint16 FromNode, uint16 ToNode)
 {
 	FDirectionalEdge Edge(FromNode, ToNode);
 
@@ -263,36 +268,36 @@ void FMovieSceneEntitySystemDirectedGraph::DestroyEdge(uint16 FromNode, uint16 T
 	}
 }
 
-void FMovieSceneEntitySystemDirectedGraph::DestroyAllEdges()
+void FDirectedGraph::DestroyAllEdges()
 {
 	SortedEdges.Reset();
 	bHasDanglingEdges = false;
 }
 
-int32 FMovieSceneEntitySystemDirectedGraph::FindEdgeStart(uint16 FromNode) const
+int32 FDirectedGraph::FindEdgeStart(uint16 FromNode) const
 {
 	return Algo::LowerBoundBy(SortedEdges, FromNode, &FDirectionalEdge::FromNode);
 }
 
-TArrayView<const FMovieSceneEntitySystemDirectedGraph::FDirectionalEdge> FMovieSceneEntitySystemDirectedGraph::GetEdges() const
+TArrayView<const FDirectedGraph::FDirectionalEdge> FDirectedGraph::GetEdges() const
 {
 	check(!bHasDanglingEdges);
 	return SortedEdges;
 }
 
-bool FMovieSceneEntitySystemDirectedGraph::HasEdgeFrom(uint16 InNode) const
+bool FDirectedGraph::HasEdgeFrom(uint16 InNode) const
 {
 	const int32 ExpectedIndex = FindEdgeStart(InNode);
 	return SortedEdges.IsValidIndex(ExpectedIndex) && SortedEdges[ExpectedIndex].FromNode == InNode;
 }
 
-bool FMovieSceneEntitySystemDirectedGraph::HasEdgeTo(uint16 InNode) const
+bool FDirectedGraph::HasEdgeTo(uint16 InNode) const
 {
 	check(InNode != TNumericLimits<uint16>::Max());
 	return Algo::FindBy(SortedEdges, InNode, &FDirectionalEdge::ToNode) != nullptr;
 }
 
-TArrayView<const FMovieSceneEntitySystemDirectedGraph::FDirectionalEdge> FMovieSceneEntitySystemDirectedGraph::GetEdgesFrom(uint16 InNodeID) const
+TArrayView<const FDirectedGraph::FDirectionalEdge> FDirectedGraph::GetEdgesFrom(uint16 InNodeID) const
 {
 	check(!bHasDanglingEdges);
 
@@ -311,7 +316,7 @@ TArrayView<const FMovieSceneEntitySystemDirectedGraph::FDirectionalEdge> FMovieS
 	return TArrayView<const FDirectionalEdge>();
 }
 
-TBitArray<> FMovieSceneEntitySystemDirectedGraph::FindEdgeUpstreamNodes() const
+TBitArray<> FDirectedGraph::FindEdgeUpstreamNodes() const
 {
 	check(!bHasDanglingEdges);
 
@@ -328,13 +333,13 @@ TBitArray<> FMovieSceneEntitySystemDirectedGraph::FindEdgeUpstreamNodes() const
 	return TBitArray<>::BitwiseAND(EdgeNodes, Nodes, EBitwiseOperatorFlags::MaxSize);
 }
 
-int32 FMovieSceneEntitySystemDirectedGraph::FindEdgeIndex(const FDirectionalEdge& Edge) const
+int32 FDirectedGraph::FindEdgeIndex(const FDirectionalEdge& Edge) const
 {
 	check(!bHasDanglingEdges);
 	return Algo::LowerBound(SortedEdges, Edge);
 }
 
-bool FMovieSceneEntitySystemDirectedGraph::EdgeExists(const FDirectionalEdge& Edge) const
+bool FDirectedGraph::EdgeExists(const FDirectionalEdge& Edge) const
 {
 	check(!bHasDanglingEdges);
 
@@ -342,3 +347,61 @@ bool FMovieSceneEntitySystemDirectedGraph::EdgeExists(const FDirectionalEdge& Ed
 	return EdgeIndex < SortedEdges.Num() && SortedEdges[EdgeIndex] == Edge;
 }
 
+FString FDirectedGraph::ToString(const FDirectedGraphStringParameters& Parameters, TFunctionRef<void(uint16, FStringBuilderBase&)> EmitLabel) const
+{
+	TStringBuilder<256> String;
+
+	String += TEXT("\ndigraph FDirectedGraph {\n");
+	String += TEXT("\tnode [shape=record,height=.1];\n");
+	String += TEXT("\tsubgraph cluster_flow_0\n\t{\n");
+	if (Parameters.ClusterName.Len() != 0)
+	{
+		String.Append(TEXT("\tlabel=\""));
+		String.Append(Parameters.ClusterName);
+		String.Append(TEXT("\";\n"));
+	}
+	if (Parameters.Color != FColor::White)
+	{
+		String.Appendf(TEXT("\tcolor=\"%s\";\n"), *Parameters.Color.ToHex());
+	}
+	String += TEXT("\t}\n");
+
+	// Add the nodes to the graph
+	for (TConstSetBitIterator<> SetBitIt(Nodes); SetBitIt; ++SetBitIt)
+	{
+		const uint16 NodeID = static_cast<uint16>(SetBitIt.GetIndex());
+
+		String += TEXT("\t\tnode%d[label=\"");
+		EmitLabel(NodeID, String);
+		String += TEXT("\"];\n");
+	}
+
+	String += TEXT("\t}\n\n");
+
+	// Make edges
+	{
+		FDirectedGraph::FDiscoverCyclicEdges CyclicEdges(this);
+		CyclicEdges.Search();
+
+		TArrayView<const FDirectedGraph::FDirectionalEdge> Edges = GetEdges();
+		for (int32 EdgeIndex = 0; EdgeIndex < Edges.Num(); ++EdgeIndex)
+		{
+			FDirectedGraph::FDirectionalEdge Edge = Edges[EdgeIndex];
+			const bool bIsCyclic = CyclicEdges.IsCyclic(EdgeIndex);
+
+			String += FString::Printf(TEXT("\tnode%d -> node%d [color=\"%s\"];\n"), (int32)Edge.FromNode, (int32)Edge.ToNode, bIsCyclic ? TEXT("#FF0000") : TEXT("#3992ad"));
+		}
+	}
+
+	String += TEXT("}\n\n");
+	return String.ToString();
+}
+
+FString FDirectedGraph::ToString(const FDirectedGraphStringParameters& Parameters) const
+{
+	return ToString(Parameters, [](uint16 NodeID, FStringBuilderBase& OutStringBuilder){
+		OutStringBuilder.Appendf(TEXT("Node %d"), NodeID);
+	});
+}
+
+} // namespace UE::MovieScene
\ No newline at end of file
diff --git a/Engine/Source/Runtime/MovieScene/Private/EntitySystem/MovieSceneEntitySystemGraphs.cpp b/Engine/Source/Runtime/MovieScene/Private/EntitySystem/MovieSceneEntitySystemGraphs.cpp
--- a/Engine/Source/Runtime/MovieScene/Private/EntitySystem/MovieSceneEntitySystemGraphs.cpp
+++ b/Engine/Source/Runtime/MovieScene/Private/EntitySystem/MovieSceneEntitySystemGraphs.cpp
@@ -94,7 +94,7 @@ int32 FMovieSceneEntitySystemGraph::RemoveIrrelevantSystems(UMovieSceneEntitySys
 
 	int32 NumRemoved = 0;
 
-	FMovieSceneEntitySystemDirectedGraph::FBreadthFirstSearch Search(&ReferenceGraph);
+	UE::MovieScene::FDirectedGraph::FBreadthFirstSearch Search(&ReferenceGraph);
 
 	// Search from all non-intermediate systems and mark systems that are still referenced
 	for (const FMovieSceneEntitySystemGraphNode& Node : Nodes.Array)
@@ -269,7 +269,7 @@ FString FMovieSceneEntitySystemGraph::ToString() const
 	String += TEXT("\t}\n");
 
 	{
-		FMovieSceneEntitySystemDirectedGraph::FDiscoverCyclicEdges CyclicEdges(&ReferenceGraph);
+		FDirectedGraph::FDiscoverCyclicEdges CyclicEdges(&ReferenceGraph);
 		CyclicEdges.Search();
 
 		TArrayView<const FDirectionalEdge> ReferenceEdges = ReferenceGraph.GetEdges();
diff --git a/Engine/Source/Runtime/MovieScene/Private/EntitySystem/MovieSceneRootInstantiatorSystem.cpp b/Engine/Source/Runtime/MovieScene/Private/EntitySystem/MovieSceneRootInstantiatorSystem.cpp
--- a/Engine/Source/Runtime/MovieScene/Private/EntitySystem/MovieSceneRootInstantiatorSystem.cpp
+++ b/Engine/Source/Runtime/MovieScene/Private/EntitySystem/MovieSceneRootInstantiatorSystem.cpp
@@ -49,6 +49,7 @@ void UMovieSceneRootInstantiatorSystem::InstantiateAllocation(const UE::MovieSce
 
 	const FComponentMask PreservationMask = Linker->EntityManager.GetComponents()->GetPreservationMask();
 
+	FMutualComponentInitializers MutualInitializers;
 	FComponentMask DerivedEntityType;
 
 	FComponentMask ParentType;
@@ -57,7 +58,7 @@ void UMovieSceneRootInstantiatorSystem::InstantiateAllocation(const UE::MovieSce
 		ParentType.Set(Header.ComponentType);
 	}
 	Linker->EntityManager.GetComponents()->Factories.ComputeChildComponents(ParentType, DerivedEntityType);
-	Linker->EntityManager.GetComponents()->Factories.ComputeMutuallyInclusiveComponents(DerivedEntityType);
+	Linker->EntityManager.GetComponents()->Factories.ComputeMutuallyInclusiveComponents(EMutuallyInclusiveComponentType::All, DerivedEntityType, MutualInitializers);
 
 	const bool bHasAnyType = DerivedEntityType.Find(true) != INDEX_NONE;
 	if (!bHasAnyType)
@@ -76,6 +77,7 @@ void UMovieSceneRootInstantiatorSystem::InstantiateAllocation(const UE::MovieSce
 
 	TArrayView<const FMovieSceneEntityID> ParentIDs = ParentAllocation->GetEntityIDs();
 
+	FEntityAllocationWriteContext WriteContext(Linker->EntityManager);
 	int32 CurrentParentOffset = 0;
 
 	// We attempt to allocate all the linker entities contiguously in memory for efficient initialization,
@@ -88,6 +90,8 @@ void UMovieSceneRootInstantiatorSystem::InstantiateAllocation(const UE::MovieSce
 		FEntityDataLocation NewLinkerEntities = Linker->EntityManager.AllocateContiguousEntities(DerivedEntityType, &NumAdded);
 		FEntityRange ChildRange{ NewLinkerEntities.Allocation, NewLinkerEntities.ComponentOffset, NumAdded };
 
+		MutualInitializers.Execute(ChildRange, WriteContext);
+
 		Linker->EntityManager.InitializeChildAllocation(ParentType, DerivedEntityType, ParentAllocation, MakeArrayView(ParentOffsets.GetData() + CurrentParentOffset, NumAdded), ChildRange);
 
 		TArrayView<const FMovieSceneEntityID> ChildIDs = NewLinkerEntities.Allocation->GetEntityIDs();
diff --git a/Engine/Source/Runtime/MovieScene/Private/MovieSceneMutualComponentInclusivity.cpp b/Engine/Source/Runtime/MovieScene/Private/MovieSceneMutualComponentInclusivity.cpp
new file mode 100644
--- /dev/null
+++ b/Engine/Source/Runtime/MovieScene/Private/MovieSceneMutualComponentInclusivity.cpp
@@ -0,0 +1,593 @@
+// Copyright Epic Games, Inc. All Rights Reserved.
+
+#include "EntitySystem/MovieSceneMutualComponentInclusivity.h"
+#include "EntitySystem/MovieSceneMutualComponentInitializer.h"
+#include "EntitySystem/MovieSceneEntitySystemLinker.h"
+#include "EntitySystem/MovieSceneComponentRegistry.h"
+#include "EntitySystem/MovieSceneComponentTypeInfo.h"
+#include "Misc/StringBuilder.h"
+#include "MovieSceneFwd.h"
+
+namespace UE::MovieScene
+{
+
+void FMutualComponentInitializers::Add(IMutualComponentInitializer* InitializerToAdd)
+{
+	// No need to check for uniqueness here because initializers can only
+	// exist in the command buffer once
+	Initializers.Add(InitializerToAdd);
+}
+
+void FMutualComponentInitializers::Reset()
+{
+	Initializers.Reset();
+}
+
+void FMutualComponentInitializers::Execute(const FEntityRange& Range, const FEntityAllocationWriteContext& WriteContext) const
+{
+	for (IMutualComponentInitializer* Initializer : Initializers)
+	{
+		Initializer->Run(Range, WriteContext);
+	}
+}
+
+int32 FMutualInclusivityGraphCommandBuffer::ComputeMutuallyInclusiveComponents(EMutuallyInclusiveComponentType TypesToCompute, const FComponentMask& InMask, FComponentMask& OutMask, FMutualComponentInitializers& OutInitializers) const
+{
+	int32 NumNewComponents = 0;
+
+	FComponentMask WorkingMask = InMask;
+
+	// NOTE: It is possible to call this function with the same reference for InMask and OutMask
+	//       such that populating the mutual components adds them to the same input. This function is written to allow that while still being performant.
+	//       This flexibility allows us to compute mutually inclusive components to a separate mask, as well as the same mask.
+
+	// Command buffer will take the form of 1 or more match commands followed by 1 or more include commands.
+	// All match commands must be satisfied for subsequent include commands to pass
+	// 
+	// CommandBuffers adhere to the following schema:
+	// [
+	//	  { type = mandatory | optional }  // Specifies the type of the commands that follow
+	//	  { simple },
+	// 
+	//	  { match_all[1..N] -> include[1..N] -> (initializer) }
+	//            |------------------short circuits to end-----|
+	// 
+	//	  { match_any[1..N] -> short_circuit -> include[1..N] -> (initializer) }
+	//            |                  |-----------|----short circuits to end-----|
+	//            |---short circuits to includes-|
+	for (int32 Index = 0; Index < Commands.Num(); ++Index)
+	{
+		FCommand Command = Commands[Index];
+
+		switch (Command.CommandType)
+		{
+		case ECommandType::ShortCircuit:
+			// - 1 because Index will be incremented as part of the loop iteration
+			Index = static_cast<int32>(Command.ShortCircuit.ShortCircuitIndex) - 1;
+			break;
+
+		case ECommandType::Type:
+			// If the type doesn't match, short-circuit the entire stack of commands
+			if (!EnumHasAnyFlags(Command.Type.Type, EMutuallyInclusiveComponentType::All))
+			{
+				// - 1 because Index will be incremented as part of the loop iteration
+				Index = static_cast<int32>(Command.Type.ShortCircuitIndex) - 1;
+			}
+			break;
+
+		case ECommandType::MatchAll:
+			if (!WorkingMask.Contains(Command.Match.ComponentTypeID))
+			{
+				// ShortCircuitIndex points to the end of the Include and initializer commands for this stack
+				// - 1 because Index will be incremented as part of the loop iteration
+				Index = static_cast<int32>(Command.Match.ShortCircuitIndex) - 1;
+			}
+			break;
+
+		case ECommandType::MatchAny:
+			if (WorkingMask.Contains(Command.Match.ComponentTypeID))
+			{
+				// ShortCircuitIndex points to the first Include command in this stack
+				// - 1 because Index will be incremented as part of the loop iteration
+				Index = static_cast<int32>(Command.Match.ShortCircuitIndex) - 1;
+			}
+			break;
+
+		case ECommandType::Simple:
+			if (WorkingMask.Contains(Command.Simple.MatchID) && !WorkingMask.Contains(Command.Simple.IncludeID))
+			{
+				++NumNewComponents;
+				OutMask.Set(Command.Simple.IncludeID);
+				WorkingMask.Set(Command.Simple.IncludeID);
+			}
+			break;
+
+		case ECommandType::Include:
+			if (!WorkingMask.Contains(Command.Include.ComponentTypeID))
+			{
+				++NumNewComponents;
+				OutMask.Set(Command.Include.ComponentTypeID);
+				WorkingMask.Set(Command.Include.ComponentTypeID);
+			}
+			break;
+
+		case ECommandType::Initialize:
+			OutInitializers.Add(Command.Initialize.Initializer);
+			break;
+		}
+	}
+
+	return NumNewComponents;
+}
+
+
+void FMutualInclusivityGraphCommandBuffer::CheckInvariants() const
+{
+	// Check invariants
+#if DO_CHECK
+	auto TestShortCircuitIndex = [this](int32 Index)
+	{
+		checkf(Index != 0 && Index <= this->Commands.Num(),
+			TEXT("Invalid short circuit index %d for command buffer length"), Index, this->Commands.Num());
+	};
+
+	for (int32 Index = 0; Index < Commands.Num(); ++Index)
+	{
+		FCommand Command = Commands[Index];
+		switch (Command.CommandType)
+		{
+		case ECommandType::ShortCircuit:
+			TestShortCircuitIndex(Command.ShortCircuit.ShortCircuitIndex);
+			break;
+
+		case ECommandType::Type:
+			TestShortCircuitIndex(Command.Type.ShortCircuitIndex);
+			break;
+
+		case ECommandType::MatchAll:
+			TestShortCircuitIndex(Command.Match.ShortCircuitIndex);
+			break;
+
+		case ECommandType::MatchAny:
+			TestShortCircuitIndex(Command.Match.ShortCircuitIndex);
+			checkf(Commands[Command.Match.ShortCircuitIndex-1].CommandType == ECommandType::ShortCircuit,
+				TEXT("Expected short circuit command to proceed match any commands"));
+			break;
+
+		default:
+			break;
+		}
+	}
+#endif
+}
+
+void FMutualInclusivityGraph::DefineMutualInclusionRule(FComponentTypeID Predicate, std::initializer_list<FComponentTypeID> Dependents)
+{
+	DefineMutualInclusionRule(Predicate, Dependents, FMutuallyInclusiveComponentParams());
+}
+
+void FMutualInclusivityGraph::DefineMutualInclusionRule(FComponentTypeID Predicate, std::initializer_list<FComponentTypeID> Dependents, FMutuallyInclusiveComponentParams&& Parameters)
+{
+	bCommandBufferInvalidated = true;
+
+	FIncludePair& IncludePair = SimpleIncludes.FindOrAdd(Predicate);
+
+	FIncludes& Includes = Parameters.Type == EMutuallyInclusiveComponentType::Mandatory
+		? IncludePair.MandatoryIncludes
+		: IncludePair.OptionalIncludes;
+
+	for (FComponentTypeID Dependent : Dependents)
+	{
+		Includes.Add(Dependent);
+	}
+
+	if (Parameters.CustomInitializer)
+	{
+		Includes.Initializers.Emplace(MoveTemp(Parameters.CustomInitializer));
+	}
+}
+
+void FMutualInclusivityGraph::DefineComplexInclusionRule(const FComplexInclusivityFilter& InFilter, std::initializer_list<FComponentTypeID> Dependents)
+{
+	DefineComplexInclusionRule(InFilter, Dependents, FMutuallyInclusiveComponentParams());
+}
+
+void FMutualInclusivityGraph::DefineComplexInclusionRule(const FComplexInclusivityFilter& InFilter, std::initializer_list<FComponentTypeID> Dependents, FMutuallyInclusiveComponentParams&& Parameters)
+{
+	bCommandBufferInvalidated = true;
+
+	FIncludePair& IncludePair = ComplexIncludes.FindOrAdd(InFilter);
+
+	FIncludes& Includes = Parameters.Type == EMutuallyInclusiveComponentType::Mandatory
+		? IncludePair.MandatoryIncludes
+		: IncludePair.OptionalIncludes;
+
+	for (FComponentTypeID Dependent : Dependents)
+	{
+		Includes.Add(Dependent);
+	}
+
+	if (Parameters.CustomInitializer)
+	{
+		Includes.Initializers.Emplace(MoveTemp(Parameters.CustomInitializer));
+	}
+}
+
+int32 FMutualInclusivityGraph::ComputeMutuallyInclusiveComponents(EMutuallyInclusiveComponentType TypesToCompute, const FComponentMask& InMask, FComponentMask& OutMask, FMutualComponentInitializers& OutInitializers) const
+{
+	if (bCommandBufferInvalidated)
+	{
+		ReconstructCommandBuffer();
+	}
+
+	return CommandBuffer.ComputeMutuallyInclusiveComponents(TypesToCompute, InMask, OutMask, OutInitializers);
+}
+
+int32 FMutualInclusivityGraph::FindPrerequisiteChainLength(const FDirectedGraph& Graph, FComponentTypeID Component, TMap<FComponentTypeID, int32>& InOutCache)
+{
+	if (const int32* Depth = InOutCache.Find(Component))
+	{
+		return *Depth;
+	}
+
+	TArrayView<const FDirectedGraph::FDirectionalEdge> Edges = Graph.GetEdgesFrom(Component.BitIndex());
+
+	int32 MaxDepth = 0;
+	if (Edges.Num() != 0)
+	{
+		for (FDirectedGraph::FDirectionalEdge Edge : Edges)
+		{
+			MaxDepth = FMath::Max(MaxDepth, FindPrerequisiteChainLength(Graph, FComponentTypeID::FromBitIndex(Edge.ToNode), InOutCache));
+		}
+	}
+
+	InOutCache.Add(Component, MaxDepth + 1);
+	return MaxDepth + 1;
+}
+
+void FMutualInclusivityGraph::ReconstructCommandBuffer() const
+{
+	using FCommand     = FMutualInclusivityGraphCommandBuffer::FCommand;
+	using ECommandType = FMutualInclusivityGraphCommandBuffer::ECommandType;
+
+	FDirectedGraph Graph;
+
+	// Step 1: Allocate nodes that can be produced from mutual includes
+	for (const TTuple<FComponentTypeID, FIncludePair>& Pair : SimpleIncludes)
+	{
+		for (FComponentTypeID Dependent : Pair.Value.MandatoryIncludes.Includes)
+		{
+			Graph.AllocateNode(Dependent.BitIndex());
+		}
+		for (FComponentTypeID Dependent : Pair.Value.OptionalIncludes.Includes)
+		{
+			Graph.AllocateNode(Dependent.BitIndex());
+		}
+	}
+	for (const TTuple<FComplexInclusivityFilter, FIncludePair>& Pair : ComplexIncludes)
+	{
+		for (FComponentTypeID Dependent : Pair.Value.MandatoryIncludes.Includes)
+		{
+			Graph.AllocateNode(Dependent.BitIndex());
+		}
+		for (FComponentTypeID Dependent : Pair.Value.OptionalIncludes.Includes)
+		{
+			Graph.AllocateNode(Dependent.BitIndex());
+		}
+	}
+
+	// Step 2: Define edges between dependent nodes. Dependencies point from dependent to predicate so we can walk _up_ the graph
+	for (const TTuple<FComponentTypeID, FIncludePair>& Pair : SimpleIncludes)
+	{
+		if (Graph.IsNodeAllocated(Pair.Key.BitIndex()))
+		{
+			for (FComponentTypeID Dependent : Pair.Value.MandatoryIncludes.Includes)
+			{
+				Graph.MakeEdge(Dependent.BitIndex(), Pair.Key.BitIndex());
+			}
+			for (FComponentTypeID Dependent : Pair.Value.OptionalIncludes.Includes)
+			{
+				Graph.MakeEdge(Dependent.BitIndex(), Pair.Key.BitIndex());
+			}
+		}
+	}
+	for (const TTuple<FComplexInclusivityFilter, FIncludePair>& Pair : ComplexIncludes)
+	{
+		for (FComponentMaskIterator It(Pair.Key.Mask.Iterate()); It; ++It)
+		{
+			FComponentTypeID Predicate = FComponentTypeID::FromBitIndex(It.GetIndex());
+			if (Graph.IsNodeAllocated(Predicate.BitIndex()))
+			{
+				for (FComponentTypeID Dependent : Pair.Value.MandatoryIncludes.Includes)
+				{
+					Graph.MakeEdge(Dependent.BitIndex(), Predicate.BitIndex());
+				}
+				for (FComponentTypeID Dependent : Pair.Value.OptionalIncludes.Includes)
+				{
+					Graph.MakeEdge(Dependent.BitIndex(), Predicate.BitIndex());
+				}
+			}
+		}
+	}
+
+	if (!ensureMsgf(!Graph.IsCyclic(), TEXT("Mutual component inclusion graph is cyclic! See log for graph output.")))
+	{
+		FComponentRegistry* ComponentRegistry = UMovieSceneEntitySystemLinker::GetComponents();
+
+		auto EmitLabel = [=](uint16 NodeID, FStringBuilderBase& OutStringBuilder)
+		{
+			FComponentTypeID ComponentTypeID = FComponentTypeID::FromBitIndex(NodeID);
+
+			const FComponentTypeInfo& ComponentTypeInfo = ComponentRegistry->GetComponentTypeChecked(ComponentTypeID);
+		#if UE_MOVIESCENE_ENTITY_DEBUG
+			OutStringBuilder.Appendf(TEXT("\t\tnode%d[label=\"%s\"];\n"), ComponentTypeID.BitIndex(), *ComponentTypeInfo.DebugInfo->DebugName);
+		#else
+			OutStringBuilder.Appendf(TEXT("\t\tnode%d[label=\"%d\"];\n"), ComponentTypeID.BitIndex(), ComponentTypeID.BitIndex());
+		#endif
+		};
+
+		FDirectedGraphStringParameters Parameters{ TEXT("Mutual Includes") };
+		UE_LOG(LogMovieScene, Warning, TEXT("Printing Graph: \n %s"), *Graph.ToString(Parameters, EmitLabel));
+	}
+
+	// Step 3: Compute ordering constraints based on depth within the graph
+	TMap<FComponentTypeID, int32> DependencyDepths;
+	for (TConstSetBitIterator<> SetBitIter(Graph.GetNodeMask()); SetBitIter; ++SetBitIter)
+	{
+		FindPrerequisiteChainLength(Graph, FComponentTypeID::FromBitIndex(SetBitIter.GetIndex()), DependencyDepths);
+	}
+
+	// Step 4: Organize simple inclusion by depth using the predicate component, so that it is ordered after
+	//         any other include that might introduce this component
+	TArray<TTuple<FComponentTypeID, int32>> OrderedSimpleInclusion;
+	for (const TTuple<FComponentTypeID, FIncludePair>& Pair : SimpleIncludes)
+	{
+		const int32* GraphDepth = DependencyDepths.Find(Pair.Key);
+
+		OrderedSimpleInclusion.Add(MakeTuple(Pair.Key, GraphDepth ? *GraphDepth : 0));
+	}
+	Algo::SortBy(OrderedSimpleInclusion, [](const TTuple<FComponentTypeID, int32>& In){return In.Get<1>();});
+
+	// Step 5: Organize complex inclusion by depth using the predicate components, so that it is ordered after
+	//         any other include that might introduce this component
+	TArray<TTuple<FComplexInclusivityFilter, int32>> OrderedComplexInclusion;
+	for (const TTuple<FComplexInclusivityFilter, FIncludePair>& Pair : ComplexIncludes)
+	{
+		int32 Depth = 0;
+		for (FComponentMaskIterator It(Pair.Key.Mask.Iterate()); It; ++It)
+		{
+			FComponentTypeID Predicate = FComponentTypeID::FromBitIndex(It.GetIndex());
+			if (const int32* GraphDepth = DependencyDepths.Find(Predicate))
+			{
+				Depth = FMath::Max(Depth, *GraphDepth);
+			}
+		}
+
+		OrderedComplexInclusion.Add(MakeTuple(Pair.Key, Depth));
+	}
+	Algo::SortBy(OrderedComplexInclusion, [](const TTuple<FComplexInclusivityFilter, int32>& In){return In.Get<1>();});
+
+	// Step 6: Generate the command buffer from the two ordered arrays, ensuring simple includes come before complex includes
+	CommandBuffer.Commands.Empty();
+
+	int32 SimpleIndex = 0;
+	int32 ComplexIndex = 0;
+
+	int32 LastTypeIndex = INDEX_NONE;
+	EMutuallyInclusiveComponentType CurrentType = EMutuallyInclusiveComponentType::All;
+
+	auto EmitTypeCommand = [this, &LastTypeIndex, &CurrentType](EMutuallyInclusiveComponentType Type)
+	{
+		// Emit a type command if this is a different type from the last one added
+		if (Type != CurrentType)
+		{
+			if (LastTypeIndex != INDEX_NONE)
+			{
+				this->CommandBuffer.Commands[LastTypeIndex].Type.ShortCircuitIndex = this->CommandBuffer.Commands.Num();
+			}
+
+			LastTypeIndex = this->CommandBuffer.Commands.Num();
+			// Add the type command - we will populate the ShortCircuitIndex later
+			this->CommandBuffer.Commands.Emplace(FMutualInclusivityGraphCommandBuffer::FTypeCommand{
+				0,
+				Type
+			});
+			CurrentType = Type;
+		}
+	};
+
+	auto AddSimpleCommand = [this, &EmitTypeCommand](FComponentTypeID PredicateComponent, const FIncludes& Dependencies, EMutuallyInclusiveComponentType Type)
+	{
+		// Emit a type command if this is a different type from the last one added
+		EmitTypeCommand(Type);
+
+		// If the match fails, skip the match command and the includes
+		const int32  NumPredicateComponents = 1;
+		const int32  NumIncludeComponents   = Dependencies.Includes.Num();
+		const int32  NumInitializers        = Dependencies.Initializers.Num();
+		const uint16 ShortCircuitIndex      = static_cast<uint16>(this->CommandBuffer.Commands.Num() + NumPredicateComponents + NumIncludeComponents + NumInitializers);
+
+		// Use a simple command if possible to reduce the size of the array
+		if (NumIncludeComponents == 1 && NumInitializers == 0)
+		{
+			this->CommandBuffer.Commands.Emplace(FMutualInclusivityGraphCommandBuffer::FSimpleCommand{
+				PredicateComponent,
+				Dependencies.Includes[0]
+			});
+		}
+		else
+		{
+			this->CommandBuffer.Commands.Emplace(FMutualInclusivityGraphCommandBuffer::FMatchCommand{
+				PredicateComponent,
+				ShortCircuitIndex
+			}, EComplexInclusivityFilterMode::AllOf);
+
+			for (FComponentTypeID Dependency : Dependencies.Includes)
+			{
+				this->CommandBuffer.Commands.Add(FMutualInclusivityGraphCommandBuffer::FIncludeCommand{ Dependency });
+			}
+
+			for (const TUniquePtr<IMutualComponentInitializer>& Initializer : Dependencies.Initializers)
+			{
+				this->CommandBuffer.Commands.Add(FMutualInclusivityGraphCommandBuffer::FInitializeCommand{ Initializer.Get() });
+			}
+		}
+	};
+
+	auto AddSimpleCommands = [this, &AddSimpleCommand](FComponentTypeID Predicate)
+	{
+		const FIncludePair* Includes = this->SimpleIncludes.Find(Predicate);
+		if (Includes)
+		{
+			// Handle mandatory includes first
+			if (Includes->MandatoryIncludes.Includes.Num() != 0)
+			{
+				AddSimpleCommand(Predicate, Includes->MandatoryIncludes, EMutuallyInclusiveComponentType::Mandatory);
+			}
+			if (Includes->OptionalIncludes.Includes.Num() != 0)
+			{
+				AddSimpleCommand(Predicate, Includes->OptionalIncludes, EMutuallyInclusiveComponentType::Optional);
+			}
+		}
+	};
+	auto AddComplexCommand = [this, &EmitTypeCommand](const FComplexInclusivityFilter& Predicate, const FIncludes& Dependencies, EMutuallyInclusiveComponentType Type)
+	{
+		// Emit a type command if this is a different type from the last one added
+		EmitTypeCommand(Type);
+
+		// If the match fails, skip the match command and the includes
+		const int32  NumPredicateComponents = Predicate.Mask.NumComponents();
+		const int32  NumInitializers        = Dependencies.Initializers.Num();
+
+		if (Predicate.Mode == EComplexInclusivityFilterMode::AnyOf)
+		{
+			// The AllOf command is a stream of matches that jump to the start of the includes on success,
+			// followed by a shortcircuit that skips the includes if it is reached
+			const uint16 IncludeStartIndex = static_cast<uint16>(this->CommandBuffer.Commands.Num() + NumPredicateComponents + 1);
+			const uint16 ShortCircuitIndex = static_cast<uint16>(IncludeStartIndex + Dependencies.Includes.Num() + NumInitializers);
+
+			for (FComponentMaskIterator It(Predicate.Mask.Iterate()); It; ++It)
+			{
+				FComponentTypeID PredicateComponent = FComponentTypeID::FromBitIndex(It.GetIndex());
+
+				this->CommandBuffer.Commands.Emplace(FMutualInclusivityGraphCommandBuffer::FMatchCommand{
+					PredicateComponent,
+					IncludeStartIndex
+				}, Predicate.Mode);
+			}
+
+			this->CommandBuffer.Commands.Emplace(FMutualInclusivityGraphCommandBuffer::FShortCircuitCommand{
+				ShortCircuitIndex
+			});
+		}
+		else
+		{
+			// The Any command is a stream of matches that skip to the end of the command on failure
+			const uint16 ShortCircuitIndex = static_cast<uint16>(this->CommandBuffer.Commands.Num() + NumPredicateComponents + Dependencies.Includes.Num() + NumInitializers);
+
+			for (FComponentMaskIterator It(Predicate.Mask.Iterate()); It; ++It)
+			{
+				FComponentTypeID PredicateComponent = FComponentTypeID::FromBitIndex(It.GetIndex());
+
+				this->CommandBuffer.Commands.Emplace(FMutualInclusivityGraphCommandBuffer::FMatchCommand{
+					PredicateComponent,
+					ShortCircuitIndex
+				}, Predicate.Mode);
+			}
+		}
+
+		// Add includes
+		for (FComponentTypeID Dependency : Dependencies.Includes)
+		{
+			this->CommandBuffer.Commands.Add(FMutualInclusivityGraphCommandBuffer::FIncludeCommand{ Dependency });
+		}
+
+		// Add custom initializers
+		for (const TUniquePtr<IMutualComponentInitializer>& Initializer : Dependencies.Initializers)
+		{
+			this->CommandBuffer.Commands.Add(FMutualInclusivityGraphCommandBuffer::FInitializeCommand{ Initializer.Get() });
+		}
+	};
+	auto AddComplexCommands = [this, &AddComplexCommand](const FComplexInclusivityFilter& Predicate)
+	{
+		const FIncludePair* Includes = this->ComplexIncludes.Find(Predicate);
+		if (Includes)
+		{
+			// Handle mandatory includes first
+			if (Includes->MandatoryIncludes.Includes.Num() != 0)
+			{
+				AddComplexCommand(Predicate, Includes->MandatoryIncludes, EMutuallyInclusiveComponentType::Mandatory);
+			}
+			if (Includes->OptionalIncludes.Includes.Num() != 0)
+			{
+				AddComplexCommand(Predicate, Includes->OptionalIncludes, EMutuallyInclusiveComponentType::Optional);
+			}
+		}
+	};
+
+	// Keep going until we have nothing left
+	while (SimpleIndex < OrderedSimpleInclusion.Num() || ComplexIndex < OrderedComplexInclusion.Num())
+	{
+		const bool bHasSimple  = SimpleIndex < OrderedSimpleInclusion.Num();
+		const bool bHasComplex = ComplexIndex < OrderedComplexInclusion.Num();
+
+		if (bHasSimple && bHasComplex)
+		{
+			// Decide which to add first based on the order
+			const int32 SimpleOrder = OrderedSimpleInclusion[SimpleIndex].Get<1>();
+			const int32 ComplexOrder = OrderedComplexInclusion[ComplexIndex].Get<1>();
+
+			if (SimpleOrder <= ComplexOrder)
+			{
+				while (SimpleIndex < OrderedSimpleInclusion.Num() && OrderedSimpleInclusion[SimpleIndex].Get<1>() == SimpleOrder)
+				{
+					FComponentTypeID Predicate = OrderedSimpleInclusion[SimpleIndex].Get<0>();
+					AddSimpleCommands(Predicate);
+
+					++SimpleIndex;
+				}
+			}
+			else
+			{
+				while (ComplexIndex < OrderedComplexInclusion.Num() && OrderedComplexInclusion[ComplexIndex].Get<1>() == ComplexOrder)
+				{
+					const FComplexInclusivityFilter& Predicate = OrderedComplexInclusion[ComplexIndex].Get<0>();
+					AddComplexCommands(Predicate);
+
+					++ComplexIndex;
+				}
+			}
+		}
+		else if (bHasSimple)
+		{
+			while (SimpleIndex < OrderedSimpleInclusion.Num())
+			{
+				FComponentTypeID Predicate = OrderedSimpleInclusion[SimpleIndex].Get<0>();
+				AddSimpleCommands(Predicate);
+
+				++SimpleIndex;
+			}
+		}
+		else
+		{
+			while (ComplexIndex < OrderedComplexInclusion.Num())
+			{
+				const FComplexInclusivityFilter& Predicate = OrderedComplexInclusion[ComplexIndex].Get<0>();
+				AddComplexCommands(Predicate);
+
+				++ComplexIndex;
+			}
+		}
+	}
+
+	// Populate the most recent type command if one was added
+	if (LastTypeIndex != INDEX_NONE)
+	{
+		CommandBuffer.Commands[LastTypeIndex].Type.ShortCircuitIndex = CommandBuffer.Commands.Num();
+	}
+
+	bCommandBufferInvalidated = false;
+	CommandBuffer.CheckInvariants();
+}
+
+} // namespace UE::MovieScene
\ No newline at end of file
diff --git a/Engine/Source/Runtime/MovieScene/Private/Tests/MovieSceneEntitySystemTypesTests.cpp b/Engine/Source/Runtime/MovieScene/Private/Tests/MovieSceneEntitySystemTypesTests.cpp
--- a/Engine/Source/Runtime/MovieScene/Private/Tests/MovieSceneEntitySystemTypesTests.cpp
+++ b/Engine/Source/Runtime/MovieScene/Private/Tests/MovieSceneEntitySystemTypesTests.cpp
@@ -62,6 +62,98 @@ bool FMovieSceneEntityComponentFilterTest::RunTest(const FString& Parameters)
 
 	return true;
 }
+IMPLEMENT_SIMPLE_AUTOMATION_TEST(FMovieSceneMutuallyInclusiveComponentsTest, 
+		"System.Engine.Sequencer.EntitySystem.MutuallyInclusiveComponents", 
+		EAutomationTestFlags::EditorContext | EAutomationTestFlags::EngineFilter)
+bool FMovieSceneMutuallyInclusiveComponentsTest::RunTest(const FString& Parameters)
+{
+	using namespace UE::MovieScene;
+
+	FComponentTypeID ComponentTypes[] = {
+		FComponentTypeID::FromBitIndex(0),
+		FComponentTypeID::FromBitIndex(1),
+		FComponentTypeID::FromBitIndex(2),
+		FComponentTypeID::FromBitIndex(3),
+		FComponentTypeID::FromBitIndex(4),
+		FComponentTypeID::FromBitIndex(5),
+		FComponentTypeID::FromBitIndex(6),
+		FComponentTypeID::FromBitIndex(7),
+		FComponentTypeID::FromBitIndex(8),
+	};
+
+	FMutualInclusivityGraph InclusivityGraph;
+
+	// The inclusivity rules for this test are as follows:
+	//       0 requires 1 (Rule A)
+	//       3 requires 4 (Rule B)
+	//       4 requires 5 and 6 (Rule C)
+	//       if 2 and 6 are present, 5 and 7 must be present (Rule D)
+	//       if 1 or 5 are present, 8 must be present (Rule E)
+	// We define theese rules out-of-order to try and ensure order doesn't impact the algorithm
+
+	InclusivityGraph.DefineMutualInclusionRule(ComponentTypes[4], { ComponentTypes[5], ComponentTypes[6] });
+	InclusivityGraph.DefineMutualInclusionRule(ComponentTypes[3], { ComponentTypes[4] });
+	InclusivityGraph.DefineMutualInclusionRule(ComponentTypes[0], { ComponentTypes[1] });
+
+	InclusivityGraph.DefineComplexInclusionRule(FComplexInclusivityFilter({ ComponentTypes[6], ComponentTypes[2] }, EComplexInclusivityFilterMode::AllOf), { ComponentTypes[7], ComponentTypes[5] });
+	InclusivityGraph.DefineComplexInclusionRule(FComplexInclusivityFilter({ ComponentTypes[5], ComponentTypes[1] }, EComplexInclusivityFilterMode::AnyOf), { ComponentTypes[8] });
+
+	struct FTest
+	{
+		FComponentMask Input;
+		FComponentMask ExpectedOutput;
+	};
+	FTest Tests[] = {
+		{ { ComponentTypes[0] }, { ComponentTypes[0], ComponentTypes[1], ComponentTypes[8]} }, // Test Rule A feeding E
+		{ { ComponentTypes[3] }, { ComponentTypes[3], ComponentTypes[4], ComponentTypes[5], ComponentTypes[6], ComponentTypes[8] } }, // Test Rule B feeding C feeding E
+		{ { ComponentTypes[4] }, { ComponentTypes[4], ComponentTypes[5], ComponentTypes[6], ComponentTypes[8] } }, // Test Rule C and E
+		{ { ComponentTypes[2], ComponentTypes[6] }, { ComponentTypes[2], ComponentTypes[5], ComponentTypes[6], ComponentTypes[7], ComponentTypes[8] } }, // Test Rule D which feeds E
+		{ { ComponentTypes[2], ComponentTypes[4] }, { ComponentTypes[2], ComponentTypes[4], ComponentTypes[5], ComponentTypes[6], ComponentTypes[7], ComponentTypes[8] } }, // Test Rule C feeding rule D whcih feeds rule E
+		{ { ComponentTypes[0], ComponentTypes[2], ComponentTypes[6] }, { ComponentTypes[0], ComponentTypes[1], ComponentTypes[2], ComponentTypes[5], ComponentTypes[6], ComponentTypes[7], ComponentTypes[8] } }, // Test Rule A, C, D and E (A feeds E, C feeds D)
+	};
+
+
+	for (int32 Index = 0; Index < UE_ARRAY_COUNT(Tests); ++Index)
+	{
+		const FTest& Test = Tests[Index];
+
+		FComponentMask Result;
+		FMutualComponentInitializers MutualInitializers;
+		InclusivityGraph.ComputeMutuallyInclusiveComponents(EMutuallyInclusiveComponentType::All, Test.Input, Result, MutualInitializers);
+
+		// Combine the result with the input since ComputeMutuallyInclusiveComponents only adds to Result
+		// if they did not exist in Input
+		Result.CombineWithBitwiseOR(Test.Input, EBitwiseOperatorFlags::MaxSize);
+
+		if (!Result.CompareSetBits(Test.ExpectedOutput) || Result.NumComponents() != Test.ExpectedOutput.NumComponents())
+		{
+			auto ToString = [](const FComponentMask& Mask)
+			{
+
+				TStringBuilder<32> String;
+				for (FComponentMaskIterator It(Mask.Iterate()); It; ++It)
+				{
+					if (String.Len() != 0)
+					{
+						String += TEXT("|");
+					}
+					String += LexToString(It.GetIndex());
+				}
+				return FString(String.ToString());
+			};
+
+			AddError(FString::Printf(TEXT("Test %d failed: result from input %s was %s: expected %s."),
+				Index,
+				*ToString(Test.Input),
+				*ToString(Result),
+				*ToString(Test.ExpectedOutput)
+				));
+		}
+	}
+
+	return true;
+}
+
 
 IMPLEMENT_SIMPLE_AUTOMATION_TEST(FMovieSceneEntityMigrationTest, 
 		"System.Engine.Sequencer.EntitySystem.Migration", 
diff --git a/Engine/Source/Runtime/MovieScene/Public/EntitySystem/MovieSceneEntityBuilder.h b/Engine/Source/Runtime/MovieScene/Public/EntitySystem/MovieSceneEntityBuilder.h
--- a/Engine/Source/Runtime/MovieScene/Public/EntitySystem/MovieSceneEntityBuilder.h
+++ b/Engine/Source/Runtime/MovieScene/Public/EntitySystem/MovieSceneEntityBuilder.h
@@ -6,6 +6,8 @@
 #include "EntitySystem/MovieSceneEntitySystemTypes.h"
 #include "EntitySystem/MovieSceneEntityManager.h"
 #include "EntitySystem/MovieSceneComponentRegistry.h"
+#include "EntitySystem/MovieSceneEntityFactory.h"
+#include "EntitySystem/MovieSceneMutualComponentInclusivity.h"
 #include "Delegates/IntegerSequence.h"
 
 #include <initializer_list>
@@ -75,7 +77,7 @@ struct IEntityBuilder
 	virtual ~IEntityBuilder() {}
 
 	virtual FMovieSceneEntityID Create(FEntityManager* EntityManager) = 0;
-	virtual void GenerateType(FEntityManager* EntityManager, FComponentMask& OutMask) = 0;
+	virtual void GenerateType(FEntityManager* EntityManager, FComponentMask& OutMask, bool& OutAddMutualComponents) = 0;
 	virtual void Initialize(FEntityManager* EntityManager, const FEntityInfo& EntityInfo) = 0;
 };
 
@@ -207,12 +209,13 @@ struct TEntityBuilderImpl<TIntegerSequence<int, Indices...>, T...> : IEntityBuil
 		return CreateEntity(EntityManager);
 	}
 
-	virtual void GenerateType(FEntityManager* EntityManager, FComponentMask& OutMask) override final
+	virtual void GenerateType(FEntityManager* EntityManager, FComponentMask& OutMask, bool& OutAddMutualComponents) override final
 	{
 		VisitTupleElements([&OutMask](auto& In){ In.AccumulateMask(OutMask); }, this->Payload);
+
 		if (bAddMutualComponents)
 		{
-			EntityManager->GetComponents()->Factories.ComputeMutuallyInclusiveComponents(OutMask);
+			OutAddMutualComponents = true;
 		}
 	}
 
@@ -234,11 +237,22 @@ struct TEntityBuilderImpl<TIntegerSequence<int, Indices...>, T...> : IEntityBuil
 	 */
 	FMovieSceneEntityID CreateEntity(FEntityManager* EntityManager, FComponentMask NewType = FComponentMask())
 	{
-		GenerateType(EntityManager, NewType);
+		bool bLocalAddMutualComponents = false;
+		GenerateType(EntityManager, NewType, bLocalAddMutualComponents);
+
+		FMutualComponentInitializers MutualInitializers;
+		FEntityAllocationWriteContext WriteContext(*EntityManager);
+		EMutuallyInclusiveComponentType MutualTypes = bLocalAddMutualComponents ? EMutuallyInclusiveComponentType::All : EMutuallyInclusiveComponentType::Mandatory;
+
+		EntityManager->GetComponents()->Factories.ComputeMutuallyInclusiveComponents(MutualTypes, NewType, MutualInitializers);
 
 		FEntityInfo Entry = EntityManager->AllocateEntity(NewType);
 		Initialize(EntityManager, Entry);
 
+		// Run mutual initializers after the builder has actually constructed the entity
+		// otherwise mutual components would be reading garbage
+		MutualInitializers.Execute(Entry.Data.AsRange(), WriteContext);
+
 		return Entry.EntityID;
 	}
 
@@ -273,6 +287,12 @@ struct TEntityBuilderImpl<TIntegerSequence<int, Indices...>, T...> : IEntityBuil
 
 		VisitTupleElements([&NewMask](auto& In){ In.AccumulateMask(NewMask); }, this->Payload);
 
+		FMutualComponentInitializers MutualInitializers;
+		FEntityAllocationWriteContext WriteContext(*EntityManager);
+
+		EMutuallyInclusiveComponentType MutualTypes = bAddMutualComponents ? EMutuallyInclusiveComponentType::All : EMutuallyInclusiveComponentType::Mandatory;
+		EntityManager->GetComponents()->Factories.ComputeMutuallyInclusiveComponents(MutualTypes, NewMask, MutualInitializers);
+
 		if (!NewMask.CompareSetBits(OldMask))
 		{
 			EntityManager->ChangeEntityType(EntityID, NewMask);
@@ -283,6 +303,10 @@ struct TEntityBuilderImpl<TIntegerSequence<int, Indices...>, T...> : IEntityBuil
 		if (Entry.Data.Allocation != nullptr)
 		{
 			VisitTupleElements([Entry](auto& In){ In.Apply(Entry.Data.Allocation, Entry.Data.ComponentOffset); }, this->Payload);
+
+			// Run mutual initializers after the builder has actually constructed the entity
+			// otherwise mutual components would be reading garbage
+			MutualInitializers.Execute(Entry.Data.AsRange(), WriteContext);
 		}
 	}
 
diff --git a/Engine/Source/Runtime/MovieScene/Public/EntitySystem/MovieSceneEntityFactory.h b/Engine/Source/Runtime/MovieScene/Public/EntitySystem/MovieSceneEntityFactory.h
--- a/Engine/Source/Runtime/MovieScene/Public/EntitySystem/MovieSceneEntityFactory.h
+++ b/Engine/Source/Runtime/MovieScene/Public/EntitySystem/MovieSceneEntityFactory.h
@@ -12,6 +12,7 @@
 #include "EntitySystem/MovieSceneEntityIDs.h"
 #include "EntitySystem/MovieSceneEntitySystemDirectedGraph.h"
 #include "EntitySystem/MovieSceneEntitySystemTypes.h"
+#include "EntitySystem/MovieSceneMutualComponentInclusivity.h"
 #include "Misc/AssertionMacros.h"
 #include "Misc/InlineValue.h"
 
@@ -22,63 +23,14 @@ namespace UE
 {
 namespace MovieScene
 {
-struct FEntityAllocation;
-struct FEntityRange;
 
+struct IMutualComponentInitializer;
 
-enum class EComplexInclusivityFilterMode
-{
-	AllOf,
-	AnyOf
-};
-
-
-struct FComplexInclusivityFilter
-{
-	FComponentMask Mask;
-	EComplexInclusivityFilterMode Mode;
-
-	FComplexInclusivityFilter(const FComponentMask& InMask, EComplexInclusivityFilterMode InMode)
-		: Mask(InMask), Mode(InMode)
-	{}
-
-	static FComplexInclusivityFilter All(std::initializer_list<FComponentTypeID> InComponentTypes)
-	{
-		return FComplexInclusivityFilter(FComponentMask(InComponentTypes), EComplexInclusivityFilterMode::AllOf);
-	}
-
-	static FComplexInclusivityFilter Any(std::initializer_list<FComponentTypeID> InComponentTypes)
-	{
-		return FComplexInclusivityFilter(FComponentMask(InComponentTypes), EComplexInclusivityFilterMode::AnyOf);
-	}
-
-	bool Match(FComponentMask Input) const
-	{
-		switch (Mode)
-		{
-			case EComplexInclusivityFilterMode::AllOf:
-				{
-					FComponentMask Temp = Mask;
-					Temp.CombineWithBitwiseAND(Input, EBitwiseOperatorFlags::MaintainSize);
-					return Temp == Mask;
-				}
-				break;
-			case EComplexInclusivityFilterMode::AnyOf:
-				{
-					FComponentMask Temp = Mask;
-					Temp.CombineWithBitwiseAND(Input, EBitwiseOperatorFlags::MaintainSize);
-					return Temp.Find(true) != INDEX_NONE;
-				}
-				break;
-			default:
-				checkf(false, TEXT("Not implemented"));
-				return false;
-		}
-	}
-};
-
+struct FEntityRange;
+struct FEntityAllocation;
+struct FMutualComponentInitializers;
 
-struct FComplexInclusivity
+struct UE_DEPRECATED(5.2, "Please use DefineComplexInclusiveComponents()") FComplexInclusivity
 {
 	FComplexInclusivityFilter Filter;
 	FComponentMask ComponentsToInclude;
@@ -134,15 +86,24 @@ struct FEntityFactories
 	 * that entity too.
 	 *
 	 * @note: the inverse is not implied (ie B can still exist without A)
-     */
-	MOVIESCENE_API void DefineMutuallyInclusiveComponent(FComponentTypeID InComponentA, FComponentTypeID InComponentB);
+	 */
+	void DefineMutuallyInclusiveComponent(FComponentTypeID InComponentA, FComponentTypeID InComponentB)
+	{
+		DefineMutuallyInclusiveComponents(InComponentA, { InComponentB });
+	}
 
 	/**
-	 * Specifies that if an entity matches the given filter, the specified components should be created on it.
-	 * @note: include "EntitySystem/MovieSceneEntityFactoryTemplates.h" for definition
+	 * Indicates that if the first component exists on an entity, the specified components should be created on
+	 * that entity too.
+	 *
+	 * @note: the inverse is not implied (ie B can still exist without A)
 	 */
-	template<typename... ComponentTypes>
-	void DefineComplexInclusiveComponents(const FComplexInclusivityFilter& InFilter, ComponentTypes... InComponents);
+	MOVIESCENE_API void DefineMutuallyInclusiveComponents(FComponentTypeID InComponentA, std::initializer_list<FComponentTypeID> InMutualComponents);
+
+	/**
+	 * Specifies a mutual inclusivity relationship along with a custom initializer for initializing the mutual component(s)
+	 */
+	MOVIESCENE_API void DefineMutuallyInclusiveComponents(FComponentTypeID InComponentA, std::initializer_list<FComponentTypeID> InMutualComponents, FMutuallyInclusiveComponentParams&& Params);
 
 	/**
 	 * Specifies that if an entity matches the given filter, the specified component should be created on it.
@@ -150,9 +111,9 @@ struct FEntityFactories
 	MOVIESCENE_API void DefineComplexInclusiveComponents(const FComplexInclusivityFilter& InFilter, FComponentTypeID InComponent);
 
 	/**
-	 * Defines a new complex inclusivity relationship. The helper methods above are easier and preferrable.
+	 * Specifies that if an entity matches the given filter, the specified components should be created on it.
 	 */
-	MOVIESCENE_API void DefineComplexInclusiveComponents(const FComplexInclusivity& InInclusivity);
+	MOVIESCENE_API void DefineComplexInclusiveComponents(const FComplexInclusivityFilter& InFilter, std::initializer_list<FComponentTypeID> InComponents, FMutuallyInclusiveComponentParams&& Params);
 
 	/**
 	 * Given a set of components on a parent entity, compute what components should exist on a child entity.
@@ -166,22 +127,36 @@ struct FEntityFactories
 	 *
 	 * This resolves all the mutual and complex inclusivity relationships.
 	 */
-	MOVIESCENE_API int32 ComputeMutuallyInclusiveComponents(FComponentMask& ComponentMask);
+	MOVIESCENE_API int32 ComputeMutuallyInclusiveComponents(EMutuallyInclusiveComponentType MutualTypes, FComponentMask& ComponentMask, FMutualComponentInitializers& OutInitializers);
 
 	MOVIESCENE_API void RunInitializers(const FComponentMask& ParentType, const FComponentMask& ChildType, const FEntityAllocation* ParentAllocation, TArrayView<const int32> ParentAllocationOffsets, const FEntityRange& InChildEntityRange);
 
-	TArray<TInlineValue<FChildEntityInitializer>> ChildInitializers;
-
-	TMultiMap<FComponentTypeID, FComponentTypeID> ParentToChildComponentTypes;
-	FMovieSceneEntitySystemDirectedGraph MutualInclusivityGraph;
-	TArray<FComplexInclusivity> ComplexInclusivity;
+public:
 
-	struct
+	PRAGMA_DISABLE_DEPRECATION_WARNINGS
+		using FDeprecatedComplexInclusivity = FComplexInclusivity;
+	PRAGMA_ENABLE_DEPRECATION_WARNINGS
+	/**
+	 * Specifies that if an entity matches the given filter, the specified components should be created on it.
+	 */
+	template<typename... ComponentTypes>
+	UE_DEPRECATED(5.2, "Please use DefineComplexInclusiveComponents(const FComplexInclusivityFilter&, initializer_list<FComponentTypeID>)")
+	void DefineComplexInclusiveComponents(const FComplexInclusivityFilter& InFilter, FComponentTypeID InComponent, ComponentTypes... InComponents)
 	{
-		FComponentMask AllMutualFirsts;
-		FComponentMask AllComplexFirsts;
+		DefineComplexInclusiveComponents(InFilter, std::initializer_list<FComponentTypeID>({ InComponent, InComponents... }), FMutuallyInclusiveComponentParams());
 	}
-	Masks;
+
+	/**
+	 * Defines a new complex inclusivity relationship. The helper methods above are easier and preferrable.
+	 */
+	UE_DEPRECATED(5.2, "Please use DefineComplexInclusiveComponents(const FComplexInclusivityFilter&, FComponentTypeID)")
+	MOVIESCENE_API void DefineComplexInclusiveComponents(const FDeprecatedComplexInclusivity& InInclusivity);
+
+private:
+
+	TArray<TInlineValue<FChildEntityInitializer>> ChildInitializers;
+	TMultiMap<FComponentTypeID, FComponentTypeID> ParentToChildComponentTypes;
+	UE::MovieScene::FMutualInclusivityGraph MutualInclusivityGraph;
 };
 
 
diff --git a/Engine/Source/Runtime/MovieScene/Public/EntitySystem/MovieSceneEntityFactoryTemplates.h b/Engine/Source/Runtime/MovieScene/Public/EntitySystem/MovieSceneEntityFactoryTemplates.h
--- a/Engine/Source/Runtime/MovieScene/Public/EntitySystem/MovieSceneEntityFactoryTemplates.h
+++ b/Engine/Source/Runtime/MovieScene/Public/EntitySystem/MovieSceneEntityFactoryTemplates.h
@@ -206,14 +206,6 @@ inline void FEntityFactories::DefineChildComponent(TComponentTypeID<ParentCompon
 	ChildInitializers.Add(FInitializer(InParentType, InChildType, Forward<InitializerCallback>(InInitializer)));
 }
 
-template<typename... ComponentTypes>
-inline void FEntityFactories::DefineComplexInclusiveComponents(const FComplexInclusivityFilter& InFilter, ComponentTypes... InComponents)
-{
-	FComponentMask ComponentsToInclude { InComponents... };
-	FComplexInclusivity NewComplexInclusivity { InFilter, ComponentsToInclude };
-	DefineComplexInclusiveComponents(NewComplexInclusivity);
-}
-
 template<typename T>
 TComponentTypeID<T> FComponentRegistry::NewComponentType(const TCHAR* const DebugName, EComponentTypeFlags Flags)
 {
diff --git a/Engine/Source/Runtime/MovieScene/Public/EntitySystem/MovieSceneEntityManager.h b/Engine/Source/Runtime/MovieScene/Public/EntitySystem/MovieSceneEntityManager.h
--- a/Engine/Source/Runtime/MovieScene/Public/EntitySystem/MovieSceneEntityManager.h
+++ b/Engine/Source/Runtime/MovieScene/Public/EntitySystem/MovieSceneEntityManager.h
@@ -47,12 +47,14 @@ struct FChildEntityInitializer;
 struct FComponentRegistry;
 struct FEntityAllocationIteratorProxy;
 struct FFreeEntityOperation;
+struct FMutualEntityInitializer;
 struct IComponentTypeHandler;
 struct IMovieSceneEntityMutation;
 template <typename T> struct TReadOptional;
 template <typename T> struct TWriteOptional;
 struct IMovieSceneConditionalEntityMutation;
 
+enum class EMutuallyInclusiveComponentType : uint8;
 
 enum class EEntityRecursion : uint8
 {
@@ -583,7 +585,7 @@ public:
 	 * @param Mutation    Implementation that defines how to mutate the entities that match the filter
 	 * @return The number of entities that were mutated, or 0 if none were matched
 	 */
-	int32 MutateAll(const FEntityComponentFilter& Filter, const IMovieSceneEntityMutation& Mutation);
+	int32 MutateAll(const FEntityComponentFilter& Filter, const IMovieSceneEntityMutation& Mutation, EMutuallyInclusiveComponentType MutualTypes = EMutuallyInclusiveComponentType::Mandatory);
 
 
 	/**
@@ -593,7 +595,7 @@ public:
 	 * @param Mutation    Implementation that defines how to mutate the entities that match the filter
 	 * @return The number of entities that were mutated, or 0 if none were matched
 	 */
-	int32 MutateConditional(const FEntityComponentFilter& Filter, const IMovieSceneConditionalEntityMutation& Mutation);
+	int32 MutateConditional(const FEntityComponentFilter& Filter, const IMovieSceneConditionalEntityMutation& Mutation, EMutuallyInclusiveComponentType MutualTypes = EMutuallyInclusiveComponentType::Mandatory);
 
 
 	/**
diff --git a/Engine/Source/Runtime/MovieScene/Public/EntitySystem/MovieSceneEntitySystemDirectedGraph.h b/Engine/Source/Runtime/MovieScene/Public/EntitySystem/MovieSceneEntitySystemDirectedGraph.h
--- a/Engine/Source/Runtime/MovieScene/Public/EntitySystem/MovieSceneEntitySystemDirectedGraph.h
+++ b/Engine/Source/Runtime/MovieScene/Public/EntitySystem/MovieSceneEntitySystemDirectedGraph.h
@@ -4,10 +4,24 @@
 
 #include "Containers/Array.h"
 #include "Containers/BitArray.h"
+#include "Templates/Function.h"
+#include "Containers/StringFwd.h"
 
 #include "EntitySystem/MovieSceneEntitySystemTypes.h"
 
-struct MOVIESCENE_API FMovieSceneEntitySystemDirectedGraph
+namespace UE::MovieScene
+{
+
+struct FDirectedGraphStringParameters
+{
+	FStringView ClusterName;
+	FColor Color = FColor::Black;
+};
+
+/**
+ * Directed graph represented as a bitarray for allocated nodes, and edges defined by pairs of integers (from->to).
+ */
+struct MOVIESCENE_API FDirectedGraph
 {
 	struct FDirectionalEdge
 	{
@@ -46,7 +60,7 @@ public:
 	{
 		TArray<uint16> PostNodes;
 
-		explicit FDepthFirstSearch(const FMovieSceneEntitySystemDirectedGraph* InGraph);
+		explicit FDepthFirstSearch(const FDirectedGraph* InGraph);
 
 		void Search(uint16 Node);
 
@@ -60,14 +74,14 @@ public:
 		TBitArray<> Visited;
 		TBitArray<> IsVisiting;
 
-		const FMovieSceneEntitySystemDirectedGraph* Graph;
+		const FDirectedGraph* Graph;
 	};
 
 	struct MOVIESCENE_API FBreadthFirstSearch
 	{
 		TArray<uint16> Nodes;
 
-		explicit FBreadthFirstSearch(const FMovieSceneEntitySystemDirectedGraph* InGraph);
+		explicit FBreadthFirstSearch(const FDirectedGraph* InGraph);
 
 		void Search(uint16 Node);
 
@@ -79,13 +93,13 @@ public:
 	private:
 
 		TBitArray<> Visited;
-		const FMovieSceneEntitySystemDirectedGraph* Graph;
+		const FDirectedGraph* Graph;
 		int32 StackIndex;
 	};
 
 	struct MOVIESCENE_API FDiscoverCyclicEdges
 	{
-		explicit FDiscoverCyclicEdges(const FMovieSceneEntitySystemDirectedGraph* InGraph);
+		explicit FDiscoverCyclicEdges(const FDirectedGraph* InGraph);
 
 		void Search(uint16 Node);
 
@@ -110,12 +124,12 @@ public:
 		TBitArray<> CyclicEdges;
 		TBitArray<> VisitedEdges;
 		TArray<uint16, TInlineAllocator<16>> EdgeChain;
-		const FMovieSceneEntitySystemDirectedGraph* Graph;
+		const FDirectedGraph* Graph;
 	};
 
 public:
 
-	FMovieSceneEntitySystemDirectedGraph()
+	FDirectedGraph()
 		: bHasDanglingEdges(false)
 	{}
 
@@ -150,6 +164,10 @@ public:
 
 	bool HasEdgeTo(uint16 InNode) const;
 
+	FString ToString(const UE::MovieScene::FDirectedGraphStringParameters& Parameters) const;
+
+	FString ToString(const UE::MovieScene::FDirectedGraphStringParameters& Parameters, TFunctionRef<void(uint16, FStringBuilderBase&)> EmitLabel) const;
+
 private:
 
 	int32 FindEdgeStart(uint16 FromNode) const;
@@ -171,3 +189,5 @@ private:
 	bool bHasDanglingEdges;
 };
 
+
+} // namespace UE::MovieScene
\ No newline at end of file
diff --git a/Engine/Source/Runtime/MovieScene/Public/EntitySystem/MovieSceneEntitySystemGraphs.h b/Engine/Source/Runtime/MovieScene/Public/EntitySystem/MovieSceneEntitySystemGraphs.h
--- a/Engine/Source/Runtime/MovieScene/Public/EntitySystem/MovieSceneEntitySystemGraphs.h
+++ b/Engine/Source/Runtime/MovieScene/Public/EntitySystem/MovieSceneEntitySystemGraphs.h
@@ -83,7 +83,7 @@ struct TStructOpsTypeTraits<FMovieSceneEntitySystemGraphNodes> : public TStructO
 USTRUCT()
 struct MOVIESCENE_API FMovieSceneEntitySystemGraph
 {
-	using FDirectionalEdge = FMovieSceneEntitySystemDirectedGraph::FDirectionalEdge;
+	using FDirectionalEdge = UE::MovieScene::FDirectedGraph::FDirectionalEdge;
 
 	GENERATED_BODY()
 
@@ -192,7 +192,7 @@ private:
 
 	TMap<uint16, uint16> GlobalToLocalNodeIDs;
 
-	FMovieSceneEntitySystemDirectedGraph ReferenceGraph;
+	UE::MovieScene::FDirectedGraph ReferenceGraph;
 
 	uint32 SerialNumber = 0;
 	uint32 PreviousSerialNumber = 0;
diff --git a/Engine/Source/Runtime/MovieScene/Public/EntitySystem/MovieSceneEntitySystemTypes.h b/Engine/Source/Runtime/MovieScene/Public/EntitySystem/MovieSceneEntitySystemTypes.h
--- a/Engine/Source/Runtime/MovieScene/Public/EntitySystem/MovieSceneEntitySystemTypes.h
+++ b/Engine/Source/Runtime/MovieScene/Public/EntitySystem/MovieSceneEntitySystemTypes.h
@@ -92,6 +92,14 @@ enum class EComplexFilterMode : uint8
 };
 ENUM_CLASS_FLAGS(EComplexFilterMode);
 
+enum class EMutuallyInclusiveComponentType : uint8
+{
+	Mandatory = 1u << 0,
+	Optional  = 1u << 1,
+
+	All = Mandatory | Optional,
+};
+ENUM_CLASS_FLAGS(EMutuallyInclusiveComponentType)
 
 /**
  * Enumeration specifying the locking mechanism to use when accessing component data
diff --git a/Engine/Source/Runtime/MovieScene/Public/EntitySystem/MovieSceneMutualComponentInclusivity.h b/Engine/Source/Runtime/MovieScene/Public/EntitySystem/MovieSceneMutualComponentInclusivity.h
new file mode 100644
--- /dev/null
+++ b/Engine/Source/Runtime/MovieScene/Public/EntitySystem/MovieSceneMutualComponentInclusivity.h
@@ -0,0 +1,307 @@
+// Copyright Epic Games, Inc. All Rights Reserved.
+
+#pragma once
+
+#include "CoreTypes.h"
+#include "Containers/Array.h"
+#include "EntitySystem/MovieSceneEntityIDs.h"
+#include "EntitySystem/MovieSceneEntitySystemTypes.h"
+#include "EntitySystem/MovieSceneMutualComponentInitializer.h"
+#include "EntitySystem/MovieSceneEntitySystemDirectedGraph.h"
+
+#include <initializer_list>
+
+namespace UE::MovieScene
+{
+
+struct FMutualInclusivityGraph;
+struct IMutualComponentInitializer;
+
+/** Enumeration specifying how to apply a complex filter for mutual inclusion */
+enum class EComplexInclusivityFilterMode
+{
+	/** Specifies that all the components must be present for a mutal component to be included */
+	AllOf,
+	/** Specifies that the mutal component should be included if any of the specified components are present */
+	AnyOf
+};
+
+struct FMutuallyInclusiveComponentParams
+{
+	FMutuallyInclusiveComponentParams()
+	{}
+
+	/** Custom initializer for this included component */
+	TUniquePtr<IMutualComponentInitializer> CustomInitializer;
+
+	/** Optional by default because generally we do not want mutual components on Imported entities */
+	EMutuallyInclusiveComponentType Type = EMutuallyInclusiveComponentType::Optional;
+};
+
+
+/** Filter used for specifying more complex mutual inclusivity rules that depend on more than one component type */
+struct FComplexInclusivityFilter
+{
+	/** The component mask defining the set of components to test for */
+	FComponentMask Mask;
+	/** The filter mode used for matching */
+	EComplexInclusivityFilterMode Mode;
+
+	FComplexInclusivityFilter(const FComponentMask& InMask, EComplexInclusivityFilterMode InMode)
+		: Mask(InMask), Mode(InMode)
+	{}
+
+	friend uint32 GetTypeHash(const FComplexInclusivityFilter& Filter)
+	{
+		return HashCombine(GetTypeHash(Filter.Mask), GetTypeHash(Filter.Mode));
+	}
+
+	friend bool operator==(const FComplexInclusivityFilter& A, const FComplexInclusivityFilter& B)
+	{
+		return A.Mask == B.Mask && A.Mode == B.Mode;
+	}
+
+	static FComplexInclusivityFilter All(std::initializer_list<FComponentTypeID> InComponentTypes)
+	{
+		return FComplexInclusivityFilter(FComponentMask(InComponentTypes), EComplexInclusivityFilterMode::AllOf);
+	}
+
+	static FComplexInclusivityFilter Any(std::initializer_list<FComponentTypeID> InComponentTypes)
+	{
+		return FComplexInclusivityFilter(FComponentMask(InComponentTypes), EComplexInclusivityFilterMode::AnyOf);
+	}
+
+	UE_DEPRECATED(5.2, "This function will no longer be called")
+	bool Match(FComponentMask Input) const
+	{
+		switch (Mode)
+		{
+			case EComplexInclusivityFilterMode::AllOf:
+				{
+					FComponentMask Temp = Mask;
+					Temp.CombineWithBitwiseAND(Input, EBitwiseOperatorFlags::MaintainSize);
+					return Temp == Mask;
+				}
+				break;
+			case EComplexInclusivityFilterMode::AnyOf:
+				{
+					FComponentMask Temp = Mask;
+					Temp.CombineWithBitwiseAND(Input, EBitwiseOperatorFlags::MaintainSize);
+					return Temp.Find(true) != INDEX_NONE;
+				}
+				break;
+			default:
+				checkf(false, TEXT("Not implemented"));
+				return false;
+		}
+	}
+};
+
+
+/**
+ * Container for keeping track of which initializers need to be run for a given call to
+ * FMutualInclusivityGraphCommandBuffer::ComputeMutuallyInclusiveComponents
+ */
+struct FMutualComponentInitializers
+{
+	/**
+	 * Run the mutual initializers over the specified entity range
+	 */
+	MOVIESCENE_API void Execute(const FEntityRange& Range, const FEntityAllocationWriteContext& WriteContext) const;
+
+	/**
+	 * Add a new initializer to this cache
+	 */
+	void Add(IMutualComponentInitializer* InitializerToAdd);
+
+	/**
+	 * Reset this container. Will not free memory.
+	 */
+	void Reset();
+
+private:
+
+	TArray<IMutualComponentInitializer*, TInlineAllocator<2>> Initializers;
+};
+
+
+/**
+ * Command buffer used for computing mutually inclusive components based on a complex set of rules and dependencies
+ */
+struct FMutualInclusivityGraphCommandBuffer
+{
+	/**
+	 * Populate a component mask with all the necessary mutually inclusive components given an input type.
+	 * Mutual components that already existed in InMask will not be added or counted.
+	 * Can be used to populate a new mask with only the mutual components, or to add the components to an existing mask.
+	 *
+	 * @param InMask           The input component type. Rules that are relevant to this type (directly or indirectly) will be applied to the output.
+	 * @param TypesToCompute   Allows computation of mandatory, optional or all rules
+	 * @param OutMask          Component mask to receieve the result. May point to the same mask as InMask in order to add the mutual components directly.
+	 * @param OutInitializers  Array to receieve initializers that must be called on the resulting entity range(s)
+	 * @return The number of new components added to OutMask. Only components that did not exist before are considered.
+	 */
+	int32 ComputeMutuallyInclusiveComponents(EMutuallyInclusiveComponentType TypesToCompute, const FComponentMask& InMask, FComponentMask& OutMask, FMutualComponentInitializers& OutInitializers) const;
+	
+	/**
+	 * Check this buffer's invariants are satisfied
+	 */
+	void CheckInvariants() const;
+
+private:
+
+	// Only FMutualInclusivityGraph can populate this command buffer
+	friend FMutualInclusivityGraph;
+
+	enum class ECommandType
+	{
+		Simple,
+		Type,
+		MatchAny,
+		MatchAll,
+		ShortCircuit,
+		Include,
+		Initialize,
+	};
+
+	struct FSimpleCommand
+	{
+		FComponentTypeID MatchID;
+		FComponentTypeID IncludeID;
+	};
+	struct FTypeCommand
+	{
+		/** Array index within FMutualInclusivityGraphCommandBuffer::Commands to jump to if this command fails */
+		uint16 ShortCircuitIndex;
+		/** The type to match */
+		EMutuallyInclusiveComponentType Type;
+	};
+
+	struct FMatchCommand
+	{
+		/** The ID of the component we are to match against, or include (depending on the value of CommandType) */
+		FComponentTypeID ComponentTypeID;
+		/** Array index within FMutualInclusivityGraphCommandBuffer::Commands to jump to if this command fails */
+		uint16 ShortCircuitIndex;
+	};
+
+	struct FShortCircuitCommand
+	{
+		/** Array index within FMutualInclusivityGraphCommandBuffer::Commands to jump to */
+		uint16 ShortCircuitIndex;
+	};
+
+	struct FIncludeCommand
+	{
+		/** The ID of the component we are to match against, or include (depending on the value of CommandType) */
+		FComponentTypeID ComponentTypeID;
+	};
+
+	struct FInitializeCommand
+	{
+		/** Pointer to an initializer to use */
+		IMutualComponentInitializer* Initializer;
+	};
+
+	/** A single command within a command buffer (FMutualInclusivityGraphCommandBuffer) that either matches or includes a component type ID */
+	struct FCommand
+	{
+		FCommand(FSimpleCommand InSimple)                                   : Simple(InSimple),             CommandType(ECommandType::Simple)       {}
+		FCommand(FTypeCommand InType)                                       : Type(InType),                 CommandType(ECommandType::Type)         {}
+		FCommand(FMatchCommand InMatch, EComplexInclusivityFilterMode Mode) : Match(InMatch),               CommandType(Mode == EComplexInclusivityFilterMode::AllOf ? ECommandType::MatchAll : ECommandType::MatchAny)   {}
+		FCommand(FShortCircuitCommand InShortCircuit)                       : ShortCircuit(InShortCircuit), CommandType(ECommandType::ShortCircuit) {}
+		FCommand(FIncludeCommand InInclude)                                 : Include(InInclude),           CommandType(ECommandType::Include)      {}
+		FCommand(FInitializeCommand InInitialize)                           : Initialize(InInitialize),     CommandType(ECommandType::Initialize)   {}
+
+		union
+		{
+			FTypeCommand         Type;
+			FSimpleCommand       Simple;
+			FMatchCommand        Match;
+			FShortCircuitCommand ShortCircuit;
+			FIncludeCommand      Include;
+			FInitializeCommand   Initialize;
+		};
+
+		/** The type of the command: either a match or include */
+		ECommandType CommandType;
+	};
+
+	/** Ordered command buffer specifying matches and includes */
+	TArray<FCommand> Commands;
+};
+
+
+/**
+ * A mutual inclusion graph for adding components to an entity based on the presence of other components
+ */
+struct FMutualInclusivityGraph
+{
+	/**
+	 * Define a new rule specifying that all of Dependents should always exist if Predicate exists
+	 */
+	void DefineMutualInclusionRule(FComponentTypeID Predicate, std::initializer_list<FComponentTypeID> Dependents);
+	void DefineMutualInclusionRule(FComponentTypeID Predicate, std::initializer_list<FComponentTypeID> Dependents, FMutuallyInclusiveComponentParams&& InParams);
+
+	/**
+	 * Define a new complex rule specifying that all of Dependents should exist if InFilter is matched
+	 */
+	void DefineComplexInclusionRule(const FComplexInclusivityFilter& InFilter, std::initializer_list<FComponentTypeID> Dependents);
+	void DefineComplexInclusionRule(const FComplexInclusivityFilter& InFilter, std::initializer_list<FComponentTypeID> Dependents, FMutuallyInclusiveComponentParams&& InParams);
+
+public:
+
+	/**
+	 * Populate a component mask with all the necessary mutually inclusive components given an input type.
+	 * Mutual components that already existed in InMask will not be added or counted.
+	 * Can be used to populate a new mask with only the mutual components, or to add the components to an existing mask.
+	 *
+	 * @param InMask           The input component type. Rules that are relevant to this type (directly or indirectly) will be applied to the output.
+	 * @param OutMask          Component mask to receieve the result. May point to the same mask as InMask in order to add the mutual components directly.
+	 * @param OutInitializers  Initializers that must be run for any new entities created from OutMask
+	 * @return The number of new components added to OutMask. Only components that did not exist before are considered.
+	 */
+	int32 ComputeMutuallyInclusiveComponents(EMutuallyInclusiveComponentType TypesToCompute, const FComponentMask& InMask, FComponentMask& OutMask, FMutualComponentInitializers& OutInitializers) const;
+
+private:
+
+	/**
+	 * Follows the longest chain in the graph starting from Component
+	 */
+	static int32 FindPrerequisiteChainLength(const FDirectedGraph& Graph, FComponentTypeID Component, TMap<FComponentTypeID, int32>& InOutCache);
+
+	/**
+	 * Reconstructs the command buffer
+	 */
+	void ReconstructCommandBuffer() const;
+
+private:
+
+	struct FIncludes
+	{
+		void Add(FComponentTypeID Component)
+		{
+			Includes.AddUnique(Component);
+		}
+		TArray<FComponentTypeID, TInlineAllocator<1>> Includes;
+		TArray<TUniquePtr<IMutualComponentInitializer>> Initializers;
+	};
+	struct FIncludePair
+	{
+		FIncludes MandatoryIncludes;
+		FIncludes OptionalIncludes;
+	};
+
+	/** 1->many simple include relationships */
+	TMap<FComponentTypeID, FIncludePair> SimpleIncludes;
+	/** many->many complex include relationships */
+	TMap<FComplexInclusivityFilter, FIncludePair> ComplexIncludes;
+
+	/** Mutable command buffer that is repopulated by ReconstructCommandBuffer */
+	mutable FMutualInclusivityGraphCommandBuffer CommandBuffer;
+	/** Flag indicating whether our command buffer needs re-generating or not */
+	mutable bool bCommandBufferInvalidated = false;
+};
+
+
+} // namespace UE::MovieScene
\ No newline at end of file
diff --git a/Engine/Source/Runtime/MovieScene/Public/EntitySystem/MovieSceneMutualComponentInitializer.h b/Engine/Source/Runtime/MovieScene/Public/EntitySystem/MovieSceneMutualComponentInitializer.h
new file mode 100644
--- /dev/null
+++ b/Engine/Source/Runtime/MovieScene/Public/EntitySystem/MovieSceneMutualComponentInitializer.h
@@ -0,0 +1,17 @@
+// Copyright Epic Games, Inc. All Rights Reserved.
+
+#pragma once
+
+namespace UE::MovieScene
+{
+
+struct FEntityRange;
+struct FEntityAllocationWriteContext;
+
+struct IMutualComponentInitializer
+{
+	virtual ~IMutualComponentInitializer(){}
+	virtual void Run(const FEntityRange& Range, const FEntityAllocationWriteContext& WriteContext) = 0;
+};
+
+} // namespace UE::MovieScene
\ No newline at end of file
diff --git a/Engine/Source/Runtime/MovieSceneTracks/Private/Systems/WeightAndEasingEvaluatorSystem.cpp b/Engine/Source/Runtime/MovieSceneTracks/Private/Systems/WeightAndEasingEvaluatorSystem.cpp
--- a/Engine/Source/Runtime/MovieSceneTracks/Private/Systems/WeightAndEasingEvaluatorSystem.cpp
+++ b/Engine/Source/Runtime/MovieSceneTracks/Private/Systems/WeightAndEasingEvaluatorSystem.cpp
@@ -238,7 +238,6 @@ struct FEasingChannelMutationBase : IMovieSceneConditionalEntityMutation
 	void CreateMutation(FEntityManager* InEntityManager, FComponentMask* InOutEntityComponentTypes) const override
 	{
 		InOutEntityComponentTypes->Set(BuiltInComponents->HierarchicalEasingChannel);
-		Factories->ComputeMutuallyInclusiveComponents(*InOutEntityComponentTypes);
 		InOutEntityComponentTypes->Set(BuiltInComponents->Tags.NeedsLink);
 	}
 
@@ -676,7 +675,7 @@ void UMovieSceneHierarchicalEasingInstantiatorSystem::OnRun(FSystemTaskPrerequis
 			ProviderFilter.None({ BuiltInComponents->Tags.NeedsUnlink });
 
 			FProviderEasingChannelMutation ProviderMutation(Linker, &SequenceDataMap);
-			Linker->EntityManager.MutateConditional(ProviderFilter, ProviderMutation);
+			Linker->EntityManager.MutateConditional(ProviderFilter, ProviderMutation, EMutuallyInclusiveComponentType::All);
 
 			ProviderMutation.RemoveStaleComponents(Linker);
 		}
@@ -689,7 +688,7 @@ void UMovieSceneHierarchicalEasingInstantiatorSystem::OnRun(FSystemTaskPrerequis
 			NewFilter.None({ BuiltInComponents->Tags.NeedsUnlink, BuiltInComponents->HierarchicalEasingProvider });
 
 			FConsumerEasingChannelMutation ConsumerMutation(Linker, &CachedInstanceIDToChannel);
-			Linker->EntityManager.MutateConditional(NewFilter, ConsumerMutation);
+			Linker->EntityManager.MutateConditional(NewFilter, ConsumerMutation, EMutuallyInclusiveComponentType::All);
 
 			ConsumerMutation.RemoveStaleComponents(Linker);
 		}
@@ -702,7 +701,7 @@ void UMovieSceneHierarchicalEasingInstantiatorSystem::OnRun(FSystemTaskPrerequis
 		Filter.None({ BuiltInComponents->HierarchicalEasingProvider });
 
 		FConsumerEasingChannelMutation ConsumerMutation(Linker, &CachedInstanceIDToChannel);
-		Linker->EntityManager.MutateConditional(Filter, ConsumerMutation);
+		Linker->EntityManager.MutateConditional(Filter, ConsumerMutation, EMutuallyInclusiveComponentType::All);
 
 		ensure(!ConsumerMutation.HasStaleComponents());
 		ConsumerMutation.RemoveStaleComponents(Linker);
