Sequencer: Miscellaneous performance improvements

  - Added update flags so that we can opt-out of calling virtual Pre/PostEvaluation and dissection functions if we will never need them
  - Removed bIsEvaluating flags from players in favor of a centrally stored bitfield. This reduces cache misses and indirections during the update phase.
  - Optimized ContainsAnyComponent by searching for the first set bit rather than counting all set bits

diff --git a/Engine/Source/Runtime/MovieScene/Private/EntitySystem/MovieSceneEntityManager.cpp b/Engine/Source/Runtime/MovieScene/Private/EntitySystem/MovieSceneEntityManager.cpp
--- a/Engine/Source/Runtime/MovieScene/Private/EntitySystem/MovieSceneEntityManager.cpp
+++ b/Engine/Source/Runtime/MovieScene/Private/EntitySystem/MovieSceneEntityManager.cpp
@@ -869,6 +869,8 @@ void FEntityManager::Compact()
 		Swap(AllocationsWithCapacity, Temp);
 	}
 
+	EntityLocations.Shrink();
+
 	CheckInvariants();
 
 	OnStructureChanged();
diff --git a/Engine/Source/Runtime/MovieScene/Private/EntitySystem/MovieSceneEntitySystemRunner.cpp b/Engine/Source/Runtime/MovieScene/Private/EntitySystem/MovieSceneEntitySystemRunner.cpp
--- a/Engine/Source/Runtime/MovieScene/Private/EntitySystem/MovieSceneEntitySystemRunner.cpp
+++ b/Engine/Source/Runtime/MovieScene/Private/EntitySystem/MovieSceneEntitySystemRunner.cpp
@@ -15,6 +15,7 @@
 DECLARE_CYCLE_STAT(TEXT("Runner Flush"), 				MovieSceneEval_RunnerFlush, 				STATGROUP_MovieSceneEval);
 
 DECLARE_CYCLE_STAT(TEXT("Spawn Phase"),                 MovieSceneEval_SpawnPhase,              	STATGROUP_MovieSceneECS);
+DECLARE_CYCLE_STAT(TEXT("Post Spawn Event"),            MovieSceneEval_PostSpawnEvent,            	STATGROUP_MovieSceneECS);
 DECLARE_CYCLE_STAT(TEXT("Instantiation Phase"), 		MovieSceneEval_InstantiationPhase, 			STATGROUP_MovieSceneECS);
 DECLARE_CYCLE_STAT(TEXT("Instantiation Async Tasks"), 	MovieSceneEval_AsyncInstantiationTasks,		STATGROUP_MovieSceneECS);
 DECLARE_CYCLE_STAT(TEXT("Post Instantiation"), 			MovieSceneEval_PostInstantiation, 			STATGROUP_MovieSceneECS);
@@ -685,12 +686,22 @@ bool FMovieSceneEntitySystemRunner::GameThread_UpdateSequenceInstances(UMovieSce
 		}
 	}
 
+	AccumulatedUpdateFlags = ESequenceInstanceUpdateFlags::None;
+
 	// Let sequence instances do any pre-evaluation work.
 	for (const FQueuedUpdateParams& UpdatedInstance : CurrentInstances)
 	{
 		if (!EnumHasAnyFlags(UpdatedInstance.UpdateFlags, ERunnerUpdateFlags::Destroy))
 		{
-			InstanceRegistry->MutateInstance(UpdatedInstance.InstanceHandle).PreEvaluation(Linker);
+			FSequenceInstance& SequenceInstance = InstanceRegistry->MutateInstance(UpdatedInstance.InstanceHandle);
+
+			AccumulatedUpdateFlags |= SequenceInstance.GetUpdateFlags();
+			IMovieScenePlayer::SetIsEvaluatingFlag(SequenceInstance.GetPlayerIndex(), true);
+
+			if (EnumHasAnyFlags(SequenceInstance.GetUpdateFlags(), ESequenceInstanceUpdateFlags::NeedsPreEvaluation))
+			{
+				SequenceInstance.PreEvaluation(Linker);
+			}
 		}
 	}
 
@@ -711,6 +722,9 @@ bool FMovieSceneEntitySystemRunner::GameThread_UpdateSequenceInstances(UMovieSce
 			if (Update.OnFlushed.IsBound())
 			{
 				OnFlushedDelegates.Add(MoveTemp(Update.OnFlushed));
+
+				// If we have any on-flushed delegates then we have to do a full Post-Eval phase
+				AccumulatedUpdateFlags |= ESequenceInstanceUpdateFlags::NeedsPostEvaluation;
 			}
 
 			if (EnumHasAnyFlags(Update.UpdateFlags, ERunnerUpdateFlags::Destroy))
@@ -805,8 +819,10 @@ bool FMovieSceneEntitySystemRunner::GameThread_SpawnPhase(UMovieSceneEntitySyste
 		}
 	}
 
-
-	Linker->Events.PostSpawnEvent.Broadcast(Linker);
+	{
+		SCOPE_CYCLE_COUNTER(MovieSceneEval_PostSpawnEvent);
+		Linker->Events.PostSpawnEvent.Broadcast(Linker);
+	}
 
 	// --------------------------------------------------------------------------------------------------------------------------------------------
 	// Only run the instantiation phase if there is anything to instantiate. This must come after the spawn phase because new instantiations may
@@ -1013,7 +1029,14 @@ void FMovieSceneEntitySystemRunner::GameThread_PostEvaluationPhase(UMovieSceneEn
 			if (InstanceRegistry->IsHandleValid(UpdateParams.InstanceHandle))
 			{
 				FSequenceInstance& Instance = InstanceRegistry->MutateInstance(UpdateParams.InstanceHandle);
-				Instance.PostEvaluation(Linker);
+
+				Instance.Ledger.UnlinkOneShots(Linker);
+				IMovieScenePlayer::SetIsEvaluatingFlag(Instance.GetPlayerIndex(), false);
+
+				if (EnumHasAnyFlags(Instance.GetUpdateFlags(), ESequenceInstanceUpdateFlags::NeedsPostEvaluation))
+				{
+					Instance.PostEvaluation(Linker);
+				}
 
 				if (EnumHasAnyFlags(UpdateParams.UpdateFlags, ERunnerUpdateFlags::Destroy))
 				{
diff --git a/Engine/Source/Runtime/MovieScene/Private/EntitySystem/MovieSceneSequenceInstance.cpp b/Engine/Source/Runtime/MovieScene/Private/EntitySystem/MovieSceneSequenceInstance.cpp
--- a/Engine/Source/Runtime/MovieScene/Private/EntitySystem/MovieSceneSequenceInstance.cpp
+++ b/Engine/Source/Runtime/MovieScene/Private/EntitySystem/MovieSceneSequenceInstance.cpp
@@ -74,6 +74,8 @@ FSequenceInstance::FSequenceInstance(UMovieSceneEntitySystemLinker* Linker, IMov
 	, InstanceHandle(InInstanceHandle)
 	, RootInstanceHandle(InInstanceHandle)
 {
+	UpdateFlags = ESequenceInstanceUpdateFlags::None;
+
 	// Root instances always start in a finished state in order to ensure that 'Start'
 	// is called correctly for the top level instance. This is subtly different from
 	// bHasEverUpdated since a sequence instance can be Finished and restarted multiple times
@@ -99,6 +101,8 @@ FSequenceInstance::FSequenceInstance(UMovieSceneEntitySystemLinker* Linker, IMov
 	, ParentInstanceHandle(InParentInstanceHandle)
 	, RootInstanceHandle(InRootInstanceHandle)
 {
+	UpdateFlags = ESequenceInstanceUpdateFlags::None;
+
 	// Sub Sequence instances always start in a non-finished state because they will only ever
 	// be created if they are active, and the Start/Update/Finish loop does not apply to sub-instances
 	bFinished = false;
@@ -152,6 +156,8 @@ void FSequenceInstance::InvalidateCachedData(UMovieSceneEntitySystemLinker* Link
 	IMovieScenePlayer* Player = GetPlayer();
 	check(Player);
 
+	UpdateFlags = ESequenceInstanceUpdateFlags::None;
+
 	FMovieSceneRootEvaluationTemplateInstance& RootTemplate        = Player->GetEvaluationTemplate();
 	UMovieSceneCompiledDataManager*            CompiledDataManager = RootTemplate.GetCompiledDataManager();
 	FMovieSceneCompiledDataID                  RootCompiledDataID  = RootTemplate.GetCompiledDataID();
@@ -173,6 +179,7 @@ void FSequenceInstance::InvalidateCachedData(UMovieSceneEntitySystemLinker* Link
 		ISequenceUpdater::FactoryInstance(SequenceUpdater, CompiledDataManager, RootCompiledDataID);
 
 		SequenceUpdater->InvalidateCachedData(Linker);
+		SequenceUpdater->PopulateUpdateFlags(Linker, Player, UpdateFlags);
 
 		if (LegacyEvaluator)
 		{
@@ -189,6 +196,11 @@ void FSequenceInstance::InvalidateCachedData(UMovieSceneEntitySystemLinker* Link
 
 void FSequenceInstance::DissectContext(UMovieSceneEntitySystemLinker* Linker, const FMovieSceneContext& InContext, TArray<TRange<FFrameTime>>& OutDissections)
 {
+	if (!EnumHasAnyFlags(UpdateFlags, ESequenceInstanceUpdateFlags::NeedsDissection))
+	{
+		return;
+	}
+
 	check(SequenceID == MovieSceneSequenceID::Root);
 
 	IMovieScenePlayer* Player = GetPlayer();
@@ -296,6 +308,11 @@ void FSequenceInstance::Finish(UMovieSceneEntitySystemLinker* Linker)
 
 void FSequenceInstance::PreEvaluation(UMovieSceneEntitySystemLinker* Linker)
 {
+	if (!EnumHasAnyFlags(UpdateFlags, ESequenceInstanceUpdateFlags::NeedsPreEvaluation))
+	{
+		return;
+	}
+
 	if (IsRootSequence())
 	{
 		IMovieScenePlayer* Player = GetPlayer();
@@ -327,9 +344,7 @@ void FSequenceInstance::RunLegacyTrackTemplates()
 
 void FSequenceInstance::PostEvaluation(UMovieSceneEntitySystemLinker* Linker)
 {
-	Ledger.UnlinkOneShots(Linker);
-
-	if (IsRootSequence())
+	if (IsRootSequence() && EnumHasAnyFlags(UpdateFlags, ESequenceInstanceUpdateFlags::NeedsPostEvaluation))
 	{
 		IMovieScenePlayer* Player = GetPlayer();
 		if (ensure(Player))
diff --git a/Engine/Source/Runtime/MovieScene/Private/EntitySystem/MovieSceneSequenceUpdaters.cpp b/Engine/Source/Runtime/MovieScene/Private/EntitySystem/MovieSceneSequenceUpdaters.cpp
--- a/Engine/Source/Runtime/MovieScene/Private/EntitySystem/MovieSceneSequenceUpdaters.cpp
+++ b/Engine/Source/Runtime/MovieScene/Private/EntitySystem/MovieSceneSequenceUpdaters.cpp
@@ -35,6 +35,7 @@ struct FSequenceUpdater_Flat : ISequenceUpdater
 	explicit FSequenceUpdater_Flat(FMovieSceneCompiledDataID InCompiledDataID);
 	~FSequenceUpdater_Flat();
 
+	virtual void PopulateUpdateFlags(UMovieSceneEntitySystemLinker* InLinker, IMovieScenePlayer* InPlayer, ESequenceInstanceUpdateFlags& OutUpdateFlags) override;
 	virtual void DissectContext(UMovieSceneEntitySystemLinker* Linker, IMovieScenePlayer* InPlayer, const FMovieSceneContext& Context, TArray<TRange<FFrameTime>>& OutDissections) override;
 	virtual void Start(UMovieSceneEntitySystemLinker* Linker, FRootInstanceHandle InstanceHandle, IMovieScenePlayer* InPlayer, const FMovieSceneContext& InContext) override;
 	virtual void Update(UMovieSceneEntitySystemLinker* Linker, FRootInstanceHandle InstanceHandle, IMovieScenePlayer* InPlayer, const FMovieSceneContext& Context) override;
@@ -63,6 +64,7 @@ struct FSequenceUpdater_Hierarchical : ISequenceUpdater
 
 	~FSequenceUpdater_Hierarchical();
 
+	virtual void PopulateUpdateFlags(UMovieSceneEntitySystemLinker* InLinker, IMovieScenePlayer* InPlayer, ESequenceInstanceUpdateFlags& OutUpdateFlags) override;
 	virtual void DissectContext(UMovieSceneEntitySystemLinker* Linker, IMovieScenePlayer* InPlayer, const FMovieSceneContext& Context, TArray<TRange<FFrameTime>>& OutDissections) override;
 	virtual void Start(UMovieSceneEntitySystemLinker* Linker, FRootInstanceHandle InstanceHandle, IMovieScenePlayer* InPlayer, const FMovieSceneContext& InContext) override;
 	virtual void Update(UMovieSceneEntitySystemLinker* Linker, FRootInstanceHandle InstanceHandle, IMovieScenePlayer* InPlayer, const FMovieSceneContext& Context) override;
@@ -194,6 +196,31 @@ TUniquePtr<ISequenceUpdater> FSequenceUpdater_Flat::MigrateToHierarchical()
 	return MakeUnique<FSequenceUpdater_Hierarchical>(CompiledDataID);
 }
 
+void FSequenceUpdater_Flat::PopulateUpdateFlags(UMovieSceneEntitySystemLinker* InLinker, IMovieScenePlayer* InPlayer, ESequenceInstanceUpdateFlags& OutUpdateFlags)
+{
+	if (!CachedDeterminismFences.IsSet())
+	{
+		UMovieSceneCompiledDataManager* CompiledDataManager = InPlayer->GetEvaluationTemplate().GetCompiledDataManager();
+		TArrayView<const FFrameTime>    DeterminismFences   = CompiledDataManager->GetEntryRef(CompiledDataID).DeterminismFences;
+
+		if (DeterminismFences.Num() != 0)
+		{
+			CachedDeterminismFences = TArray<FFrameTime>(DeterminismFences.GetData(), DeterminismFences.Num());
+		}
+		else
+		{
+			CachedDeterminismFences.Emplace();
+		}
+	}
+
+	InPlayer->PopulateUpdateFlags(OutUpdateFlags);
+
+	if (CachedDeterminismFences.IsSet() && CachedDeterminismFences->Num() > 0)
+	{
+		OutUpdateFlags |= ESequenceInstanceUpdateFlags::NeedsDissection;
+	}
+}
+
 void FSequenceUpdater_Flat::DissectContext(UMovieSceneEntitySystemLinker* Linker, IMovieScenePlayer* InPlayer, const FMovieSceneContext& Context, TArray<TRange<FFrameTime>>& OutDissections)
 {
 	if (!CachedDeterminismFences.IsSet())
@@ -325,6 +352,42 @@ FSequenceUpdater_Hierarchical::~FSequenceUpdater_Hierarchical()
 {
 }
 
+void FSequenceUpdater_Hierarchical::PopulateUpdateFlags(UMovieSceneEntitySystemLinker* InLinker, IMovieScenePlayer* InPlayer, ESequenceInstanceUpdateFlags& OutUpdateFlags)
+{
+	InPlayer->PopulateUpdateFlags(OutUpdateFlags);
+
+	// If we've already been told we need dissection there's nothing left to do
+	if (EnumHasAnyFlags(OutUpdateFlags, ESequenceInstanceUpdateFlags::NeedsDissection))
+	{
+		return;
+	}
+
+	UMovieSceneCompiledDataManager* CompiledDataManager = InPlayer->GetEvaluationTemplate().GetCompiledDataManager();
+
+	{
+		const FMovieSceneCompiledDataEntry& RootDataEntry = CompiledDataManager->GetEntryRef(CompiledDataID);
+		if (RootDataEntry.DeterminismFences.Num() > 0)
+		{
+			OutUpdateFlags |= ESequenceInstanceUpdateFlags::NeedsDissection;
+		}
+	}
+
+	if (const FMovieSceneSequenceHierarchy* Hierarchy = CompiledDataManager->FindHierarchy(CompiledDataID))
+	{
+		for (const TPair<FMovieSceneSequenceID, FMovieSceneSubSequenceData>& Pair : Hierarchy->AllSubSequenceData())
+		{
+			UMovieSceneSequence*      SubSequence = Pair.Value.GetSequence();
+			FMovieSceneCompiledDataID SubDataID   = SubSequence ? CompiledDataManager->GetDataID(SubSequence) : FMovieSceneCompiledDataID();
+
+			if (SubDataID.IsValid() && CompiledDataManager->GetEntryRef(SubDataID).DeterminismFences.Num() > 0)
+			{
+				OutUpdateFlags |= ESequenceInstanceUpdateFlags::NeedsDissection;
+				break;
+			}
+		}
+	}
+}
+
 void FSequenceUpdater_Hierarchical::DissectContext(UMovieSceneEntitySystemLinker* Linker, IMovieScenePlayer* InPlayer, const FMovieSceneContext& Context, TArray<TRange<FFrameTime>>& OutDissections)
 {
 	UMovieSceneCompiledDataManager* CompiledDataManager = InPlayer->GetEvaluationTemplate().GetCompiledDataManager();
diff --git a/Engine/Source/Runtime/MovieScene/Private/MovieScenePlayer.cpp b/Engine/Source/Runtime/MovieScene/Private/MovieScenePlayer.cpp
--- a/Engine/Source/Runtime/MovieScene/Private/MovieScenePlayer.cpp
+++ b/Engine/Source/Runtime/MovieScene/Private/MovieScenePlayer.cpp
@@ -15,6 +15,7 @@ namespace MovieScene
 
 static FRWLock                          GGlobalPlayerRegistryLock;
 static TSparseArray<IMovieScenePlayer*> GGlobalPlayerRegistry;
+static TBitArray<> GGlobalPlayerUpdateFlags;
 
 } // namespace MovieScene
 } // namespace UE
@@ -22,7 +23,11 @@ static TSparseArray<IMovieScenePlayer*> GGlobalPlayerRegistry;
 IMovieScenePlayer::IMovieScenePlayer()
 {
 	FWriteScopeLock ScopeLock(UE::MovieScene::GGlobalPlayerRegistryLock);
+
+	UE::MovieScene::GGlobalPlayerRegistry.Shrink();
 	UniqueIndex = UE::MovieScene::GGlobalPlayerRegistry.Add(this);
+
+	UE::MovieScene::GGlobalPlayerUpdateFlags.PadToNum(UniqueIndex + 1, false);
 }
 
 IMovieScenePlayer::~IMovieScenePlayer()
@@ -38,6 +43,24 @@ IMovieScenePlayer* IMovieScenePlayer::Get(uint16 InUniqueIndex)
 	return UE::MovieScene::GGlobalPlayerRegistry[InUniqueIndex];
 }
 
+void IMovieScenePlayer::SetIsEvaluatingFlag(uint16 InUniqueIndex, bool bIsUpdating)
+{
+	check(UE::MovieScene::GGlobalPlayerUpdateFlags.IsValidIndex(InUniqueIndex));
+	UE::MovieScene::GGlobalPlayerUpdateFlags[InUniqueIndex] = bIsUpdating;
+}
+
+bool IMovieScenePlayer::IsEvaluating() const
+{
+	return UE::MovieScene::GGlobalPlayerUpdateFlags[UniqueIndex];
+}
+
+void IMovieScenePlayer::PopulateUpdateFlags(UE::MovieScene::ESequenceInstanceUpdateFlags& OutFlags)
+{
+	using namespace UE::MovieScene;
+
+	OutFlags |= ESequenceInstanceUpdateFlags::NeedsPreEvaluation | ESequenceInstanceUpdateFlags::NeedsPostEvaluation;
+}
+
 void IMovieScenePlayer::ResolveBoundObjects(const FGuid& InBindingId, FMovieSceneSequenceID SequenceID, UMovieSceneSequence& Sequence, UObject* ResolutionContext, TArray<UObject*, TInlineAllocator<1>>& OutObjects) const
 {
 	Sequence.LocateBoundObjects(InBindingId, ResolutionContext, OutObjects);
diff --git a/Engine/Source/Runtime/MovieScene/Private/MovieSceneSequencePlayer.cpp b/Engine/Source/Runtime/MovieScene/Private/MovieSceneSequencePlayer.cpp
--- a/Engine/Source/Runtime/MovieScene/Private/MovieSceneSequencePlayer.cpp
+++ b/Engine/Source/Runtime/MovieScene/Private/MovieSceneSequencePlayer.cpp
@@ -116,7 +116,6 @@ UMovieSceneSequencePlayer::UMovieSceneSequencePlayer(const FObjectInitializer& I
 	, Status(EMovieScenePlayerStatus::Stopped)
 	, bReversePlayback(false)
 	, bPendingOnStartedPlaying(false)
-	, bIsEvaluating(false)
 	, bIsAsyncUpdate(false)
 	, bSkipNextUpdate(false)
 	, Sequence(nullptr)
@@ -785,7 +784,7 @@ void UMovieSceneSequencePlayer::Initialize(UMovieSceneSequence* InSequence, cons
 void UMovieSceneSequencePlayer::Initialize(UMovieSceneSequence* InSequence)
 {
 	check(InSequence);
-	check(!bIsEvaluating);
+	check(!IsEvaluating());
 
 	// If we have a valid sequence that may have been played back,
 	// Explicitly stop and tear down the template instance before 
@@ -938,7 +937,7 @@ void UMovieSceneSequencePlayer::Update(const float DeltaSeconds)
 
 		if (!bSkipNextUpdate)
 		{
-			check(!bIsEvaluating);
+			check(!IsEvaluating());
 
 			FFrameTime NewTime = TimeController->RequestCurrentTime(GetCurrentTime(), PlayRate, GetDisplayRate());
 			UpdateTimeCursorPosition(NewTime, EUpdatePositionMethod::Play);
@@ -974,7 +973,7 @@ void UMovieSceneSequencePlayer::UpdateAsync(const float DeltaSeconds)
 
 void UMovieSceneSequencePlayer::UpdateTimeCursorPosition(FFrameTime NewPosition, EUpdatePositionMethod Method, bool bHasJumpedOverride)
 {
-	if (ensure(!bIsEvaluating))
+	if (ensure(!IsEvaluating()))
 	{
 		UpdateTimeCursorPosition_Internal(NewPosition, Method, bHasJumpedOverride);
 	}
@@ -1154,7 +1153,7 @@ void UMovieSceneSequencePlayer::UpdateTimeCursorPosition_Internal(FFrameTime New
 
 		// Now that the evaluation has taken place we call Pause (to trigger delegates, etc.), however if we're running the evaluation
 		// as async we actually need to queue a latent action to do it so that it happens after the data is actually updated. We can't
-		// use the QueueLatentAction inside of Pause() because bIsEvaluating is only set during actual evaluation so it won't be set yet.
+		// use the QueueLatentAction inside of Pause() because IsEvaluating() is only set during actual evaluation so it won't be set yet.
 		if (PauseRange.IsSet())
 		{
 			if (Args.bIsAsync)
@@ -1249,8 +1248,6 @@ bool UMovieSceneSequencePlayer::HasDynamicWeighting() const
 void UMovieSceneSequencePlayer::PreEvaluation(const FMovieSceneContext& Context)
 {
 	RunPreEvaluationCallbacks();
-
-	bIsEvaluating = true;
 }
 
 void UMovieSceneSequencePlayer::PostEvaluation(const FMovieSceneContext& Context)
@@ -1260,8 +1257,6 @@ void UMovieSceneSequencePlayer::PostEvaluation(const FMovieSceneContext& Context
 	OnMovieSceneSequencePlayerUpdate.Broadcast(*this, CurrentTime, PreviousTime);
 
 	RunPostEvaluationCallbacks();
-
-	bIsEvaluating = false;
 }
 
 void UMovieSceneSequencePlayer::RunPreEvaluationCallbacks()
@@ -1790,7 +1785,7 @@ bool UMovieSceneSequencePlayer::CallRemoteFunction(UFunction* Function, void* Pa
 
 bool UMovieSceneSequencePlayer::NeedsQueueLatentAction() const
 {
-	return bIsEvaluating;
+	return IsEvaluating();
 }
 
 void UMovieSceneSequencePlayer::QueueLatentAction(FMovieSceneSequenceLatentActionDelegate Delegate)
diff --git a/Engine/Source/Runtime/MovieScene/Public/EntitySystem/MovieSceneEntityManager.h b/Engine/Source/Runtime/MovieScene/Public/EntitySystem/MovieSceneEntityManager.h
--- a/Engine/Source/Runtime/MovieScene/Public/EntitySystem/MovieSceneEntityManager.h
+++ b/Engine/Source/Runtime/MovieScene/Public/EntitySystem/MovieSceneEntityManager.h
@@ -741,7 +741,7 @@ public:
 	bool ContainsAnyComponent(const FComponentMask& ComponentTypeIDs) const
 	{
 		const FComponentMask& Mask = GetAccumulatedMask();
-		return FComponentMask::BitwiseAND(Mask, ComponentTypeIDs, EBitwiseOperatorFlags::MinSize).NumComponents() != 0;
+		return FComponentMask::BitwiseAND(Mask, ComponentTypeIDs, EBitwiseOperatorFlags::MinSize).First() != FComponentTypeID::Invalid();
 	}
 
 	/**
diff --git a/Engine/Source/Runtime/MovieScene/Public/EntitySystem/MovieSceneEntitySystemRunner.h b/Engine/Source/Runtime/MovieScene/Public/EntitySystem/MovieSceneEntitySystemRunner.h
--- a/Engine/Source/Runtime/MovieScene/Public/EntitySystem/MovieSceneEntitySystemRunner.h
+++ b/Engine/Source/Runtime/MovieScene/Public/EntitySystem/MovieSceneEntitySystemRunner.h
@@ -258,6 +258,8 @@ private:
 	TArray<FDissectedUpdate> DissectedUpdates;
 	TArray<FSimpleDelegate> OnFlushedDelegates;
 	FMovieSceneEntitySystemEventTriggers EventTriggers;
+	/** Update flags that are accumulated for all currently active sequence instances being evaluated */
+	UE::MovieScene::ESequenceInstanceUpdateFlags AccumulatedUpdateFlags;
 
 	/** The number of times this runner has been re-entrant */
 	uint32 ReentrancyCount;
diff --git a/Engine/Source/Runtime/MovieScene/Public/EntitySystem/MovieSceneSequenceInstance.h b/Engine/Source/Runtime/MovieScene/Public/EntitySystem/MovieSceneSequenceInstance.h
--- a/Engine/Source/Runtime/MovieScene/Public/EntitySystem/MovieSceneSequenceInstance.h
+++ b/Engine/Source/Runtime/MovieScene/Public/EntitySystem/MovieSceneSequenceInstance.h
@@ -34,6 +34,15 @@ struct FSequenceInstance;
 struct FSubSequencePath;
 struct ISequenceUpdater;
 
+enum class ESequenceInstanceUpdateFlags : uint8
+{
+	None                = 0,
+	NeedsDissection     = 1u << 0,
+	NeedsPreEvaluation  = 1u << 1,
+	NeedsPostEvaluation = 1u << 2,
+};
+ENUM_CLASS_FLAGS(ESequenceInstanceUpdateFlags);
+
 /**
  * A sequence instance represents a specific instance of a currently playing sequence, either as a top-level sequence in an IMovieScenePlayer, or as a sub sequence.
  * Any given sequence asset may have any number of instances created for it at any given time depending on how many times it is referenced by playing sequences
@@ -119,6 +128,14 @@ public:
 	 */
 	IMovieScenePlayer* GetPlayer() const;
 
+	/**
+	 * Retrieve the IMovieScenePlayer's unique index
+	 */
+	uint16 GetPlayerIndex() const
+	{
+		return PlayerIndex;
+	}
+
 	/**
 	 * Retrieve the SequenceID for this instance
 	 *
@@ -211,6 +228,14 @@ public:
 		return bHasEverUpdated;
 	}
 
+	/**
+	 * Retrieve this sequence's update flags
+	 */
+	ESequenceInstanceUpdateFlags GetUpdateFlags() const
+	{
+		return UpdateFlags;
+	}
+
 	/**
 	 * Attempt to locate a sub instance based on its sequence ID
 	 */
@@ -321,7 +346,9 @@ private:
 	/** When SequenceID != MovieSceneSequenceID::Root, specifies an ID to override as a simulated root. */
 	FMovieSceneSequenceID RootOverrideSequenceID;
 	/** The index of this instance's IMovieScenePlayer retrievable through IMovieScenePlayer::Get(). */
-	int32 PlayerIndex;
+	uint16 PlayerIndex;
+	/** Cached update flags denoting what kinds of updates are required by this instance */
+	ESequenceInstanceUpdateFlags UpdateFlags;
 	/** This instance's handle. */
 	FInstanceHandle InstanceHandle;
 	/** This instance's parent handle. */
diff --git a/Engine/Source/Runtime/MovieScene/Public/EntitySystem/MovieSceneSequenceUpdaters.h b/Engine/Source/Runtime/MovieScene/Public/EntitySystem/MovieSceneSequenceUpdaters.h
--- a/Engine/Source/Runtime/MovieScene/Public/EntitySystem/MovieSceneSequenceUpdaters.h
+++ b/Engine/Source/Runtime/MovieScene/Public/EntitySystem/MovieSceneSequenceUpdaters.h
@@ -52,6 +52,11 @@ public:
 	/** Virtual destructor */
 	virtual ~ISequenceUpdater() {}
 
+	/**
+	 * Called to initialize the flag structure that denotes what functions need to be called on this updater
+	 */
+	virtual void PopulateUpdateFlags(UMovieSceneEntitySystemLinker* InLinker, IMovieScenePlayer* InPlayer, ESequenceInstanceUpdateFlags& OutUpdateFlags) = 0;
+
 	/**
 	 * Called before any updates to the sequence to allow this updater to dissect the context into smaller ranges that should be evaluated independently for the purpose of ensuring determinism.
 	 * If the resulting array is empty the whole context will be used by default.
diff --git a/Engine/Source/Runtime/MovieScene/Public/IMovieScenePlayer.h b/Engine/Source/Runtime/MovieScene/Public/IMovieScenePlayer.h
--- a/Engine/Source/Runtime/MovieScene/Public/IMovieScenePlayer.h
+++ b/Engine/Source/Runtime/MovieScene/Public/IMovieScenePlayer.h
@@ -31,6 +31,10 @@ struct FMovieSceneRootEvaluationTemplateInstance;
 class FMovieSceneSequenceInstance;
 class IMovieScenePlayer;
 
+namespace UE::MovieScene
+{
+	enum class ESequenceInstanceUpdateFlags : uint8;
+}
 
 struct EMovieSceneViewportParams
 {
@@ -192,6 +196,11 @@ public:
 	 */
 	virtual void NotifyBindingUpdate(const FGuid& InGuid, FMovieSceneSequenceIDRef InSequenceID, TArrayView<TWeakObjectPtr<>> Objects) { NotifyBindingsChanged(); }
 
+	/**
+	 * Called to initialize the flag structure that denotes what functions need to be called on this updater
+	 */
+	MOVIESCENE_API virtual void PopulateUpdateFlags(UE::MovieScene::ESequenceInstanceUpdateFlags& OutFlags);
+
 	/**
 	 * Called whenever any object bindings have changed
 	 */
@@ -337,6 +346,10 @@ public:
 
 	MOVIESCENE_API static IMovieScenePlayer* Get(uint16 InUniqueIndex);
 
+	MOVIESCENE_API static void SetIsEvaluatingFlag(uint16 InUniqueIndex, bool bIsUpdating);
+
+	MOVIESCENE_API bool IsEvaluating() const;
+
 	uint16 GetUniqueIndex() const
 	{
 		return UniqueIndex;
diff --git a/Engine/Source/Runtime/MovieScene/Public/MovieSceneSequencePlayer.h b/Engine/Source/Runtime/MovieScene/Public/MovieSceneSequencePlayer.h
--- a/Engine/Source/Runtime/MovieScene/Public/MovieSceneSequencePlayer.h
+++ b/Engine/Source/Runtime/MovieScene/Public/MovieSceneSequencePlayer.h
@@ -644,9 +644,6 @@ protected:
 	/** Set to true to invoke OnStartedPlaying on first update tick for started playing */
 	uint32 bPendingOnStartedPlaying : 1;
 
-	/** Set to true while the player's sequence is being evaluated */
-	uint32 bIsEvaluating : 1;
-
 	/** Set to true when the player is currently in the main level update */
 	uint32 bIsAsyncUpdate : 1;
 
diff --git a/Engine/Source/Runtime/UMG/Private/Animation/UMGSequencePlayer.cpp b/Engine/Source/Runtime/UMG/Private/Animation/UMGSequencePlayer.cpp
--- a/Engine/Source/Runtime/UMG/Private/Animation/UMGSequencePlayer.cpp
+++ b/Engine/Source/Runtime/UMG/Private/Animation/UMGSequencePlayer.cpp
@@ -42,7 +42,6 @@ UUMGSequencePlayer::UUMGSequencePlayer(const FObjectInitializer& ObjectInitializ
 	PlaybackSpeed = 1;
 	bRestoreState = false;
 	Animation = nullptr;
-	bIsEvaluating = false;
 	bIsStopping = false;
 	bIsBeginningPlay = false;
 	bCompleteOnPostEvaluation = false;
@@ -84,12 +83,14 @@ UMovieSceneEntitySystemLinker* UUMGSequencePlayer::ConstructEntitySystemLinker()
 
 void UUMGSequencePlayer::Tick(float DeltaTime)
 {
-	if (bIsStopping || bIsBeginningPlay || bIsEvaluating)
+	if (IsEvaluating())
+	{
+		BlockedDeltaTimeCompensation += DeltaTime;
+		return;
+	}
+
+	if (bIsStopping || bIsBeginningPlay)
 	{
-		if (bIsEvaluating)
-		{
-			BlockedDeltaTimeCompensation += DeltaTime;
-		}
 		return;
 	}
 
@@ -501,19 +502,14 @@ void UUMGSequencePlayer::SetPlaybackStatus(EMovieScenePlayerStatus::Type InPlayb
 	PlayerStatus = InPlaybackStatus;
 }
 
-void UUMGSequencePlayer::PreEvaluation(const FMovieSceneContext& Context)
-{
-	bIsEvaluating = true;
-}
-
-void UUMGSequencePlayer::PostEvaluation(const FMovieSceneContext& Context)
+void UUMGSequencePlayer::PopulateUpdateFlags(UE::MovieScene::ESequenceInstanceUpdateFlags& OutFlags)
 {
-	bIsEvaluating = false;
+	// Leave empty so Pre and PostEvaluation are not called
 }
 
 bool UUMGSequencePlayer::NeedsQueueLatentAction() const
 {
-	return bIsEvaluating;
+	return IsEvaluating();
 }
 
 void UUMGSequencePlayer::QueueLatentAction(FMovieSceneSequenceLatentActionDelegate Delegate)
diff --git a/Engine/Source/Runtime/UMG/Public/Animation/UMGSequencePlayer.h b/Engine/Source/Runtime/UMG/Public/Animation/UMGSequencePlayer.h
--- a/Engine/Source/Runtime/UMG/Public/Animation/UMGSequencePlayer.h
+++ b/Engine/Source/Runtime/UMG/Public/Animation/UMGSequencePlayer.h
@@ -79,8 +79,7 @@ public:
 	virtual UObject* GetPlaybackContext() const override;
 	virtual TArray<UObject*> GetEventContexts() const override;
 	virtual void SetPlaybackStatus(EMovieScenePlayerStatus::Type InPlaybackStatus) override;
-	virtual void PreEvaluation(const FMovieSceneContext& Context) override;
-	virtual void PostEvaluation(const FMovieSceneContext& Context) override;
+	virtual void PopulateUpdateFlags(UE::MovieScene::ESequenceInstanceUpdateFlags& OutFlags) override;
 
 	/** UObject interface */
 	virtual void BeginDestroy() override;
@@ -167,9 +166,6 @@ private:
 	/** True if the animation is playing forward, otherwise false and it's playing in reverse. */
 	bool bIsPlayingForward;
 
-	/** Set to true while evaluating to prevent reentrancy */
-	bool bIsEvaluating : 1;
-
 	/** Set to true when this player is in the process of being stopped */
 	bool bIsStopping : 1;
 
