Added missing fwd declaration

diff --git a/Engine/Source/Runtime/MovieScene/Public/EntitySystem/MovieSceneSequenceUpdaters.h b/Engine/Source/Runtime/MovieScene/Public/EntitySystem/MovieSceneSequenceUpdaters.h
--- a/Engine/Source/Runtime/MovieScene/Public/EntitySystem/MovieSceneSequenceUpdaters.h
+++ b/Engine/Source/Runtime/MovieScene/Public/EntitySystem/MovieSceneSequenceUpdaters.h
@@ -25,6 +25,8 @@ namespace MovieScene
 struct FInstanceHandle;
 struct FRootInstanceHandle;
 
+enum class ESequenceInstanceUpdateFlags : uint8;
+
 
 /**
  * Interface for an object that will update a sequence based on the current context. It holds several responsibilities:
