Sequencer: Always reset IsEvaluating flags when creating or destroying a MovieScenePlayer

diff --git a/Engine/Source/Runtime/MovieScene/Private/MovieScenePlayer.cpp b/Engine/Source/Runtime/MovieScene/Private/MovieScenePlayer.cpp
--- a/Engine/Source/Runtime/MovieScene/Private/MovieScenePlayer.cpp
+++ b/Engine/Source/Runtime/MovieScene/Private/MovieScenePlayer.cpp
@@ -28,11 +28,14 @@ IMovieScenePlayer::IMovieScenePlayer()
 	UniqueIndex = UE::MovieScene::GGlobalPlayerRegistry.Add(this);
 
 	UE::MovieScene::GGlobalPlayerUpdateFlags.PadToNum(UniqueIndex + 1, false);
+	UE::MovieScene::GGlobalPlayerUpdateFlags[UniqueIndex] = 0;
 }
 
 IMovieScenePlayer::~IMovieScenePlayer()
-{
+{	
 	FWriteScopeLock ScopeLock(UE::MovieScene::GGlobalPlayerRegistryLock);
+
+	UE::MovieScene::GGlobalPlayerUpdateFlags[UniqueIndex] = 0;
 	UE::MovieScene::GGlobalPlayerRegistry.RemoveAt(UniqueIndex, 1);
 }
 
