Use a unique object name when creating linkers in the MovieSceneSequenceTickManager to avoid reusing linkers that might have been torn down in the middle of a multi-frame evaluation. Linkers torn down in this way can be left in a bad state (dead ActiveRunner entries).

diff --git a/Engine/Source/Runtime/MovieScene/Private/MovieSceneSequenceTickManager.cpp b/Engine/Source/Runtime/MovieScene/Private/MovieSceneSequenceTickManager.cpp
--- a/Engine/Source/Runtime/MovieScene/Private/MovieSceneSequenceTickManager.cpp
+++ b/Engine/Source/Runtime/MovieScene/Private/MovieSceneSequenceTickManager.cpp
@@ -152,8 +152,11 @@ void UMovieSceneSequenceTickManager::RegisterTickClient(const FMovieSceneSequenc
 			LinkerName.Appendf(TEXT("_%i_ms"), DesiredTickIntervalMs);
 		}
 
-
-		UMovieSceneEntitySystemLinker* Linker = UMovieSceneEntitySystemLinker::FindOrCreateLinker(GetWorld(), EEntitySystemLinkerRole::LevelSequences, LinkerName.ToString());
+		// With support for multi-frame evaluations, it is possible for the linker group
+		// to be torn down mid evaluation which can leave the linker in a bad state. Use a unique
+		// linker name to avoid reusing those linkers.
+		const FName UniqueLinkerName = MakeUniqueObjectName(GetWorld(), UMovieSceneEntitySystemLinker::StaticClass(), LinkerName.ToString());
+		UMovieSceneEntitySystemLinker* Linker = UMovieSceneEntitySystemLinker::FindOrCreateLinker(GetWorld(), EEntitySystemLinkerRole::LevelSequences, *UniqueLinkerName.ToString());
 		check(Linker);
 
 		FLinkerGroup& NewGroup = LinkerGroups[LinkerIndex];
