Sequencer: Fixed volatility manager getting skipped when dissection was not required

diff --git a/Engine/Source/Runtime/MovieScene/Private/EntitySystem/MovieSceneEntitySystemRunner.cpp b/Engine/Source/Runtime/MovieScene/Private/EntitySystem/MovieSceneEntitySystemRunner.cpp
--- a/Engine/Source/Runtime/MovieScene/Private/EntitySystem/MovieSceneEntitySystemRunner.cpp
+++ b/Engine/Source/Runtime/MovieScene/Private/EntitySystem/MovieSceneEntitySystemRunner.cpp
@@ -614,6 +614,8 @@ bool FMovieSceneEntitySystemRunner::GameThread_UpdateSequenceInstances(UMovieSce
 				);
 				continue;
 			}
+
+			Instance.ConditionalRecompile(Linker);
 			Instance.DissectContext(Linker, Request.Context, Dissections);
 
 			if (Dissections.Num() != 0)
diff --git a/Engine/Source/Runtime/MovieScene/Private/EntitySystem/MovieSceneSequenceInstance.cpp b/Engine/Source/Runtime/MovieScene/Private/EntitySystem/MovieSceneSequenceInstance.cpp
--- a/Engine/Source/Runtime/MovieScene/Private/EntitySystem/MovieSceneSequenceInstance.cpp
+++ b/Engine/Source/Runtime/MovieScene/Private/EntitySystem/MovieSceneSequenceInstance.cpp
@@ -194,19 +194,11 @@ void FSequenceInstance::InvalidateCachedData(UMovieSceneEntitySystemLinker* Link
 	}
 }
 
-void FSequenceInstance::DissectContext(UMovieSceneEntitySystemLinker* Linker, const FMovieSceneContext& InContext, TArray<TRange<FFrameTime>>& OutDissections)
+void FSequenceInstance::ConditionalRecompile(UMovieSceneEntitySystemLinker* Linker)
 {
-	if (!EnumHasAnyFlags(UpdateFlags, ESequenceInstanceUpdateFlags::NeedsDissection))
-	{
-		return;
-	}
-
-	check(SequenceID == MovieSceneSequenceID::Root);
-
-	IMovieScenePlayer* Player = GetPlayer();
-
 	if (VolatilityManager)
 	{
+		IMovieScenePlayer*                         Player              = GetPlayer();
 		FMovieSceneRootEvaluationTemplateInstance& RootTemplate        = Player->GetEvaluationTemplate();
 		UMovieSceneCompiledDataManager*            CompiledDataManager = RootTemplate.GetCompiledDataManager();
 		FMovieSceneCompiledDataID                  RootCompiledDataID  = RootTemplate.GetCompiledDataID();
@@ -216,8 +208,15 @@ void FSequenceInstance::DissectContext(UMovieSceneEntitySystemLinker* Linker, co
 			InvalidateCachedData(Linker);
 		}
 	}
+}
 
-	SequenceUpdater->DissectContext(Linker, Player, InContext, OutDissections);
+void FSequenceInstance::DissectContext(UMovieSceneEntitySystemLinker* Linker, const FMovieSceneContext& InContext, TArray<TRange<FFrameTime>>& OutDissections)
+{
+	if (EnumHasAnyFlags(UpdateFlags, ESequenceInstanceUpdateFlags::NeedsDissection))
+	{
+		check(SequenceID == MovieSceneSequenceID::Root);
+		SequenceUpdater->DissectContext(Linker, GetPlayer(), InContext, OutDissections);
+	}
 }
 
 void FSequenceInstance::Start(UMovieSceneEntitySystemLinker* Linker, const FMovieSceneContext& InContext)
diff --git a/Engine/Source/Runtime/MovieScene/Public/EntitySystem/MovieSceneSequenceInstance.h b/Engine/Source/Runtime/MovieScene/Public/EntitySystem/MovieSceneSequenceInstance.h
--- a/Engine/Source/Runtime/MovieScene/Public/EntitySystem/MovieSceneSequenceInstance.h
+++ b/Engine/Source/Runtime/MovieScene/Public/EntitySystem/MovieSceneSequenceInstance.h
@@ -54,6 +54,13 @@ struct MOVIESCENE_API FSequenceInstance
 
 public:
 
+	/**
+	 * Conditionally recompile this sequence if it needs to be
+	 *
+	 * @param Linker     The linker that owns this sequence instance
+	 */
+	void ConditionalRecompile(UMovieSceneEntitySystemLinker* Linker);
+
 	/**
 	 * Called only for top-level sequence instances before any updates to it with the specified context.
 	 * This allows the sequence an opportunity to dissect the context into a series of distinct evaluations to force determinism.
