Sequencer: Only re-register with the tick manager if our desired tick interval differs with that already registered

This was causing an issue with MRQ that re-initializes a UMovieSceneSequencePlayer. It was ending up calling RegisterTickClient multiple times, which ends up removing the tick client, emptying the linker groups in MovieSceneSequenceTickManager, then re-registering it with the same tick interval and re-creating a _new_ linker without destroying the old instance.

Now we only unregister the tick client if the desired tick interval has changed, and we destroy the instance data if so (so it can be re-created later)

diff --git a/Engine/Source/Runtime/MovieScene/Private/MovieSceneSequencePlayer.cpp b/Engine/Source/Runtime/MovieScene/Private/MovieSceneSequencePlayer.cpp
--- a/Engine/Source/Runtime/MovieScene/Private/MovieSceneSequencePlayer.cpp
+++ b/Engine/Source/Runtime/MovieScene/Private/MovieSceneSequencePlayer.cpp
@@ -875,13 +875,21 @@ void UMovieSceneSequencePlayer::Initialize(UMovieSceneSequence* InSequence)
 		? FMovieSceneSequenceTickInterval::GetInheritedInterval(this)
 		: PlaybackSettings.TickInterval;
 
-	if (RegisteredTickInterval.IsSet() && RegisteredTickInterval.GetValue() != TickInterval)
+	// If we haven't registered with the tick manager yet, register directly
+	if (!RegisteredTickInterval.IsSet())
 	{
+		TickManager->RegisterTickClient(TickInterval, this);
+	}
+	// If we were already registered with a different Tick Interval we need to re-register with the new one, which involves tearing everything down and setting up a new instance
+	else if (RegisteredTickInterval.GetValue() != TickInterval)
+	{
+		RootTemplateInstance.BeginDestroy();
 		TickManager->UnregisterTickClient(this);
+
+		TickManager->RegisterTickClient(RegisteredTickInterval.GetValue(), this);
 	}
 
 	RegisteredTickInterval = TickInterval;
-	TickManager->RegisterTickClient(RegisteredTickInterval.GetValue(), this);
 
 	TSharedPtr<FMovieSceneEntitySystemRunner> RunnerToUse = TickManager->GetRunner(RegisteredTickInterval.GetValue());
 	if (EnumHasAnyFlags(Sequence->GetFlags(), EMovieSceneSequenceFlags::BlockingEvaluation))
