Sequencer: Fix crash on end PIE

nullptr checks on BoundObject in pre animated state

diff --git a/Engine/Source/Runtime/MovieScene/Public/Evaluation/PreAnimatedState/MovieScenePreAnimatedObjectStorage.h b/Engine/Source/Runtime/MovieScene/Public/Evaluation/PreAnimatedState/MovieScenePreAnimatedObjectStorage.h
--- a/Engine/Source/Runtime/MovieScene/Public/Evaluation/PreAnimatedState/MovieScenePreAnimatedObjectStorage.h
+++ b/Engine/Source/Runtime/MovieScene/Public/Evaluation/PreAnimatedState/MovieScenePreAnimatedObjectStorage.h
@@ -138,7 +138,10 @@ public:
 	{
 		for (UObject* BoundObject : BoundObjects)
 		{
-			CachePreAnimatedValue(Params, BoundObject);
+			if (BoundObject)
+			{
+				CachePreAnimatedValue(Params, BoundObject);
+			}
 		}
 	}
 
diff --git a/Engine/Source/Runtime/MovieScene/Public/Evaluation/PreAnimatedState/MovieScenePreAnimatedPropertyStorage.h b/Engine/Source/Runtime/MovieScene/Public/Evaluation/PreAnimatedState/MovieScenePreAnimatedPropertyStorage.h
--- a/Engine/Source/Runtime/MovieScene/Public/Evaluation/PreAnimatedState/MovieScenePreAnimatedPropertyStorage.h
+++ b/Engine/Source/Runtime/MovieScene/Public/Evaluation/PreAnimatedState/MovieScenePreAnimatedPropertyStorage.h
@@ -167,6 +167,11 @@ struct TPreAnimatedPropertyStorageImpl<PropertyTraits, TPropertyMetaData<MetaDat
 		for (int32 Index = 0; Index < Num; ++Index)
 		{
 			UObject* BoundObject  = BoundObjects[Index];
+			if (!BoundObject)
+			{
+				continue;
+			}
+
 			FName    PropertyPath = PropertyBindings[Index].PropertyPath;
 
 			TTuple<FObjectKey, FName> Key{ BoundObject, PropertyPath };
@@ -195,6 +200,11 @@ struct TPreAnimatedPropertyStorageImpl<PropertyTraits, TPropertyMetaData<MetaDat
 		for (int32 Index = 0; Index < Num; ++Index)
 		{
 			UObject* BoundObject  = BoundObjects[Index];
+			if (!BoundObject)
+			{
+				continue;
+			}
+
 			FName    PropertyPath = PropertyBindings[Index].PropertyPath;
 
 			TTuple<FObjectKey, FName> Key{ BoundObject, PropertyPath };
diff --git a/Engine/Source/Runtime/MovieSceneTracks/Private/PreAnimatedState/MovieScenePreAnimatedComponentTransformStorage.cpp b/Engine/Source/Runtime/MovieSceneTracks/Private/PreAnimatedState/MovieScenePreAnimatedComponentTransformStorage.cpp
--- a/Engine/Source/Runtime/MovieSceneTracks/Private/PreAnimatedState/MovieScenePreAnimatedComponentTransformStorage.cpp
+++ b/Engine/Source/Runtime/MovieSceneTracks/Private/PreAnimatedState/MovieScenePreAnimatedComponentTransformStorage.cpp
@@ -76,6 +76,11 @@ void FPreAnimatedComponentTransformStorage::CachePreAnimatedTransforms(const FCa
 
 	for (UObject* BoundObject : BoundObjects)
 	{
+		if (!BoundObject)
+		{
+			continue;
+		}
+
 		TOptional<FResolvedProperty> Property = FPropertyRegistry::ResolveProperty(BoundObject, PropertyBinding, CustomAccessors);
 		if (!Property)
 		{
