Sequencer: Fixed mutual component type not being honored when computing mutually inclusive components

diff --git a/Engine/Source/Runtime/MovieScene/Private/MovieSceneMutualComponentInclusivity.cpp b/Engine/Source/Runtime/MovieScene/Private/MovieSceneMutualComponentInclusivity.cpp
--- a/Engine/Source/Runtime/MovieScene/Private/MovieSceneMutualComponentInclusivity.cpp
+++ b/Engine/Source/Runtime/MovieScene/Private/MovieSceneMutualComponentInclusivity.cpp
@@ -68,7 +68,7 @@ int32 FMutualInclusivityGraphCommandBuffer::ComputeMutuallyInclusiveComponents(E
 
 		case ECommandType::Type:
 			// If the type doesn't match, short-circuit the entire stack of commands
-			if (!EnumHasAnyFlags(Command.Type.Type, EMutuallyInclusiveComponentType::All))
+			if (!EnumHasAnyFlags(Command.Type.Type, TypesToCompute))
 			{
 				// - 1 because Index will be incremented as part of the loop iteration
 				Index = static_cast<int32>(Command.Type.ShortCircuitIndex) - 1;
