Sequencer: Remove WorldCleanup code since this now actively works against Sequencer using weak bound objects now

diff --git a/Engine/Source/Runtime/MovieScene/Private/EntitySystem/MovieSceneEntitySystemLinker.cpp b/Engine/Source/Runtime/MovieScene/Private/EntitySystem/MovieSceneEntitySystemLinker.cpp
--- a/Engine/Source/Runtime/MovieScene/Private/EntitySystem/MovieSceneEntitySystemLinker.cpp
+++ b/Engine/Source/Runtime/MovieScene/Private/EntitySystem/MovieSceneEntitySystemLinker.cpp
@@ -161,8 +161,6 @@ UMovieSceneEntitySystemLinker::UMovieSceneEntitySystemLinker(const FObjectInitia
 		EntityManager.SetDebugName(GetName() + TEXT("[Entity Manager]"));
 		EntityManager.SetComponentRegistry(&GComponentRegistry);
 
-		FWorldDelegates::OnWorldCleanup.AddUObject(this, &UMovieSceneEntitySystemLinker::OnWorldCleanup);
-
 		InstanceRegistry.Reset(new FInstanceRegistry(this));
 
 #if WITH_EDITOR
@@ -185,7 +183,6 @@ void UMovieSceneEntitySystemLinker::Reset()
 	Events.CleanTaggedGarbage.Clear();
 	Events.AddReferencedObjects.Clear();
 	Events.AbandonLinker.Clear();
-	Events.CleanUpWorld.Clear();
 
 	SystemGraph.Shutdown();
 	EntitySystemsByGlobalGraphID.Reset();
@@ -254,7 +251,6 @@ void UMovieSceneEntitySystemLinker::SystemUnlinked(UMovieSceneEntitySystem* InSy
 	Events.CleanTaggedGarbage.RemoveAll(InSystem);
 	Events.AddReferencedObjects.RemoveAll(InSystem);
 	Events.AbandonLinker.RemoveAll(InSystem);
-	Events.CleanUpWorld.RemoveAll(InSystem);
 
 	// Add the system to the recycling pool.
 	ensure(EntitySystemsRecyclingPool.Contains(InSystem->GetClass()) == false);
@@ -470,14 +466,6 @@ void UMovieSceneEntitySystemLinker::OnObjectsReplaced(const TMap<UObject*, UObje
 #endif
 }
 
-void UMovieSceneEntitySystemLinker::OnWorldCleanup(UWorld* InWorld, bool bSessionEnded, bool bCleanupResources)
-{
-	Events.CleanUpWorld.Broadcast(this, InWorld);
-	InstanceRegistry->WorldCleanup(InWorld);
-
-	HandlePostGarbageCollection();
-}
-
 void UMovieSceneEntitySystemLinker::AddReferencedObjects(UObject* Object, FReferenceCollector& Collector)
 {
 	Super::AddReferencedObjects(Object, Collector);
diff --git a/Engine/Source/Runtime/MovieScene/Private/EntitySystem/MovieSceneInstanceRegistry.cpp b/Engine/Source/Runtime/MovieScene/Private/EntitySystem/MovieSceneInstanceRegistry.cpp
--- a/Engine/Source/Runtime/MovieScene/Private/EntitySystem/MovieSceneInstanceRegistry.cpp
+++ b/Engine/Source/Runtime/MovieScene/Private/EntitySystem/MovieSceneInstanceRegistry.cpp
@@ -110,22 +110,6 @@ void FInstanceRegistry::TagGarbage()
 	}
 }
 
-void FInstanceRegistry::WorldCleanup(UWorld* World)
-{
-	auto Iter = [World](FMovieSceneEntityID EntityID, UObject*& BoundObject)
-	{
-		if (BoundObject && BoundObject->IsIn(World))
-		{
-			BoundObject = nullptr;
-		}
-	};
-
-	FEntityTaskBuilder()
-	.ReadEntityIDs()
-	.Write(FBuiltInComponentTypes::Get()->BoundObject)
-	.Iterate_PerEntity(&Linker->EntityManager, Iter);
-}
-
 void FInstanceRegistry::CleanupLinkerEntities(const TSet<FMovieSceneEntityID>& ExpiredBoundObjects)
 {
 	if (ExpiredBoundObjects.Num() != 0)
diff --git a/Engine/Source/Runtime/MovieScene/Public/EntitySystem/MovieSceneInstanceRegistry.h b/Engine/Source/Runtime/MovieScene/Public/EntitySystem/MovieSceneInstanceRegistry.h
--- a/Engine/Source/Runtime/MovieScene/Public/EntitySystem/MovieSceneInstanceRegistry.h
+++ b/Engine/Source/Runtime/MovieScene/Public/EntitySystem/MovieSceneInstanceRegistry.h
@@ -109,8 +109,6 @@ struct FInstanceRegistry
 
 	void TagGarbage();
 
-	void WorldCleanup(UWorld* World);
-
 private:
 
 	UMovieSceneEntitySystemLinker* Linker;
