Sequencer: Null pointer check on movie scene to fix a crash a subsection has no movie scene

This can happen if you create a new take shot and undo it, or if the asset just goes missing (or is removed)

diff --git a/Engine/Source/Editor/MovieSceneTools/Private/Sections/CinematicShotSection.cpp b/Engine/Source/Editor/MovieSceneTools/Private/Sections/CinematicShotSection.cpp
--- a/Engine/Source/Editor/MovieSceneTools/Private/Sections/CinematicShotSection.cpp
+++ b/Engine/Source/Editor/MovieSceneTools/Private/Sections/CinematicShotSection.cpp
@@ -37,7 +37,7 @@ FCinematicShotSection::FCinematicSectionCache::FCinematicSectionCache(UMovieScen
 	if (Section)
 	{
 		UMovieSceneSequence* InnerSequence = Section->GetSequence();
-		if (InnerSequence)
+		if (InnerSequence&& InnerSequence->GetMovieScene())
 		{
 			InnerFrameRate = InnerSequence->GetMovieScene()->GetTickResolution();
 		}
diff --git a/Engine/Source/Editor/MovieSceneTools/Private/TrackEditors/SubTrackEditorBase.cpp b/Engine/Source/Editor/MovieSceneTools/Private/TrackEditors/SubTrackEditorBase.cpp
--- a/Engine/Source/Editor/MovieSceneTools/Private/TrackEditors/SubTrackEditorBase.cpp
+++ b/Engine/Source/Editor/MovieSceneTools/Private/TrackEditors/SubTrackEditorBase.cpp
@@ -26,7 +26,7 @@ FSubSectionPainterResult FSubSectionPainterUtil::PaintSection(TSharedPtr<const I
     }
 
     UMovieSceneSequence* InnerSequence = SectionObject.GetSequence();
-    if (InnerSequence == nullptr)
+    if (InnerSequence == nullptr || InnerSequence->GetMovieScene() == nullptr)
     {
         return FSSPR_NoInnerSequence;
     }
diff --git a/Engine/Source/Runtime/MovieScene/Private/Compilation/MovieSceneCompiledDataManager.cpp b/Engine/Source/Runtime/MovieScene/Private/Compilation/MovieSceneCompiledDataManager.cpp
--- a/Engine/Source/Runtime/MovieScene/Private/Compilation/MovieSceneCompiledDataManager.cpp
+++ b/Engine/Source/Runtime/MovieScene/Private/Compilation/MovieSceneCompiledDataManager.cpp
@@ -1553,7 +1553,7 @@ void UMovieSceneCompiledDataManager::PopulateSubSequenceTree(UMovieSceneSubTrack
 	for (const FMovieSceneTrackEvaluationFieldEntry& Entry : SubTrack->GetEvaluationField().Entries)
 	{
 		UMovieSceneSubSection* SubSection  = Cast<UMovieSceneSubSection>(Entry.Section);
-		if (!SubSection || SubSection->GetSequence() == nullptr)
+		if (!SubSection || SubSection->GetSequence() == nullptr || SubSection->GetSequence()->GetMovieScene() == nullptr)
 		{
 			continue;
 		}
