Sequencer: Remove garbage input entity references from overlapping entity trackers

  - Failing to do so would trigger a non-fatal assert when the output was updated because the input was still being tracked when it no longer existed.

diff --git a/Engine/Source/Runtime/MovieScene/Public/EntitySystem/MovieSceneOverlappingEntityTracker.h b/Engine/Source/Runtime/MovieScene/Public/EntitySystem/MovieSceneOverlappingEntityTracker.h
--- a/Engine/Source/Runtime/MovieScene/Public/EntitySystem/MovieSceneOverlappingEntityTracker.h
+++ b/Engine/Source/Runtime/MovieScene/Public/EntitySystem/MovieSceneOverlappingEntityTracker.h
@@ -528,6 +528,7 @@ struct TOverlappingEntityTracker_WithGarbage : TOverlappingEntityTrackerImpl<Out
 		if (Linker)
 		{
 			Linker->Events.TagGarbage.RemoveAll(this);
+			Linker->Events.CleanTaggedGarbage.RemoveAll(this);
 			Linker->Events.AddReferencedObjects.RemoveAll(this);
 		}
 	}
@@ -542,11 +543,12 @@ struct TOverlappingEntityTracker_WithGarbage : TOverlappingEntityTrackerImpl<Out
 		this->bIsInitialized = true;
 		WeakOwningSystem = OwningSystem;
 
-		OwningSystem->GetLinker()->Events.TagGarbage.AddRaw(this, &ThisType::CleanupGarbage);
+		OwningSystem->GetLinker()->Events.TagGarbage.AddRaw(this, &ThisType::TagGarbage);
+		OwningSystem->GetLinker()->Events.CleanTaggedGarbage.AddRaw(this, &ThisType::CleanTaggedGarbage);
 		OwningSystem->GetLinker()->Events.AddReferencedObjects.AddRaw(this, &ThisType::AddReferencedObjects);
 	}
 
-	void CleanupGarbage(UMovieSceneEntitySystemLinker* Linker)
+	void TagGarbage(UMovieSceneEntitySystemLinker* Linker)
 	{
 		for (int32 Index = this->Outputs.GetMaxIndex()-1; Index >= 0; --Index)
 		{
@@ -580,6 +582,26 @@ struct TOverlappingEntityTracker_WithGarbage : TOverlappingEntityTrackerImpl<Out
 		}
 	}
 
+	void CleanTaggedGarbage(UMovieSceneEntitySystemLinker* Linker)
+	{
+		FComponentTypeID NeedsUnlink = FBuiltInComponentTypes::Get()->Tags.NeedsUnlink;
+
+		// Check whether any of our inputs are about to be destroyed
+		for (auto EntityToOutputIt = this->EntityToOutput.CreateIterator(); EntityToOutputIt; ++EntityToOutputIt)
+		{
+			const FMovieSceneEntityID EntityID    = EntityToOutputIt.Key();
+			const uint16              OutputIndex = EntityToOutputIt.Value();
+			if (OutputIndex != Super::NO_OUTPUT && Linker->EntityManager.HasComponent(EntityID, NeedsUnlink))
+			{
+				this->OutputToEntity.Remove(OutputIndex, EntityID);
+				this->EntityToOutput.Remove(EntityID);
+
+				this->InvalidatedOutputs.PadToNum(OutputIndex + 1, false);
+				this->InvalidatedOutputs[OutputIndex] = true;
+			}
+		}
+	}
+
 	void AddReferencedObjects(UMovieSceneEntitySystemLinker* Linker, FReferenceCollector& ReferenceCollector)
 	{
 		constexpr bool bKeyCanBeGarbage = (THasAddReferencedObjectForComponent<InputTypes>::Value || ...);
