Sequencer: Fix ease-in curves not showing when a section starts at ime 0 and have no ease-out

This was caused by a previous change to fix easing curves not drawing their very last point properly. But ease-out curves ended up being computed as NaN even when they had no duration.
Now the previous change has been reverted, and easing ranges are computed to be inclusive at both start and end times. This lets us draw curves up to their end time.

diff --git a/Engine/Source/Runtime/MovieScene/Private/MovieSceneSection.cpp b/Engine/Source/Runtime/MovieScene/Private/MovieSceneSection.cpp
--- a/Engine/Source/Runtime/MovieScene/Private/MovieSceneSection.cpp
+++ b/Engine/Source/Runtime/MovieScene/Private/MovieSceneSection.cpp
@@ -762,7 +762,7 @@ void UMovieSceneSection::EvaluateEasing(FFrameTime InTime, TOptional<float>& Out
 		}
 	}
 
-	if (HasEndFrame() && Easing.EaseOut.GetObject() && (GetEaseOutRange().Contains(InTime.FrameNumber) || GetEaseOutRange().GetUpperBoundValue() == InTime.FrameNumber))
+	if (HasEndFrame() && Easing.EaseOut.GetObject() && GetEaseOutRange().Contains(InTime.FrameNumber))
 	{
 		const int32  EaseFrame     = (InTime.FrameNumber - GetExclusiveEndFrame() + Easing.GetEaseOutDuration()).Value;
 		const double EaseOutInterp = (double(EaseFrame) + InTime.GetSubFrame()) / Easing.GetEaseOutDuration();
@@ -782,7 +782,7 @@ TRange<FFrameNumber> UMovieSceneSection::GetEaseInRange() const
 	if (HasStartFrame() && Easing.GetEaseInDuration() > 0)
 	{
 		TRangeBound<FFrameNumber> LowerBound = TRangeBound<FFrameNumber>::Inclusive(GetInclusiveStartFrame());
-		TRangeBound<FFrameNumber> UpperBound = TRangeBound<FFrameNumber>::Exclusive(GetInclusiveStartFrame() + Easing.GetEaseInDuration());
+		TRangeBound<FFrameNumber> UpperBound = TRangeBound<FFrameNumber>::Inclusive(GetInclusiveStartFrame() + Easing.GetEaseInDuration());
 
 		UpperBound = TRangeBound<FFrameNumber>::MinUpper(UpperBound, SectionRange.Value.GetUpperBound());
 		return TRange<FFrameNumber>(LowerBound, UpperBound);
@@ -796,7 +796,7 @@ TRange<FFrameNumber> UMovieSceneSection::GetEaseOutRange() const
 {
 	if (HasEndFrame() && Easing.GetEaseOutDuration() > 0)
 	{
-		TRangeBound<FFrameNumber> UpperBound = TRangeBound<FFrameNumber>::Exclusive(GetExclusiveEndFrame());
+		TRangeBound<FFrameNumber> UpperBound = TRangeBound<FFrameNumber>::Inclusive(GetExclusiveEndFrame());
 		TRangeBound<FFrameNumber> LowerBound = TRangeBound<FFrameNumber>::Inclusive(GetExclusiveEndFrame() - Easing.GetEaseOutDuration());
 
 		LowerBound = TRangeBound<FFrameNumber>::MaxLower(LowerBound, SectionRange.Value.GetLowerBound());
