Sequencer: Override post paste to clear transient flags on paste

diff --git a/Engine/Source/Runtime/MovieScene/Private/MovieSceneSection.cpp b/Engine/Source/Runtime/MovieScene/Private/MovieSceneSection.cpp
--- a/Engine/Source/Runtime/MovieScene/Private/MovieSceneSection.cpp
+++ b/Engine/Source/Runtime/MovieScene/Private/MovieSceneSection.cpp
@@ -829,6 +829,19 @@ void UMovieSceneSection::PostEditChangeProperty(FPropertyChangedEvent& PropertyC
 	}
 }
 
+
+void UMovieSceneSection::PostPaste()
+{
+	if (UObject* DefaultEaseIn = Easing.EaseIn.GetObject())
+	{
+		DefaultEaseIn->ClearFlags(RF_Transient);
+	}
+	if (UObject* DefaultEaseOut = Easing.EaseOut.GetObject())
+	{
+		DefaultEaseOut->ClearFlags(RF_Transient);
+	}
+}
+
 ECookOptimizationFlags UMovieSceneSection::GetCookOptimizationFlags() const
 {
 	UMovieSceneTrack* Track = GetTypedOuter<UMovieSceneTrack>();
diff --git a/Engine/Source/Runtime/MovieScene/Public/MovieSceneSection.h b/Engine/Source/Runtime/MovieScene/Public/MovieSceneSection.h
--- a/Engine/Source/Runtime/MovieScene/Public/MovieSceneSection.h
+++ b/Engine/Source/Runtime/MovieScene/Public/MovieSceneSection.h
@@ -668,7 +668,7 @@ public:
 #if WITH_EDITOR
 	MOVIESCENE_API virtual void PostEditChangeProperty(FPropertyChangedEvent& PropertyChangedEvent) override;
 
-	MOVIESCENE_API virtual void PostPaste() {}
+	MOVIESCENE_API virtual void PostPaste();
 #endif
 
 #if WITH_EDITOR
