Sequencer: saving pre-animated state should be able to automatically start tracking the source if needed

The code was previously written mostly with object property tracks in mind, for which the tracking and caching of values is split between inputs (contributors to the value) and outputs (the property being read from and written to). For other tracks like the CVar track, however, there is no difference. We don't want people to have to explicitly start tracking sources, so this change adds an extra parameter to the pre-animated state caching APIs to auto-track sources if needed.

diff --git a/Engine/Source/Runtime/MovieScene/Private/EntitySystem/TrackInstance/MovieSceneTrackInstance.cpp b/Engine/Source/Runtime/MovieScene/Private/EntitySystem/TrackInstance/MovieSceneTrackInstance.cpp
--- a/Engine/Source/Runtime/MovieScene/Private/EntitySystem/TrackInstance/MovieSceneTrackInstance.cpp
+++ b/Engine/Source/Runtime/MovieScene/Private/EntitySystem/TrackInstance/MovieSceneTrackInstance.cpp
@@ -92,6 +92,8 @@ void UMovieSceneTrackInstance::UpdateInputs(TArray<FMovieSceneTrackInstanceInput
 					InputMetaData->StopTrackingCaptureSource(Inputs[OldIndex]);
 				}
 
+				FScopedPreAnimatedCaptureSource CaptureSource(PrivateLinker, InNewInputs[NewIndex]);
+
 				InNewInputs[NewIndex].bInputHasBeenProcessed = true;
 				OnInputAdded(InNewInputs[NewIndex]);
 			}
@@ -139,6 +141,8 @@ void UMovieSceneTrackInstance::UpdateInputs(TArray<FMovieSceneTrackInstanceInput
 		else if (ensure(NewIndex < NewNum))
 		{
 			// and in with the new
+			FScopedPreAnimatedCaptureSource CaptureSource(PrivateLinker, InNewInputs[NewIndex]);
+
 			InNewInputs[NewIndex].bInputHasBeenProcessed = true;
 			OnInputAdded(InNewInputs[NewIndex]);
 			++NewIndex;
diff --git a/Engine/Source/Runtime/MovieScene/Public/Evaluation/MovieScenePreAnimatedState.h b/Engine/Source/Runtime/MovieScene/Public/Evaluation/MovieScenePreAnimatedState.h
--- a/Engine/Source/Runtime/MovieScene/Public/Evaluation/MovieScenePreAnimatedState.h
+++ b/Engine/Source/Runtime/MovieScene/Public/Evaluation/MovieScenePreAnimatedState.h
@@ -133,8 +133,6 @@ private:
 
 	void ConditionalInitializeEntityStorage(bool bOverrideWantsRestoreState);
 
-	void AddSourceMetaData(const UE::MovieScene::FPreAnimatedStateEntry& Entry);
-
 private:
 
 	friend struct FScopedPreAnimatedCaptureSource;
diff --git a/Engine/Source/Runtime/MovieScene/Public/Evaluation/PreAnimatedState/MovieScenePreAnimatedObjectStorage.h b/Engine/Source/Runtime/MovieScene/Public/Evaluation/PreAnimatedState/MovieScenePreAnimatedObjectStorage.h
--- a/Engine/Source/Runtime/MovieScene/Public/Evaluation/PreAnimatedState/MovieScenePreAnimatedObjectStorage.h
+++ b/Engine/Source/Runtime/MovieScene/Public/Evaluation/PreAnimatedState/MovieScenePreAnimatedObjectStorage.h
@@ -116,6 +116,15 @@ public:
 		}
 	}
 
+	FPreAnimatedStateEntry MakeEntry(UObject* BoundObject)
+	{
+		FObjectKey Key{ BoundObject };
+	
+		FPreAnimatedStorageIndex       StorageIndex = this->GetOrCreateStorageIndex(Key);
+		FPreAnimatedStorageGroupHandle GroupHandle  = this->Traits.MakeGroup(BoundObject);
+		return FPreAnimatedStateEntry{ GroupHandle, FPreAnimatedStateCachedValueHandle{ this->StorageID, StorageIndex } };
+	}
+
 	void BeginTrackingEntity(FMovieSceneEntityID EntityID, bool bWantsRestoreState, FRootInstanceHandle RootInstanceHandle, UObject* BoundObject) override
 	{
 		if (!this->ParentExtension->IsCapturingGlobalState() && !bWantsRestoreState)
@@ -125,12 +134,7 @@ public:
 
 		FPreAnimatedEntityCaptureSource* EntityMetaData = this->ParentExtension->GetOrCreateEntityMetaData();
 
-		FObjectKey Key{ BoundObject };
-
-		FPreAnimatedStorageGroupHandle GroupHandle  = this->Traits.MakeGroup(BoundObject);
-		FPreAnimatedStorageIndex       StorageIndex = this->GetOrCreateStorageIndex(Key);
-
-		FPreAnimatedStateEntry Entry{ GroupHandle, FPreAnimatedStateCachedValueHandle{ this->StorageID, StorageIndex } };
+		FPreAnimatedStateEntry Entry = MakeEntry(BoundObject);
 		EntityMetaData->BeginTrackingEntity(Entry, EntityID, RootInstanceHandle, bWantsRestoreState);
 	}
 
@@ -145,29 +149,19 @@ public:
 		}
 	}
 
-	void CachePreAnimatedValue(const FCachePreAnimatedValueParams& Params, UObject* BoundObject)
+	void CachePreAnimatedValue(const FCachePreAnimatedValueParams& Params, UObject* BoundObject, EPreAnimatedCaptureSourceTracking TrackingMode = EPreAnimatedCaptureSourceTracking::CacheIfTracked)
 	{
-		FObjectKey Key{ BoundObject };
-	
-		FPreAnimatedStorageIndex       StorageIndex = this->GetOrCreateStorageIndex(Key);
-		FPreAnimatedStorageGroupHandle GroupHandle  = this->Traits.MakeGroup(BoundObject);
-		FPreAnimatedStateEntry         Entry        = FPreAnimatedStateEntry{ GroupHandle, FPreAnimatedStateCachedValueHandle{ this->StorageID, StorageIndex } };
-	
-		if (this->ParentExtension->IsCapturingGlobalState())
-		{
-			this->ParentExtension->EnsureMetaData(Entry);
-		}
-		else if (!this->ParentExtension->MetaDataExists(Entry))
+		FPreAnimatedStateEntry Entry = MakeEntry(BoundObject);
+
+		if (this->TrackCaptureSource(Entry, TrackingMode))
 		{
-			return;
+			CachePreAnimatedValue(Params, Entry, BoundObject);
 		}
-	
-		CachePreAnimatedValue(Params, Entry, BoundObject);
 	}
 	
 	void CachePreAnimatedValue(const FCachePreAnimatedValueParams& Params, const FPreAnimatedStateEntry& Entry, UObject* BoundObject)
 	{
-		FPreAnimatedStorageIndex StorageIndex = Entry.ValueHandle.StorageIndex;
+		const FPreAnimatedStorageIndex StorageIndex = Entry.ValueHandle.StorageIndex;
 	
 		EPreAnimatedStorageRequirement StorageRequirement = this->ParentExtension->GetStorageRequirement(Entry);
 		if (!this->IsStorageRequirementSatisfied(StorageIndex, StorageRequirement))
@@ -183,33 +177,26 @@ public:
 	}
 
 	template<typename OnCacheValue /* StorageType(const KeyType&) */>
-	void CachePreAnimatedValue(const FCachePreAnimatedValueParams& Params, UObject* BoundObject, OnCacheValue&& CacheCallback)
+	void CachePreAnimatedValue(const FCachePreAnimatedValueParams& Params, UObject* BoundObject, OnCacheValue&& CacheCallback, EPreAnimatedCaptureSourceTracking TrackingMode = EPreAnimatedCaptureSourceTracking::CacheIfTracked)
 	{
-		FObjectKey Key{ BoundObject };
-
-		FPreAnimatedStorageIndex       StorageIndex = this->GetOrCreateStorageIndex(Key);
-		FPreAnimatedStorageGroupHandle GroupHandle  = this->Traits.MakeGroup(BoundObject);
-		FPreAnimatedStateEntry         Entry        = FPreAnimatedStateEntry{ GroupHandle, FPreAnimatedStateCachedValueHandle{ this->StorageID, StorageIndex } };
+		FPreAnimatedStateEntry Entry = MakeEntry(BoundObject);
 
-		if (this->ParentExtension->IsCapturingGlobalState())
+		if (this->TrackCaptureSource(Entry, TrackingMode))
 		{
-			this->ParentExtension->EnsureMetaData(Entry);
-		}
-		else if (!this->ParentExtension->MetaDataExists(Entry))
-		{
-			return;
-		}
+			const FPreAnimatedStorageIndex StorageIndex = Entry.ValueHandle.StorageIndex;
 
-		EPreAnimatedStorageRequirement StorageRequirement = this->ParentExtension->GetStorageRequirement(Entry);
-		if (!this->IsStorageRequirementSatisfied(StorageIndex, StorageRequirement))
-		{
-			StorageType NewValue = CacheCallback(Key);
-			this->AssignPreAnimatedValue(StorageIndex, StorageRequirement, MoveTemp(NewValue));
-		}
+			EPreAnimatedStorageRequirement StorageRequirement = this->ParentExtension->GetStorageRequirement(Entry);
+			if (!this->IsStorageRequirementSatisfied(StorageIndex, StorageRequirement))
+			{
+				KeyType Key { BoundObject };
+				StorageType NewValue = CacheCallback(Key);
+				this->AssignPreAnimatedValue(StorageIndex, StorageRequirement, MoveTemp(NewValue));
+			}
 
-		if (Params.bForcePersist)
-		{
-			this->ForciblyPersistStorage(StorageIndex);
+			if (Params.bForcePersist)
+			{
+				this->ForciblyPersistStorage(StorageIndex);
+			}
 		}
 	}
 };
diff --git a/Engine/Source/Runtime/MovieScene/Public/Evaluation/PreAnimatedState/MovieScenePreAnimatedPropertyStorage.h b/Engine/Source/Runtime/MovieScene/Public/Evaluation/PreAnimatedState/MovieScenePreAnimatedPropertyStorage.h
--- a/Engine/Source/Runtime/MovieScene/Public/Evaluation/PreAnimatedState/MovieScenePreAnimatedPropertyStorage.h
+++ b/Engine/Source/Runtime/MovieScene/Public/Evaluation/PreAnimatedState/MovieScenePreAnimatedPropertyStorage.h
@@ -172,14 +172,7 @@ struct TPreAnimatedPropertyStorageImpl<PropertyTraits, TPropertyMetaData<MetaDat
 				continue;
 			}
 
-			FName    PropertyPath = PropertyBindings[Index].PropertyPath;
-
-			TTuple<FObjectKey, FName> Key{ BoundObject, PropertyPath };
-
-			FPreAnimatedStorageGroupHandle GroupHandle  = this->Traits.MakeGroup(BoundObject);
-			FPreAnimatedStorageIndex       StorageIndex = this->GetOrCreateStorageIndex(Key);
-
-			FPreAnimatedStateEntry Entry{ GroupHandle, FPreAnimatedStateCachedValueHandle{ this->StorageID, StorageIndex } };
+			FPreAnimatedStateEntry Entry = this->MakeEntry(BoundObject, PropertyBindings[Index].PropertyPath);
 			EntityMetaData->BeginTrackingEntity(Entry, EntityIDs[Index], InstanceHandles[Index], bWantsRestore);
 		}
 	}
@@ -205,26 +198,15 @@ struct TPreAnimatedPropertyStorageImpl<PropertyTraits, TPropertyMetaData<MetaDat
 				continue;
 			}
 
-			FName    PropertyPath = PropertyBindings[Index].PropertyPath;
-
-			TTuple<FObjectKey, FName> Key{ BoundObject, PropertyPath };
+			FPreAnimatedStateEntry Entry = this->MakeEntry(BoundObject, PropertyBindings[Index].PropertyPath);
 
-			FPreAnimatedStorageGroupHandle GroupHandle  = this->Traits.MakeGroup(BoundObject);
-			FPreAnimatedStorageIndex       StorageIndex = this->GetOrCreateStorageIndex(Key);
-
-			FPreAnimatedStateEntry Entry{ GroupHandle, FPreAnimatedStateCachedValueHandle{ this->StorageID, StorageIndex } };
-
-			if (this->ParentExtension->IsCapturingGlobalState())
-			{
-				this->ParentExtension->EnsureMetaData(Entry);
-			}
-			else if (!this->ParentExtension->MetaDataExists(Entry))
+			if (!this->TrackCaptureSource(Entry, EPreAnimatedCaptureSourceTracking::CacheIfTracked))
 			{
 				continue;
 			}
 
 			EPreAnimatedStorageRequirement StorageRequirement = this->ParentExtension->GetStorageRequirement(Entry);
-			if (!this->IsStorageRequirementSatisfied(StorageIndex, StorageRequirement))
+			if (!this->IsStorageRequirementSatisfied(Entry.ValueHandle.StorageIndex, StorageRequirement))
 			{
 				StorageType NewValue(MetaData.template Get<MetaDataIndices>()[Index]...);
 
@@ -248,12 +230,12 @@ struct TPreAnimatedPropertyStorageImpl<PropertyTraits, TPropertyMetaData<MetaDat
 					PropertyTraits::GetObjectPropertyValue(BoundObjects[Index], MetaData.template Get<MetaDataIndices>()[Index]..., Bindings.Get(), NewValue.Data);
 				}
 
-				this->AssignPreAnimatedValue(StorageIndex, StorageRequirement, MoveTemp(NewValue));
+				this->AssignPreAnimatedValue(Entry.ValueHandle.StorageIndex, StorageRequirement, MoveTemp(NewValue));
 			}
 
 			if (Params.bForcePersist)
 			{
-				this->ForciblyPersistStorage(StorageIndex);
+				this->ForciblyPersistStorage(Entry.ValueHandle.StorageIndex);
 			}
 		}
 	}
diff --git a/Engine/Source/Runtime/MovieScene/Public/Evaluation/PreAnimatedState/MovieScenePreAnimatedStateStorage.h b/Engine/Source/Runtime/MovieScene/Public/Evaluation/PreAnimatedState/MovieScenePreAnimatedStateStorage.h
--- a/Engine/Source/Runtime/MovieScene/Public/Evaluation/PreAnimatedState/MovieScenePreAnimatedStateStorage.h
+++ b/Engine/Source/Runtime/MovieScene/Public/Evaluation/PreAnimatedState/MovieScenePreAnimatedStateStorage.h
@@ -40,6 +40,15 @@ struct FPreAnimatedObjectGroupManager;
 struct FRestoreStateParams;
 struct FRootInstanceHandle;
 
+/** Enum that defines how to cache pre-animated values based on their capture source */
+enum class EPreAnimatedCaptureSourceTracking
+{
+	/** Cache the pre-animated value only if we are already tracking a capture source for it (or if global capture is enabled) */
+	CacheIfTracked,
+	/** Always cache the pre-animated value, potentially adding tracking meta-data for the scoped capture source */
+	AlwaysCache,
+};
+
 /**
  * Default task that is provided to BeginTrackingEntitiesTask or CachePreAnimatedValuesTask to constrain
  * the task using an additional filter, or to provide custom short-circuit behavior that allows the client
@@ -49,9 +58,11 @@ template<typename... InputTypes>
 struct TPreAnimatedStateTaskParams
 {
 	FEntityComponentFilter AdditionalFilter;
+	EPreAnimatedCaptureSourceTracking TrackingMode;
 
 	TPreAnimatedStateTaskParams()
 	{
+		TrackingMode = EPreAnimatedCaptureSourceTracking::CacheIfTracked;
 		AdditionalFilter.All({ FBuiltInComponentTypes::Get()->Tags.NeedsLink });
 	}
 
@@ -446,6 +457,27 @@ public:
 		return true;
 	}
 
+	/**
+	 * Creates a new pre-animated state entry from a list of arguments suitable for building the storage key.
+	 *
+	 * Note that the arguments are also used for getting the group handle, if the storage traits required grouping.
+	 */
+	template<typename... KeyArgs>
+	FPreAnimatedStateEntry MakeEntry(KeyArgs&&... InKeyArgs)
+	{
+		KeyType Key{ Forward<KeyArgs>(InKeyArgs)... };
+
+		FPreAnimatedStorageGroupHandle GroupHandle;
+		if constexpr(StorageTraits::SupportsGrouping)
+		{
+			GroupHandle = this->Traits.MakeGroup(InKeyArgs...);
+		}
+
+		FPreAnimatedStorageIndex StorageIndex = this->GetOrCreateStorageIndex(Key);
+
+		return FPreAnimatedStateEntry{ GroupHandle, FPreAnimatedStateCachedValueHandle{ this->StorageID, StorageIndex } };
+	}
+
 	/**
 	 * Cause a previously cached value to always outlive any actively animating sources.
 	 * This is called when a Restore State track overlaps a Keep State track. The Restore State track will initially
@@ -535,17 +567,7 @@ public:
 					continue;
 				}
 
-				KeyType Key{ Inputs[Index]... };
-
-				FPreAnimatedStorageGroupHandle GroupHandle;
-				if constexpr(StorageTraits::SupportsGrouping)
-				{
-					GroupHandle = this->Traits.MakeGroup(Inputs[Index]...);
-				}
-
-				FPreAnimatedStorageIndex StorageIndex = this->GetOrCreateStorageIndex(Key);
-
-				FPreAnimatedStateEntry Entry{ GroupHandle, FPreAnimatedStateCachedValueHandle{ this->StorageID, StorageIndex } };
+				FPreAnimatedStateEntry Entry = MakeEntry(Inputs[Index]...);
 				EntityMetaData->BeginTrackingEntity(Entry, EntityIDs[Index], RootInstanceHandles[Index], bWantsRestore);
 			}
 		};
@@ -580,17 +602,7 @@ public:
 			return;
 		}
 
-		KeyType Key{ InComponents... };
-
-		FPreAnimatedStorageGroupHandle GroupHandle;
-		if constexpr(StorageTraits::SupportsGrouping)
-		{
-			GroupHandle = this->Traits.MakeGroup(InComponents...);
-		}
-
-		FPreAnimatedStorageIndex StorageIndex = this->GetOrCreateStorageIndex(Key);
-
-		FPreAnimatedStateEntry Entry{ GroupHandle, FPreAnimatedStateCachedValueHandle{ this->StorageID, StorageIndex } };
+		FPreAnimatedStateEntry Entry = MakeEntry(InComponents...);
 		FPreAnimatedEntityCaptureSource* EntityMetaData = this->ParentExtension->GetOrCreateEntityMetaData();
 		EntityMetaData->BeginTrackingEntity(Entry, EntityID, RootInstanceHandle, bWantsRestoreState);
 	}
@@ -610,8 +622,6 @@ public:
 	template<typename TaskType, typename... ContributorTypes>
 	void CachePreAnimatedValuesTask(UMovieSceneEntitySystemLinker* Linker, const TaskType& InParams, TComponentTypeID<ContributorTypes>... InComponentTypes)
 	{
-		FBuiltInComponentTypes* BuiltInComponents = FBuiltInComponentTypes::Get();
-
 		auto VisitAllocation = [this, &InParams](FEntityAllocationIteratorItem Item, TRead<ContributorTypes>... Values)
 		{
 			const int32 Num = Item.GetAllocation()->Num();
@@ -622,32 +632,15 @@ public:
 					continue;
 				}
 
-				KeyType Key{ Values[Index]... };
-
-				FPreAnimatedStorageGroupHandle GroupHandle;
-				if constexpr(StorageTraits::SupportsGrouping)
-				{
-					GroupHandle = this->Traits.MakeGroup(Values[Index]...);
-				}
-		
-				FPreAnimatedStorageIndex StorageIndex = this->GetOrCreateStorageIndex(Key);
-
-				FPreAnimatedStateEntry Entry{ GroupHandle, FPreAnimatedStateCachedValueHandle{ this->StorageID, StorageIndex } };
-
-				if (this->ParentExtension->IsCapturingGlobalState())
-				{
-					this->ParentExtension->EnsureMetaData(Entry);
-				}
-				else if (!this->ParentExtension->MetaDataExists(Entry))
-				{
-					continue;
-				}
-
-				EPreAnimatedStorageRequirement StorageRequirement = this->ParentExtension->GetStorageRequirement(Entry);
-				if (!this->IsStorageRequirementSatisfied(StorageIndex, StorageRequirement))
+				FPreAnimatedStateEntry Entry = MakeEntry(Values[Index]...);
+				if (TrackCaptureSource(Entry, InParams.TrackingMode))
 				{
-					StorageType NewValue = this->Traits.CachePreAnimatedValue(Values[Index]...);
-					this->AssignPreAnimatedValue(StorageIndex, StorageRequirement, MoveTemp(NewValue));
+					EPreAnimatedStorageRequirement StorageRequirement = this->ParentExtension->GetStorageRequirement(Entry);
+					if (!this->IsStorageRequirementSatisfied(Entry.ValueHandle.StorageIndex, StorageRequirement))
+					{
+						StorageType NewValue = this->Traits.CachePreAnimatedValue(Values[Index]...);
+						this->AssignPreAnimatedValue(Entry.ValueHandle.StorageIndex, StorageRequirement, MoveTemp(NewValue));
+					}
 				}
 			}
 		};
@@ -659,44 +652,34 @@ public:
 	}
 
 	/**
-	 * Save pre-animated state for the specified values
+	 * Save pre-animated state for the specified values, using CacheIfTracked tracking
 	 * Requires that the traits class implements CachePreAnimatedValue(ContributorTypes...)
 	 */
 	template<typename... ContributorTypes>
 	void CachePreAnimatedValue(ContributorTypes... Values)
 	{
-		KeyType Key{ Values... };
-
-		FPreAnimatedStorageGroupHandle GroupHandle;
-		if constexpr(StorageTraits::SupportsGrouping)
-		{
-			GroupHandle = this->Traits.MakeGroup(Values...);
-		}
-
-		FPreAnimatedStorageIndex StorageIndex = this->GetOrCreateStorageIndex(Key);
-
-		FPreAnimatedStateEntry Entry{ GroupHandle, FPreAnimatedStateCachedValueHandle{ this->StorageID, StorageIndex } };
+		CacheTrackedPreAnimatedValue(EPreAnimatedCaptureSourceTracking::CacheIfTracked, Values...);
+	}
 
-		if (this->ParentExtension->IsCapturingGlobalState())
-		{
-			this->ParentExtension->EnsureMetaData(Entry);
-		}
-		else if (!this->ParentExtension->MetaDataExists(Entry))
+	/**
+	 * Save pre-animated state for the specified values
+	 * Requires that the traits class implements CachePreAnimatedValue(ContributorTypes...)
+	 */
+	template<typename... ContributorTypes>
+	void CacheTrackedPreAnimatedValue(EPreAnimatedCaptureSourceTracking TrackingMode, ContributorTypes... Values)
+	{
+		FPreAnimatedStateEntry Entry = MakeEntry(Values...);
+		if (TrackCaptureSource(Entry, TrackingMode))
 		{
-			return;
-		}
-
-		// Begin tracking the entry from its capture source
-		ParentExtension->AddSourceMetaData(Entry);
-
-		EPreAnimatedStorageRequirement Requirement = ParentExtension->HasActiveCaptureSource()
-			? EPreAnimatedStorageRequirement::Transient
-			: EPreAnimatedStorageRequirement::Persistent;
+			EPreAnimatedStorageRequirement Requirement = ParentExtension->HasActiveCaptureSource()
+				? EPreAnimatedStorageRequirement::Transient
+				: EPreAnimatedStorageRequirement::Persistent;
 
-		if (!IsStorageRequirementSatisfied(Entry.ValueHandle.StorageIndex, Requirement))
-		{
-			StorageType NewValue = this->Traits.CachePreAnimatedValue(Values...);
-			this->AssignPreAnimatedValue(StorageIndex, Requirement, MoveTemp(NewValue));
+			if (!IsStorageRequirementSatisfied(Entry.ValueHandle.StorageIndex, Requirement))
+			{
+				StorageType NewValue = this->Traits.CachePreAnimatedValue(Values...);
+				this->AssignPreAnimatedValue(Entry.ValueHandle.StorageIndex, Requirement, MoveTemp(NewValue));
+			}
 		}
 	}
 
@@ -705,10 +688,10 @@ public:
 	 * Callback will only be invoked if state has not already been saved.
 	 */
 	template<typename OnCacheValue /* StorageType(const KeyType&) */>
-	void CachePreAnimatedValue(const KeyType& InKey, OnCacheValue&& CacheCallback)
+	void CachePreAnimatedValue(const KeyType& InKey, OnCacheValue&& CacheCallback, EPreAnimatedCaptureSourceTracking TrackingMode = EPreAnimatedCaptureSourceTracking::CacheIfTracked)
 	{
 		static_assert(StorageTraits::SupportsGrouping == false, "Grouped pre-animated state requires passing a group handle");
-		CachePreAnimatedValue(FPreAnimatedStorageGroupHandle(), InKey, Forward<OnCacheValue>(CacheCallback));
+		CachePreAnimatedValue(FPreAnimatedStorageGroupHandle(), InKey, Forward<OnCacheValue>(CacheCallback), TrackingMode);
 	}
 
 	/**
@@ -716,27 +699,18 @@ public:
 	 * Callback will only be invoked if state has not already been saved.
 	 */
 	template<typename OnCacheValue /* StorageType(const KeyType&) */>
-	void CachePreAnimatedValue(FPreAnimatedStorageGroupHandle GroupHandle, const KeyType& InKey, OnCacheValue&& CacheCallback)
+	void CachePreAnimatedValue(FPreAnimatedStorageGroupHandle GroupHandle, const KeyType& InKey, OnCacheValue&& CacheCallback, EPreAnimatedCaptureSourceTracking TrackingMode = EPreAnimatedCaptureSourceTracking::CacheIfTracked)
 	{
 		ensureMsgf(GroupHandle.IsValid() || !StorageTraits::SupportsGrouping, TEXT("The group handle must be valid for pre-animated state that supports grouping, and invalid if not"));
 
 		// Find the storage index for the specific key we're animating
 		FPreAnimatedStorageIndex StorageIndex = GetOrCreateStorageIndex(InKey);
-
 		FPreAnimatedStateEntry Entry{ GroupHandle, FPreAnimatedStateCachedValueHandle{ this->StorageID, StorageIndex } };
-
-		if (this->ParentExtension->IsCapturingGlobalState())
-		{
-			this->ParentExtension->EnsureMetaData(Entry);
-		}
-		else if (!this->ParentExtension->MetaDataExists(Entry))
+		if (!TrackCaptureSource(Entry, TrackingMode))
 		{
 			return;
 		}
 
-		// Begin tracking the entry from its capture source
-		ParentExtension->AddSourceMetaData(Entry);
-
 		EPreAnimatedStorageRequirement Requirement = ParentExtension->HasActiveCaptureSource()
 			? EPreAnimatedStorageRequirement::Transient
 			: EPreAnimatedStorageRequirement::Persistent;
@@ -783,24 +757,15 @@ public:
 					continue;
 				}
 
-				KeyType Key{ Inputs[Index]... };
-
-				FPreAnimatedStorageGroupHandle GroupHandle;
-				if constexpr(StorageTraits::SupportsGrouping)
-				{
-					GroupHandle = this->Traits.MakeGroup(Inputs[Index]...);
-				}
-
-				FPreAnimatedStorageIndex StorageIndex = this->GetOrCreateStorageIndex(Key);
+				FPreAnimatedStateEntry Entry = MakeEntry(Inputs[Index]...);
 
-				FPreAnimatedStateEntry Entry{ GroupHandle, FPreAnimatedStateCachedValueHandle{ this->StorageID, StorageIndex } };
 				EntityMetaData->BeginTrackingEntity(Entry, EntityIDs[Index], RootInstanceHandles[Index], bWantsRestore);
 
 				EPreAnimatedStorageRequirement StorageRequirement = this->ParentExtension->GetStorageRequirement(Entry);
-				if (!this->IsStorageRequirementSatisfied(StorageIndex, StorageRequirement))
+				if (!this->IsStorageRequirementSatisfied(Entry.ValueHandle.StorageIndex, StorageRequirement))
 				{
 					StorageType NewValue = this->Traits.CachePreAnimatedValue(Inputs[Index]...);
-					this->AssignPreAnimatedValue(StorageIndex, StorageRequirement, MoveTemp(NewValue));
+					this->AssignPreAnimatedValue(Entry.ValueHandle.StorageIndex, StorageRequirement, MoveTemp(NewValue));
 				}
 
 			}
@@ -828,6 +793,44 @@ public:
 
 protected:
 
+	/**
+	 * Begins tracking the current entry with the currently set tracking source, if necessary/desirable.
+	 *
+	 * Returns whether caching pre-animated state value is wanted.
+	 */
+	bool TrackCaptureSource(const FPreAnimatedStateEntry& Entry, EPreAnimatedCaptureSourceTracking TrackingMode)
+	{
+		switch (TrackingMode)
+		{
+		case EPreAnimatedCaptureSourceTracking::CacheIfTracked:
+			{
+				// If we're capturing global state we always track changes
+				if (this->ParentExtension->IsCapturingGlobalState())
+				{
+					this->ParentExtension->EnsureMetaData(Entry);
+					return true;
+				}
+
+				// Only cache the value if tracking meta-data exists for this entry
+				// (ie, something is actively animating this and wants restore state)
+				return this->ParentExtension->MetaDataExists(Entry);
+			}
+			break;
+
+		case EPreAnimatedCaptureSourceTracking::AlwaysCache:
+			{
+				// Always cache the source meta data. If there is no capture source this calls EnsureMetaData as above
+				this->ParentExtension->AddSourceMetaData(Entry);
+				return true;
+			}
+			break;
+
+		default:
+			ensureMsgf(false, TEXT("Unsupported tracking mode, no pre-animated state caching will occur!"));
+			return false;
+		}
+	}
+
 	struct FCachedData
 	{
 		FCachedData()
diff --git a/Engine/Source/Runtime/MovieSceneTracks/Private/TrackInstances/MovieSceneCVarTrackInstance.cpp b/Engine/Source/Runtime/MovieSceneTracks/Private/TrackInstances/MovieSceneCVarTrackInstance.cpp
--- a/Engine/Source/Runtime/MovieSceneTracks/Private/TrackInstances/MovieSceneCVarTrackInstance.cpp
+++ b/Engine/Source/Runtime/MovieSceneTracks/Private/TrackInstances/MovieSceneCVarTrackInstance.cpp
@@ -32,7 +32,7 @@ struct FPreAnimatedCVarTraits : FPreAnimatedStateTraits
 		IConsoleVariable* ConsoleVariable = ConsoleManager.FindConsoleVariable(*InKey);
 		if (ConsoleVariable)
 		{
-			ConsoleVariable->Set(*CachedValue);
+			ConsoleVariable->SetWithCurrentPriority(*CachedValue);
 		}
 	}
 
@@ -90,7 +90,7 @@ void UMovieSceneCVarTrackInstance::OnInputAdded(const FMovieSceneTrackInstanceIn
 
 		for (const TPair<FString, FString>& Pair : CommandsAndValues)
 		{
-			CVarStorage->CachePreAnimatedValue(UngroupedHandle, Pair.Key, FPreAnimatedCVarTraits::CachePreAnimatedValue);
+			CVarStorage->CachePreAnimatedValue(UngroupedHandle, Pair.Key, FPreAnimatedCVarTraits::CachePreAnimatedValue, EPreAnimatedCaptureSourceTracking::AlwaysCache);
 		}
 	}
 
@@ -174,7 +174,7 @@ void UMovieSceneCVarTrackInstance::OnEndUpdateInputs()
 		IConsoleVariable* ConsoleVariable = ConsoleManager.FindConsoleVariable(*Pair.Key);
 		if (ConsoleVariable)
 		{
-			ConsoleVariable->Set(*Pair.Value.Value);
+			ConsoleVariable->SetWithCurrentPriority(*Pair.Value.Value);
 		}
 		else
 		{
