Sequencer: Reset the director instance if reinitializing to a new sequence

This fixes a crash when using the same player to playback a different sequence. The player's director instance is still pointing to the original sequence.

diff --git a/Engine/Source/Runtime/MovieScene/Private/Evaluation/MovieSceneEvaluationTemplateInstance.cpp b/Engine/Source/Runtime/MovieScene/Private/Evaluation/MovieSceneEvaluationTemplateInstance.cpp
--- a/Engine/Source/Runtime/MovieScene/Private/Evaluation/MovieSceneEvaluationTemplateInstance.cpp
+++ b/Engine/Source/Runtime/MovieScene/Private/Evaluation/MovieSceneEvaluationTemplateInstance.cpp
@@ -161,6 +161,8 @@ void FMovieSceneRootEvaluationTemplateInstance::Initialize(UMovieSceneSequence&
 		}
 
 		Player.PreAnimatedState.Initialize(EntitySystemLinker, RootInstanceHandle);
+
+		ResetDirectorInstances();
 	}
 }
 
