Sequencer: Fix filter not updating after adding/removing node groups

diff --git a/Engine/Source/Runtime/MovieScene/Private/MovieScene.cpp b/Engine/Source/Runtime/MovieScene/Private/MovieScene.cpp
--- a/Engine/Source/Runtime/MovieScene/Private/MovieScene.cpp
+++ b/Engine/Source/Runtime/MovieScene/Private/MovieScene.cpp
@@ -792,7 +792,7 @@ void UMovieSceneNodeGroupCollection::AddNodeGroup(UMovieSceneNodeGroup* NodeGrou
 		NodeGroups.Add(NodeGroup);
 		NodeGroup->OnNodeGroupChanged().AddUObject(this, &UMovieSceneNodeGroupCollection::OnNodeGroupChanged);
 		
-		OnNodeGroupCollectionChangedEvent.Broadcast();
+		OnNodeGroupChanged();
 	}
 }
 
@@ -804,7 +804,7 @@ void UMovieSceneNodeGroupCollection::RemoveNodeGroup(UMovieSceneNodeGroup* NodeG
 
 	if (NodeGroups.RemoveSingle(NodeGroup))
 	{
-		OnNodeGroupCollectionChangedEvent.Broadcast();
+		OnNodeGroupChanged();
 	}
 }
 
