Sequencer: Set Tangent if key type is auto or broken fix

diff --git a/Engine/Source/Runtime/MovieScene/Private/Channels/MovieSceneCurveChannelImpl.cpp b/Engine/Source/Runtime/MovieScene/Private/Channels/MovieSceneCurveChannelImpl.cpp
--- a/Engine/Source/Runtime/MovieScene/Private/Channels/MovieSceneCurveChannelImpl.cpp
+++ b/Engine/Source/Runtime/MovieScene/Private/Channels/MovieSceneCurveChannelImpl.cpp
@@ -1200,7 +1200,7 @@ FKeyHandle TMovieSceneCurveChannelImpl<ChannelType>::AddKeyToChannel(ChannelType
 	{
 		FMovieSceneTangentData TangentData;
 		if ((Interpolation == EMovieSceneKeyInterpolation::User || Interpolation == EMovieSceneKeyInterpolation::Break)
-			&& ChannelData.GetTimes().Num() > 2)
+			&& ChannelData.GetTimes().Num() >= 2)
 		{
 			const double DeltaTime = 0.1;
 
