Sequencer: fix autosizing of subsections, audio sections, and skeletal animation sections

These were all missing taking into account start/end offsets.

diff --git a/Engine/Source/Runtime/MovieScene/Private/Sections/MovieSceneSubSection.cpp b/Engine/Source/Runtime/MovieScene/Private/Sections/MovieSceneSubSection.cpp
--- a/Engine/Source/Runtime/MovieScene/Private/Sections/MovieSceneSubSection.cpp
+++ b/Engine/Source/Runtime/MovieScene/Private/Sections/MovieSceneSubSection.cpp
@@ -347,15 +347,16 @@ UMovieSceneSection* UMovieSceneSubSection::SplitSection( FQualifiedFrameTime Spl
 
 TOptional<TRange<FFrameNumber> > UMovieSceneSubSection::GetAutoSizeRange() const
 {
-	if (SubSequence && SubSequence->GetMovieScene())
+	UMovieScene* MovieScene = SubSequence ? SubSequence->GetMovieScene() : nullptr;
+	if (MovieScene)
 	{
 		// We probably want to just auto-size the section to the sub-sequence's scaled playback range... if this section
-		// is looping, however, it's hard to know what we want to do.
-		FMovieSceneTimeTransform InnerToOuter = OuterToInnerTransform().InverseLinearOnly();
-		UMovieScene* InnerMovieScene = SubSequence->GetMovieScene();
+		// is looping, however, it's hard to know what we want to do. Let's just size it to one loop.
+		const FMovieSceneTimeTransform InnerToOuter = OuterToInnerTransform().InverseLinearOnly();
+		const TRange<FFrameNumber> InnerPlaybackRange = UMovieSceneSubSection::GetValidatedInnerPlaybackRange(Parameters, *MovieScene);
 
-		FFrameTime IncAutoStartTime = FFrameTime(UE::MovieScene::DiscreteInclusiveLower(InnerMovieScene->GetPlaybackRange())) * InnerToOuter;
-		FFrameTime ExcAutoEndTime   = FFrameTime(UE::MovieScene::DiscreteExclusiveUpper(InnerMovieScene->GetPlaybackRange())) * InnerToOuter;
+		const FFrameTime IncAutoStartTime = FFrameTime(UE::MovieScene::DiscreteInclusiveLower(InnerPlaybackRange)) * InnerToOuter;
+		const FFrameTime ExcAutoEndTime = FFrameTime(UE::MovieScene::DiscreteExclusiveUpper(InnerPlaybackRange)) * InnerToOuter;
 
 		return TRange<FFrameNumber>(GetInclusiveStartFrame(), GetInclusiveStartFrame() + (ExcAutoEndTime.RoundToFrame() - IncAutoStartTime.RoundToFrame()));
 	}
diff --git a/Engine/Source/Runtime/MovieSceneTracks/Private/Sections/MovieSceneAudioSection.cpp b/Engine/Source/Runtime/MovieSceneTracks/Private/Sections/MovieSceneAudioSection.cpp
--- a/Engine/Source/Runtime/MovieSceneTracks/Private/Sections/MovieSceneAudioSection.cpp
+++ b/Engine/Source/Runtime/MovieSceneTracks/Private/Sections/MovieSceneAudioSection.cpp
@@ -256,9 +256,9 @@ TOptional<TRange<FFrameNumber> > UMovieSceneAudioSection::GetAutoSizeRange() con
 		return TRange<FFrameNumber>();
 	}
 
-	float SoundDuration = MovieSceneHelpers::GetSoundDuration(Sound);
+	const float SoundDuration = MovieSceneHelpers::GetSoundDuration(Sound);
 
-	FFrameRate FrameRate = GetTypedOuter<UMovieScene>()->GetTickResolution();
+	const FFrameRate FrameRate = GetTypedOuter<UMovieScene>()->GetTickResolution();
 
 	// determine initial duration
 	// @todo Once we have infinite sections, we can remove this
@@ -267,10 +267,11 @@ TOptional<TRange<FFrameNumber> > UMovieSceneAudioSection::GetAutoSizeRange() con
 
 	if (SoundDuration != INDEFINITELY_LOOPING_DURATION)
 	{
-		DurationToUse = SoundDuration * FrameRate;
+		DurationToUse = FMath::Max(SoundDuration * FrameRate - StartFrameOffset, FFrameTime(1));
 	}
 
-	return TRange<FFrameNumber>(GetInclusiveStartFrame(), GetInclusiveStartFrame() + DurationToUse.FrameNumber);
+	const int32 IFrameNumber = DurationToUse.FrameNumber.Value + (int)(DurationToUse.GetSubFrame() + 0.5f);
+	return TRange<FFrameNumber>(GetInclusiveStartFrame(), GetInclusiveStartFrame() + IFrameNumber);
 }
 
 	
diff --git a/Engine/Source/Runtime/MovieSceneTracks/Private/Sections/MovieSceneSkeletalAnimationSection.cpp b/Engine/Source/Runtime/MovieSceneTracks/Private/Sections/MovieSceneSkeletalAnimationSection.cpp
--- a/Engine/Source/Runtime/MovieSceneTracks/Private/Sections/MovieSceneSkeletalAnimationSection.cpp
+++ b/Engine/Source/Runtime/MovieSceneTracks/Private/Sections/MovieSceneSkeletalAnimationSection.cpp
@@ -274,10 +274,12 @@ void UMovieSceneSkeletalAnimationSection::PostLoad()
 
 TOptional<TRange<FFrameNumber> > UMovieSceneSkeletalAnimationSection::GetAutoSizeRange() const
 {
-	FFrameRate FrameRate = GetTypedOuter<UMovieScene>()->GetTickResolution();
+	const FFrameRate FrameRate = GetTypedOuter<UMovieScene>()->GetTickResolution();
+	const float AnimPlayRate = FMath::IsNearlyZero(Params.PlayRate) || Params.Animation == nullptr ? 1.0f : Params.PlayRate * Params.Animation->RateScale;
 
-	FFrameTime AnimationLength = Params.GetSequenceLength() * FrameRate;
-	int32 IFrameNumber = AnimationLength.FrameNumber.Value + (int)(AnimationLength.GetSubFrame() + 0.5f);
+	const FFrameTime UnscaledAnimationLength = FMath::Max(Params.GetSequenceLength() * FrameRate - Params.FirstLoopStartFrameOffset - Params.StartFrameOffset - Params.EndFrameOffset, FFrameTime(1));
+	const FFrameTime AnimationLength = UnscaledAnimationLength / AnimPlayRate;
+	const int32 IFrameNumber = AnimationLength.FrameNumber.Value + (int)(AnimationLength.GetSubFrame() + 0.5f);
 
 	return TRange<FFrameNumber>(GetInclusiveStartFrame(), GetInclusiveStartFrame() + IFrameNumber + 1);
 }
