Sequencer: Fixed a crash caused by unloading a spawnable level instance

Destroying a spawnable level instance will cause a GC to run, which in turn will cause the sequencer ECS to tag garbage. Unfortunately this is invalid because the ECS is currently being iterated so is locked for mutation. The fix is to cache all the destroyed spawnables into an array on the stack, then destroy them outside of an iteration.

diff --git a/Engine/Source/Runtime/MovieScene/Private/EntitySystem/MovieSceneSpawnablesSystem.cpp b/Engine/Source/Runtime/MovieScene/Private/EntitySystem/MovieSceneSpawnablesSystem.cpp
--- a/Engine/Source/Runtime/MovieScene/Private/EntitySystem/MovieSceneSpawnablesSystem.cpp
+++ b/Engine/Source/Runtime/MovieScene/Private/EntitySystem/MovieSceneSpawnablesSystem.cpp
@@ -182,7 +182,10 @@ void UMovieSceneSpawnablesSystem::OnRun(FSystemTaskPrerequisites& InPrerequisite
 
 	// ----------------------------------------------------------------------------------------------------------------------------------------
 	// Step 2 - destroy any spawnable objects that are no longer relevant
-	auto DestroyOldSpawnables = [InstanceRegistry](FInstanceHandle InstanceHandle, const FGuid& SpawnableObjectID)
+	//          NOTE: We gather the objects into an array because destroying an object can potentially cause a garbage collection to run (ie
+	//                if the spawnable is a level instance), and that could assert because we are currently iterating the ECS
+	TArray<TTuple<FGuid, FMovieSceneSequenceID, FInstanceHandle>> DestroyedObjects;
+	auto DestroyOldSpawnables = [&DestroyedObjects, InstanceRegistry](FInstanceHandle InstanceHandle, const FGuid& SpawnableObjectID)
 	{
 		SCOPE_CYCLE_COUNTER(MovieSceneEval_DestroySpawnables)
 
@@ -203,7 +206,7 @@ void UMovieSceneSpawnablesSystem::OnRun(FSystemTaskPrerequisites& InPrerequisite
 				}
 			}
 
-			Player->GetSpawnRegister().DestroySpawnedObject(SpawnableObjectID, Instance.GetSequenceID(), *Player);
+			DestroyedObjects.Emplace(SpawnableObjectID, Instance.GetSequenceID(), InstanceHandle);
 		}
 	};
 
@@ -212,5 +215,15 @@ void UMovieSceneSpawnablesSystem::OnRun(FSystemTaskPrerequisites& InPrerequisite
 	.Read(BuiltInComponents->SpawnableBinding)
 	.FilterAll({ BuiltInComponents->Tags.NeedsUnlink })
 	.Iterate_PerEntity(&Linker->EntityManager, DestroyOldSpawnables);
+
+	for (TTuple<FGuid, FMovieSceneSequenceID, FInstanceHandle> Tuple : DestroyedObjects)
+	{
+		// Have to check whether the player is still valid because there is a possibility it got cleaned up
+		if (InstanceRegistry->IsHandleValid(Tuple.Get<2>()))
+		{
+			IMovieScenePlayer* Player = InstanceRegistry->GetInstance(Tuple.Get<2>()).GetPlayer();
+			Player->GetSpawnRegister().DestroySpawnedObject(Tuple.Get<0>(), Tuple.Get<1>(), *Player);
+		}
+	}
 }
 
