Sequencer: Fixed track instances not being able to compile world context pins correctly

diff --git a/Engine/Source/Runtime/MovieScene/Public/EntitySystem/TrackInstance/MovieSceneTrackInstance.h b/Engine/Source/Runtime/MovieScene/Public/EntitySystem/TrackInstance/MovieSceneTrackInstance.h
--- a/Engine/Source/Runtime/MovieScene/Public/EntitySystem/TrackInstance/MovieSceneTrackInstance.h
+++ b/Engine/Source/Runtime/MovieScene/Public/EntitySystem/TrackInstance/MovieSceneTrackInstance.h
@@ -163,6 +163,11 @@ public:
 		return Inputs;
 	}
 
+	virtual UWorld* GetWorld() const override;
+#if WITH_EDITOR
+	virtual bool ImplementsGetWorld() const override { return true; }
+#endif
+
 private:
 
 	/** Called when this animator is first created to perform any initialization */
@@ -186,8 +191,6 @@ private:
 	/** Called when this animator is being destroyed in order that it can perform any final fixup */
 	virtual void OnDestroyed() {}
 
-	virtual UWorld* GetWorld() const override;
-
 private:
 
 	UPROPERTY()
