MovieScene: Move NetSyncProps processing to MovieSceneSequencePlayer.Update to ensure that network synchronization only occurs once per tick.

diff --git a/Engine/Source/Runtime/MovieScene/Private/MovieSceneSequencePlayer.cpp b/Engine/Source/Runtime/MovieScene/Private/MovieSceneSequencePlayer.cpp
--- a/Engine/Source/Runtime/MovieScene/Private/MovieSceneSequencePlayer.cpp
+++ b/Engine/Source/Runtime/MovieScene/Private/MovieSceneSequencePlayer.cpp
@@ -118,6 +118,7 @@ UMovieSceneSequencePlayer::UMovieSceneSequencePlayer(const FObjectInitializer& I
 	, bPendingOnStartedPlaying(false)
 	, bIsAsyncUpdate(false)
 	, bSkipNextUpdate(false)
+	, bUpdateNetSync(false)
 	, Sequence(nullptr)
 	, StartTime(0)
 	, DurationFrames(0)
@@ -929,6 +930,8 @@ void UMovieSceneSequencePlayer::Update(const float DeltaSeconds)
 		CurrentWorldTime = World->GetTimeSeconds();
 	}
 
+	UpdateNetworkSync();
+
 	if (IsPlaying())
 	{
 		// Delta seconds has already been multiplied by GetEffectiveTimeDilation at this point, so don't pass that through to Tick
@@ -1585,22 +1588,11 @@ void UMovieSceneSequencePlayer::PostNetReceive()
 		return;
 	}
 
-	float PingMs = 0.f;
-
-	UWorld* PlayWorld = GetPlaybackWorld();
-	if (PlayWorld)
-	{
-		UNetDriver* NetDriver = PlayWorld->GetNetDriver();
-		if (NetDriver && NetDriver->ServerConnection && NetDriver->ServerConnection->PlayerController && NetDriver->ServerConnection->PlayerController->PlayerState)
-		{
-			PingMs = NetDriver->ServerConnection->PlayerController->PlayerState->ExactPing * (bReversePlayback ? -1.f : 1.f);
-		}
-	}
-
 	const bool bHasStartedPlaying = NetSyncProps.LastKnownStatus == EMovieScenePlayerStatus::Playing && Status != EMovieScenePlayerStatus::Playing;
 	const bool bHasChangedStatus  = NetSyncProps.LastKnownStatus   != Status;
 	const bool bHasChangedTime    = NetSyncProps.LastKnownPosition != PlayPosition.GetCurrentPosition();
 
+	const float PingMs            = GetPing();
 	const FFrameTime PingLag      = (PingMs/1000.f) * PlayPosition.GetInputRate();
 	//const FFrameTime LagThreshold = 0.2f * PlayPosition.GetInputRate();
 	//const FFrameTime LagDisparity = FMath::Abs(PlayPosition.GetCurrentPosition() - NetSyncProps.LastKnownPosition);
@@ -1650,83 +1642,16 @@ void UMovieSceneSequencePlayer::PostNetReceive()
 			// Make sure the client time matches the server according to the client's current status
 			if (Status == EMovieScenePlayerStatus::Playing)
 			{
-				// When the server has looped back to the start but a client is near the end (and is thus about to loop), we don't want to forcibly synchronize the time unless
-				// the *real* difference in time is above the threshold. We compute the real-time difference by adding SequenceDuration*LoopCountDifference to the server position:
-				//		start	srv_time																																clt_time		end
-				//		0		1		2		3		4		5		6		7		8		9		10		11		12		13		14		15		16		17		18		19		20
-				//		|		|																																		|				|
-				//
-				//		Let NetSyncProps.LastKnownNumLoops = 1, CurrentNumLoops = 0, bReversePlayback = false
-				//			=> LoopOffset = 1
-				//			   OffsetServerTime = srv_time + FrameDuration*LoopOffset = 1 + 20*1 = 21
-				//			   Difference = 21 - 18 = 3 frames
-				const int32        LoopOffset       = (NetSyncProps.LastKnownNumLoops - CurrentNumLoops) * (bReversePlayback ? -1 : 1);
-				const FFrameTime   OffsetServerTime = (NetSyncProps.LastKnownPosition + PingLag) + GetFrameDuration()*LoopOffset;
-
-				if (LoopOffset != 0)
-				{
-					// If we crossed a loop boundary, reset the samples
-					ServerTimeSamples.Reset();
-				}
-
-				const bool bUseSmoothing = GSequencerMaxSmoothedNetSyncSampleAge != 0;
-				if (bUseSmoothing)
-				{
-					ServerTimeSamples.Add(FServerTimeSample{ OffsetServerTime / PlayPosition.GetInputRate(), FPlatformTime::Seconds() });
-				}
-
-				const FFrameTime SmoothedServerTime = bUseSmoothing ? UpdateServerTimeSamples() : OffsetServerTime;
-				const FFrameTime Difference         = FMath::Abs(PlayPosition.GetCurrentPosition() - SmoothedServerTime);
-
-				SET_DWORD_STAT(MovieSceneRepl_NumServerSamples, ServerTimeSamples.Num());
-				SET_FLOAT_STAT(MovieSceneRepl_SmoothedServerTime, SmoothedServerTime.AsDecimal());
-
 				if (bHasChangedStatus)
 				{
 					// If the status has changed forcibly play to the server position before setting the new status
 					SetPlaybackPosition(FMovieSceneSequencePlaybackParams(NetSyncProps.LastKnownPosition + PingLag, EUpdatePositionMethod::Play));
 				}
-				else if (Difference > LagThreshold + PingLag)
+				else
 				{
-#if !NO_LOGGING
-					if (UE_LOG_ACTIVE(LogMovieSceneRepl, Log))
-					{
-						const FFrameTime CurrentTime = PlayPosition.GetCurrentPosition();
-						const FString SequenceName = GetSequenceName(true);
-						UE_LOG(LogMovieSceneRepl, Log, TEXT("Correcting de-synced play position for sequence %s %s @ frame %d, subframe %f. Server is %s @ frame %d, subframe %f, (smoothed: frame %d, sf %f). Client ping is %.2fms."),
-							*SequenceName, *UEnum::GetValueAsString(TEXT("MovieScene.EMovieScenePlayerStatus"), Status.GetValue()), CurrentTime.FrameNumber.Value, CurrentTime.GetSubFrame(),
-							*UEnum::GetValueAsString(TEXT("MovieScene.EMovieScenePlayerStatus"), NetSyncProps.LastKnownStatus.GetValue()),
-							NetSyncProps.LastKnownPosition.FrameNumber.Value, NetSyncProps.LastKnownPosition.GetSubFrame(),
-							SmoothedServerTime.FrameNumber.Value, SmoothedServerTime.GetSubFrame(),
-							PingMs);
-					}
-#endif
-					// We're drastically out of sync with the server so we need to forcibly set the time.
-					const FFrameTime LastPosition = FFrameRate::TransformTime(
-							PlayPosition.GetCurrentPosition(), PlayPosition.GetInputRate(), PlayPosition.GetOutputRate());
-
-					// Play to the time only if it is further on in the sequence (in our play direction).
-					// Otherwise, jump backwards in time (in our play direction).
-					const bool bPlayToFrame = bReversePlayback ? SmoothedServerTime < PlayPosition.GetCurrentPosition() : SmoothedServerTime > PlayPosition.GetCurrentPosition();
-					if (bPlayToFrame)
-					{
-						FMovieSceneSequencePlaybackParams Params(SmoothedServerTime, EUpdatePositionMethod::Play);
-						// Indicate that the sequence may have jumped a considerable distance.
-						// This especially helps the audio track to stay in-sync after a correction
-						Params.bHasJumped = true;
-						SetPlaybackPosition(Params);
-					}
-					else
-					{
-						SetPlaybackPosition(FMovieSceneSequencePlaybackParams(SmoothedServerTime, EUpdatePositionMethod::Jump));
-					}
-
-					// When playing back we skip this sequence's ticked update to avoid queuing 2 updates this frame
-					bSkipNextUpdate = true;
-
-					// Also skip all events up to the last known position, otherwise if we skipped back in time we
-					// will re-trigger events again.
-					DisableEventTriggersUntilTime = LastPosition;
+					// Delay net synchronization until next Update call to ensure that we only issue
+					// one desync correction per tick.
+					bUpdateNetSync = true;
 				}
 			}
 			else if (Status == EMovieScenePlayerStatus::Stopped)
@@ -1755,6 +1680,111 @@ void UMovieSceneSequencePlayer::PostNetReceive()
 	}
 }
 
+void UMovieSceneSequencePlayer::UpdateNetworkSync()
+{
+	if (!bUpdateNetSync)
+	{
+		return;
+	}
+	bUpdateNetSync = false;
+
+	// Only process net playback synchronization if we are still Playing.
+	if (Status == EMovieScenePlayerStatus::Playing)
+	{
+		const float PingMs            = GetPing();
+		const FFrameTime PingLag      = (PingMs/1000.f) * PlayPosition.GetInputRate();
+		const FFrameTime LagThreshold = (GSequencerNetSyncThresholdMS * 0.001f) * PlayPosition.GetInputRate();
+		
+		// When the server has looped back to the start but a client is near the end (and is thus about to loop), we don't want to forcibly synchronize the time unless
+		// the *real* difference in time is above the threshold. We compute the real-time difference by adding SequenceDuration*LoopCountDifference to the server position:
+		//		start	srv_time																																clt_time		end
+		//		0		1		2		3		4		5		6		7		8		9		10		11		12		13		14		15		16		17		18		19		20
+		//		|		|																																		|				|
+		//
+		//		Let NetSyncProps.LastKnownNumLoops = 1, CurrentNumLoops = 0, bReversePlayback = false
+		//			=> LoopOffset = 1
+		//			   OffsetServerTime = srv_time + FrameDuration*LoopOffset = 1 + 20*1 = 21
+		//			   Difference = 21 - 18 = 3 frames
+		const int32        LoopOffset       = (NetSyncProps.LastKnownNumLoops - CurrentNumLoops) * (bReversePlayback ? -1 : 1);
+		const FFrameTime   OffsetServerTime = (NetSyncProps.LastKnownPosition + PingLag) + GetFrameDuration()*LoopOffset;
+
+		if (LoopOffset != 0)
+		{
+			// If we crossed a loop boundary, reset the samples
+			ServerTimeSamples.Reset();
+		}
+
+		const bool bUseSmoothing = GSequencerMaxSmoothedNetSyncSampleAge != 0;
+		if (bUseSmoothing)
+		{
+			ServerTimeSamples.Add(FServerTimeSample{ OffsetServerTime / PlayPosition.GetInputRate(), FPlatformTime::Seconds() });
+		}
+
+		const FFrameTime SmoothedServerTime = bUseSmoothing ? UpdateServerTimeSamples() : OffsetServerTime;
+		const FFrameTime Difference         = FMath::Abs(PlayPosition.GetCurrentPosition() - SmoothedServerTime);
+
+		SET_DWORD_STAT(MovieSceneRepl_NumServerSamples, ServerTimeSamples.Num());
+		SET_FLOAT_STAT(MovieSceneRepl_SmoothedServerTime, SmoothedServerTime.AsDecimal());
+		
+		if (Difference > LagThreshold + PingLag)
+		{
+#if !NO_LOGGING
+			if (UE_LOG_ACTIVE(LogMovieSceneRepl, Log))
+			{
+				const FFrameTime CurrentTime = PlayPosition.GetCurrentPosition();
+				const FString SequenceName = GetSequenceName(true);
+				UE_LOG(LogMovieSceneRepl, Log, TEXT("Correcting de-synced play position for sequence %s %s @ frame %d, subframe %f. Server is %s @ frame %d, subframe %f, (smoothed: frame %d, sf %f). Client ping is %.2fms."),
+					*SequenceName, *UEnum::GetValueAsString(TEXT("MovieScene.EMovieScenePlayerStatus"), Status.GetValue()), CurrentTime.FrameNumber.Value, CurrentTime.GetSubFrame(),
+					*UEnum::GetValueAsString(TEXT("MovieScene.EMovieScenePlayerStatus"), NetSyncProps.LastKnownStatus.GetValue()),
+					NetSyncProps.LastKnownPosition.FrameNumber.Value, NetSyncProps.LastKnownPosition.GetSubFrame(),
+					SmoothedServerTime.FrameNumber.Value, SmoothedServerTime.GetSubFrame(),
+					PingMs);
+			}
+#endif
+			// We're drastically out of sync with the server so we need to forcibly set the time.
+			const FFrameTime LastPosition = FFrameRate::TransformTime(
+					PlayPosition.GetCurrentPosition(), PlayPosition.GetInputRate(), PlayPosition.GetOutputRate());
+
+			// Play to the time only if it is further on in the sequence (in our play direction).
+			// Otherwise, jump backwards in time (in our play direction).
+			const bool bPlayToFrame = bReversePlayback ? SmoothedServerTime < PlayPosition.GetCurrentPosition() : SmoothedServerTime > PlayPosition.GetCurrentPosition();
+			if (bPlayToFrame)
+			{
+				FMovieSceneSequencePlaybackParams Params(SmoothedServerTime, EUpdatePositionMethod::Play);
+				// Indicate that the sequence may have jumped a considerable distance.
+				// This especially helps the audio track to stay in-sync after a correction
+				Params.bHasJumped = true;
+				SetPlaybackPosition(Params);
+			}
+			else
+			{
+				SetPlaybackPosition(FMovieSceneSequencePlaybackParams(SmoothedServerTime, EUpdatePositionMethod::Jump));
+			}
+
+			// When playing back we skip this sequence's ticked update to avoid queuing 2 updates this frame
+			bSkipNextUpdate = true;
+
+			// Also skip all events up to the last known position, otherwise if we skipped back in time we
+			// will re-trigger events again.
+			DisableEventTriggersUntilTime = LastPosition;
+		}
+	}
+}
+
+float UMovieSceneSequencePlayer::GetPing() const
+{
+	float PingMs = 0.0f;
+	if (const UWorld* PlayWorld = GetPlaybackWorld())
+	{
+		const UNetDriver* NetDriver = PlayWorld->GetNetDriver();
+		if (NetDriver && NetDriver->ServerConnection && NetDriver->ServerConnection->PlayerController && NetDriver->ServerConnection->PlayerController->PlayerState)
+		{
+			PingMs = NetDriver->ServerConnection->PlayerController->PlayerState->ExactPing * (bReversePlayback ? -1.f : 1.f);
+		}
+	}
+	return PingMs;
+}
+
 void UMovieSceneSequencePlayer::BeginDestroy()
 {
 	RootTemplateInstance.BeginDestroy();
diff --git a/Engine/Source/Runtime/MovieScene/Public/MovieSceneSequencePlayer.h b/Engine/Source/Runtime/MovieScene/Public/MovieSceneSequencePlayer.h
--- a/Engine/Source/Runtime/MovieScene/Public/MovieSceneSequencePlayer.h
+++ b/Engine/Source/Runtime/MovieScene/Public/MovieSceneSequencePlayer.h
@@ -593,6 +593,16 @@ private:
 	 */
 	FFrameTime UpdateServerTimeSamples();
 
+	/**
+	 * Check and correct network synchronization for the clients of this sequence player.
+	 */
+	void UpdateNetworkSync();
+
+	/**
+	 * Compute the latency for the client connection.
+	 */
+	float GetPing() const;
+
 protected:
 
 	/** Movie player status. */
@@ -612,6 +622,9 @@ protected:
 	/** Flag that allows the player to tick its time controller without actually evaluating the sequence */
 	uint32 bSkipNextUpdate : 1;
 
+	/** Flag that notifies the player to check network synchronization on next update */
+	uint32 bUpdateNetSync : 1;
+
 	/** The sequence to play back */
 	UPROPERTY(transient)
 	TObjectPtr<UMovieSceneSequence> Sequence;
