Sequencer: Added the ability for playback clients to manipulate evaluation ranges

  - IMovieScenePlaybackClient::WarpEvaluationRange allows clients to perform last-minute warping of evaluation ranges without persistently changing any of the stored playback positions.
  - This can be used to provide time-wraping functionality to a player

diff --git a/Engine/Source/Editor/Sequencer/Private/Sequencer.cpp b/Engine/Source/Editor/Sequencer/Private/Sequencer.cpp
--- a/Engine/Source/Editor/Sequencer/Private/Sequencer.cpp
+++ b/Engine/Source/Editor/Sequencer/Private/Sequencer.cpp
@@ -3271,6 +3271,11 @@ void FSequencer::EvaluateInternal(FMovieSceneEvaluationRange InRange, bool bHasJ
 		}
 	}
 
+	if (IMovieScenePlaybackClient* PlaybackClient = GetPlaybackClient())
+	{
+		PlaybackClient->WarpEvaluationRange(InRange);
+	}
+
 	FMovieSceneContext Context = FMovieSceneContext(InRange, PlaybackState).SetIsSilent(SilentModeCount != 0);
 	Context.SetHasJumped(bHasJumped);
 
diff --git a/Engine/Source/Runtime/MovieScene/Private/Evaluation/MovieScenePlayback.cpp b/Engine/Source/Runtime/MovieScene/Private/Evaluation/MovieScenePlayback.cpp
--- a/Engine/Source/Runtime/MovieScene/Private/Evaluation/MovieScenePlayback.cpp
+++ b/Engine/Source/Runtime/MovieScene/Private/Evaluation/MovieScenePlayback.cpp
@@ -55,6 +55,13 @@ FMovieSceneEvaluationRange::FMovieSceneEvaluationRange(FFrameTime InCurrentTime,
 {
 }
 
+void FMovieSceneEvaluationRange::ResetRange(const TRange<FFrameTime>& NewRange)
+{
+	ensureMsgf(TimeOverride == TNumericLimits<int32>::Lowest(),
+		TEXT("Should reset time ranges on a range with a fixed time override. This should never happen because such ranges are internal to the movie scene compiler."));
+	EvaluationRange = NewRange;
+}
+
 TRange<FFrameNumber> FMovieSceneEvaluationRange::GetTraversedFrameNumberRange() const
 {
 	TRange<FFrameNumber> FrameNumberRange;
diff --git a/Engine/Source/Runtime/MovieScene/Private/MovieSceneSequencePlayer.cpp b/Engine/Source/Runtime/MovieScene/Private/MovieSceneSequencePlayer.cpp
--- a/Engine/Source/Runtime/MovieScene/Private/MovieSceneSequencePlayer.cpp
+++ b/Engine/Source/Runtime/MovieScene/Private/MovieSceneSequencePlayer.cpp
@@ -1210,6 +1210,11 @@ void UMovieSceneSequencePlayer::UpdateMovieSceneInstance(FMovieSceneEvaluationRa
 	UE_LOG(LogMovieScene, VeryVerbose, TEXT("Evaluating sequence %s at frame %d, subframe %f (%f fps)."), *SequenceName, CurrentTime.Time.FrameNumber.Value, CurrentTime.Time.GetSubFrame(), CurrentTime.Rate.AsDecimal());
 #endif
 
+	if (PlaybackClient)
+	{
+		PlaybackClient->WarpEvaluationRange(InRange);
+	}
+
 	// Once we have updated we must no longer skip updates
 	bSkipNextUpdate = false;
 
diff --git a/Engine/Source/Runtime/MovieScene/Public/Evaluation/MovieScenePlayback.h b/Engine/Source/Runtime/MovieScene/Public/Evaluation/MovieScenePlayback.h
--- a/Engine/Source/Runtime/MovieScene/Public/Evaluation/MovieScenePlayback.h
+++ b/Engine/Source/Runtime/MovieScene/Public/Evaluation/MovieScenePlayback.h
@@ -159,6 +159,11 @@ struct MOVIESCENE_API FMovieSceneEvaluationRange
 		return CurrentFrameRate;
 	}
 
+	/**
+	 * Reset this range to a new range while retaining the current framerate and direction
+	 */
+	void ResetRange(const TRange<FFrameTime>& NewRange);
+
 protected:
 
 	/** The range to evaluate */
diff --git a/Engine/Source/Runtime/MovieScene/Public/IMovieScenePlaybackClient.h b/Engine/Source/Runtime/MovieScene/Public/IMovieScenePlaybackClient.h
--- a/Engine/Source/Runtime/MovieScene/Public/IMovieScenePlaybackClient.h
+++ b/Engine/Source/Runtime/MovieScene/Public/IMovieScenePlaybackClient.h
@@ -8,6 +8,7 @@
 #include "MovieSceneSequenceID.h"
 #include "IMovieScenePlaybackClient.generated.h"
 
+struct FMovieSceneEvaluationRange;
 
 /** Movie scene binding overrides interface */
 UINTERFACE()
@@ -53,4 +54,10 @@ public:
 	{
 		return false;
 	}
+
+	/*
+	 * Warp the time range right before we evaluate
+	 * This gives clients an opportunity to apply time-warping effects without manipulating the actual user-facing playback position
+	 */
+	virtual void WarpEvaluationRange(FMovieSceneEvaluationRange& InOutRange) const {}
 };
