Sequencer: add ability for dynamic resolvers to set ownership of the returned object

This makes is possible to use spawnables in the editor while possessing objects dynamically at runtime, without having the spawnable system try to destroy the actor once the sequence is finished.

diff --git a/Engine/Source/Editor/MovieSceneTools/Private/MovieSceneDirectorBlueprintEndpointCustomization.cpp b/Engine/Source/Editor/MovieSceneTools/Private/MovieSceneDirectorBlueprintEndpointCustomization.cpp
--- a/Engine/Source/Editor/MovieSceneTools/Private/MovieSceneDirectorBlueprintEndpointCustomization.cpp
+++ b/Engine/Source/Editor/MovieSceneTools/Private/MovieSceneDirectorBlueprintEndpointCustomization.cpp
@@ -47,6 +47,40 @@
 
 #define LOCTEXT_NAMESPACE "MovieSceneDirectorBlueprintEndpointCustomization"
 
+namespace UE::Sequencer
+{
+
+TFunction<bool(const FBlueprintActionFilter& Filter, FBlueprintActionInfo& BlueprintAction)> MakeRejectAnyIncompatibleReturnValuesFilter(FProperty* ReturnProperty)
+{
+	if (!ReturnProperty)
+	{
+		return [](const FBlueprintActionFilter& Filter, FBlueprintActionInfo& BlueprintAction) { return false; };
+	}
+
+	FStructProperty* ReturnStructProperty = CastField<FStructProperty>(ReturnProperty);
+	FObjectPropertyBase* ReturnObjectProperty = CastField<FObjectPropertyBase>(ReturnProperty);
+	auto RejectAnyIncompatibleReturnValues = [=](const FBlueprintActionFilter& Filter, FBlueprintActionInfo& BlueprintAction)
+	{
+		const UFunction* Function = BlueprintAction.GetAssociatedFunction();
+		const FProperty* FunctionReturnProperty = Function->GetReturnProperty();
+		if (ReturnProperty->SameType(FunctionReturnProperty))
+		{
+			if (ReturnStructProperty && ReturnStructProperty->Struct == CastField<FStructProperty>(FunctionReturnProperty)->Struct)
+			{
+				return false;
+			}
+			else if (ReturnObjectProperty && ReturnObjectProperty->PropertyClass == CastField<FObjectPropertyBase>(FunctionReturnProperty)->PropertyClass)
+			{
+				return false;
+			}
+		}
+		return true;
+	};
+	return RejectAnyIncompatibleReturnValues;
+};
+
+} // namespace UE::Sequencer
+
 void FMovieSceneDirectorBlueprintEndpointCustomization::CustomizeHeader(TSharedRef<IPropertyHandle> InPropertyHandle, FDetailWidgetRow& HeaderRow, IPropertyTypeCustomizationUtils& CustomizationUtils)
 {
 }
@@ -794,6 +828,11 @@ void FMovieSceneDirectorBlueprintEndpointCustomization::CollectQuickBindActions(
 {
 	// Build up the context object
 	auto RejectAnyNonFunctions = [](const FBlueprintActionFilter& Filter, FBlueprintActionInfo& BlueprintAction)
+	{
+		const UFunction* Function = BlueprintAction.GetAssociatedFunction();
+		return Function == nullptr;
+	};
+	auto RejectAnyNonPureFunctions = [](const FBlueprintActionFilter& Filter, FBlueprintActionInfo& BlueprintAction)
 	{
 		const UFunction* Function = BlueprintAction.GetAssociatedFunction();
 		return Function == nullptr || Function->HasAnyFunctionFlags(FUNC_BlueprintPure);
@@ -803,14 +842,8 @@ void FMovieSceneDirectorBlueprintEndpointCustomization::CollectQuickBindActions(
 		return BlueprintAction.GetBindings().Num() <= 0;
 	};
 
-	FObjectProperty* ReturnProperty = EndpointDefinition.EndpointSignature ?
-		CastField<FObjectProperty>(EndpointDefinition.EndpointSignature->GetReturnProperty()) : nullptr;
-	auto RejectAnyIncompatibleReturnValues = [ReturnProperty](const FBlueprintActionFilter& Filter, FBlueprintActionInfo& BlueprintAction)
-	{
-		const UFunction* Function = BlueprintAction.GetAssociatedFunction();
-		const FObjectProperty* FunctionReturnProperty = CastField<FObjectProperty>(Function->GetReturnProperty());
-		return !FunctionReturnProperty || FunctionReturnProperty->PropertyClass != ReturnProperty->PropertyClass;
-	};
+	FProperty* ReturnProperty = EndpointDefinition.EndpointSignature ? EndpointDefinition.EndpointSignature->GetReturnProperty() : nullptr;
+	auto RejectAnyIncompatibleReturnValues = UE::Sequencer::MakeRejectAnyIncompatibleReturnValuesFilter(ReturnProperty);
 
 	FBlueprintActionMenuBuilder ContextMenuBuilder;
 
@@ -836,7 +869,7 @@ void FMovieSceneDirectorBlueprintEndpointCustomization::CollectQuickBindActions(
 		// but also prevents displaying functions on BP components. Comment out for now.
 		//CallOnMemberFilter.AddRejectionTest(FBlueprintActionFilter::FRejectionTestDelegate::CreateStatic(RejectAnyUnboundActions));
 		
-		CallOnMemberFilter.AddRejectionTest(FBlueprintActionFilter::FRejectionTestDelegate::CreateStatic(RejectAnyNonFunctions));
+		CallOnMemberFilter.AddRejectionTest(FBlueprintActionFilter::FRejectionTestDelegate::CreateStatic(RejectAnyNonPureFunctions));
 
 		if (ReturnProperty)
 		{
@@ -1041,14 +1074,8 @@ void FMovieSceneDirectorBlueprintEndpointCustomization::CollectAllRebindActions(
 		return OuterClass->ClassGeneratedBy != Blueprint;
 	};
 
-	FObjectProperty* ReturnProperty = EndpointDefinition.EndpointSignature ?
-		CastField<FObjectProperty>(EndpointDefinition.EndpointSignature->GetReturnProperty()) : nullptr;
-	auto RejectAnyIncompatibleReturnValues = [ReturnProperty](const FBlueprintActionFilter& Filter, FBlueprintActionInfo& BlueprintAction)
-	{
-		const UFunction* Function = BlueprintAction.GetAssociatedFunction();
-		const FObjectProperty* FunctionReturnProperty = CastField<FObjectProperty>(Function->GetReturnProperty());
-		return !FunctionReturnProperty || FunctionReturnProperty->PropertyClass != ReturnProperty->PropertyClass;
-	};
+	FProperty* ReturnProperty = EndpointDefinition.EndpointSignature ? EndpointDefinition.EndpointSignature->GetReturnProperty() : nullptr;
+	auto RejectAnyIncompatibleReturnValues = UE::Sequencer::MakeRejectAnyIncompatibleReturnValuesFilter(ReturnProperty);
 
 	FBlueprintActionMenuBuilder ContextMenuBuilder;
 
diff --git a/Engine/Source/Editor/MovieSceneTools/Private/MovieSceneDynamicBindingCustomization.cpp b/Engine/Source/Editor/MovieSceneTools/Private/MovieSceneDynamicBindingCustomization.cpp
--- a/Engine/Source/Editor/MovieSceneTools/Private/MovieSceneDynamicBindingCustomization.cpp
+++ b/Engine/Source/Editor/MovieSceneTools/Private/MovieSceneDynamicBindingCustomization.cpp
@@ -234,11 +234,10 @@ void FMovieSceneDynamicBindingCustomization::CollectResolverLibraryBindActions(U
 	auto RejectAnyIncompatibleReturnValues = [](const FBlueprintActionFilter& Filter, FBlueprintActionInfo& BlueprintAction)
 	{
 		const UFunction* Function = BlueprintAction.GetAssociatedFunction();
-		const FObjectProperty* FunctionReturnProperty = CastField<FObjectProperty>(Function->GetReturnProperty());
+		const FStructProperty* FunctionReturnProperty = CastField<FStructProperty>(Function->GetReturnProperty());
 		return 
-			!FunctionReturnProperty || 
-			!FunctionReturnProperty->PropertyClass || 
-			!FunctionReturnProperty->PropertyClass->IsChildOf<UObject>();
+			FunctionReturnProperty == nullptr || 
+			FunctionReturnProperty->Struct != FMovieSceneDynamicBindingResolveResult::StaticStruct();
 	};
 
 	MenuFilter.AddRejectionTest(FBlueprintActionFilter::FRejectionTestDelegate::CreateLambda(RejectAnyIncompatibleReturnValues));
diff --git a/Engine/Source/Editor/MovieSceneTools/Public/MovieSceneDynamicBindingUtils.h b/Engine/Source/Editor/MovieSceneTools/Public/MovieSceneDynamicBindingUtils.h
--- a/Engine/Source/Editor/MovieSceneTools/Public/MovieSceneDynamicBindingUtils.h
+++ b/Engine/Source/Editor/MovieSceneTools/Public/MovieSceneDynamicBindingUtils.h
@@ -67,6 +67,6 @@ class UMovieSceneDynamicBindingEndpointUtil : public UObject
 	GENERATED_BODY()
 
 	UFUNCTION()
-	UObject* SampleResolveBinding() { return nullptr; }
+	FMovieSceneDynamicBindingResolveResult SampleResolveBinding() { return FMovieSceneDynamicBindingResolveResult(); }
 };
 
diff --git a/Engine/Source/Runtime/MovieScene/Private/Evaluation/MovieSceneEvaluationState.cpp b/Engine/Source/Runtime/MovieScene/Private/Evaluation/MovieSceneEvaluationState.cpp
--- a/Engine/Source/Runtime/MovieScene/Private/Evaluation/MovieSceneEvaluationState.cpp
+++ b/Engine/Source/Runtime/MovieScene/Private/Evaluation/MovieSceneEvaluationState.cpp
@@ -391,8 +391,13 @@ void FMovieSceneObjectCache::UpdateBindings(const FGuid& InGuid, IMovieScenePlay
 				}
 				else
 				{
-					const bool bUseDefault = FMovieSceneDynamicBindingInvoker::ResolveDynamicBinding(Player, Sequence, SequenceID, InGuid, Possessable->DynamicBinding, FoundObjects);
-					if (bUseDefault)
+					FMovieSceneDynamicBindingResolveResult ResolveResult = FMovieSceneDynamicBindingInvoker::ResolveDynamicBinding(Player, Sequence, SequenceID, InGuid, Possessable->DynamicBinding);
+					if (ResolveResult.Object)
+					{
+						ensureMsgf(ResolveResult.bIsPossessedObject, TEXT("Possessables don't support spawnable-type ownership"));
+						FoundObjects.Add(ResolveResult.Object);
+					}
+					else
 					{
 						Player.ResolveBoundObjects(InGuid, SequenceID, *Sequence, ResolutionContext, FoundObjects);
 					}
diff --git a/Engine/Source/Runtime/MovieScene/Private/MovieSceneDynamicBinding.cpp b/Engine/Source/Runtime/MovieScene/Private/MovieSceneDynamicBinding.cpp
--- a/Engine/Source/Runtime/MovieScene/Private/MovieSceneDynamicBinding.cpp
+++ b/Engine/Source/Runtime/MovieScene/Private/MovieSceneDynamicBinding.cpp
@@ -5,6 +5,7 @@
 #include "EngineUtils.h"
 #include "Engine/Engine.h"
 #include "Engine/World.h"
+#include "GameFramework/Pawn.h"
 #include "GameFramework/PlayerController.h"
 #include "MovieScene.h"
 #include "MovieSceneBinding.h"
@@ -14,7 +15,7 @@
 
 #include UE_INLINE_GENERATED_CPP_BY_NAME(MovieSceneDynamicBinding)
 
-UObject* UBuiltInDynamicBindingResolverLibrary::ResolveToPlayerPawn(UObject* WorldContextObject, int32 PlayerControllerIndex)
+FMovieSceneDynamicBindingResolveResult UBuiltInDynamicBindingResolverLibrary::ResolveToPlayerPawn(UObject* WorldContextObject, int32 PlayerControllerIndex)
 {
 	if (UWorld* World = GEngine->GetWorldFromContextObject(WorldContextObject, EGetWorldErrorMode::LogAndReturnNull))
 	{
@@ -29,11 +30,14 @@ UObject* UBuiltInDynamicBindingResolverLibrary::ResolveToPlayerPawn(UObject* Wor
 			if (APlayerController* PlayerController = It->Get())
 			{
 				APawn* Pawn = PlayerController->GetPawn();
-				return (UObject*)Pawn;
+				FMovieSceneDynamicBindingResolveResult Result;
+				Result.Object = Pawn;
+				Result.bIsPossessedObject = true;
+				return Result;
 			}
 		}
 	}
 
-	return nullptr;
+	return FMovieSceneDynamicBindingResolveResult();
 }
 
diff --git a/Engine/Source/Runtime/MovieScene/Private/MovieSceneDynamicBindingInvoker.cpp b/Engine/Source/Runtime/MovieScene/Private/MovieSceneDynamicBindingInvoker.cpp
--- a/Engine/Source/Runtime/MovieScene/Private/MovieSceneDynamicBindingInvoker.cpp
+++ b/Engine/Source/Runtime/MovieScene/Private/MovieSceneDynamicBindingInvoker.cpp
@@ -7,35 +7,21 @@
 #include "Evaluation/MovieSceneEvaluationTemplateInstance.h"
 #include "IMovieScenePlayer.h"
 #include "MovieSceneDynamicBinding.h"
+#include "UObject/UnrealType.h"
 
-bool FMovieSceneDynamicBindingInvoker::ResolveDynamicBinding(IMovieScenePlayer& Player, UMovieSceneSequence* Sequence, const FMovieSceneSequenceID& SequenceID, const FGuid& InGuid, const FMovieSceneDynamicBinding& DynamicBinding, TArray<UObject*, TInlineAllocator<1>>& OutObjects)
-{
-	UObject* Object = ResolveDynamicBinding(Player, Sequence, SequenceID, InGuid, DynamicBinding);
-	if (Object != nullptr)
-	{
-		OutObjects.Add(Object);
-
-		// We have successfully called the dynamic binding function, don't use the default behavior.
-		return false;
-	}
-
-	// No valid object found, fallback to default behavior.
-	return true;
-}
-
-UObject* FMovieSceneDynamicBindingInvoker::ResolveDynamicBinding(IMovieScenePlayer& Player, UMovieSceneSequence* Sequence, const FMovieSceneSequenceID& SequenceID, const FGuid& InGuid, const FMovieSceneDynamicBinding& DynamicBinding)
+FMovieSceneDynamicBindingResolveResult FMovieSceneDynamicBindingInvoker::ResolveDynamicBinding(IMovieScenePlayer& Player, UMovieSceneSequence* Sequence, const FMovieSceneSequenceID& SequenceID, const FGuid& InGuid, const FMovieSceneDynamicBinding& DynamicBinding)
 {
 	if (!ensure(Sequence))
 	{
 		// Sequence is somehow null... fallback to default behavior.
-		return nullptr;
+		return FMovieSceneDynamicBindingResolveResult();
 	}
 
 	UFunction* DynamicBindingFunc = DynamicBinding.Function.Get();
 	if (!DynamicBindingFunc)
 	{
 		// No dynamic binding, fallback to default behavior.
-		return nullptr;
+		return FMovieSceneDynamicBindingResolveResult();
 	}
 
 	UObject* DirectorInstance = Player.GetEvaluationTemplate().GetOrCreateDirectorInstance(SequenceID, Player);
@@ -47,7 +33,7 @@ UObject* FMovieSceneDynamicBindingInvoker::ResolveDynamicBinding(IMovieScenePlay
 				*Sequence->GetName(), *DynamicBindingFunc->GetName());
 #endif
 		// Fallback to default behavior.
-		return nullptr;
+		return FMovieSceneDynamicBindingResolveResult();
 	}
 
 #if WITH_EDITOR
@@ -62,7 +48,7 @@ UObject* FMovieSceneDynamicBindingInvoker::ResolveDynamicBinding(IMovieScenePlay
 				TEXT("%s: Refusing to resolve dynamic binding '%s' in editor world because function '%s' has 'Call in Editor' set to false."),
 				*Sequence->GetName(), *LexToString(InGuid), *DynamicBindingFunc->GetName());
 		// Fallback to default behavior.
-		return nullptr;
+		return FMovieSceneDynamicBindingResolveResult();
 	}
 #endif // WITH_EDITOR
 
@@ -74,12 +60,12 @@ UObject* FMovieSceneDynamicBindingInvoker::ResolveDynamicBinding(IMovieScenePlay
 	ResolveParams.ObjectBindingID = InGuid;
 	ResolveParams.Sequence = Sequence;
 	ResolveParams.RootSequence = Player.GetEvaluationTemplate().GetRootSequence();
-	UObject* Object = InvokeDynamicBinding(DirectorInstance, DynamicBinding, ResolveParams);
+	FMovieSceneDynamicBindingResolveResult Result = InvokeDynamicBinding(DirectorInstance, DynamicBinding, ResolveParams);
 
-	return Object;
+	return Result;
 }
 
-UObject* FMovieSceneDynamicBindingInvoker::InvokeDynamicBinding(UObject* DirectorInstance, const FMovieSceneDynamicBinding& DynamicBinding, const FMovieSceneDynamicBindingResolveParams& ResolveParams)
+FMovieSceneDynamicBindingResolveResult FMovieSceneDynamicBindingInvoker::InvokeDynamicBinding(UObject* DirectorInstance, const FMovieSceneDynamicBinding& DynamicBinding, const FMovieSceneDynamicBindingResolveParams& ResolveParams)
 {
 	// Parse all function parameters.
 	UFunction* DynamicBindingFunc = DynamicBinding.Function.Get();
@@ -90,7 +76,7 @@ UObject* FMovieSceneDynamicBindingInvoker::InvokeDynamicBinding(UObject* Directo
 	// Initialize parameters.
 	FMemory::Memzero(Parameters, DynamicBindingFunc->ParmsSize);
 
-	FObjectPropertyBase* ReturnProp = nullptr;
+	FStructProperty* ReturnProp = nullptr;
 	for (TFieldIterator<FProperty> It(DynamicBindingFunc); It; ++It)
 	{
 		FProperty* LocalProp = *It;
@@ -102,8 +88,9 @@ UObject* FMovieSceneDynamicBindingInvoker::InvokeDynamicBinding(UObject* Directo
 
 		if (LocalProp->HasAnyPropertyFlags(CPF_ReturnParm))
 		{
-			ensure(ReturnProp == nullptr);
-			ReturnProp = CastFieldChecked<FObjectPropertyBase>(LocalProp);
+			ensureMsgf(ReturnProp == nullptr,
+					TEXT("Found more than one return parameter in dynamic binding resolver function!"));
+			ReturnProp = CastFieldChecked<FStructProperty>(LocalProp);
 		}
 	}
 	
@@ -117,10 +104,11 @@ UObject* FMovieSceneDynamicBindingInvoker::InvokeDynamicBinding(UObject* Directo
 	DirectorInstance->ProcessEvent(DynamicBindingFunc, Parameters);
 
 	// Grab the result value.
-	UObject* ResultObject = nullptr;
-	if (ensure(ReturnProp))
+	FMovieSceneDynamicBindingResolveResult Result;
+	if (ensureMsgf(ReturnProp != nullptr && ReturnProp->Struct == FMovieSceneDynamicBindingResolveResult::StaticStruct(),
+			TEXT("The dynamic binding resolver function has no return value of type FMovieSceneDynamicBindingResolveResult")))
 	{
-		ResultObject = ReturnProp->GetObjectPropertyValue_InContainer(Parameters);
+		ReturnProp->GetValue_InContainer(Parameters, static_cast<void*>(&Result));
 	}
 
 	// Destroy parameters.
@@ -129,6 +117,6 @@ UObject* FMovieSceneDynamicBindingInvoker::InvokeDynamicBinding(UObject* Directo
 		It->DestroyValue_InContainer(Parameters);
 	}
 
-	return ResultObject;
+	return Result;
 }
 
diff --git a/Engine/Source/Runtime/MovieScene/Private/MovieSceneDynamicBindingInvoker.h b/Engine/Source/Runtime/MovieScene/Private/MovieSceneDynamicBindingInvoker.h
--- a/Engine/Source/Runtime/MovieScene/Private/MovieSceneDynamicBindingInvoker.h
+++ b/Engine/Source/Runtime/MovieScene/Private/MovieSceneDynamicBindingInvoker.h
@@ -4,26 +4,23 @@
 
 #include "Containers/Array.h"
 #include "CoreTypes.h"
+#include "MovieSceneDynamicBinding.h"
 
 class UObject;
 class IMovieScenePlayer;
 class UMovieSceneSequence;
 struct FGuid;
 struct FMovieSceneSequenceID;
-struct FMovieSceneDynamicBinding;
-struct FMovieSceneDynamicBindingResolveParams;
 
 /**
  * Utility class for invoking dynamic binding endpoints.
  */
 struct FMovieSceneDynamicBindingInvoker
 {
-	/** Invoke the dynamic binding, if any, and add the result to the given array of objects */
-	static bool ResolveDynamicBinding(IMovieScenePlayer& Player, UMovieSceneSequence* Sequence, const FMovieSceneSequenceID& SequenceID, const FGuid& InGuid, const FMovieSceneDynamicBinding& DynamicBinding, TArray<UObject*, TInlineAllocator<1>>& OutObjects);
 	/** Invoke the dynamic binding, if any, and return the result */
-	static UObject* ResolveDynamicBinding(IMovieScenePlayer& Player, UMovieSceneSequence* Sequence, const FMovieSceneSequenceID& SequenceID, const FGuid& InGuid, const FMovieSceneDynamicBinding& DynamicBinding);
+	static FMovieSceneDynamicBindingResolveResult ResolveDynamicBinding(IMovieScenePlayer& Player, UMovieSceneSequence* Sequence, const FMovieSceneSequenceID& SequenceID, const FGuid& InGuid, const FMovieSceneDynamicBinding& DynamicBinding);
 
 private:
-	static UObject* InvokeDynamicBinding(UObject* DirectorInstance, const FMovieSceneDynamicBinding& DynamicBinding, const FMovieSceneDynamicBindingResolveParams& Params);
+	static FMovieSceneDynamicBindingResolveResult InvokeDynamicBinding(UObject* DirectorInstance, const FMovieSceneDynamicBinding& DynamicBinding, const FMovieSceneDynamicBindingResolveParams& Params);
 };
 
diff --git a/Engine/Source/Runtime/MovieScene/Private/MovieSceneSpawnRegister.cpp b/Engine/Source/Runtime/MovieScene/Private/MovieSceneSpawnRegister.cpp
--- a/Engine/Source/Runtime/MovieScene/Private/MovieSceneSpawnRegister.cpp
+++ b/Engine/Source/Runtime/MovieScene/Private/MovieSceneSpawnRegister.cpp
@@ -43,8 +43,19 @@ UObject* FMovieSceneSpawnRegister::SpawnObject(const FGuid& BindingId, UMovieSce
 		return nullptr;
 	}
 
+	UObject* SpawnedActor = nullptr;
+	ESpawnOwnership SpawnOwnership = Spawnable->GetSpawnOwnership();
+
 	// See if there is some dynamic binding logic to invoke, otherwise spawn the actor
-	UObject* SpawnedActor = FMovieSceneDynamicBindingInvoker::ResolveDynamicBinding(Player, Sequence, TemplateID, BindingId, Spawnable->DynamicBinding);
+	FMovieSceneDynamicBindingResolveResult ResolveResult = FMovieSceneDynamicBindingInvoker::ResolveDynamicBinding(Player, Sequence, TemplateID, BindingId, Spawnable->DynamicBinding);
+	if (ResolveResult.Object)
+	{
+		SpawnedActor = ResolveResult.Object;
+		if (ResolveResult.bIsPossessedObject)
+		{
+			SpawnOwnership = ESpawnOwnership::External;
+		}
+	}
 	if (!SpawnedActor)
 	{
 		SpawnedActor = SpawnObject(*Spawnable, TemplateID, Player);
@@ -55,7 +66,7 @@ UObject* FMovieSceneSpawnRegister::SpawnObject(const FGuid& BindingId, UMovieSce
 		FMovieSceneSpawnableAnnotation::Add(SpawnedActor, BindingId, TemplateID, Sequence);
 
 		FMovieSceneSpawnRegisterKey Key(TemplateID, BindingId);
-		Register.Add(Key, FSpawnedObject(BindingId, *SpawnedActor, Spawnable->GetSpawnOwnership()));
+		Register.Add(Key, FSpawnedObject(BindingId, *SpawnedActor, SpawnOwnership));
 
 		Player.State.Invalidate(BindingId, TemplateID);
 	}
diff --git a/Engine/Source/Runtime/MovieScene/Public/MovieSceneDynamicBinding.h b/Engine/Source/Runtime/MovieScene/Public/MovieSceneDynamicBinding.h
--- a/Engine/Source/Runtime/MovieScene/Public/MovieSceneDynamicBinding.h
+++ b/Engine/Source/Runtime/MovieScene/Public/MovieSceneDynamicBinding.h
@@ -71,20 +71,42 @@ struct MOVIESCENE_API FMovieSceneDynamicBindingResolveParams
 	GENERATED_BODY()
 
 	/** The sequence that contains the object binding being resolved */
-	UPROPERTY(EditAnywhere, Category="General")
+	UPROPERTY(EditAnywhere, BlueprintReadWrite, Category="General")
 	TObjectPtr<UMovieSceneSequence> Sequence;
 
 	/** The ID of the object binding being resolved */
-	UPROPERTY(EditAnywhere, Category="General")
+	UPROPERTY(EditAnywhere, BlueprintReadWrite, Category="General")
 	FGuid ObjectBindingID;
 
 	/** The root sequence */
-	UPROPERTY(EditAnywhere, Category="General")
+	UPROPERTY(EditAnywhere, BlueprintReadWrite, Category="General")
 	TObjectPtr<UMovieSceneSequence> RootSequence;
 };
 
+USTRUCT(BlueprintType)
+struct MOVIESCENE_API FMovieSceneDynamicBindingResolveResult
+{
+	GENERATED_BODY()
+
+	/** The resolved object */
+	UPROPERTY(EditAnywhere, BlueprintReadWrite, Category="General")
+	TObjectPtr<UObject> Object = nullptr;
+
+	/**
+	 * Whether the resolved object is external to the sequence
+	 *
+	 * This property is ignored for possessables, who are always treated as external.
+	 * When resolving a spawnable, setting this to true will not destroy the object
+	 * when the spawnable track ends, or the sequence finishes.
+	 */
+	UPROPERTY(EditAnywhere, BlueprintReadWrite, Category="General")
+	bool bIsPossessedObject = true;
+};
+
 /**
- * Dummy structure for showing an FMovieSceneDynamicBinding inside a details view.
+ * Dummy structure for showing an FMovieSceneDynamicBinding inside a details view,
+ * and having a function signature that we can get a valid UFunction from to prepare
+ * blueprint function graphs.
  */
 USTRUCT()
 struct MOVIESCENE_API FMovieSceneDynamicBindingContainer
@@ -107,6 +129,6 @@ public:
 
 	/** Resolve the bound object to the player's pawn */
 	UFUNCTION(BlueprintPure, Category="Sequencer|Dynamic Binding", meta=(WorldContext="WorldContextObject"))
-	static UObject* ResolveToPlayerPawn(UObject* WorldContextObject, int32 PlayerControllerIndex = 0);
+	static FMovieSceneDynamicBindingResolveResult ResolveToPlayerPawn(UObject* WorldContextObject, int32 PlayerControllerIndex = 0);
 };
 
