Sequencer: Node group collection should be transactional and need to be refreshed when undone.

diff --git a/Engine/Source/Runtime/MovieScene/Private/MovieScene.cpp b/Engine/Source/Runtime/MovieScene/Private/MovieScene.cpp
--- a/Engine/Source/Runtime/MovieScene/Private/MovieScene.cpp
+++ b/Engine/Source/Runtime/MovieScene/Private/MovieScene.cpp
@@ -79,6 +79,7 @@ UMovieScene::UMovieScene(const FObjectInitializer& ObjectInitializer)
 	FixedFrameInterval_DEPRECATED = 0.0f;
 
 	NodeGroupCollection = CreateEditorOnlyDefaultSubobject<UMovieSceneNodeGroupCollection>("NodeGroupCollection");
+	NodeGroupCollection->SetFlags(RF_Transactional);
 
 	InTime_DEPRECATED    =  FLT_MAX;
 	OutTime_DEPRECATED   = -FLT_MAX;
@@ -768,22 +769,34 @@ void UMovieSceneNodeGroup::SetEnableFilter(bool bInEnableFilter)
 	}
 }
 
-void UMovieSceneNodeGroupCollection::PostLoad()
+void UMovieSceneNodeGroupCollection::Refresh()
 {
 	bAnyActiveFilter = false;
 	for (UMovieSceneNodeGroup* NodeGroup : NodeGroups)
 	{
 		NodeGroup->OnNodeGroupChanged().AddUObject(this, &UMovieSceneNodeGroupCollection::OnNodeGroupChanged);
-		
+
 		if (NodeGroup->GetEnableFilter())
 		{
 			bAnyActiveFilter = true;
 		}
 	}
-	
+}
+
+void UMovieSceneNodeGroupCollection::PostLoad()
+{
+	Refresh();
+
 	Super::PostLoad();
 }
 
+void UMovieSceneNodeGroupCollection::PostEditUndo()
+{
+	Refresh();
+
+	Super::PostEditUndo();
+}
+
 void UMovieSceneNodeGroupCollection::AddNodeGroup(UMovieSceneNodeGroup* NodeGroup)
 {
 	Modify();
diff --git a/Engine/Source/Runtime/MovieScene/Public/MovieScene.h b/Engine/Source/Runtime/MovieScene/Public/MovieScene.h
--- a/Engine/Source/Runtime/MovieScene/Public/MovieScene.h
+++ b/Engine/Source/Runtime/MovieScene/Public/MovieScene.h
@@ -291,6 +291,7 @@ class MOVIESCENE_API UMovieSceneNodeGroupCollection : public UObject
 public:
 	/** Called after this object has been deserialized */
 	virtual void PostLoad() override;
+	virtual void PostEditUndo() override;
 
 	void AddNodeGroup(UMovieSceneNodeGroup* NodeGroup);
 	void RemoveNodeGroup(UMovieSceneNodeGroup* NodeGroup);
@@ -306,6 +307,9 @@ public:
 	FOnNodeGroupCollectionChanged& OnNodeGroupCollectionChanged() { return OnNodeGroupCollectionChangedEvent; }
 
 private:
+
+	void Refresh();
+
 	UPROPERTY()
 	TArray<TObjectPtr<UMovieSceneNodeGroup>> NodeGroups;
 
