Sequencer: Fix regression with node group collections

#rnx

diff --git a/Engine/Source/Runtime/MovieScene/Private/MovieScene.cpp b/Engine/Source/Runtime/MovieScene/Private/MovieScene.cpp
--- a/Engine/Source/Runtime/MovieScene/Private/MovieScene.cpp
+++ b/Engine/Source/Runtime/MovieScene/Private/MovieScene.cpp
@@ -79,7 +79,6 @@ UMovieScene::UMovieScene(const FObjectInitializer& ObjectInitializer)
 	FixedFrameInterval_DEPRECATED = 0.0f;
 
 	NodeGroupCollection = CreateEditorOnlyDefaultSubobject<UMovieSceneNodeGroupCollection>("NodeGroupCollection");
-	NodeGroupCollection->SetFlags(RF_Transactional);
 
 	InTime_DEPRECATED    =  FLT_MAX;
 	OutTime_DEPRECATED   = -FLT_MAX;
