Sequencer: Fix crash in SortConsecutiveSections

If sections are null, MovieSceneTrack::PostLoad calls RemoveAt on the null section and that can cause the sections to be sorted.
Changed this to use Algo::SortBy and check for valid UMovieSceneSection pointer.

diff --git a/Engine/Source/Runtime/MovieScene/Private/MovieSceneCommonHelpers.cpp b/Engine/Source/Runtime/MovieScene/Private/MovieSceneCommonHelpers.cpp
--- a/Engine/Source/Runtime/MovieScene/Private/MovieSceneCommonHelpers.cpp
+++ b/Engine/Source/Runtime/MovieScene/Private/MovieSceneCommonHelpers.cpp
@@ -78,14 +78,14 @@ UMovieSceneSection* MovieSceneHelpers::FindNearestSectionAtTime( TArrayView<UMov
 	{
 		Algo::SortBy(NonOverlappingSections, [](const UMovieSceneSection* A) { return A->GetRange().GetUpperBound(); }, SortUpperBounds);
 
-		const int32 PreviousIndex = Algo::UpperBoundBy(NonOverlappingSections, TRangeBound<FFrameNumber>(Time), [](const UMovieSceneSection* A){ return A->GetRange().GetUpperBound(); }, SortUpperBounds)-1;
+		const int32 PreviousIndex = Algo::UpperBoundBy(NonOverlappingSections, TRangeBound<FFrameNumber>(Time), [](const UMovieSceneSection* A){ return A ? A->GetRange().GetUpperBound() : FFrameNumber(0); }, SortUpperBounds)-1;
 		if (NonOverlappingSections.IsValidIndex(PreviousIndex))
 		{
 			return NonOverlappingSections[PreviousIndex];
 		}
 		else
 		{
-			Algo::SortBy(NonOverlappingSections, [](const UMovieSceneSection* A) { return A->GetRange().GetLowerBound(); }, SortLowerBounds);
+			Algo::SortBy(NonOverlappingSections, [](const UMovieSceneSection* A) { return A ? A->GetRange().GetLowerBound() : FFrameNumber(0); }, SortLowerBounds);
 			return NonOverlappingSections[0];
 		}
 	}
@@ -102,12 +102,7 @@ bool MovieSceneHelpers::SortOverlappingSections(const UMovieSceneSection* A, con
 
 void MovieSceneHelpers::SortConsecutiveSections(TArray<UMovieSceneSection*>& Sections)
 {
-	Sections.Sort([](const UMovieSceneSection& A, const UMovieSceneSection& B)
-		{
-			TRangeBound<FFrameNumber> LowerBoundA = A.GetRange().GetLowerBound();
-			return TRangeBound<FFrameNumber>::MinLower(LowerBoundA, B.GetRange().GetLowerBound()) == LowerBoundA;
-		}
-	);
+	Algo::SortBy(Sections, [](const UMovieSceneSection* A) { return A ? A->GetRange().GetLowerBound() : FFrameNumber(0); }, SortLowerBounds);
 }
 
 bool MovieSceneHelpers::FixupConsecutiveSections(TArray<UMovieSceneSection*>& Sections, UMovieSceneSection& Section, bool bDelete, bool bCleanUp)
