Sequencer: fix issues with runners detaching and attaching to linkers in interwoven ways

diff --git a/Engine/Source/Runtime/MovieScene/Private/EntitySystem/MovieSceneEntitySystemRunner.cpp b/Engine/Source/Runtime/MovieScene/Private/EntitySystem/MovieSceneEntitySystemRunner.cpp
--- a/Engine/Source/Runtime/MovieScene/Private/EntitySystem/MovieSceneEntitySystemRunner.cpp
+++ b/Engine/Source/Runtime/MovieScene/Private/EntitySystem/MovieSceneEntitySystemRunner.cpp
@@ -116,11 +116,17 @@ void FMovieSceneEntitySystemRunner::DetachFromLinker()
 {
 	if (ensureMsgf(!WeakLinker.IsExplicitlyNull(), TEXT("This runner is not attached to any linker")))
 	{
-		if (WeakLinker.IsValid())
+		// Abandon our previous linker. We need to do so even for a linker that is pending kill, otherwise
+		// we will later get a OnLinkerAbandon call, which _could_ come _after_ we've been re-attached to a
+		// new valid linker. This would in turn trip the ensure in OnLinkerAbandon that checks that we are
+		// abandoning the linker we have, instead of an unrelated linker.
+		if (UMovieSceneEntitySystemLinker* Linker = WeakLinker.Get(true))
 		{
-			OnLinkerAbandon(WeakLinker.Get());
+			OnLinkerAbandon(Linker);
 		}
 	}
+
+	WeakLinker.Reset();
 }
 
 UMovieSceneEntitySystemLinker* FMovieSceneEntitySystemRunner::GetLinker() const
@@ -1094,10 +1100,14 @@ void FMovieSceneEntitySystemRunner::MarkForUpdate(FInstanceHandle InInstanceHand
 
 void FMovieSceneEntitySystemRunner::OnLinkerAbandon(UMovieSceneEntitySystemLinker* InLinker)
 {
+	// WARNING: this can be called with a linker that is PendingKill
+
 	if (ensure(InLinker))
 	{
 		InLinker->Events.AbandonLinker.RemoveAll(this);
 	}
+
+	ensure(InLinker == WeakLinker.Get(true));
 	WeakLinker.Reset();
 }
 
