Sequencer: remove assert because of false positives

The weak pointer won't resolve to a valid pointer (even when using "Get(true)") if the pointed object is being destroyed, so the assert was firing during legitimate calls.

diff --git a/Engine/Source/Runtime/MovieScene/Private/EntitySystem/MovieSceneEntitySystemRunner.cpp b/Engine/Source/Runtime/MovieScene/Private/EntitySystem/MovieSceneEntitySystemRunner.cpp
--- a/Engine/Source/Runtime/MovieScene/Private/EntitySystem/MovieSceneEntitySystemRunner.cpp
+++ b/Engine/Source/Runtime/MovieScene/Private/EntitySystem/MovieSceneEntitySystemRunner.cpp
@@ -1107,7 +1107,6 @@ void FMovieSceneEntitySystemRunner::OnLinkerAbandon(UMovieSceneEntitySystemLinke
 		InLinker->Events.AbandonLinker.RemoveAll(this);
 	}
 
-	ensure(InLinker == WeakLinker.Get(true));
 	WeakLinker.Reset();
 }
 
