Sequencer: Invalidate cached BP director instance if it has been recompiled and a new version has superceded it

diff --git a/Engine/Source/Runtime/MovieScene/Private/Evaluation/MovieSceneEvaluationTemplateInstance.cpp b/Engine/Source/Runtime/MovieScene/Private/Evaluation/MovieSceneEvaluationTemplateInstance.cpp
--- a/Engine/Source/Runtime/MovieScene/Private/Evaluation/MovieSceneEvaluationTemplateInstance.cpp
+++ b/Engine/Source/Runtime/MovieScene/Private/Evaluation/MovieSceneEvaluationTemplateInstance.cpp
@@ -352,6 +352,17 @@ void FMovieSceneRootEvaluationTemplateInstance::FindEntitiesFromOwner(UObject* O
 UObject* FMovieSceneRootEvaluationTemplateInstance::GetOrCreateDirectorInstance(FMovieSceneSequenceIDRef SequenceID, IMovieScenePlayer& Player)
 {
 	UObject* ExistingDirectorInstance = DirectorInstances.FindRef(SequenceID);
+#if WITH_EDITOR
+	if (ExistingDirectorInstance)
+	{
+		// Invalidate our cached director instance if it has been recompiled.
+		UClass* DirectorClass = ExistingDirectorInstance->GetClass();
+		if (!DirectorClass || DirectorClass->HasAnyClassFlags(CLASS_NewerVersionExists))
+		{
+			ExistingDirectorInstance = nullptr;
+		}
+	}
+#endif
 	if (ExistingDirectorInstance)
 	{
 		return ExistingDirectorInstance;
