PR #10369: Update MovieSceneSequencePlayer.cpp (Contributed by [REDACTED_PERSON])

diff --git a/Engine/Source/Runtime/MovieScene/Private/MovieSceneSequencePlayer.cpp b/Engine/Source/Runtime/MovieScene/Private/MovieSceneSequencePlayer.cpp
--- a/Engine/Source/Runtime/MovieScene/Private/MovieSceneSequencePlayer.cpp
+++ b/Engine/Source/Runtime/MovieScene/Private/MovieSceneSequencePlayer.cpp
@@ -1118,7 +1118,7 @@ void UMovieSceneSequencePlayer::UpdateTimeCursorPosition_Internal(FFrameTime New
 			// 3. We are the client side of a replicated sequence, but playing is only happening on our side (i.e. the Play() method was
 			//    called only on the client, and the server sequence is stopped)
 			const bool bHasAuthorityToFinish = (
-				(PlaybackClient == nullptr || !PlaybackClient->GetIsReplicatedPlayback()) ||
+				(!PlaybackClient || !PlaybackClient->GetIsReplicatedPlayback()) ||
 				HasAuthority() ||
 				NetSyncProps.LastKnownStatus == EMovieScenePlayerStatus::Stopped);
 			if (bHasAuthorityToFinish)
