Sequencer: Smart Auto Key
New algorithm is in MovieScenCurveChannelImpl where we calculate auto tangents. Since they work off tangents and not key positons, need to make sure to recalc them if the tangent changes now.  Since FRichCurve will not support these new auto tangents, they only do 2 key interps with the old auto key algorithm that doesn't average the tangents at all, had to add a SupportedTangentType to FCurveEditor so that we can use that with the curve editor context menu/commands/SCurveEditorPanel since they are all shared between Sequencer and the Curve Asset viewer. Most of the code is making sure we treat smart auto's the same as the old autos

diff --git a/Engine/Source/Editor/CurveEditor/Private/CurveEditor.cpp b/Engine/Source/Editor/CurveEditor/Private/CurveEditor.cpp
--- a/Engine/Source/Editor/CurveEditor/Private/CurveEditor.cpp
+++ b/Engine/Source/Editor/CurveEditor/Private/CurveEditor.cpp
@@ -129,6 +129,17 @@ void FCurveEditor::InitCurveEditor(const FCurveEditorInitParams& InInitParams)
 	GEditor->RegisterForUndo(this);
 }
 
+int32 FCurveEditor::GetSupportedTangentTypes()
+{
+	return ((int32)ECurveEditorTangentTypes::InterpolationConstant |
+		(int32)ECurveEditorTangentTypes::InterpolationLinear |
+		(int32)ECurveEditorTangentTypes::InterpolationCubicAuto |
+		(int32)ECurveEditorTangentTypes::InterpolationCubicUser |
+		(int32)ECurveEditorTangentTypes::InterpolationCubicBreak |
+		(int32)ECurveEditorTangentTypes::InterpolationCubicWeighted);
+		//nope we don't support smart auto by default, FRichCurve doesn't support i
+}
+
 void FCurveEditor::SetPanel(TSharedPtr<SCurveEditorPanel> InPanel)
 {
 	WeakPanel = InPanel;
@@ -1775,7 +1786,7 @@ void FCurveEditor::FlattenSelection()
 				if (Attributes.HasTangentMode() && (Attributes.HasArriveTangent() || Attributes.HasLeaveTangent()))
 				{
 					Attributes.SetArriveTangent(0.f).SetLeaveTangent(0.f);
-					if (Attributes.GetTangentMode() == RCTM_Auto)
+					if (Attributes.GetTangentMode() == RCTM_Auto || Attributes.GetTangentMode() == RCTM_SmartAuto)
 					{
 						Attributes.SetTangentMode(RCTM_User);
 					}
@@ -1852,7 +1863,7 @@ void FCurveEditor::StraightenSelection()
 				{
 					float NewTangent = (Attributes.GetLeaveTangent() + Attributes.GetArriveTangent()) * 0.5f;
 					Attributes.SetArriveTangent(NewTangent).SetLeaveTangent(NewTangent);
-					if (Attributes.GetTangentMode() == RCTM_Auto)
+					if (Attributes.GetTangentMode() == RCTM_Auto || Attributes.GetTangentMode() == RCTM_SmartAuto)
 					{
 						Attributes.SetTangentMode(RCTM_User);
 					}
diff --git a/Engine/Source/Editor/CurveEditor/Private/CurveEditorCommands.cpp b/Engine/Source/Editor/CurveEditor/Private/CurveEditorCommands.cpp
--- a/Engine/Source/Editor/CurveEditor/Private/CurveEditorCommands.cpp
+++ b/Engine/Source/Editor/CurveEditor/Private/CurveEditorCommands.cpp
@@ -24,6 +24,7 @@ void FCurveEditorCommands::RegisterCommands()
 
 	UI_COMMAND(InterpolationConstant, "Constant", "Constant interpolation", EUserInterfaceActionType::ToggleButton, FInputChord(EKeys::Five));
 	UI_COMMAND(InterpolationLinear, "Linear", "Linear interpolation", EUserInterfaceActionType::ToggleButton, FInputChord(EKeys::Four));
+	UI_COMMAND(InterpolationCubicSmartAuto, "Auto", "Cubic interpolation - Smart Automatic tangents", EUserInterfaceActionType::ToggleButton, FInputChord(EKeys::Zero));
 	UI_COMMAND(InterpolationCubicAuto, "Auto", "Cubic interpolation - Automatic tangents", EUserInterfaceActionType::ToggleButton, FInputChord(EKeys::One));
 	UI_COMMAND(InterpolationCubicUser, "User", "Cubic interpolation - User flat tangents", EUserInterfaceActionType::ToggleButton, FInputChord(EKeys::Two));
 	UI_COMMAND(InterpolationCubicBreak, "Break", "Cubic interpolation - User broken tangents", EUserInterfaceActionType::ToggleButton, FInputChord(EKeys::Three));
diff --git a/Engine/Source/Editor/CurveEditor/Private/CurveEditorContextMenu.cpp b/Engine/Source/Editor/CurveEditor/Private/CurveEditorContextMenu.cpp
--- a/Engine/Source/Editor/CurveEditor/Private/CurveEditorContextMenu.cpp
+++ b/Engine/Source/Editor/CurveEditor/Private/CurveEditorContextMenu.cpp
@@ -86,12 +86,35 @@ void FCurveEditorContextMenu::BuildMenu(FMenuBuilder& MenuBuilder, TSharedRef<FC
 				MenuBuilder.AddMenuSeparator();
 
 				// Tangent Types
-				MenuBuilder.AddMenuEntry(FCurveEditorCommands::Get().InterpolationCubicAuto);
-				MenuBuilder.AddMenuEntry(FCurveEditorCommands::Get().InterpolationCubicUser);
-				MenuBuilder.AddMenuEntry(FCurveEditorCommands::Get().InterpolationCubicBreak);
-				MenuBuilder.AddMenuEntry(FCurveEditorCommands::Get().InterpolationLinear);
-				MenuBuilder.AddMenuEntry(FCurveEditorCommands::Get().InterpolationConstant);
-				MenuBuilder.AddMenuEntry(FCurveEditorCommands::Get().InterpolationToggleWeighted);
+				int32 SupportedTangentTypes = CurveEditor->GetSupportedTangentTypes();
+				if (SupportedTangentTypes & (int32)ECurveEditorTangentTypes::InterpolationCubicSmartAuto)
+				{
+					MenuBuilder.AddMenuEntry(FCurveEditorCommands::Get().InterpolationCubicSmartAuto);
+				};
+				if (SupportedTangentTypes & (int32)ECurveEditorTangentTypes::InterpolationCubicAuto)
+				{
+					MenuBuilder.AddMenuEntry(FCurveEditorCommands::Get().InterpolationCubicAuto);
+				};
+				if (SupportedTangentTypes & (int32)ECurveEditorTangentTypes::InterpolationCubicUser)
+				{
+					MenuBuilder.AddMenuEntry(FCurveEditorCommands::Get().InterpolationCubicUser);
+				}
+				if (SupportedTangentTypes & (int32)ECurveEditorTangentTypes::InterpolationCubicBreak)
+				{
+					MenuBuilder.AddMenuEntry(FCurveEditorCommands::Get().InterpolationCubicBreak);
+				}
+				if (SupportedTangentTypes & (int32)ECurveEditorTangentTypes::InterpolationLinear)
+				{
+					MenuBuilder.AddMenuEntry(FCurveEditorCommands::Get().InterpolationLinear);
+				}
+				if (SupportedTangentTypes & (int32)ECurveEditorTangentTypes::InterpolationConstant)
+				{
+					MenuBuilder.AddMenuEntry(FCurveEditorCommands::Get().InterpolationConstant);
+				}
+				if (SupportedTangentTypes & (int32)ECurveEditorTangentTypes::InterpolationCubicWeighted)
+				{
+					MenuBuilder.AddMenuEntry(FCurveEditorCommands::Get().InterpolationToggleWeighted);
+				}
 
 				MenuBuilder.AddMenuSeparator();
 			}
diff --git a/Engine/Source/Editor/CurveEditor/Private/SCurveEditorPanel.cpp b/Engine/Source/Editor/CurveEditor/Private/SCurveEditorPanel.cpp
--- a/Engine/Source/Editor/CurveEditor/Private/SCurveEditorPanel.cpp
+++ b/Engine/Source/Editor/CurveEditor/Private/SCurveEditorPanel.cpp
@@ -349,6 +349,7 @@ void SCurveEditorPanel::BindCommands()
 		FExecuteAction SetConstant   = FExecuteAction::CreateSP(this, &SCurveEditorPanel::SetKeyAttributes, FKeyAttributes().SetInterpMode(RCIM_Constant).SetTangentMode(RCTM_Auto), LOCTEXT("SetInterpConstant", "Set Interp Constant"));
 		FExecuteAction SetLinear     = FExecuteAction::CreateSP(this, &SCurveEditorPanel::SetKeyAttributes, FKeyAttributes().SetInterpMode(RCIM_Linear).SetTangentMode(RCTM_Auto),   LOCTEXT("SetInterpLinear",   "Set Interp Linear"));
 		FExecuteAction SetCubicAuto  = FExecuteAction::CreateSP(this, &SCurveEditorPanel::SetKeyAttributes, FKeyAttributes().SetInterpMode(RCIM_Cubic).SetTangentMode(RCTM_Auto),    LOCTEXT("SetInterpCubic",    "Set Interp Auto"));
+		FExecuteAction SetCubicSmartAuto = FExecuteAction::CreateSP(this, &SCurveEditorPanel::SetKeyAttributes, FKeyAttributes().SetInterpMode(RCIM_Cubic).SetTangentMode(RCTM_SmartAuto),	LOCTEXT("SetInterpSmartAuto", "Set Interp Smart Auto"));
 		FExecuteAction SetCubicUser  = FExecuteAction::CreateSP(this, &SCurveEditorPanel::SetKeyAttributes, FKeyAttributes().SetInterpMode(RCIM_Cubic).SetTangentMode(RCTM_User),    LOCTEXT("SetInterpUser",     "Set Interp User"));
 		FExecuteAction SetCubicBreak = FExecuteAction::CreateSP(this, &SCurveEditorPanel::SetKeyAttributes, FKeyAttributes().SetInterpMode(RCIM_Cubic).SetTangentMode(RCTM_Break),   LOCTEXT("SetInterpBreak",    "Set Interp Break"));
 
@@ -358,19 +359,42 @@ void SCurveEditorPanel::BindCommands()
 		FIsActionChecked IsConstantCommon   = FIsActionChecked::CreateSP(this, &SCurveEditorPanel::CompareCommonInterpolationMode, RCIM_Constant);
 		FIsActionChecked IsLinearCommon     = FIsActionChecked::CreateSP(this, &SCurveEditorPanel::CompareCommonInterpolationMode, RCIM_Linear);
 		FIsActionChecked IsCubicAutoCommon  = FIsActionChecked::CreateSP(this, &SCurveEditorPanel::CompareCommonTangentMode, RCIM_Cubic, RCTM_Auto);
+		FIsActionChecked IsCubicSmartAutoCommon = FIsActionChecked::CreateSP(this, &SCurveEditorPanel::CompareCommonTangentMode, RCIM_Cubic, RCTM_SmartAuto);
 		FIsActionChecked IsCubicUserCommon  = FIsActionChecked::CreateSP(this, &SCurveEditorPanel::CompareCommonTangentMode, RCIM_Cubic, RCTM_User);
 		FIsActionChecked IsCubicBreakCommon = FIsActionChecked::CreateSP(this, &SCurveEditorPanel::CompareCommonTangentMode, RCIM_Cubic, RCTM_Break);
 		FIsActionChecked IsCubicWeightCommon = FIsActionChecked::CreateSP(this, &SCurveEditorPanel::CompareCommonTangentWeightMode, RCIM_Cubic, RCTWM_WeightedBoth);
 
 		FCanExecuteAction CanSetKeyTangent = FIsActionChecked::CreateSP(this, &SCurveEditorPanel::CanSetKeyInterpolation);
 
-		CommandList->MapAction(FCurveEditorCommands::Get().InterpolationConstant, SetConstant, CanSetKeyTangent, IsConstantCommon);
-		CommandList->MapAction(FCurveEditorCommands::Get().InterpolationLinear, SetLinear, CanSetKeyTangent, IsLinearCommon);
-		CommandList->MapAction(FCurveEditorCommands::Get().InterpolationCubicAuto, SetCubicAuto, CanSetKeyTangent, IsCubicAutoCommon);
-		CommandList->MapAction(FCurveEditorCommands::Get().InterpolationCubicUser, SetCubicUser, CanSetKeyTangent, IsCubicUserCommon);
-		CommandList->MapAction(FCurveEditorCommands::Get().InterpolationCubicBreak, SetCubicBreak, CanSetKeyTangent, IsCubicBreakCommon);
-		CommandList->MapAction(FCurveEditorCommands::Get().InterpolationToggleWeighted, ToggleWeighted, CanToggleWeighted, IsCubicWeightCommon);
-
+		int32 SupportedTangentTypes = CurveEditor->GetSupportedTangentTypes();
+		if (SupportedTangentTypes & (int32)ECurveEditorTangentTypes::InterpolationCubicSmartAuto)
+		{
+			CommandList->MapAction(FCurveEditorCommands::Get().InterpolationCubicSmartAuto, SetCubicSmartAuto, CanSetKeyTangent, IsCubicSmartAutoCommon);
+		};
+		if (SupportedTangentTypes & (int32)ECurveEditorTangentTypes::InterpolationCubicAuto)
+		{
+			CommandList->MapAction(FCurveEditorCommands::Get().InterpolationCubicAuto, SetCubicAuto, CanSetKeyTangent, IsCubicAutoCommon);
+		};
+		if (SupportedTangentTypes & (int32)ECurveEditorTangentTypes::InterpolationCubicUser)
+		{
+			CommandList->MapAction(FCurveEditorCommands::Get().InterpolationCubicUser, SetCubicUser, CanSetKeyTangent, IsCubicUserCommon);
+		}
+		if (SupportedTangentTypes & (int32)ECurveEditorTangentTypes::InterpolationCubicBreak)
+		{
+			CommandList->MapAction(FCurveEditorCommands::Get().InterpolationCubicBreak, SetCubicBreak, CanSetKeyTangent, IsCubicBreakCommon);
+		}
+		if (SupportedTangentTypes & (int32)ECurveEditorTangentTypes::InterpolationLinear)
+		{
+			CommandList->MapAction(FCurveEditorCommands::Get().InterpolationLinear, SetLinear, CanSetKeyTangent, IsLinearCommon);
+		}
+		if (SupportedTangentTypes & (int32)ECurveEditorTangentTypes::InterpolationConstant)
+		{
+			CommandList->MapAction(FCurveEditorCommands::Get().InterpolationConstant, SetConstant, CanSetKeyTangent, IsConstantCommon);
+		}
+		if (SupportedTangentTypes & (int32)ECurveEditorTangentTypes::InterpolationCubicWeighted)
+		{
+			CommandList->MapAction(FCurveEditorCommands::Get().InterpolationToggleWeighted, ToggleWeighted, CanToggleWeighted, IsCubicWeightCommon);
+		}
 	}
 
 	// Pre Extrapolation Modes
@@ -1265,13 +1289,36 @@ TSharedPtr<FExtender> SCurveEditorPanel::GetToolbarExtender()
 
 			ToolBarBuilder.BeginSection("Tangents");
 			{
-				ToolBarBuilder.AddToolBarButton(FCurveEditorCommands::Get().InterpolationCubicAuto);
-				ToolBarBuilder.AddToolBarButton(FCurveEditorCommands::Get().InterpolationCubicUser);
-				ToolBarBuilder.AddToolBarButton(FCurveEditorCommands::Get().InterpolationCubicBreak);
-				ToolBarBuilder.AddToolBarButton(FCurveEditorCommands::Get().InterpolationLinear);
-				ToolBarBuilder.AddToolBarButton(FCurveEditorCommands::Get().InterpolationConstant);
-				ToolBarBuilder.AddToolBarButton(FCurveEditorCommands::Get().InterpolationToggleWeighted);
-				
+				int32 SupportedTangentTypes = InEditorPanel->CurveEditor->GetSupportedTangentTypes();
+				if (SupportedTangentTypes & (int32)ECurveEditorTangentTypes::InterpolationCubicSmartAuto)
+				{
+					ToolBarBuilder.AddToolBarButton(FCurveEditorCommands::Get().InterpolationCubicSmartAuto);
+				};
+				if (SupportedTangentTypes & (int32)ECurveEditorTangentTypes::InterpolationCubicAuto)
+				{
+					ToolBarBuilder.AddToolBarButton(FCurveEditorCommands::Get().InterpolationCubicAuto);
+				};
+				if (SupportedTangentTypes & (int32)ECurveEditorTangentTypes::InterpolationCubicUser)
+				{
+					ToolBarBuilder.AddToolBarButton(FCurveEditorCommands::Get().InterpolationCubicUser);
+				}
+				if (SupportedTangentTypes & (int32)ECurveEditorTangentTypes::InterpolationCubicBreak)
+				{
+					ToolBarBuilder.AddToolBarButton(FCurveEditorCommands::Get().InterpolationCubicBreak);
+				}
+				if (SupportedTangentTypes & (int32)ECurveEditorTangentTypes::InterpolationLinear)
+				{
+					ToolBarBuilder.AddToolBarButton(FCurveEditorCommands::Get().InterpolationLinear);
+				}
+				if (SupportedTangentTypes & (int32)ECurveEditorTangentTypes::InterpolationConstant)
+				{
+					ToolBarBuilder.AddToolBarButton(FCurveEditorCommands::Get().InterpolationConstant);
+				}
+				if (SupportedTangentTypes & (int32)ECurveEditorTangentTypes::InterpolationCubicWeighted)
+				{
+					ToolBarBuilder.AddToolBarButton(FCurveEditorCommands::Get().InterpolationToggleWeighted);
+				}
+
 				// We re-use key interpolation checks here, as you can set them under the same conditions.
 				FCanExecuteAction CanSetInfinities = FCanExecuteAction::CreateSP(InEditorPanel, &SCurveEditorPanel::CanSetKeyInterpolation);
 				
diff --git a/Engine/Source/Editor/CurveEditor/Public/CurveEditor.h b/Engine/Source/Editor/CurveEditor/Public/CurveEditor.h
--- a/Engine/Source/Editor/CurveEditor/Public/CurveEditor.h
+++ b/Engine/Source/Editor/CurveEditor/Public/CurveEditor.h
@@ -58,6 +58,19 @@ DECLARE_DELEGATE_OneParam(FOnSetBoolean, bool)
 DECLARE_MULTICAST_DELEGATE_OneParam(FOnActiveToolChanged, FCurveEditorToolID)
 DECLARE_MULTICAST_DELEGATE_ThreeParams(FOnCurveArrayChanged, FCurveModel*, bool /*displayed*/,const FCurveEditor*);
 
+/** Enums to describe supported tangent types, by default support all but smart auto*/
+enum class ECurveEditorTangentTypes : int32
+{
+	InterpolationConstant = 0x1,
+	InterpolationLinear = 0x2,
+	InterpolationCubicAuto = 0x4,
+
+	InterpolationCubicUser = 0x8,
+	InterpolationCubicBreak = 0x10,
+	InterpolationCubicWeighted = 0x20,
+	InterpolationCubicSmartAuto = 0x40,
+
+};
 
 class CURVEEDITOR_API FCurveEditor 
 	: public FEditorUndoClient
@@ -121,6 +134,8 @@ public:
 
 	void InitCurveEditor(const FCurveEditorInitParams& InInitParams);
 
+	virtual int32 GetSupportedTangentTypes();
+
 public:
 
 	void SetPanel(TSharedPtr<SCurveEditorPanel> InPanel);
@@ -134,6 +149,7 @@ public:
 	FCurveEditorScreenSpaceH GetPanelInputSpace() const;
 
 	void ResetMinMaxes();
+
 public:
 	/**
 	 * Zoom the curve editor to fit all the selected curves (or all curves if none selected)
diff --git a/Engine/Source/Editor/CurveEditor/Public/CurveEditorCommands.h b/Engine/Source/Editor/CurveEditor/Public/CurveEditorCommands.h
--- a/Engine/Source/Editor/CurveEditor/Public/CurveEditorCommands.h
+++ b/Engine/Source/Editor/CurveEditor/Public/CurveEditorCommands.h
@@ -42,6 +42,7 @@ public:
 	TSharedPtr<FUICommandInfo> InterpolationConstant;
 	TSharedPtr<FUICommandInfo> InterpolationLinear;
 	TSharedPtr<FUICommandInfo> InterpolationCubicAuto;
+	TSharedPtr<FUICommandInfo> InterpolationCubicSmartAuto;
 	TSharedPtr<FUICommandInfo> InterpolationCubicUser;
 	TSharedPtr<FUICommandInfo> InterpolationCubicBreak;
 	TSharedPtr<FUICommandInfo> InterpolationToggleWeighted;
diff --git a/Engine/Source/Editor/EditorStyle/Private/SlateEditorStyle.cpp b/Engine/Source/Editor/EditorStyle/Private/SlateEditorStyle.cpp
--- a/Engine/Source/Editor/EditorStyle/Private/SlateEditorStyle.cpp
+++ b/Engine/Source/Editor/EditorStyle/Private/SlateEditorStyle.cpp
@@ -2312,6 +2312,7 @@ void FSlateEditorStyle::FStyle::SetupSequencerStyles()
 	// Sequencer
 	if (IncludeEditorSpecificStyles())
 	{
+		Set("Sequencer.IconKeySmartAuto", new IMAGE_BRUSH("Sequencer/IconKeySmartAuto", Icon12x12));
 		Set("Sequencer.IconKeyAuto", new IMAGE_BRUSH("Sequencer/IconKeyAuto", Icon12x12));
 		Set("Sequencer.IconKeyBreak", new IMAGE_BRUSH("Sequencer/IconKeyBreak", Icon12x12));
 		Set("Sequencer.IconKeyConstant", new IMAGE_BRUSH("Sequencer/IconKeyConstant", Icon12x12));
@@ -7411,6 +7412,8 @@ void FSlateEditorStyle::FStyle::SetupToolkitStyles()
 		Set("GenericCurveEditor.ToggleOutputSnapping.Small", new IMAGE_BRUSH("GenericCurveEditor/Icons/AdjustKeySnapValue_48x", Icon24x24));
 
 		// Tangent Types
+		Set("GenericCurveEditor.InterpolationCubicSmartAuto", new IMAGE_BRUSH("GenericCurveEditor/Icons/CubicTangentSmartAuto_48x", Icon48x48));
+		Set("GenericCurveEditor.InterpolationCubicSmartAuto.Small", new IMAGE_BRUSH("GenericCurveEditor/Icons/CubicTangentSmartAuto_48x", Icon24x24));
 		Set("GenericCurveEditor.InterpolationCubicAuto", new IMAGE_BRUSH("GenericCurveEditor/Icons/CubicTangentAuto_48x", Icon48x48));
 		Set("GenericCurveEditor.InterpolationCubicAuto.Small", new IMAGE_BRUSH("GenericCurveEditor/Icons/CubicTangentAuto_48x", Icon24x24));
 		Set("GenericCurveEditor.InterpolationCubicUser", new IMAGE_BRUSH("GenericCurveEditor/Icons/TangentsCubic_48x", Icon48x48));
diff --git a/Engine/Source/Editor/EditorStyle/Private/StarshipStyle.cpp b/Engine/Source/Editor/EditorStyle/Private/StarshipStyle.cpp
--- a/Engine/Source/Editor/EditorStyle/Private/StarshipStyle.cpp
+++ b/Engine/Source/Editor/EditorStyle/Private/StarshipStyle.cpp
@@ -2057,6 +2057,7 @@ void FStarshipEditorStyle::FStyle::SetupSequencerStyles()
 
 		Set("SequencerToolBar", SequencerToolbar);
 
+		Set("Sequencer.IconKeySmartAuto", new IMAGE_BRUSH("Sequencer/IconKeySmartAuto", Icon12x12));
 		Set("Sequencer.IconKeyAuto", new IMAGE_BRUSH("Sequencer/IconKeyAuto", Icon12x12));
 		Set("Sequencer.IconKeyBreak", new IMAGE_BRUSH("Sequencer/IconKeyBreak", Icon12x12));
 		Set("Sequencer.IconKeyConstant", new IMAGE_BRUSH("Sequencer/IconKeyConstant", Icon12x12));
@@ -6924,6 +6925,7 @@ void FStarshipEditorStyle::FStyle::SetupToolkitStyles()
 		Set("GenericCurveEditor.ToggleOutputSnapping", new IMAGE_BRUSH_SVG("Starship/GenericCurveEditor/AdjustKeySnapFrameVertical", Icon20x20));
 
 		// Tangent Types
+		Set("GenericCurveEditor.InterpolationCubicSmartAuto", new IMAGE_BRUSH_SVG("Starship/GenericCurveEditor/TangentsCubicSmartAuto", Icon20x20));
 		Set("GenericCurveEditor.InterpolationCubicAuto", new IMAGE_BRUSH_SVG("Starship/GenericCurveEditor/TangentsCubicAuto", Icon20x20));
 		Set("GenericCurveEditor.InterpolationCubicUser", new IMAGE_BRUSH_SVG("Starship/GenericCurveEditor/TangentsCubicUser", Icon20x20));
 		Set("GenericCurveEditor.InterpolationCubicBreak", new IMAGE_BRUSH_SVG("Starship/GenericCurveEditor/TangentsCubicBreak", Icon20x20));
diff --git a/Engine/Source/Editor/MovieSceneTools/Private/Channels/BezierChannelCurveModel.cpp b/Engine/Source/Editor/MovieSceneTools/Private/Channels/BezierChannelCurveModel.cpp
--- a/Engine/Source/Editor/MovieSceneTools/Private/Channels/BezierChannelCurveModel.cpp
+++ b/Engine/Source/Editor/MovieSceneTools/Private/Channels/BezierChannelCurveModel.cpp
@@ -201,6 +201,10 @@ void FBezierChannelCurveModel<ChannelType, ChannelValue, KeyType>::GetKeyAttribu
 		}
 	}
 }
+static bool IsAuto(ERichCurveTangentMode TangentMode)
+{
+	return (TangentMode == RCTM_Auto || TangentMode == RCTM_SmartAuto);
+}
 
 template<typename ChannelType, typename ChannelValue, typename KeyType> 
 void FBezierChannelCurveModel<ChannelType, ChannelValue, KeyType>::SetKeyAttributes(TArrayView<const FKeyHandle> InKeys, TArrayView<const FKeyAttributes> InAttributes, EPropertyChangeType::Type ChangeType)
@@ -229,7 +233,7 @@ void FBezierChannelCurveModel<ChannelType, ChannelValue, KeyType>::SetKeyAttribu
 				if (Attributes.HasTangentMode())
 				{
 					KeyValue.TangentMode = Attributes.GetTangentMode();
-					if (KeyValue.TangentMode == RCTM_Auto)
+					if (IsAuto(KeyValue.TangentMode))
 					{
 						KeyValue.Tangent.TangentWeightMode = RCTWM_WeightedNone;
 					}
@@ -269,11 +273,12 @@ void FBezierChannelCurveModel<ChannelType, ChannelValue, KeyType>::SetKeyAttribu
 							KeyValue.TangentMode = RCTM_User;
 						}
 					}
+					bAutoSetTangents = true;
 				}
 
 				if (Attributes.HasArriveTangent())
 				{
-					if (KeyValue.TangentMode == RCTM_Auto)
+					if (IsAuto(KeyValue.TangentMode))
 					{
 						KeyValue.TangentMode = RCTM_User;
 						KeyValue.Tangent.TangentWeightMode = RCTWM_WeightedNone;
@@ -284,11 +289,12 @@ void FBezierChannelCurveModel<ChannelType, ChannelValue, KeyType>::SetKeyAttribu
 					{
 						KeyValue.Tangent.LeaveTangent = KeyValue.Tangent.ArriveTangent;
 					}
+					bAutoSetTangents = true;
 				}
 
 				if (Attributes.HasLeaveTangent())
 				{
-					if (KeyValue.TangentMode == RCTM_Auto)
+					if (IsAuto(KeyValue.TangentMode))
 					{
 						KeyValue.TangentMode = RCTM_User;
 						KeyValue.Tangent.TangentWeightMode = RCTWM_WeightedNone;
@@ -299,11 +305,12 @@ void FBezierChannelCurveModel<ChannelType, ChannelValue, KeyType>::SetKeyAttribu
 					{
 						KeyValue.Tangent.ArriveTangent = KeyValue.Tangent.LeaveTangent;
 					}
+					bAutoSetTangents = true;
 				}
 
 				if (Attributes.HasArriveTangentWeight())
 				{
-					if (KeyValue.TangentMode == RCTM_Auto)
+					if (IsAuto(KeyValue.TangentMode))
 					{
 						KeyValue.TangentMode = RCTM_User;
 						KeyValue.Tangent.TangentWeightMode = RCTWM_WeightedNone;
@@ -314,12 +321,12 @@ void FBezierChannelCurveModel<ChannelType, ChannelValue, KeyType>::SetKeyAttribu
 					{
 						KeyValue.Tangent.LeaveTangentWeight = KeyValue.Tangent.ArriveTangentWeight;
 					}
+					bAutoSetTangents = true;
 				}
 
 				if (Attributes.HasLeaveTangentWeight())
 				{
-				
-					if (KeyValue.TangentMode == RCTM_Auto)
+					if (IsAuto(KeyValue.TangentMode))
 					{
 						KeyValue.TangentMode = RCTM_User;
 						KeyValue.Tangent.TangentWeightMode = RCTWM_WeightedNone;
@@ -330,6 +337,7 @@ void FBezierChannelCurveModel<ChannelType, ChannelValue, KeyType>::SetKeyAttribu
 					{
 						KeyValue.Tangent.ArriveTangentWeight = KeyValue.Tangent.LeaveTangentWeight;
 					}
+					bAutoSetTangents = true;
 				}
 			}
 		}
diff --git a/Engine/Source/Editor/MovieSceneTools/Private/Channels/BuiltInChannelEditors.cpp b/Engine/Source/Editor/MovieSceneTools/Private/Channels/BuiltInChannelEditors.cpp
--- a/Engine/Source/Editor/MovieSceneTools/Private/Channels/BuiltInChannelEditors.cpp
+++ b/Engine/Source/Editor/MovieSceneTools/Private/Channels/BuiltInChannelEditors.cpp
@@ -855,6 +855,7 @@ void DrawKeysImpl(ChannelType* Channel, TArrayView<const FKeyHandle> InKeyHandle
 
 			switch (TangentMode)
 			{
+			case RCTM_SmartAuto:  TempParams.FillTint = FLinearColor(0.759f, 0.176f, 0.67f, 1.0f);break; // little vermillion
 			case RCTM_Auto:  TempParams.FillTint = FLinearColor(0.972f, 0.2f, 0.2f, 1.0f);     break; // vermillion
 			case RCTM_Break: TempParams.FillTint = FLinearColor(0.336f, 0.703f, 0.5f, 0.91f);  break; // sky blue
 			case RCTM_User:  TempParams.FillTint = FLinearColor(0.797f, 0.473f, 0.5f, 0.652f); break; // reddish purple
@@ -977,6 +978,18 @@ struct TCurveChannelKeyMenuExtension : TSharedFromThis<TCurveChannelKeyMenuExten
 
 		MenuBuilder.BeginSection("SequencerInterpolation", LOCTEXT("KeyInterpolationMenu", "Key Interpolation"));
 		{
+			MenuBuilder.AddMenuEntry(
+				LOCTEXT("SetKeyInterpolationSmartAuto", "Cubic (Smart Auto)"),
+				LOCTEXT("SetKeyInterpolationSmartAutoTooltip", "Set key interpolation to smart auto"),
+				FSlateIcon(FAppStyle::GetAppStyleSetName(), "Sequencer.IconKeySmartAuto"),
+				FUIAction(
+					FExecuteAction::CreateLambda([SharedThis] { SharedThis->SetInterpTangentMode(RCIM_Cubic, RCTM_SmartAuto); }),
+					FCanExecuteAction(),
+					FIsActionChecked::CreateLambda([SharedThis] { return SharedThis->IsInterpTangentModeSelected(RCIM_Cubic, RCTM_SmartAuto); })),
+				NAME_None,
+				EUserInterfaceActionType::ToggleButton
+			);
+
 			MenuBuilder.AddMenuEntry(
 				LOCTEXT("SetKeyInterpolationAuto", "Cubic (Auto)"),
 				LOCTEXT("SetKeyInterpolationAutoTooltip", "Set key interpolation to auto"),
diff --git a/Engine/Source/Editor/Sequencer/Private/MVVM/CurveEditorExtension.cpp b/Engine/Source/Editor/Sequencer/Private/MVVM/CurveEditorExtension.cpp
--- a/Engine/Source/Editor/Sequencer/Private/MVVM/CurveEditorExtension.cpp
+++ b/Engine/Source/Editor/Sequencer/Private/MVVM/CurveEditorExtension.cpp
@@ -60,6 +60,16 @@ public:
 			}
 		}
 	}
+	int32 GetSupportedTangentTypes() override
+	{
+		return ((int32)ECurveEditorTangentTypes::InterpolationConstant	|
+			(int32)ECurveEditorTangentTypes::InterpolationLinear		|
+			(int32)ECurveEditorTangentTypes::InterpolationCubicAuto		|
+			(int32)ECurveEditorTangentTypes::InterpolationCubicUser		|
+			(int32)ECurveEditorTangentTypes::InterpolationCubicBreak	|
+			(int32)ECurveEditorTangentTypes::InterpolationCubicWeighted |
+			(int32)ECurveEditorTangentTypes::InterpolationCubicSmartAuto);
+	}
 };
 
 struct FSequencerCurveEditorBounds : ICurveEditorBounds
@@ -338,7 +348,9 @@ FKeyAttributes FCurveEditorExtension::GetDefaultKeyAttributes() const
 	case EMovieSceneKeyInterpolation::Break:    return FKeyAttributes().SetInterpMode(RCIM_Cubic).SetTangentMode(RCTM_Break);
 	case EMovieSceneKeyInterpolation::Linear:   return FKeyAttributes().SetInterpMode(RCIM_Linear).SetTangentMode(RCTM_Auto);
 	case EMovieSceneKeyInterpolation::Constant: return FKeyAttributes().SetInterpMode(RCIM_Constant).SetTangentMode(RCTM_Auto);
-	default:                                    return FKeyAttributes().SetInterpMode(RCIM_Cubic).SetTangentMode(RCTM_Auto);
+	case EMovieSceneKeyInterpolation::Auto:     return FKeyAttributes().SetInterpMode(RCIM_Cubic).SetTangentMode(RCTM_Auto);
+	case EMovieSceneKeyInterpolation::SmartAuto:
+	default:                                    return FKeyAttributes().SetInterpMode(RCIM_Cubic).SetTangentMode(RCTM_SmartAuto);
 	}
 }
 
diff --git a/Engine/Source/Editor/Sequencer/Private/SSequencer.cpp b/Engine/Source/Editor/Sequencer/Private/SSequencer.cpp
--- a/Engine/Source/Editor/Sequencer/Private/SSequencer.cpp
+++ b/Engine/Source/Editor/Sequencer/Private/SSequencer.cpp
@@ -2671,6 +2671,17 @@ TSharedRef<SWidget> SSequencer::MakeKeyGroupMenu()
 	// Interpolation
 	MenuBuilder.BeginSection("SequencerInterpolation", LOCTEXT("KeyInterpolationMenu", "Default Key Interpolation"));
 	{
+		MenuBuilder.AddMenuEntry(
+			LOCTEXT("SetKeyInterpolationSmartAuto", "Cubic (Smart Auto)"),
+			LOCTEXT("SetKeyInterpolationAutoTooltip", "Set key interpolation to smart auto"),
+			FSlateIcon(FAppStyle::GetAppStyleSetName(), "Sequencer.IconKeySmartAuto"),
+			FUIAction(
+				FExecuteAction::CreateLambda([this] { SequencerPtr.Pin()->SetKeyInterpolation(EMovieSceneKeyInterpolation::SmartAuto); }),
+				FCanExecuteAction(),
+				FIsActionChecked::CreateLambda([this] { return SequencerPtr.Pin()->GetKeyInterpolation() == EMovieSceneKeyInterpolation::SmartAuto; })),
+			NAME_None,
+			EUserInterfaceActionType::RadioButton
+		);
 		MenuBuilder.AddMenuEntry(
 			LOCTEXT("SetKeyInterpolationAuto", "Cubic (Auto)"),
 			LOCTEXT("SetKeyInterpolationAutoTooltip", "Set key interpolation to auto"),
diff --git a/Engine/Source/Editor/Sequencer/Private/Sequencer.cpp b/Engine/Source/Editor/Sequencer/Private/Sequencer.cpp
--- a/Engine/Source/Editor/Sequencer/Private/Sequencer.cpp
+++ b/Engine/Source/Editor/Sequencer/Private/Sequencer.cpp
@@ -11469,6 +11469,10 @@ void FSequencer::BindCommands()
 		FExecuteAction::CreateSP(this, &FSequencer::JumpBackward),
 		EUIActionRepeatMode::RepeatEnabled);
 
+	SequencerCommandBindings->MapAction(
+		Commands.SetInterpolationCubicAuto,
+		FExecuteAction::CreateSP(this, &FSequencer::SetInterpTangentMode, ERichCurveInterpMode::RCIM_Cubic, ERichCurveTangentMode::RCTM_SmartAuto));
+
 	SequencerCommandBindings->MapAction(
 		Commands.SetInterpolationCubicAuto,
 		FExecuteAction::CreateSP(this, &FSequencer::SetInterpTangentMode, ERichCurveInterpMode::RCIM_Cubic, ERichCurveTangentMode::RCTM_Auto));
diff --git a/Engine/Source/Editor/Sequencer/Private/SequencerCommands.cpp b/Engine/Source/Editor/Sequencer/Private/SequencerCommands.cpp
--- a/Engine/Source/Editor/Sequencer/Private/SequencerCommands.cpp
+++ b/Engine/Source/Editor/Sequencer/Private/SequencerCommands.cpp
@@ -79,6 +79,7 @@ void FSequencerCommands::RegisterCommands()
 
 	UI_COMMAND( AddActorsToSequencer, "Add Actors", "Add actors to sequencer", EUserInterfaceActionType::Button, FInputChord(EModifierKey::Control, EKeys::A) );
 	UI_COMMAND( SetKey, "Set Key", "Sets a key on the selected tracks", EUserInterfaceActionType::Button, FInputChord(EKeys::Enter) );
+	UI_COMMAND(SetInterpolationCubicSmartAuto, "Set Key Smart Auto", "Cubic interpolation - Smart Automatic tangents", EUserInterfaceActionType::Button, FInputChord(EKeys::Zero));
 	UI_COMMAND( SetInterpolationCubicAuto, "Set Key Auto", "Cubic interpolation - Automatic tangents", EUserInterfaceActionType::Button, FInputChord(EKeys::One));
 	UI_COMMAND( SetInterpolationCubicUser, "Set Key User", "Cubic interpolation - User flat tangents", EUserInterfaceActionType::Button, FInputChord(EKeys::Two));
 	UI_COMMAND( SetInterpolationCubicBreak, "Set Key Break", "Cubic interpolation - User broken tangents", EUserInterfaceActionType::Button, FInputChord(EKeys::Three));
diff --git a/Engine/Source/Editor/Sequencer/Private/SequencerContextMenus.cpp b/Engine/Source/Editor/Sequencer/Private/SequencerContextMenus.cpp
--- a/Engine/Source/Editor/Sequencer/Private/SequencerContextMenus.cpp
+++ b/Engine/Source/Editor/Sequencer/Private/SequencerContextMenus.cpp
@@ -471,6 +471,15 @@ void FSectionContextMenu::AddEditMenu(FMenuBuilder& MenuBuilder)
 
 	MenuBuilder.BeginSection("SequencerInterpolation", LOCTEXT("KeyInterpolationMenu", "Key Interpolation"));
 	
+	MenuBuilder.AddMenuEntry(
+		LOCTEXT("SetKeyInterpolationSmartAuto", "Cubic (Smart Auto)"),
+		LOCTEXT("SetKeyInterpolationSmartAutoTooltip", "Set key interpolation to smart auto"),
+		FSlateIcon(FAppStyle::GetAppStyleSetName(), "Sequencer.IconKeySmartAuto"),
+		FUIAction(
+			FExecuteAction::CreateLambda([=] { Shared->SetInterpTangentMode(RCIM_Cubic, RCTM_SmartAuto); }),
+			FCanExecuteAction::CreateLambda([=] { return Shared->CanSetInterpTangentMode(); }))
+	);
+
 	MenuBuilder.AddMenuEntry(
 		LOCTEXT("SetKeyInterpolationAuto", "Cubic (Auto)"),
 		LOCTEXT("SetKeyInterpolationAutoTooltip", "Set key interpolation to auto"),
diff --git a/Engine/Source/Editor/Sequencer/Public/SequencerCommands.h b/Engine/Source/Editor/Sequencer/Public/SequencerCommands.h
--- a/Engine/Source/Editor/Sequencer/Public/SequencerCommands.h
+++ b/Engine/Source/Editor/Sequencer/Public/SequencerCommands.h
@@ -208,6 +208,9 @@ public:
 	/** Sets a key at the current time for the selected actor */
 	TSharedPtr< FUICommandInfo > SetKey;
 
+	/** Sets the interp tangent mode for the selected keys to smart auto */
+	TSharedPtr< FUICommandInfo > SetInterpolationCubicSmartAuto;
+
 	/** Sets the interp tangent mode for the selected keys to auto */
 	TSharedPtr< FUICommandInfo > SetInterpolationCubicAuto;
 
diff --git a/Engine/Source/Editor/UnrealEd/Private/SCurveEditor.cpp b/Engine/Source/Editor/UnrealEd/Private/SCurveEditor.cpp
--- a/Engine/Source/Editor/UnrealEd/Private/SCurveEditor.cpp
+++ b/Engine/Source/Editor/UnrealEd/Private/SCurveEditor.cpp
@@ -163,6 +163,11 @@ void SCurveEditor::Construct(const FArguments& InArgs)
 		FCanExecuteAction::CreateSP(this, &SCurveEditor::HasRichCurves),
 		FIsActionChecked::CreateSP(this, &SCurveEditor::IsInterpolationModeSelected, RCIM_Cubic, RCTM_Auto));
 
+	Commands->MapAction(FCurveEditorCommands::Get().InterpolationCubicSmartAuto,
+		FExecuteAction::CreateSP(this, &SCurveEditor::OnSelectInterpolationMode, RCIM_Cubic, RCTM_SmartAuto),
+		FCanExecuteAction::CreateSP(this, &SCurveEditor::HasRichCurves),
+		FIsActionChecked::CreateSP(this, &SCurveEditor::IsInterpolationModeSelected, RCIM_Cubic, RCTM_SmartAuto));
+
 	Commands->MapAction(FCurveEditorCommands::Get().InterpolationCubicUser,
 		FExecuteAction::CreateSP(this, &SCurveEditor::OnSelectInterpolationMode, RCIM_Cubic, RCTM_User),
 		FCanExecuteAction::CreateSP(this, &SCurveEditor::HasRichCurves),
diff --git a/Engine/Source/Runtime/Engine/Classes/Curves/RichCurve.h b/Engine/Source/Runtime/Engine/Classes/Curves/RichCurve.h
--- a/Engine/Source/Runtime/Engine/Classes/Curves/RichCurve.h
+++ b/Engine/Source/Runtime/Engine/Classes/Curves/RichCurve.h
@@ -20,7 +20,9 @@ enum ERichCurveTangentMode : int
 	/** User specifies the tangent as two separate broken tangents on each side of the key which can allow a sharp change in evaluation before or after. */
 	RCTM_Break UMETA(DisplayName="Break"),
 	/** No tangents. */
-	RCTM_None UMETA(Hidden)
+	RCTM_None UMETA(Hidden),
+	/** New Auto tangent that creates smoother curves than Auto. */
+	RCTM_SmartAuto UMETA(DisplayName = "Smart Auto"),
 };
 
 
diff --git a/Engine/Source/Runtime/MovieScene/Private/Channels/MovieSceneCurveChannelImpl.cpp b/Engine/Source/Runtime/MovieScene/Private/Channels/MovieSceneCurveChannelImpl.cpp
--- a/Engine/Source/Runtime/MovieScene/Private/Channels/MovieSceneCurveChannelImpl.cpp
+++ b/Engine/Source/Runtime/MovieScene/Private/Channels/MovieSceneCurveChannelImpl.cpp
@@ -23,6 +23,12 @@ static TAutoConsoleVariable<int32> CVarSequencerAutoTangentInterpolation(
 	TEXT("If 1 Auto Tangent will use new algorithm to gradually flatten maximum/minimum keys, if 0 Auto Tangent will average all keys (pre 4.23 behavior)."),
 	ECVF_Default);
 
+static TAutoConsoleVariable<float> CVarSequencerSmartAutoBlendLocationPercentage(
+	TEXT("Sequencer.SmartAutoBlendLocationPercentage"),
+	0.8,
+	TEXT("Percentage near the next value that the tangent will blend to the adjacent tangent, if over 1.0 we won't blend. Default to 0.8"),
+	ECVF_Default);
+
 
 namespace UE
 {
@@ -911,6 +917,97 @@ bool TMovieSceneCurveChannelImpl<ChannelType>::EvaluateLegacy(const ChannelType*
 
 	return true;
 }
+template<typename ChannelType>
+float TMovieSceneCurveChannelImpl<ChannelType>::CalcSmartTangent(ChannelType* InChannel, int32 Index)
+{
+	ChannelValueType& ThisKey = InChannel->Values[Index];
+	int32 PrevIndex = Index - 1;
+	ChannelValueType  PrevKey = InChannel->Values[PrevIndex];
+
+	int32 NextIndex = Index + 1;
+	ChannelValueType NextKey = InChannel->Values[NextIndex];
+
+	float NewTangent = 0.0f;
+	// if key doesn't lie between we keep it flat(0.0).
+	if ((ThisKey.Value > PrevKey.Value && ThisKey.Value < NextKey.Value) ||
+		(ThisKey.Value < PrevKey.Value && ThisKey.Value > NextKey.Value))
+	{
+		
+		while (NextIndex <= (InChannel->Values.Num() - 2)
+			&& FMath::IsNearlyZero(NextKey.Tangent.ArriveTangent)
+			&& NextKey.InterpMode == RCIM_Cubic && (NextKey.TangentMode == RCTM_Auto || NextKey.TangentMode == RCTM_SmartAuto)
+			&& ((NextKey.Value > ThisKey.Value && NextKey.Value < InChannel->Values[NextIndex + 1].Value) ||
+				(NextKey.Value < ThisKey.Value && NextKey.Value > InChannel->Values[NextIndex + 1].Value)))
+		{
+			++NextIndex;
+			NextKey = InChannel->Values[NextIndex];
+		}
+		
+
+		const CurveValueType OneThird = 1.0 / 3.0;
+		const CurveValueType TwoThird = 2.0 / 3.0;
+		//we calculate the tangent one third from the previous and next tangents
+		const double TimeToPrevious = FMath::Max<double>(KINDA_SMALL_NUMBER, InChannel->Times[Index].Value - InChannel->Times[PrevIndex].Value);
+		const double OneThirdTimeToPrevious = OneThird * TimeToPrevious;
+		CurveValueType PrevY = PrevKey.Value + (PrevKey.Tangent.LeaveTangent * OneThirdTimeToPrevious);
+
+		const double TimeToNext = FMath::Max<double>(KINDA_SMALL_NUMBER, InChannel->Times[NextIndex].Value - InChannel->Times[Index].Value);
+		const double OneThirdTimeToNext = OneThird * TimeToNext;
+		CurveValueType NextY = NextKey.Value - (NextKey.Tangent.ArriveTangent * OneThirdTimeToNext);
+		
+		//leaving ThisKey.Value - ThisKey.Value there since it is needed if we decide to add any weighting
+		NewTangent = ((ThisKey.Value - PrevY) + (NextY - ThisKey.Value))
+			/ (TwoThird * TimeToPrevious + TwoThird * TimeToNext);
+
+		const double PreviousSlope = (ThisKey.Value - PrevY) / (TwoThird * TimeToPrevious);
+		const double NextSlope = (NextY - ThisKey.Value) / (TwoThird * TimeToNext);
+		/*
+		* Leaving this in for now but averaging is over agressive.
+		
+		if (NextKey.Value > PrevKey.Value)
+		{
+			CurveValueType NextToPrev = NextKey.Value - PrevKey.Value;
+			CurveValueType ToNext = (NextKey.Value - ThisKey.Value) / NextToPrev;
+			CurveValueType ToPrev = (ThisKey.Value - PrevKey.Value) / NextToPrev;
+
+			//NewTangent = (PreviousSlope * ToNext)+ (ToPrev * NextSlope);
+
+		}
+		else
+		{
+			CurveValueType NextToPrev = PrevKey.Value - NextKey.Value;
+			CurveValueType ToNext = (ThisKey.Value - NextKey.Value) / NextToPrev;
+			CurveValueType ToPrev = (PrevKey.Value  - ThisKey.Value) / NextToPrev;
+
+			//NewTangent = (PreviousSlope * ToNext) + (ToPrev * NextSlope);
+		}
+		*/
+		const float BlendToNextRange = CVarSequencerSmartAutoBlendLocationPercentage->GetFloat();
+		const double ValDiff = FMath::Abs<double>(NextKey.Value - PrevKey.Value);
+		const double OurDiff = FMath::Abs<double>(ThisKey.Value - PrevKey.Value);
+		//ValDiff won't be zero ever due to previous check
+		double PercDiff = OurDiff / ValDiff;
+		float NextTangent = (NextKey.InterpMode == RCIM_Cubic && (NextKey.TangentMode == RCTM_Auto || NextKey.TangentMode == RCTM_SmartAuto)) ?
+			0.0 : NextKey.Tangent.ArriveTangent;
+		float PrevTangent = (PrevKey.InterpMode == RCIM_Cubic && (PrevKey.TangentMode == RCTM_Auto || PrevKey.TangentMode == RCTM_SmartAuto)) ?
+			0.0 : PrevKey.Tangent.LeaveTangent;
+		if (BlendToNextRange >= 0.0f || BlendToNextRange <= 1.0f)
+		{
+			NextTangent = PrevTangent = 0.0;
+			if (PercDiff > BlendToNextRange)
+			{
+				PercDiff = (PercDiff - BlendToNextRange) / (1.0 - BlendToNextRange);
+				NewTangent = NewTangent * (1.0 - PercDiff) + (PercDiff * NextTangent);
+			}
+			else if (PercDiff < (1.0 - BlendToNextRange))
+			{
+				PercDiff = PercDiff / (1.0 - BlendToNextRange);
+				NewTangent = NewTangent * PercDiff + (1.0 - PercDiff) * PrevTangent;
+			}
+		}
+	}
+	return NewTangent;
+}
 
 template<typename ChannelType>
 void TMovieSceneCurveChannelImpl<ChannelType>::AutoSetTangents(ChannelType* InChannel, float Tension)
@@ -921,7 +1018,6 @@ void TMovieSceneCurveChannelImpl<ChannelType>::AutoSetTangents(ChannelType* InCh
 	}
 
 	const int UseNewAutoTangent = CVarSequencerAutoTangentInterpolation->GetInt();
-
 	{
 		ChannelValueType& FirstValue = InChannel->Values[0];
 		if (FirstValue.InterpMode == RCIM_Linear)
@@ -932,7 +1028,7 @@ void TMovieSceneCurveChannelImpl<ChannelType>::AutoSetTangents(ChannelType* InCh
 			const double NewTangent = (NextKey.Value - FirstValue.Value) / NextTimeDiff;
 			FirstValue.Tangent.LeaveTangent = NewTangent;
 		}
-		else if (FirstValue.InterpMode == RCIM_Cubic && FirstValue.TangentMode == RCTM_Auto)
+		else if (FirstValue.InterpMode == RCIM_Cubic && (FirstValue.TangentMode == RCTM_Auto || FirstValue.TangentMode == RCTM_SmartAuto))
 		{
 			FirstValue.Tangent.LeaveTangent = FirstValue.Tangent.ArriveTangent = 0.0f;
 			FirstValue.Tangent.TangentWeightMode = RCTWM_WeightedNone; 
@@ -950,27 +1046,54 @@ void TMovieSceneCurveChannelImpl<ChannelType>::AutoSetTangents(ChannelType* InCh
 			const double NewTangent = (LastValue.Value - PrevKey.Value) / PrevTimeDiff;
 			LastValue.Tangent.ArriveTangent = NewTangent;
 		}
-		else if (LastValue.InterpMode == RCIM_Cubic && LastValue.TangentMode == RCTM_Auto)
+		else if (LastValue.InterpMode == RCIM_Cubic && (LastValue.TangentMode == RCTM_Auto || LastValue.TangentMode == RCTM_SmartAuto))
 		{
 			LastValue.Tangent.LeaveTangent = LastValue.Tangent.ArriveTangent = 0.0f;
 			LastValue.Tangent.TangentWeightMode = RCTWM_WeightedNone;
 		}
 	}
-
+	const int32 MaxSmartAutoIterations = 10; // with smart auto we may need to iterate tangent calculations since it depends 
+											 // upon next tangent value which may be 0 since it is new.
+	int32 SmartAutoCount = 0;
 	for (int32 Index = 1; Index < InChannel->Values.Num() - 1; ++Index)
 	{
 		ChannelValueType  PrevKey = InChannel->Values[Index-1];
 		ChannelValueType& ThisKey = InChannel->Values[Index  ];
 
-		if (ThisKey.InterpMode == RCIM_Cubic && ThisKey.TangentMode == RCTM_Auto)
+		if (ThisKey.InterpMode == RCIM_Cubic && (ThisKey.TangentMode == RCTM_Auto || ThisKey.TangentMode == RCTM_SmartAuto))
 		{
-			ChannelValueType& NextKey = InChannel->Values[Index+1];
-			CurveValueType NewTangent = 0.0;
-			const double PrevToNextTimeDiff = FMath::Max<double>(KINDA_SMALL_NUMBER, InChannel->Times[Index + 1].Value - InChannel->Times[Index - 1].Value);
-
-			if (!UseNewAutoTangent)
+			ChannelValueType NextKey = InChannel->Values[Index + 1];
+			float NewTangent = 0.0f; //flat tangent by default
+			double PrevToNextTimeDiff = FMath::Max<double>(KINDA_SMALL_NUMBER, InChannel->Times[Index + 1].Value - InChannel->Times[Index - 1].Value);
+			if (ThisKey.TangentMode == RCTM_SmartAuto)
 			{
-				AutoCalcTangent(PrevKey.Value, ThisKey.Value, NextKey.Value, Tension, NewTangent);
+				NewTangent = CalcSmartTangent(InChannel, Index);
+				//okay we changed our tangent so we need to recalc the previous tangent
+				//but only if it's non zero
+				if (FMath::IsNearlyZero(ThisKey.Tangent.ArriveTangent,1e-3f) == false && FMath::IsNearlyEqual(ThisKey.Tangent.ArriveTangent, NewTangent,1e-3f) == false)
+				{
+					if (SmartAutoCount < MaxSmartAutoIterations && Index >= 2)
+					{
+						++SmartAutoCount;
+						Index -= 2; //go back 2 since we will increment by one
+					}
+					else
+					{
+						SmartAutoCount = 0;
+					}
+				}
+				else
+				{
+					SmartAutoCount = 0;
+				}
+			}
+			else if (!UseNewAutoTangent)
+			{
+				//need to pass in the curve value type since though unfortunately tangents are just always floats (for doubles or floats),
+				//the AutoCalcTangent works with either.
+				CurveValueType ValueNewTangent = 0.0;
+				AutoCalcTangent(PrevKey.Value, ThisKey.Value, NextKey.Value, Tension, ValueNewTangent);
+				NewTangent = ValueNewTangent;
 				NewTangent /= PrevToNextTimeDiff;
 			}
 			else
@@ -979,7 +1102,9 @@ void TMovieSceneCurveChannelImpl<ChannelType>::AutoSetTangents(ChannelType* InCh
 				if ( (ThisKey.Value > PrevKey.Value && ThisKey.Value < NextKey.Value) ||
 					(ThisKey.Value < PrevKey.Value && ThisKey.Value > NextKey.Value))
 				{
-					AutoCalcTangent(PrevKey.Value, ThisKey.Value, NextKey.Value, Tension, NewTangent);
+					CurveValueType ValueNewTangent = 0.0;
+					AutoCalcTangent(PrevKey.Value, ThisKey.Value, NextKey.Value, Tension, ValueNewTangent);
+					NewTangent = ValueNewTangent;
 					NewTangent /= PrevToNextTimeDiff;
 					//if within 0 to 15% or 85% to 100% range we gradually weight tangent to zero
 					const double AverageToZeroRange = 0.85;
@@ -1220,13 +1345,14 @@ FKeyHandle TMovieSceneCurveChannelImpl<ChannelType>::AddKeyToChannel(ChannelType
 		}
 		switch (Interpolation)
 		{
-
+			case EMovieSceneKeyInterpolation::SmartAuto:     ExistingIndex = InChannel->AddCubicKey(InFrameNumber, InValue, RCTM_SmartAuto);  break;
 			case EMovieSceneKeyInterpolation::Auto:     ExistingIndex = InChannel->AddCubicKey(InFrameNumber, InValue, RCTM_Auto);  break;
-			case EMovieSceneKeyInterpolation::User:     ExistingIndex = InChannel->AddCubicKey(InFrameNumber, InValue, RCTM_User, TangentData);  break;
-			case EMovieSceneKeyInterpolation::Break:    ExistingIndex = InChannel->AddCubicKey(InFrameNumber, InValue, RCTM_Break, TangentData); break;
+			case EMovieSceneKeyInterpolation::User:     ExistingIndex = InChannel->AddCubicKey(InFrameNumber, InValue, RCTM_User);  break;
+			case EMovieSceneKeyInterpolation::Break:    ExistingIndex = InChannel->AddCubicKey(InFrameNumber, InValue, RCTM_Break); break;
 			case EMovieSceneKeyInterpolation::Linear:   ExistingIndex = InChannel->AddLinearKey(InFrameNumber, InValue);            break;
 			case EMovieSceneKeyInterpolation::Constant: ExistingIndex = InChannel->AddConstantKey(InFrameNumber, InValue);          break;
 		}
+		
 	}
 
 	return InChannel->GetData().GetHandle(ExistingIndex);
diff --git a/Engine/Source/Runtime/MovieScene/Private/Channels/MovieSceneCurveChannelImpl.h b/Engine/Source/Runtime/MovieScene/Private/Channels/MovieSceneCurveChannelImpl.h
--- a/Engine/Source/Runtime/MovieScene/Private/Channels/MovieSceneCurveChannelImpl.h
+++ b/Engine/Source/Runtime/MovieScene/Private/Channels/MovieSceneCurveChannelImpl.h
@@ -79,6 +79,9 @@ struct MOVIESCENE_API TMovieSceneCurveChannelImpl
 	/** Adds a cubic interpolation key */
 	static int32 AddCubicKey(ChannelType* InChannel, FFrameNumber InTime, CurveValueType InValue, ERichCurveTangentMode TangentMode = RCTM_Auto, const FMovieSceneTangentData& Tangent = FMovieSceneTangentData());
 
+	/** Calculated smart tangent at that index*/
+	static float CalcSmartTangent(ChannelType* InChannel, int32 Index);
+
 	/** Auto-sets all tangents in the channel's curve */
 	static void AutoSetTangents(ChannelType* InChannel, float Tension = 0.f);
 
diff --git a/Engine/Source/Runtime/MovieScene/Public/KeyParams.h b/Engine/Source/Runtime/MovieScene/Public/KeyParams.h
--- a/Engine/Source/Runtime/MovieScene/Public/KeyParams.h
+++ b/Engine/Source/Runtime/MovieScene/Public/KeyParams.h
@@ -24,6 +24,9 @@ enum class EMovieSceneKeyInterpolation : uint8
 
 	/** Constant. */
 	Constant UMETA(DisplayName="Constant"),
+
+	/** Smart Auto. */
+	SmartAuto UMETA(DisplayName = "Smart Auto"),
 };
 
 /** Defines different modes for adding keyframes in sequencer. */
