Sequencer: Fixed erroneous removal of entry from map during iteration

diff --git a/Engine/Source/Runtime/MovieScene/Public/EntitySystem/MovieSceneOverlappingEntityTracker.h b/Engine/Source/Runtime/MovieScene/Public/EntitySystem/MovieSceneOverlappingEntityTracker.h
--- a/Engine/Source/Runtime/MovieScene/Public/EntitySystem/MovieSceneOverlappingEntityTracker.h
+++ b/Engine/Source/Runtime/MovieScene/Public/EntitySystem/MovieSceneOverlappingEntityTracker.h
@@ -594,10 +594,11 @@ struct TOverlappingEntityTracker_WithGarbage : TOverlappingEntityTrackerImpl<Out
 			if (OutputIndex != Super::NO_OUTPUT && Linker->EntityManager.HasComponent(EntityID, NeedsUnlink))
 			{
 				this->OutputToEntity.Remove(OutputIndex, EntityID);
-				this->EntityToOutput.Remove(EntityID);
 
 				this->InvalidatedOutputs.PadToNum(OutputIndex + 1, false);
 				this->InvalidatedOutputs[OutputIndex] = true;
+
+				EntityToOutputIt.RemoveCurrent();
 			}
 		}
 	}
