Sequencer: Do not create sub instance entities for sub sections that have no sequence assigned

diff --git a/Engine/Source/Runtime/MovieScene/Private/Sections/MovieSceneSubSection.cpp b/Engine/Source/Runtime/MovieScene/Private/Sections/MovieSceneSubSection.cpp
--- a/Engine/Source/Runtime/MovieScene/Private/Sections/MovieSceneSubSection.cpp
+++ b/Engine/Source/Runtime/MovieScene/Private/Sections/MovieSceneSubSection.cpp
@@ -188,9 +188,12 @@ void UMovieSceneSubSection::PostLoad()
 
 bool UMovieSceneSubSection::PopulateEvaluationFieldImpl(const TRange<FFrameNumber>& EffectiveRange, const FMovieSceneEvaluationFieldEntityMetaData& InMetaData, FMovieSceneEntityComponentFieldBuilder* OutFieldBuilder)
 {
-	const int32 EntityIndex   = OutFieldBuilder->FindOrAddEntity(this, 0);
-	const int32 MetaDataIndex = OutFieldBuilder->AddMetaData(InMetaData);
-	OutFieldBuilder->AddPersistentEntity(EffectiveRange, EntityIndex, MetaDataIndex);
+	if (SubSequence)
+	{
+		const int32 EntityIndex   = OutFieldBuilder->FindOrAddEntity(this, 0);
+		const int32 MetaDataIndex = OutFieldBuilder->AddMetaData(InMetaData);
+		OutFieldBuilder->AddPersistentEntity(EffectiveRange, EntityIndex, MetaDataIndex);
+	}
 
 	return true;
 }
