Sequencer: Fixed event graph tasks getting leaked for systems that belong to different phases

diff --git a/Engine/Source/Runtime/MovieScene/Private/EntitySystem/MovieSceneEntitySystemGraphs.cpp b/Engine/Source/Runtime/MovieScene/Private/EntitySystem/MovieSceneEntitySystemGraphs.cpp
--- a/Engine/Source/Runtime/MovieScene/Private/EntitySystem/MovieSceneEntitySystemGraphs.cpp
+++ b/Engine/Source/Runtime/MovieScene/Private/EntitySystem/MovieSceneEntitySystemGraphs.cpp
@@ -315,10 +315,10 @@ void FMovieSceneEntitySystemGraph::ExecutePhase(UE::MovieScene::ESystemPhase Pha
 
 	switch (Phase)
 	{
-	case UE::MovieScene::ESystemPhase::Spawn:         ExecutePhase(SpawnPhase,         Linker, OutTasks); break;
-	case UE::MovieScene::ESystemPhase::Instantiation: ExecutePhase(InstantiationPhase, Linker, OutTasks); break;
-	case UE::MovieScene::ESystemPhase::Evaluation:    ExecutePhase(EvaluationPhase,    Linker, OutTasks); break;
-	case UE::MovieScene::ESystemPhase::Finalization:  ExecutePhase(FinalizationPhase,  Linker, OutTasks); break;
+	case UE::MovieScene::ESystemPhase::Spawn:         ExecutePhase(Phase, SpawnPhase,         Linker, OutTasks); break;
+	case UE::MovieScene::ESystemPhase::Instantiation: ExecutePhase(Phase, InstantiationPhase, Linker, OutTasks); break;
+	case UE::MovieScene::ESystemPhase::Evaluation:    ExecutePhase(Phase, EvaluationPhase,    Linker, OutTasks); break;
+	case UE::MovieScene::ESystemPhase::Finalization:  ExecutePhase(Phase, FinalizationPhase,  Linker, OutTasks); break;
 	default: ensureMsgf(false, TEXT("Invalid phase specified for execution.")); break;
 	}
 }
@@ -344,7 +344,7 @@ void FMovieSceneEntitySystemGraph::IteratePhase(UE::MovieScene::ESystemPhase Pha
 }
 
 template<typename ArrayType>
-void FMovieSceneEntitySystemGraph::ExecutePhase(const ArrayType& SortedEntries, UMovieSceneEntitySystemLinker* Linker, FGraphEventArray& OutTasks)
+void FMovieSceneEntitySystemGraph::ExecutePhase(UE::MovieScene::ESystemPhase Phase, const ArrayType& SortedEntries, UMovieSceneEntitySystemLinker* Linker, FGraphEventArray& OutTasks)
 {
 	using namespace UE::MovieScene;
 
@@ -427,13 +427,19 @@ void FMovieSceneEntitySystemGraph::ExecutePhase(const ArrayType& SortedEntries,
 				if (ToNodeID)
 				{
 					FMovieSceneEntitySystemGraphNode& ToNode = Nodes.Array[*ToNodeID];
-					if (!ToNode.Prerequisites)
+					if (EnumHasAnyFlags(ToNode.System->GetPhase(), Phase))
 					{
-						ToNode.Prerequisites = MakeShared<FSystemTaskPrerequisites>();
+						if (!ToNode.Prerequisites)
+						{
+							ToNode.Prerequisites = MakeShared<FSystemTaskPrerequisites>();
+						}
+						ToNode.Prerequisites->Consume(*DownstreamTasks.Subsequents);
 					}
-					ToNode.Prerequisites->Consume(*DownstreamTasks.Subsequents);
 				}
 			}
+
+			// Done wtih subsequents now
+			DownstreamTasks.Subsequents->Empty();
 		}
 	}
 }
diff --git a/Engine/Source/Runtime/MovieScene/Public/EntitySystem/MovieSceneEntitySystemGraphs.h b/Engine/Source/Runtime/MovieScene/Public/EntitySystem/MovieSceneEntitySystemGraphs.h
--- a/Engine/Source/Runtime/MovieScene/Public/EntitySystem/MovieSceneEntitySystemGraphs.h
+++ b/Engine/Source/Runtime/MovieScene/Public/EntitySystem/MovieSceneEntitySystemGraphs.h
@@ -177,7 +177,7 @@ private:
 	void UpdateCache();
 
 	template<typename ArrayType>
-	void ExecutePhase(const ArrayType& RetrieveEntries, UMovieSceneEntitySystemLinker* Linker, FGraphEventArray& OutTasks);
+	void ExecutePhase(UE::MovieScene::ESystemPhase Phase, const ArrayType& RetrieveEntries, UMovieSceneEntitySystemLinker* Linker, FGraphEventArray& OutTasks);
 
 private:
 	friend UE::MovieScene::FSystemSubsequentTasks;
