Sequencer: remove (probably) useless restrictions on property types, and remove unused size/align info

diff --git a/Engine/Source/Runtime/MovieScene/Public/EntitySystem/MovieScenePropertyRegistry.h b/Engine/Source/Runtime/MovieScene/Public/EntitySystem/MovieScenePropertyRegistry.h
--- a/Engine/Source/Runtime/MovieScene/Public/EntitySystem/MovieScenePropertyRegistry.h
+++ b/Engine/Source/Runtime/MovieScene/Public/EntitySystem/MovieScenePropertyRegistry.h
@@ -66,13 +66,12 @@ struct FPropertyDefinition
 	FPropertyDefinition() = default;
 
 	FPropertyDefinition(
-			uint16 InVariableSizeCompositeOffset, uint16 InSizeofStorageType, uint16 InAlignofStorageType,
+			uint16 InVariableSizeCompositeOffset,
 			FComponentTypeID InPropertyType, FComponentTypeID InInitialValueType)
 		: CustomPropertyRegistration(nullptr)
 		, DoubleCompositeMask(0)
 		, VariableSizeCompositeOffset(InVariableSizeCompositeOffset)
 		, CompositeSize(0)
-		, StorageType{ InSizeofStorageType, InAlignofStorageType }
 		, PropertyType(InPropertyType)
 		, InitialValueType(InInitialValueType)
 		, BlenderSystemClass(nullptr)
@@ -108,13 +107,6 @@ struct FPropertyDefinition
 	/** The number of channels that this property comprises */
 	uint16 CompositeSize = 0;
 
-	/** Operational type meta-data */
-	struct
-	{
-		uint16 Sizeof = 0;
-		uint16 Alignof = 0;
-	} StorageType;
-
 	/** The component type or tag of the property itself */
 	FComponentTypeID PropertyType;
 
@@ -265,7 +257,6 @@ private:
 	void DefinePropertyImpl(TPropertyComponents<PropertyTraits>& InOutPropertyComponents, const TCHAR* InStatName)
 	{
 		using StorageType = typename PropertyTraits::StorageType;
-		static_assert( TIsBitwiseConstructible<StorageType, StorageType>::Value && TIsTriviallyDestructible<StorageType>::Value, "StorageType must be trivially TIsTriviallyCopyConstructible" );
 
 		const int32 CompositeOffset = CompositeDefinitions.Num();
 		checkf(CompositeOffset <= MAX_uint16, TEXT("Maximum number of composite definitions reached"));
@@ -292,7 +283,6 @@ private:
 
 		FPropertyDefinition NewDefinition(
 			CompositeOffset, 
-			sizeof(StorageType), alignof(StorageType),
 			InOutPropertyComponents.PropertyTag,
 			InOutPropertyComponents.InitialValue);
 
