Sequencer: Easing is no longer forcibly 0 and 1 when t=0 and 1 respectively

diff --git a/Engine/Source/Runtime/MovieScene/Private/MovieSceneSection.cpp b/Engine/Source/Runtime/MovieScene/Private/MovieSceneSection.cpp
--- a/Engine/Source/Runtime/MovieScene/Private/MovieSceneSection.cpp
+++ b/Engine/Source/Runtime/MovieScene/Private/MovieSceneSection.cpp
@@ -710,11 +710,11 @@ float UMovieSceneSection::EvaluateEasing(FFrameTime InTime) const
 		const int32  EaseFrame    = (InTime.FrameNumber - GetInclusiveStartFrame()).Value;
 		const double EaseInInterp = (double(EaseFrame) + InTime.GetSubFrame()) / Easing.GetEaseInDuration();
 
-		if (EaseInInterp <= 0.0)
+		if (EaseInInterp < 0.0)
 		{
 			EaseInValue = 0.0;
 		}
-		else if (EaseInInterp >= 1.0)
+		else if (EaseInInterp > 1.0)
 		{
 			EaseInValue = 1.0;
 		}
@@ -730,11 +730,11 @@ float UMovieSceneSection::EvaluateEasing(FFrameTime InTime) const
 		const int32  EaseFrame     = (InTime.FrameNumber - GetExclusiveEndFrame() + Easing.GetEaseOutDuration()).Value;
 		const double EaseOutInterp = (double(EaseFrame) + InTime.GetSubFrame()) / Easing.GetEaseOutDuration();
 
-		if (EaseOutInterp <= 0.0)
+		if (EaseOutInterp < 0.0)
 		{
 			EaseOutValue = 1.0;
 		}
-		else if (EaseOutInterp >= 1.0)
+		else if (EaseOutInterp > 1.0)
 		{
 			EaseOutValue = 0.0;
 		}
