Sequencer: methods to find sections at a given time can now limit the search to a given row

diff --git a/Engine/Source/Runtime/MovieScene/Private/MovieSceneCommonHelpers.cpp b/Engine/Source/Runtime/MovieScene/Private/MovieSceneCommonHelpers.cpp
--- a/Engine/Source/Runtime/MovieScene/Private/MovieSceneCommonHelpers.cpp
+++ b/Engine/Source/Runtime/MovieScene/Private/MovieSceneCommonHelpers.cpp
@@ -34,14 +34,15 @@ bool MovieSceneHelpers::IsSectionKeyable(const UMovieSceneSection* Section)
 	return !Track->IsRowEvalDisabled(Section->GetRowIndex()) && !Track->IsEvalDisabled() && Section->IsActive();
 }
 
-UMovieSceneSection* MovieSceneHelpers::FindSectionAtTime( TArrayView<UMovieSceneSection* const> Sections, FFrameNumber Time )
+UMovieSceneSection* MovieSceneHelpers::FindSectionAtTime( TArrayView<UMovieSceneSection* const> Sections, FFrameNumber Time, int32 RowIndex )
 {
 	for( int32 SectionIndex = 0; SectionIndex < Sections.Num(); ++SectionIndex )
 	{
 		UMovieSceneSection* Section = Sections[SectionIndex];
 
 		//@todo sequencer: There can be multiple sections overlapping in time. Returning instantly does not account for that.
-		if( Section->IsTimeWithinSection( Time ) && IsSectionKeyable(Section) )
+		if( (RowIndex == INDEX_NONE || Section->GetRowIndex() == RowIndex) &&
+				Section->IsTimeWithinSection( Time ) && IsSectionKeyable(Section) )
 		{
 			return Section;
 		}
@@ -50,12 +51,13 @@ UMovieSceneSection* MovieSceneHelpers::FindSectionAtTime( TArrayView<UMovieScene
 	return nullptr;
 }
 
-UMovieSceneSection* MovieSceneHelpers::FindNearestSectionAtTime( TArrayView<UMovieSceneSection* const> Sections, FFrameNumber Time )
+UMovieSceneSection* MovieSceneHelpers::FindNearestSectionAtTime( TArrayView<UMovieSceneSection* const> Sections, FFrameNumber Time, int32 RowIndex )
 {
 	TArray<UMovieSceneSection*> OverlappingSections, NonOverlappingSections;
 	for (UMovieSceneSection* Section : Sections)
 	{
-		if (IsSectionKeyable(Section))
+		if ((RowIndex == INDEX_NONE || Section->GetRowIndex() == RowIndex) &&
+				IsSectionKeyable(Section))
 		{
 			if (Section->GetRange().Contains(Time))
 			{
@@ -1045,4 +1047,4 @@ template<> void FTrackInstancePropertyBindings::SetCurrentValue<UObject*>(UObjec
 	{
 		InRuntimeObject.ProcessEvent(NotifyFunction, nullptr);
 	}
-}
\ No newline at end of file
+}
diff --git a/Engine/Source/Runtime/MovieScene/Public/MovieSceneCommonHelpers.h b/Engine/Source/Runtime/MovieScene/Public/MovieSceneCommonHelpers.h
--- a/Engine/Source/Runtime/MovieScene/Public/MovieSceneCommonHelpers.h
+++ b/Engine/Source/Runtime/MovieScene/Public/MovieSceneCommonHelpers.h
@@ -51,17 +51,19 @@ public:
 	 * Finds a section that exists at a given time
 	 *
 	 * @param Time	The time to find a section at
+	 * @param RowIndex  Limit the search to a given row index
 	 * @return The found section or null
 	 */
-	static UMovieSceneSection* FindSectionAtTime( TArrayView<UMovieSceneSection* const> Sections, FFrameNumber Time );
+	static UMovieSceneSection* FindSectionAtTime( TArrayView<UMovieSceneSection* const> Sections, FFrameNumber Time, int32 RowIndex = INDEX_NONE );
 
 	/**
 	 * Finds the nearest section to the given time
 	 *
 	 * @param Time	The time to find a section at
+	 * @param RowIndex  Limit the search to a given row index
 	 * @return The found section or null
 	 */
-	static UMovieSceneSection* FindNearestSectionAtTime( TArrayView<UMovieSceneSection* const> Sections, FFrameNumber Time );
+	static UMovieSceneSection* FindNearestSectionAtTime( TArrayView<UMovieSceneSection* const> Sections, FFrameNumber Time, int32 RowIndex = INDEX_NONE );
 
 	/** Find the next section that doesn't overlap - the section that has the next closest start time to the requested start time */
 	static UMovieSceneSection* FindNextSection(TArrayView<UMovieSceneSection* const> Sections, FFrameNumber Time);
