Sequencer: Remove ensure and early out if override root field is null

#rnx

diff --git a/Engine/Source/Runtime/MovieScene/Private/Evaluation/Instances/MovieSceneTrackEvaluator.cpp b/Engine/Source/Runtime/MovieScene/Private/Evaluation/Instances/MovieSceneTrackEvaluator.cpp
--- a/Engine/Source/Runtime/MovieScene/Private/Evaluation/Instances/MovieSceneTrackEvaluator.cpp
+++ b/Engine/Source/Runtime/MovieScene/Private/Evaluation/Instances/MovieSceneTrackEvaluator.cpp
@@ -195,7 +195,7 @@ const FMovieSceneEvaluationGroup* FMovieSceneTrackEvaluator::SetupFrame(UMovieSc
 		}
 	}
 
-	if (!ensureMsgf(OverrideRootField, TEXT("Could not find valid evaluation field for supplied sequence ID.")))
+	if (OverrideRootField == nullptr)
 	{
 		return nullptr;
 	}
