Sequencer: Added custom, persistent scheduler for dispatching ECS tasks

The aim of this work is to improve overall CPU cost in 3 ways:

  - to reduce the per-frame overhead of running ECS tasks
  - to decrease latency between useful gamethread work being able to run
  - to increase thread ultilization where many cores are available

In both single-threaded and multi-threaded mode this is showing a wall-clock speed improvement of ~25%, upon which other optimizations can now be made.

Previously, the Dispatch_ methods would construct tasks and filters for their work every frame. Tasks dispatched by downstream systems that were added as implicit dependencies would always depend on all the work of their upstream systems, even if they didn't necessarily need to be (for example, float properties would need to wait for *all* float channels to be evaluated and blended before they could be applied, even ones that only applied to transform properties).

The new approach introduces a new system phase (ESystemPhase::Scheduling) which is only run as part of the post-instantiation. This phase invokes a new method (OnSchedulePersistentTasks) which allows systems to define the task-structure for the current state of the ECS. Tasks defined here are stored inside the scheduler, and live for the duration of the instantiation - if the structure of the entity manager materially changes, the task graph is re-constructed.

Entity task builders are still used in this new method, but can only call Schedule_ and Fork_ functions for running tasks. These methods define serial and parallel behavior respectively that is to be run across the ECS (as should already be familiar). Crucially, each relevant entity allocation now gets its own task - the filter is run when the task graph is constructed, and a new task is created for each match. This allows us to interleave and/or distribute tasks across as many cores as possible, as soon as prerequisites are completed.

Tasks may also 'prelock' their component data - doing so will cache the location of the component data for each task (using a new type, TRelativePtr) so that it doesn't have to go hunting for component headers when the task is run.

Implementation details:
    IMovieSceneTaskScheduler exists as the public front-end to the new routines, but allows us to keep all the implementation details private.

    Tasks can be in one of two states: a single task can have multiple prerequisites and subsequents, a parent task can have many children which must be complete in order for the parent to be considered complete. This is used for depending on the results of forked and serial tasks that may be spawned multiple times.

    When constructing the graph, we keep track of the read and write dependencies for each component type within each allocation. When a task reads from a component it must first depend on any tasks that write to the same component on the same allocation. When a task writes to a component it must depend on all upstream read and write dependencies for that component, on that allocation.

    Moved TSparseBitSet to UE::MovieScene so I can use it for task dependencies

    Introduced TRelativePtr which allows storage of pointers relative to another pointer using a smaller type than a raw ptr would, and isn't invalidated like a raw ptr would be if is pointing to an element in a dynamic array.
    This is used for our prelocked component data, but may have uses elsewhere. Note that this is not a TRelativePtr that is relative to 'this' like some other relative pointer implementations (though one could easily be written that uses this implementation if necessary.

diff --git a/Engine/Source/Editor/SequencerCore/Private/MVVM/ICastable.cpp b/Engine/Source/Editor/SequencerCore/Private/MVVM/ICastable.cpp
--- a/Engine/Source/Editor/SequencerCore/Private/MVVM/ICastable.cpp
+++ b/Engine/Source/Editor/SequencerCore/Private/MVVM/ICastable.cpp
@@ -44,8 +44,14 @@ const void* ICastable::CastRaw(FViewModelTypeID InType) const
 	return Result;
 }
 
+FCastableTypeTable::FCastableTypeTable(uint8* InTypeMaskStorage)
+	: TypeMask(InTypeMaskStorage)
+{}
+
 FCastableTypeTable* FCastableTypeTable::FCastableTypeTableGenerator::Commit(uint32 TypeID, FName InName)
 {
+	using namespace UE::MovieScene;
+
 	uint8* NewType = (uint8*)FMemory::Malloc(ByteSize(), alignof(FCastableTypeTable));
 
 	// FCastableTypeTable is a special type that uses a fixed allocation size determined by the size of the type table itself
@@ -56,14 +62,12 @@ FCastableTypeTable* FCastableTypeTable::FCastableTypeTableGenerator::Commit(uint
 	//	|	StaticTypeOffsets	|		  TypeMask			|	NumStaticTypes	|	ThisTypeID		|	ThisTypeName	|											|		 										|
 	//				\_________________________\________________________________________________________________________/										   /
 	//										   \______________________________________________TypeMask.Buckets.Storage____________________________________________/
-	FCastableTypeTable* NewTable = new (NewType) FCastableTypeTable();
+	uint8* DataOffset = NewType + sizeof(FCastableTypeTable);
+
+	FCastableTypeTable* NewTable = new (NewType) FCastableTypeTable(DataOffset);
 	NewTable->ThisTypeID = TypeID;
 	NewTable->ThisTypeName = InName;
 
-	uint8* DataOffset = NewType + sizeof(FCastableTypeTable);
-
-	// Sparse bitset storage lives at the end of the table
-	NewTable->TypeMask = TSparseBitSet<uint32, TFixedSparseBitSetBucketStorage<uint8>>({ DataOffset });
 	TypeMask.CopyToUnsafe(NewTable->TypeMask, TypeMask.NumBuckets());
 
 	DataOffset += TypeMask.NumBuckets()*sizeof(decltype(NewTable->TypeMask)::BucketType);
diff --git a/Engine/Source/Editor/SequencerCore/Public/MVVM/CastableTypeTable.h b/Engine/Source/Editor/SequencerCore/Public/MVVM/CastableTypeTable.h
--- a/Engine/Source/Editor/SequencerCore/Public/MVVM/CastableTypeTable.h
+++ b/Engine/Source/Editor/SequencerCore/Public/MVVM/CastableTypeTable.h
@@ -5,7 +5,7 @@
 #include "CoreTypes.h"
 #include "HAL/Platform.h"
 #include "Templates/UnrealTemplate.h"
-#include "MVVM/SparseBitSet.h"
+#include "Containers/SparseBitSet.h"
 #include "UObject/NameTypes.h"
 
 #include <type_traits>
@@ -24,7 +24,7 @@ class IDynamicExtension;
 
 struct FCastableTypeTable
 {
-	FCastableTypeTable() = default;
+	explicit FCastableTypeTable(uint8* InTypeMaskStorage);
 
 	FCastableTypeTable(const FCastableTypeTable&) = delete;
 	void operator=(const FCastableTypeTable&) = delete;
@@ -209,6 +209,8 @@ private:
 		template<typename StartClass, typename T, typename U = decltype(T::ID)>
 		void AddStaticType(T*, int16 StartOffset)
 		{
+			using namespace UE::MovieScene;
+
 			TViewModelTypeID<T>& ID = T::ID;
 
 			int16 OffsetFromBase = StartOffset + ClassOffsetFrom<StartClass, T>();
@@ -227,7 +229,7 @@ private:
 	private:
 
 		/* Accumulator sparse bitset of types that are implemented on the type that Generate was called with */
-		TSparseBitSet<uint32> TypeMask;
+		MovieScene::TSparseBitSet<uint32> TypeMask;
 		/* Typeoffsets from T to the ptr for each type within TypeMask */
 		TArray<int16> StaticTypeOffsets;
 	};
@@ -237,7 +239,7 @@ private:
 	/** Ptr to the static type offsets allocated at the end of this struct. Size = NumStaticTypes */
 	int16* StaticTypeOffsets;
 	/** Fixed sparse bitset of types that this table contains. Storage is allocated at the end of this struct. */
-	TSparseBitSet<uint32, TFixedSparseBitSetBucketStorage<uint8>> TypeMask;
+	MovieScene::TSparseBitSet<uint32, MovieScene::TFixedSparseBitSetBucketStorage<uint8>> TypeMask;
 	/** The number of static types implemented within this type table. */
 	uint32 NumStaticTypes;
 	/** Unique static type ID for this type. */
diff --git a/Engine/Source/Runtime/MovieScene/Private/EntitySystem/MovieSceneEntityFactory.cpp b/Engine/Source/Runtime/MovieScene/Private/EntitySystem/MovieSceneEntityFactory.cpp
--- a/Engine/Source/Runtime/MovieScene/Private/EntitySystem/MovieSceneEntityFactory.cpp
+++ b/Engine/Source/Runtime/MovieScene/Private/EntitySystem/MovieSceneEntityFactory.cpp
@@ -114,14 +114,22 @@ void FObjectFactoryBatch::InitializeAllocation(UMovieSceneEntitySystemLinker* Li
 
 	FEntityAllocationWriteContext WriteContext = FEntityAllocationWriteContext::NewAllocation();
 
+	const FEntityAllocation* Allocation = nullptr;
+	int32 ComponentStartOffset = 0;
+	int32 Num = 0;
+
+	TArrayView<const FMovieSceneEntityID> ChildEntityIDs        = InChildEntityRange.Allocation->GetEntityIDs();
+	TComponentReader<FMovieSceneEntityID> ParentIDComponents    = InChildEntityRange.Allocation->ReadComponents(ParentEntity);
+	TComponentWriter<UObject*>            BoundObjectComponents = InChildEntityRange.Allocation->WriteComponents(BoundObject, WriteContext);
+
 	int32 Index = GetCurrentIndex();
-	for (TEntityPtr<const FMovieSceneEntityID, const FMovieSceneEntityID, UObject*> Tuple : FEntityTaskBuilder().ReadEntityIDs().Read(ParentEntity).Write(BoundObject).IterateRange(InChildEntityRange, WriteContext))
+	for (int32 ChildIndex = InChildEntityRange.ComponentStartOffset; ChildIndex < InChildEntityRange.ComponentStartOffset + InChildEntityRange.Num; ++ChildIndex)
 	{
-		FMovieSceneEntityID Parent = Tuple.Get<1>();
-		FMovieSceneEntityID Child = Tuple.Get<0>();
+		FMovieSceneEntityID Parent = ParentIDComponents[ChildIndex];
+		FMovieSceneEntityID Child  = ChildEntityIDs[ChildIndex];
 
 		UObject* Object = ObjectsToAssign[Index++];
-		Tuple.Get<2>() = Object;
+		BoundObjectComponents[ChildIndex] = Object;
 
 		if (FMovieSceneEntityID OldEntityToPreserve = StaleEntitiesToPreserve->FindRef(MakeTuple(Object, Parent)))
 		{
diff --git a/Engine/Source/Runtime/MovieScene/Private/EntitySystem/MovieSceneEntityManager.cpp b/Engine/Source/Runtime/MovieScene/Private/EntitySystem/MovieSceneEntityManager.cpp
--- a/Engine/Source/Runtime/MovieScene/Private/EntitySystem/MovieSceneEntityManager.cpp
+++ b/Engine/Source/Runtime/MovieScene/Private/EntitySystem/MovieSceneEntityManager.cpp
@@ -229,6 +229,7 @@ struct FEntityInitializer
 					uint8 Alignment = FMath::Max<uint8>(PLATFORM_CACHE_LINE_SIZE, TypeInfo.Alignment);
 					ComponentDataPtr = Align(ComponentDataPtr, Alignment);
 
+					Header->ScheduledAccessCount.exchange(0, std::memory_order_relaxed);
 					Header->Components = reinterpret_cast<uint8*>(ComponentDataPtr);
 
 					check(IsAligned(Header->Components, TypeInfo.Alignment));
@@ -1094,7 +1095,7 @@ int32 FEntityManager::GetOrCreateAllocationWithSlack(const FComponentMask& Entit
 
 int32 FEntityManager::CreateAllocationWithSlack(const FComponentMask& EntityComponentMask, int32* InOutDesiredSlack)
 {
-	static const uint16 MaxCapacity = 2048;
+	static const uint16 MaxCapacity = 64;
 
 	uint16 DefaultCapacity = 4;
 	if (InOutDesiredSlack)
diff --git a/Engine/Source/Runtime/MovieScene/Private/EntitySystem/MovieSceneEntitySystem.cpp b/Engine/Source/Runtime/MovieScene/Private/EntitySystem/MovieSceneEntitySystem.cpp
--- a/Engine/Source/Runtime/MovieScene/Private/EntitySystem/MovieSceneEntitySystem.cpp
+++ b/Engine/Source/Runtime/MovieScene/Private/EntitySystem/MovieSceneEntitySystem.cpp
@@ -687,6 +687,25 @@ void UMovieSceneEntitySystem::FinishDestroy()
 	Super::FinishDestroy();
 }
 
+void UMovieSceneEntitySystem::SchedulePersistentTasks(UE::MovieScene::IEntitySystemScheduler* TaskScheduler)
+{
+	if (!bSystemIsEnabled)
+	{
+		return;
+	}
+
+	checkf(Linker != nullptr, TEXT("Attempting to evaluate a system that has been unlinked!"));
+
+	// We may have erroneously linked a system we should have done, but we must not run it in this case
+	if (!Linker->GetSystemFilter().CheckSystem(this))
+	{
+		return;
+	}
+
+	Linker->EntityManager.IncrementSystemSerial();
+	OnSchedulePersistentTasks(TaskScheduler);
+}
+
 void UMovieSceneEntitySystem::Run(FSystemTaskPrerequisites& InPrerequisites, FSystemSubsequentTasks& Subsequents)
 {
 	if (!bSystemIsEnabled)
diff --git a/Engine/Source/Runtime/MovieScene/Private/EntitySystem/MovieSceneEntitySystemGraphs.cpp b/Engine/Source/Runtime/MovieScene/Private/EntitySystem/MovieSceneEntitySystemGraphs.cpp
--- a/Engine/Source/Runtime/MovieScene/Private/EntitySystem/MovieSceneEntitySystemGraphs.cpp
+++ b/Engine/Source/Runtime/MovieScene/Private/EntitySystem/MovieSceneEntitySystemGraphs.cpp
@@ -4,6 +4,7 @@
 #include "EntitySystem/MovieSceneEntitySystem.h"
 #include "EntitySystem/MovieSceneEntitySystemLinker.h"
 #include "EntitySystem/MovieSceneSystemTaskDependencies.h"
+#include "EntitySystem/MovieSceneTaskScheduler.h"
 
 #include "Templates/SubclassOf.h"
 #include "Algo/IndexOf.h"
@@ -21,6 +22,12 @@ void FMovieSceneEntitySystemGraphNodes::AddStructReferencedObjects(FReferenceCol
 	}
 }
 
+FMovieSceneEntitySystemGraph::FMovieSceneEntitySystemGraph() = default;
+FMovieSceneEntitySystemGraph::~FMovieSceneEntitySystemGraph() = default;
+
+FMovieSceneEntitySystemGraph::FMovieSceneEntitySystemGraph(FMovieSceneEntitySystemGraph&&) = default;
+FMovieSceneEntitySystemGraph& FMovieSceneEntitySystemGraph::operator=(FMovieSceneEntitySystemGraph&&) = default;
+
 void FMovieSceneEntitySystemGraph::AddSystem(UMovieSceneEntitySystem* InSystem)
 {
 	checkf(ReentrancyGuard == 0, TEXT("Attempting to add a system to the graph recursively."));
@@ -161,6 +168,7 @@ void FMovieSceneEntitySystemGraph::UpdateCache()
 
 	SpawnPhase.Empty();
 	InstantiationPhase.Empty();
+	SchedulingPhase.Empty();
 	EvaluationPhase.Empty();
 	FinalizationPhase.Empty();
 
@@ -177,6 +185,8 @@ void FMovieSceneEntitySystemGraph::UpdateCache()
 		SortedNodeIDs.Add(GlobalToLocalNodeIDs[GlobalNodeID]);
 	}
 
+	const bool bCombineSchedulingAndEvaluation = !FEntitySystemScheduler::IsCustomSchedulingEnabled();
+
 	for (uint16 NodeID : SortedNodeIDs)
 	{
 		ESystemPhase SystemPhase = Nodes.Array[NodeID].System->GetPhase();
@@ -189,6 +199,17 @@ void FMovieSceneEntitySystemGraph::UpdateCache()
 		{
 			InstantiationPhase.Emplace(NodeID);
 		}
+		if (EnumHasAnyFlags(SystemPhase, ESystemPhase::Scheduling))
+		{
+			if (bCombineSchedulingAndEvaluation)
+			{
+				EvaluationPhase.Emplace(NodeID);
+			}
+			else
+			{
+				SchedulingPhase.Emplace(NodeID);
+			}
+		}
 		if (EnumHasAnyFlags(SystemPhase, ESystemPhase::Evaluation))
 		{
 			EvaluationPhase.Emplace(NodeID);
@@ -309,6 +330,18 @@ UMovieSceneEntitySystem* FMovieSceneEntitySystemGraph::FindSystemOfType(TSubclas
 	return nullptr;
 }
 
+int32 FMovieSceneEntitySystemGraph::NumInPhase(UE::MovieScene::ESystemPhase Phase) const
+{
+	switch (Phase)
+	{
+	case UE::MovieScene::ESystemPhase::Spawn:         return SpawnPhase.Num();
+	case UE::MovieScene::ESystemPhase::Instantiation: return InstantiationPhase.Num();
+	case UE::MovieScene::ESystemPhase::Evaluation:    return EvaluationPhase.Num();
+	case UE::MovieScene::ESystemPhase::Finalization:  return FinalizationPhase.Num();
+	default: ensureMsgf(false, TEXT("Invalid phase specified for execution.")); return 0;
+	}
+}
+
 void FMovieSceneEntitySystemGraph::ExecutePhase(UE::MovieScene::ESystemPhase Phase, UMovieSceneEntitySystemLinker* Linker, FGraphEventArray& OutTasks)
 {
 	UpdateCache();
@@ -332,6 +365,7 @@ void FMovieSceneEntitySystemGraph::IteratePhase(UE::MovieScene::ESystemPhase Pha
 	{
 	case UE::MovieScene::ESystemPhase::Spawn:         Array = SpawnPhase;         break;
 	case UE::MovieScene::ESystemPhase::Instantiation: Array = InstantiationPhase; break;
+	case UE::MovieScene::ESystemPhase::Scheduling:    Array = SchedulingPhase;    break;
 	case UE::MovieScene::ESystemPhase::Evaluation:    Array = EvaluationPhase;    break;
 	case UE::MovieScene::ESystemPhase::Finalization:  Array = FinalizationPhase;  break;
 	default: ensureMsgf(false, TEXT("Invalid phase specified for iteration."));   return;
@@ -427,7 +461,9 @@ void FMovieSceneEntitySystemGraph::ExecutePhase(UE::MovieScene::ESystemPhase Pha
 				if (ToNodeID)
 				{
 					FMovieSceneEntitySystemGraphNode& ToNode = Nodes.Array[*ToNodeID];
-					if (EnumHasAnyFlags(ToNode.System->GetPhase(), Phase))
+					if (EnumHasAnyFlags(ToNode.System->GetPhase(), Phase)
+						|| (Phase == ESystemPhase::Evaluation && EnumHasAnyFlags(ToNode.System->GetPhase(), ESystemPhase::Scheduling))
+						)
 					{
 						if (!ToNode.Prerequisites)
 						{
@@ -444,6 +480,78 @@ void FMovieSceneEntitySystemGraph::ExecutePhase(UE::MovieScene::ESystemPhase Pha
 	}
 }
 
+void FMovieSceneEntitySystemGraph::ReconstructTaskSchedule(UE::MovieScene::FEntityManager* EntityManager)
+{
+	using namespace UE::MovieScene;
+
+	if (!FEntitySystemScheduler::IsCustomSchedulingEnabled())
+	{
+		// When disabled, we combine the scheduling phase with the evaluation phase and
+		// execute them using the legacy procedural task prerequisites API through OnRun
+
+		return;
+	}
+
+	UpdateCache();
+
+	if (!TaskScheduler)
+	{
+		TaskScheduler = MakeUnique<FEntitySystemScheduler>(EntityManager);
+	}
+
+	TaskScheduler->BeginConstruction();
+
+	for (int32 CurrentIndex = 0; CurrentIndex < SchedulingPhase.Num(); ++CurrentIndex)
+	{
+		const uint16 NodeID = SchedulingPhase[CurrentIndex];
+
+		UMovieSceneEntitySystem* System = Nodes.Array[NodeID].System;
+		checkSlow(System);
+
+		// Initilaize downstream task structure for this system
+		TaskScheduler->BeginSystem(NodeID);
+
+		System->SchedulePersistentTasks(TaskScheduler.Get());
+
+		// Propagate subsequent tasks
+		if (TaskScheduler->HasAnyTasksToPropagateDownstream())
+		{
+			TArray<uint16> ToGlobalNodeIDs;
+			UMovieSceneEntitySystem::GetSubsequentSystems(System->GetGlobalDependencyGraphID(), ToGlobalNodeIDs);
+			for (uint16 ToGlobalNodeID : ToGlobalNodeIDs)
+			{
+				uint16* ToNodeID = GlobalToLocalNodeIDs.Find(ToGlobalNodeID);
+				if (ToNodeID)
+				{
+					TaskScheduler->PropagatePrerequisite(*ToNodeID);
+				}
+			}
+		}
+
+		TaskScheduler->EndSystem(NodeID);
+	}
+
+	TaskScheduler->EndConstruction();
+
+	SchedulerSerialNumber = EntityManager->GetSystemSerial();
+}
+
+void FMovieSceneEntitySystemGraph::ScheduleTasks(UE::MovieScene::FEntityManager* EntityManager)
+{
+	// @todo: First off go through and increase the WriteContexts?
+	if (TaskScheduler)
+	{
+		EntityManager->IncrementSystemSerial();
+
+		if (EntityManager->HasStructureChangedSince(SchedulerSerialNumber))
+		{
+			ReconstructTaskSchedule(EntityManager);
+		}
+
+		TaskScheduler->ExecuteTasks();
+	}
+}
+
 void FMovieSceneEntitySystemGraph::Shutdown()
 {
 	checkf(ReentrancyGuard == 0, TEXT("Attempting to shutdown the system graph while it is in use."));
diff --git a/Engine/Source/Runtime/MovieScene/Private/EntitySystem/MovieSceneEntitySystemRunner.cpp b/Engine/Source/Runtime/MovieScene/Private/EntitySystem/MovieSceneEntitySystemRunner.cpp
--- a/Engine/Source/Runtime/MovieScene/Private/EntitySystem/MovieSceneEntitySystemRunner.cpp
+++ b/Engine/Source/Runtime/MovieScene/Private/EntitySystem/MovieSceneEntitySystemRunner.cpp
@@ -5,6 +5,7 @@
 #include "EntitySystem/MovieSceneEntityMutations.h"
 #include "EntitySystem/BuiltInComponentTypes.h"
 #include "EntitySystem/MovieSceneEntitySystemTypes.h"
+#include "EntitySystem/MovieSceneTaskScheduler.h"
 #include "Evaluation/PreAnimatedState/MovieScenePreAnimatedCaptureSource.h"
 #include "Evaluation/MovieSceneEvaluationTemplateInstance.h"
 #include "IMovieScenePlayer.h"
@@ -898,6 +899,11 @@ bool FMovieSceneEntitySystemRunner::GameThread_PostInstantiation(UMovieSceneEnti
 	Linker->AutoUnlinkIrrelevantSystems();
 
 	EntityManager.Compact();
+
+	if (FEntitySystemScheduler::IsCustomSchedulingEnabled())
+	{
+		Linker->SystemGraph.ReconstructTaskSchedule(&Linker->EntityManager);
+	}
 	return true;
 }
 
@@ -909,7 +915,6 @@ bool FMovieSceneEntitySystemRunner::GameThread_EvaluationPhase(UMovieSceneEntity
 
 	CurrentPhase = ESystemPhase::Evaluation;
 
-	FGraphEventArray AllTasks;
 	// --------------------------------------------------------------------------------------------------------------------------------------------
 	// Step 2: Run the evaluation phase. The entity manager is locked down for this phase, meaning no changes to entity-component structure is allowed
 	//         This vastly simplifies the concurrent handling of entity component allocations
@@ -917,6 +922,12 @@ bool FMovieSceneEntitySystemRunner::GameThread_EvaluationPhase(UMovieSceneEntity
 
 	checkf(!Linker->EntityManager.ContainsComponent(FBuiltInComponentTypes::Get()->Tags.NeedsUnlink), TEXT("Stale entities remain in the entity manager during evaluation - these should have been destroyed during the instantiation phase. Did it run?"));
 
+	if (FEntitySystemScheduler::IsCustomSchedulingEnabled())
+	{
+		Linker->SystemGraph.ScheduleTasks(&Linker->EntityManager);
+	}
+
+	FGraphEventArray AllTasks;
 	Linker->SystemGraph.ExecutePhase(ESystemPhase::Evaluation, Linker, AllTasks);
 
 	if (AllTasks.Num() != 0)
@@ -929,6 +940,14 @@ bool FMovieSceneEntitySystemRunner::GameThread_EvaluationPhase(UMovieSceneEntity
 
 	Linker->EntityManager.ReleaseLockDown();
 
+	if ( !EnumHasAnyFlags(AccumulatedUpdateFlags, ESequenceInstanceUpdateFlags::HasLegacyTemplates)
+		 && Linker->SystemGraph.NumInPhase(ESystemPhase::Finalization) == 0
+		 && !EventTriggers.IsBound())
+	{
+		// Skip Finalization if there's no need for it
+		SkipFlushState(ERunnerFlushState::Finalization | ERunnerFlushState::EventTriggers);
+	}
+
 	return true;
 }
 
diff --git a/Engine/Source/Runtime/MovieScene/Private/EntitySystem/MovieSceneEvalTimeSystem.cpp b/Engine/Source/Runtime/MovieScene/Private/EntitySystem/MovieSceneEvalTimeSystem.cpp
--- a/Engine/Source/Runtime/MovieScene/Private/EntitySystem/MovieSceneEvalTimeSystem.cpp
+++ b/Engine/Source/Runtime/MovieScene/Private/EntitySystem/MovieSceneEvalTimeSystem.cpp
@@ -10,9 +10,66 @@
 
 #include UE_INLINE_GENERATED_CPP_BY_NAME(MovieSceneEvalTimeSystem)
 
+namespace UE::MovieScene
+{
+
+DECLARE_CYCLE_STAT(TEXT("MovieScene: Gather evaluation times"), MovieSceneEval_GatherEvalTimes, STATGROUP_MovieSceneECS);
+DECLARE_CYCLE_STAT(TEXT("MovieScene: Assign evaluation times"), MovieSceneEval_AssignEvalTimes, STATGROUP_MovieSceneECS);
+
+struct FGatherTimes
+{
+	const FInstanceRegistry* InstanceRegistry;
+	TArray<FEvaluatedTime, TInlineAllocator<16>>* EvaluatedTimes;
+
+	void Run(FEntityAllocationWriteContext WriteContext) const
+	{
+		Run();
+	}
+	void Run() const
+	{
+		EvaluatedTimes->SetNum(InstanceRegistry->GetSparseInstances().GetMaxIndex());
+		for (auto It = InstanceRegistry->GetSparseInstances().CreateConstIterator(); It; ++It)
+		{
+			const FMovieSceneContext& Context = It->GetContext();
+
+			FEvaluatedTime EvaluatedTime;
+			EvaluatedTime.FrameTime = Context.GetTime();
+			EvaluatedTime.Seconds = Context.GetFrameRate().AsSeconds(EvaluatedTime.FrameTime);
+
+			(*EvaluatedTimes)[It.GetIndex()] = EvaluatedTime;
+		}
+	}
+};
+
+struct FAssignEvalTimesTask
+{
+	const TArray<FEvaluatedTime, TInlineAllocator<16>>* EvaluatedTimes;
 
-DECLARE_CYCLE_STAT(TEXT("EvalTime System"), MovieSceneEval_EvalTimeSystem, STATGROUP_MovieSceneECS);
-DECLARE_CYCLE_STAT(TEXT("MovieScene: Assign evaluation times"), MovieSceneEval_EvalTimes, STATGROUP_MovieSceneECS);
+	FAssignEvalTimesTask(const TArray<FEvaluatedTime, TInlineAllocator<16>>* InEvaluatedTimes)
+		: EvaluatedTimes(InEvaluatedTimes)
+	{}
+
+	void ForEachEntity(FInstanceHandle InstanceHandle, FFrameTime& EvalTime) const
+	{
+		EvalTime = (*EvaluatedTimes)[InstanceHandle.InstanceID].FrameTime;
+	}
+};
+
+struct FAssignEvalSecondsTask
+{
+	const TArray<FEvaluatedTime, TInlineAllocator<16>>* EvaluatedTimes;
+
+	FAssignEvalSecondsTask(const TArray<FEvaluatedTime, TInlineAllocator<16>>* InEvaluatedTimes)
+		: EvaluatedTimes(InEvaluatedTimes)
+	{}
+
+	void ForEachEntity(FInstanceHandle InstanceHandle, double& EvalSeconds) const
+	{
+		EvalSeconds = (*EvaluatedTimes)[InstanceHandle.InstanceID].Seconds;
+	}
+};
+
+} // namespace UE::MovieScene
 
 UMovieSceneEvalTimeSystem::UMovieSceneEvalTimeSystem(const FObjectInitializer& ObjInit)
 	: Super(ObjInit)
@@ -24,6 +81,8 @@ UMovieSceneEvalTimeSystem::UMovieSceneEvalTimeSystem(const FObjectInitializer& O
 	SystemCategories = EEntitySystemCategory::Core;
 	RelevantFilter.Any({ BuiltInComponents->EvalTime, BuiltInComponents->EvalSeconds });
 
+	Phase = ESystemPhase::Scheduling;
+
 	if (HasAnyFlags(RF_ClassDefaultObject))
 	{
 		DefineComponentProducer(GetClass(), BuiltInComponents->EvalTime);
@@ -36,88 +95,49 @@ bool UMovieSceneEvalTimeSystem::IsRelevantImpl(UMovieSceneEntitySystemLinker* In
 	return InLinker->EntityManager.Contains(RelevantFilter);
 }
 
-void UMovieSceneEvalTimeSystem::OnRun(FSystemTaskPrerequisites& InPrerequisites, FSystemSubsequentTasks& Subsequents)
+void UMovieSceneEvalTimeSystem::OnSchedulePersistentTasks(UE::MovieScene::IEntitySystemScheduler* TaskScheduler)
 {
 	using namespace UE::MovieScene;
 
-	SCOPE_CYCLE_COUNTER(MovieSceneEval_EvalTimeSystem)
-
-	struct FGatherTimes
-	{
-		FGatherTimes(const FInstanceRegistry* InInstanceRegistry, TArray<FEvaluatedTime>* InEvaluatedTimes)
-			: InstanceRegistry(InInstanceRegistry)
-			, EvaluatedTimes(InEvaluatedTimes)
-		{}
-
-		const FInstanceRegistry* InstanceRegistry;
-		TArray<FEvaluatedTime>* EvaluatedTimes;
-
-		FORCEINLINE TStatId           GetStatId() const    { RETURN_QUICK_DECLARE_CYCLE_STAT(FGenericTask, STATGROUP_TaskGraphTasks); }
-		static ENamedThreads::Type    GetDesiredThread()   { return ENamedThreads::AnyHiPriThreadHiPriTask; }
-		static ESubsequentsMode::Type GetSubsequentsMode() { return ESubsequentsMode::TrackSubsequents; }
+	FTaskID GatherTask = TaskScheduler->AddTask<FGatherTimes>(
+		FTaskParams(TEXT("Gather Evaluation Times"), GET_STATID(MovieSceneEval_GatherEvalTimes)),
+		Linker->GetInstanceRegistry(),
+		&EvaluatedTimes
+	);
 
-		void DoTask(ENamedThreads::Type CurrentThread, const FGraphEventRef& MyCompletionGraphEvent)
-		{
-			Run();
-		}
-
-		void Run()
-		{
-			EvaluatedTimes->SetNum(InstanceRegistry->GetSparseInstances().GetMaxIndex());
-			for (auto It = InstanceRegistry->GetSparseInstances().CreateConstIterator(); It; ++It)
-			{
-				const FMovieSceneContext& Context = It->GetContext();
-
-				FEvaluatedTime EvaluatedTime;
-				EvaluatedTime.FrameTime = Context.GetTime();
-				EvaluatedTime.Seconds = Context.GetFrameRate().AsSeconds(EvaluatedTime.FrameTime);
-
-				(*EvaluatedTimes)[It.GetIndex()] = EvaluatedTime;
-			}
-		}
-	};
-
-	struct FAssignEvalTimesTask
-	{
-		const TArray<FEvaluatedTime>* EvaluatedTimes;
-
-		explicit FAssignEvalTimesTask(const TArray<FEvaluatedTime>* InEvaluatedTimes)
-			: EvaluatedTimes(InEvaluatedTimes)
-		{}
+	FBuiltInComponentTypes* BuiltInComponents = FBuiltInComponentTypes::Get();
 
-		void ForEachEntity(UE::MovieScene::FInstanceHandle InstanceHandle, FFrameTime& EvalTime)
-		{
-			EvalTime = (*EvaluatedTimes)[InstanceHandle.InstanceID].FrameTime;
-		}
-	};
+	FTaskID WriteTask1 = FEntityTaskBuilder()
+	.Read(BuiltInComponents->InstanceHandle)
+	.Write(BuiltInComponents->EvalTime)
+	.FilterNone({ BuiltInComponents->Tags.FixedTime })
+	.SetStat(GET_STATID(MovieSceneEval_AssignEvalTimes))
+	.Fork_PerEntity<FAssignEvalTimesTask>(&Linker->EntityManager, TaskScheduler, &EvaluatedTimes);
 
-	struct FAssignEvalSecondsTask
-	{
-		const TArray<FEvaluatedTime>* EvaluatedTimes;
+	FTaskID WriteTask2 = FEntityTaskBuilder()
+	.Read(BuiltInComponents->InstanceHandle)
+	.Write(BuiltInComponents->EvalSeconds)
+	.FilterNone({ BuiltInComponents->Tags.FixedTime })
+	.SetStat(GET_STATID(MovieSceneEval_AssignEvalTimes))
+	.Fork_PerEntity<FAssignEvalSecondsTask>(&Linker->EntityManager, TaskScheduler, &EvaluatedTimes);
 
-		explicit FAssignEvalSecondsTask(const TArray<FEvaluatedTime>* InEvaluatedTimes)
-			: EvaluatedTimes(InEvaluatedTimes)
-		{}
+	TaskScheduler->AddPrerequisite(GatherTask, WriteTask1);
+	TaskScheduler->AddPrerequisite(GatherTask, WriteTask2);
+}
 
-		void ForEachEntity(UE::MovieScene::FInstanceHandle InstanceHandle, double& EvalSeconds)
-		{
-			EvalSeconds = (*EvaluatedTimes)[InstanceHandle.InstanceID].Seconds;
-		}
-	};
+void UMovieSceneEvalTimeSystem::OnRun(FSystemTaskPrerequisites& InPrerequisites, FSystemSubsequentTasks& Subsequents)
+{
+	using namespace UE::MovieScene;
 
 	FSystemTaskPrerequisites EvalPrereqs;
 	FBuiltInComponentTypes* BuiltInComponents = FBuiltInComponentTypes::Get();
 
-	if (Linker->EntityManager.GetThreadingModel() == EEntityThreadingModel::NoThreading)
-	{
-		FGatherTimes(Linker->GetInstanceRegistry(), &EvaluatedTimes).Run();
-	}
-	else
-	{
-		// The only thing we depend on is the gather task
-		FGraphEventRef GatherEvalTimesEvent = TGraphTask<FGatherTimes>::CreateTask(nullptr, Linker->EntityManager.GetDispatchThread())
-		.ConstructAndDispatchWhenReady(Linker->GetInstanceRegistry(), &EvaluatedTimes);
+	FGraphEventRef GatherEvalTimesEvent = TEntityTaskComponents<>()
+	.SetStat(GET_STATID(MovieSceneEval_GatherEvalTimes))
+	.Dispatch<FGatherTimes>(&Linker->EntityManager, InPrerequisites, nullptr, Linker->GetInstanceRegistry(), &EvaluatedTimes);
 
+	if (GatherEvalTimesEvent)
+	{
 		EvalPrereqs.AddComponentTask(BuiltInComponents->EvalTime, GatherEvalTimesEvent);
 		EvalPrereqs.AddComponentTask(BuiltInComponents->EvalSeconds, GatherEvalTimesEvent);
 	}
@@ -126,14 +146,13 @@ void UMovieSceneEvalTimeSystem::OnRun(FSystemTaskPrerequisites& InPrerequisites,
 	.Read(BuiltInComponents->InstanceHandle)
 	.Write(BuiltInComponents->EvalTime)
 	.FilterNone({ FBuiltInComponentTypes::Get()->Tags.FixedTime })
-	.SetStat(GET_STATID(MovieSceneEval_EvalTimes))
+	.SetStat(GET_STATID(MovieSceneEval_AssignEvalTimes))
 	.Dispatch_PerEntity<FAssignEvalTimesTask>(&Linker->EntityManager, EvalPrereqs, &Subsequents, &EvaluatedTimes);
 
 	FEntityTaskBuilder()
 	.Read(BuiltInComponents->InstanceHandle)
 	.Write(BuiltInComponents->EvalSeconds)
 	.FilterNone({ FBuiltInComponentTypes::Get()->Tags.FixedTime })
-	.SetStat(GET_STATID(MovieSceneEval_EvalTimes))
+	.SetStat(GET_STATID(MovieSceneEval_AssignEvalTimes))
 	.Dispatch_PerEntity<FAssignEvalSecondsTask>(&Linker->EntityManager, EvalPrereqs, &Subsequents, &EvaluatedTimes);
-}
-
+}
\ No newline at end of file
diff --git a/Engine/Source/Runtime/MovieScene/Private/EntitySystem/MovieSceneEvaluationHookSystem.cpp b/Engine/Source/Runtime/MovieScene/Private/EntitySystem/MovieSceneEvaluationHookSystem.cpp
--- a/Engine/Source/Runtime/MovieScene/Private/EntitySystem/MovieSceneEvaluationHookSystem.cpp
+++ b/Engine/Source/Runtime/MovieScene/Private/EntitySystem/MovieSceneEvaluationHookSystem.cpp
@@ -28,7 +28,7 @@ struct FEvaluationHookUpdater
 		: HookSystem(InHookSystem), InstanceRegistry(InInstanceRegistry)
 	{}
 
-	void ForEachAllocation(FEntityAllocationProxy Item, TRead<FInstanceHandle> InstanceHandles, TRead<FMovieSceneEvaluationHookComponent> Hooks, TRead<FFrameTime> EvalTimes, TWrite<FEvaluationHookFlags> WriteFlags)
+	void ForEachAllocation(FEntityAllocationProxy Item, TRead<FInstanceHandle> InstanceHandles, TRead<FMovieSceneEvaluationHookComponent> Hooks, TRead<FFrameTime> EvalTimes, TWrite<FEvaluationHookFlags> WriteFlags) const
 	{
 		const int32 Num = Item.GetAllocation()->Num();
 		const bool bRestoreState = Item.GetAllocationType().Contains(FBuiltInComponentTypes::Get()->Tags.RestoreState);
@@ -53,26 +53,8 @@ struct FEvaluationHookUpdater
 			HookSystem->AddEvent(SequenceInstance.GetRootInstanceHandle(), NewEvent);
 		}
 	}
-};
-
-struct FEvaluationHookSorter
-{
-	UMovieSceneEvaluationHookSystem* HookSystem;
-
-	FEvaluationHookSorter(UMovieSceneEvaluationHookSystem* InHookSystem)
-		: HookSystem(InHookSystem)
-	{}
-
-	FORCEINLINE TStatId           GetStatId() const    { return GET_STATID(MovieSceneECS_GenericHooks); }
-	static ENamedThreads::Type    GetDesiredThread()   { return ENamedThreads::AnyHiPriThreadHiPriTask; }
-	static ESubsequentsMode::Type GetSubsequentsMode() { return ESubsequentsMode::TrackSubsequents; }
 
-	void DoTask(ENamedThreads::Type CurrentThread, const FGraphEventRef& MyCompletionGraphEvent)
-	{
-		Run();
-	}
-
-	void Run()
+	void PostTask()
 	{
 		HookSystem->SortEvents();
 	}
@@ -84,7 +66,7 @@ struct FEvaluationHookSorter
 UMovieSceneEvaluationHookSystem::UMovieSceneEvaluationHookSystem(const FObjectInitializer& ObjInit)
 	: Super(ObjInit)
 {
-	Phase = UE::MovieScene::ESystemPhase::Instantiation | UE::MovieScene::ESystemPhase::Evaluation | UE::MovieScene::ESystemPhase::Finalization;
+	Phase = UE::MovieScene::ESystemPhase::Instantiation | UE::MovieScene::ESystemPhase::Scheduling | UE::MovieScene::ESystemPhase::Finalization;
 
 	if (HasAnyFlags(RF_ClassDefaultObject))
 	{
@@ -107,6 +89,20 @@ bool UMovieSceneEvaluationHookSystem::IsRelevantImpl(UMovieSceneEntitySystemLink
 	return HasEvents() || InLinker->EntityManager.ContainsComponent(UE::MovieScene::FBuiltInComponentTypes::Get()->EvaluationHook);
 }
 
+void UMovieSceneEvaluationHookSystem::OnSchedulePersistentTasks(UE::MovieScene::IEntitySystemScheduler* TaskScheduler)
+{
+	using namespace UE::MovieScene;
+
+	FBuiltInComponentTypes* Components = FBuiltInComponentTypes::Get();
+
+	FEntityTaskBuilder()
+	.Read(Components->InstanceHandle)
+	.Read(Components->EvaluationHook)
+	.Read(Components->EvalTime)
+	.Write(Components->EvaluationHookFlags)
+	.Schedule_PerAllocation<FEvaluationHookUpdater>(&Linker->EntityManager, TaskScheduler, this, Linker->GetInstanceRegistry());
+}
+
 void UMovieSceneEvaluationHookSystem::OnRun(FSystemTaskPrerequisites& InPrerequisites, FSystemSubsequentTasks& Subsequents)
 {
 	using namespace UE::MovieScene;
@@ -132,20 +128,6 @@ void UMovieSceneEvaluationHookSystem::OnRun(FSystemTaskPrerequisites& InPrerequi
 		.Read(Components->EvalTime)
 		.Write(Components->EvaluationHookFlags)
 		.Dispatch_PerAllocation<FEvaluationHookUpdater>(&Linker->EntityManager, InPrerequisites, &Subsequents, this, Linker->GetInstanceRegistry());
-
-		if (Linker->EntityManager.GetThreadingModel() == EEntityThreadingModel::NoThreading)
-		{
-			this->SortEvents();
-		}
-		else
-		{
-			// The only thing we depend on is the gather task
-			FGraphEventArray Prereqs = { UpdateEvent };
-			FGraphEventRef SortTask = TGraphTask<FEvaluationHookSorter>::CreateTask(&Prereqs, Linker->EntityManager.GetDispatchThread())
-			.ConstructAndDispatchWhenReady(this);
-
-			Subsequents.AddRootTask(SortTask);
-		}
 	}
 	else if (HasEvents())
 	{
diff --git a/Engine/Source/Runtime/MovieScene/Private/EntitySystem/MovieScenePreAnimatedStateSystem.cpp b/Engine/Source/Runtime/MovieScene/Private/EntitySystem/MovieScenePreAnimatedStateSystem.cpp
--- a/Engine/Source/Runtime/MovieScene/Private/EntitySystem/MovieScenePreAnimatedStateSystem.cpp
+++ b/Engine/Source/Runtime/MovieScene/Private/EntitySystem/MovieScenePreAnimatedStateSystem.cpp
@@ -59,6 +59,7 @@ void UMovieSceneCachePreAnimatedStateSystem::OnRun(FSystemTaskPrerequisites& InP
 	};
 	Linker->SystemGraph.IteratePhase(ESystemPhase::Spawn, ForEachSystem);
 	Linker->SystemGraph.IteratePhase(ESystemPhase::Instantiation, ForEachSystem);
+	Linker->SystemGraph.IteratePhase(ESystemPhase::Scheduling, ForEachSystem);
 	Linker->SystemGraph.IteratePhase(ESystemPhase::Evaluation, ForEachSystem);
 
 	IMovieScenePreAnimatedStateSystemInterface::FPreAnimationParameters Params{ &InPrerequisites, &Subsequents, &Linker->PreAnimatedState };
@@ -115,6 +116,7 @@ void UMovieSceneRestorePreAnimatedStateSystem::OnRun(FSystemTaskPrerequisites& I
 	};
 	Linker->SystemGraph.IteratePhase(ESystemPhase::Spawn, ForEachSystem);
 	Linker->SystemGraph.IteratePhase(ESystemPhase::Instantiation, ForEachSystem);
+	Linker->SystemGraph.IteratePhase(ESystemPhase::Scheduling, ForEachSystem);
 	Linker->SystemGraph.IteratePhase(ESystemPhase::Evaluation, ForEachSystem);
 
 	IMovieScenePreAnimatedStateSystemInterface::FPreAnimationParameters Params{ &InPrerequisites, &Subsequents, &Linker->PreAnimatedState };
diff --git a/Engine/Source/Runtime/MovieScene/Private/EntitySystem/MovieSceneSequenceInstance.cpp b/Engine/Source/Runtime/MovieScene/Private/EntitySystem/MovieSceneSequenceInstance.cpp
--- a/Engine/Source/Runtime/MovieScene/Private/EntitySystem/MovieSceneSequenceInstance.cpp
+++ b/Engine/Source/Runtime/MovieScene/Private/EntitySystem/MovieSceneSequenceInstance.cpp
@@ -140,12 +140,15 @@ void FSequenceInstance::InitializeLegacyEvaluator(UMovieSceneEntitySystemLinker*
 		if (!LegacyEvaluator)
 		{
 			LegacyEvaluator = MakeUnique<FMovieSceneTrackEvaluator>(CompiledEntry.GetSequence(), RootTemplate.GetCompiledDataID(), CompiledDataManager);
+			UpdateFlags |= ESequenceInstanceUpdateFlags::HasLegacyTemplates;
 		}
 	}
 	else if (LegacyEvaluator)
 	{
 		LegacyEvaluator->Finish(*Player);
 		LegacyEvaluator = nullptr;
+
+		UpdateFlags &= ~ESequenceInstanceUpdateFlags::HasLegacyTemplates;
 	}
 }
 
diff --git a/Engine/Source/Runtime/MovieScene/Private/Tests/MovieSceneSparseBitSetTests.cpp b/Engine/Source/Runtime/MovieScene/Private/Tests/MovieSceneSparseBitSetTests.cpp
new file mode 100644
--- /dev/null
+++ b/Engine/Source/Runtime/MovieScene/Private/Tests/MovieSceneSparseBitSetTests.cpp
@@ -0,0 +1,144 @@
+// Copyright Epic Games, Inc. All Rights Reserved.
+
+#include "CoreTypes.h"
+#include "MovieSceneFwd.h"
+#include "Misc/AutomationTest.h"
+#include "Containers/SparseBitSet.h"
+#include "Math/RandomStream.h"
+#include "Misc/GeneratedTypeName.h"
+
+#if WITH_DEV_AUTOMATION_TESTS
+
+namespace UE::MovieScene::Tests
+{
+
+	template<typename SparseBitSet>
+	void TestBitSet(FAutomationTestBase* Test, SparseBitSet& InBitSet)
+	{
+		// Set random bits in the bitset
+		const int32 MaxBitIndex = static_cast<int32>(InBitSet.GetMaxNumBits());
+
+		// Static int for re-running tests using the same seed as a previous run
+		static int32 SeedOverride = 0;
+		const int32 InitialSeed = SeedOverride == 0 ? FMath::Rand() : SeedOverride;
+		FRandomStream Random(InitialSeed);
+
+		UE_LOG(LogMovieScene, Display, TEXT("Running TSparseBitSet<> tests for %s with a random seed %i..."), GetGeneratedTypeName<SparseBitSet>(), InitialSeed);
+
+		TSet<uint32> SetIndices;
+		const int32 Num = Random.RandHelper(MaxBitIndex);
+
+		for (int32 Index = 0; Index < Num; ++Index)
+		{
+			const uint32 Bit = static_cast<uint32>(Random.RandHelper(MaxBitIndex));
+			InBitSet.SetBit(Bit);
+			SetIndices.Add(Bit);
+		}
+
+		Test->TestEqual(TEXT("Num Set Bits"), InBitSet.CountSetBits(), static_cast<uint32>(SetIndices.Num()));
+
+		// Test all bits are the correct state
+		for (uint32 BitIndex = 0; BitIndex < static_cast<uint32>(MaxBitIndex); ++BitIndex)
+		{
+			Test->TestEqual(FString::Printf(TEXT("Bit index %d"), BitIndex), InBitSet.IsBitSet(BitIndex), SetIndices.Contains(BitIndex));
+		}
+
+		// Test the iterator
+		int32 NumIterated = 0;
+		for (int32 BitIndex : InBitSet)
+		{
+			++NumIterated;
+			if (!SetIndices.Contains(BitIndex))
+			{
+				Test->AddError(FString::Printf(TEXT("Bit %i was iterated but it shouldn't be set!"), BitIndex));
+			}
+		}
+
+		Test->TestEqual(TEXT("Number of iterated bits"), NumIterated, SetIndices.Num());
+	}
+
+} // namespace UE::MovieScene::Tests
+
+IMPLEMENT_SIMPLE_AUTOMATION_TEST(FMovieSceneSparseBitSetTest, 
+		"System.Engine.Sequencer.SparseBitSet", 
+		EAutomationTestFlags::EditorContext | EAutomationTestFlags::EngineFilter)
+bool FMovieSceneSparseBitSetTest::RunTest(const FString& Parameters)
+{
+	using namespace UE::MovieScene;
+	using namespace UE::MovieScene::Tests;
+
+	{
+		// 8 bit buckets
+		TSparseBitSet<uint8, TDynamicSparseBitSetBucketStorage<uint8, 4>> BitSet8_8;
+		TestEqual(TEXT("TSparseBitSet<uint8, ...<uint8>>::GetMaxNumBits"), BitSet8_8.GetMaxNumBits(), 8u*8u);
+		TestBitSet(this, BitSet8_8);
+
+		TSparseBitSet<uint16, TDynamicSparseBitSetBucketStorage<uint8, 4>> BitSet16_8;
+		TestEqual(TEXT("TSparseBitSet<uint16, ...<uint8>>::GetMaxNumBits"), BitSet16_8.GetMaxNumBits(), 16u*8u);
+		TestBitSet(this, BitSet16_8);
+
+		TSparseBitSet<uint32, TDynamicSparseBitSetBucketStorage<uint8, 4>> BitSet32_8;
+		TestEqual(TEXT("TSparseBitSet<uint32, ...<uint8>>::GetMaxNumBits"), BitSet32_8.GetMaxNumBits(), 32u*8u);
+		TestBitSet(this, BitSet32_8);
+
+		TSparseBitSet<uint64, TDynamicSparseBitSetBucketStorage<uint8, 4>> BitSet64_8;
+		TestEqual(TEXT("TSparseBitSet<uint64, ...<uint8>>::GetMaxNumBits"), BitSet64_8.GetMaxNumBits(), 64u*8u);
+		TestBitSet(this, BitSet64_8);
+
+		// 16 bit buckets
+		TSparseBitSet<uint8, TDynamicSparseBitSetBucketStorage<uint16, 4>> BitSet8_16;
+		TestEqual(TEXT("TSparseBitSet<uint8, ...<uint16>>::GetMaxNumBits"), BitSet8_16.GetMaxNumBits(), 8u*16u);
+		TestBitSet(this, BitSet8_16);
+
+		TSparseBitSet<uint16, TDynamicSparseBitSetBucketStorage<uint16, 4>> BitSet16_16;
+		TestEqual(TEXT("TSparseBitSet<uint16, ...<uint16>>::GetMaxNumBits"), BitSet16_16.GetMaxNumBits(), 16u*16u);
+		TestBitSet(this, BitSet16_16);
+
+		TSparseBitSet<uint32, TDynamicSparseBitSetBucketStorage<uint16, 4>> BitSet32_16;
+		TestEqual(TEXT("TSparseBitSet<uint32, ...<uint16>>::GetMaxNumBits"), BitSet32_16.GetMaxNumBits(), 32u*16u);
+		TestBitSet(this, BitSet32_16);
+
+		TSparseBitSet<uint64, TDynamicSparseBitSetBucketStorage<uint16, 4>> BitSet64_16;
+		TestEqual(TEXT("TSparseBitSet<uint64, ...<uint16>>::GetMaxNumBits"), BitSet64_16.GetMaxNumBits(), 64u*16u);
+		TestBitSet(this, BitSet64_16);
+
+		// 32 bit buckets
+		TSparseBitSet<uint8, TDynamicSparseBitSetBucketStorage<uint32, 4>> BitSet8_32;
+		TestEqual(TEXT("TSparseBitSet<uint8, ...<uint32>>::GetMaxNumBits"), BitSet8_32.GetMaxNumBits(), 8u*32u);
+		TestBitSet(this, BitSet8_32);
+
+		TSparseBitSet<uint16, TDynamicSparseBitSetBucketStorage<uint32, 4>> BitSet16_32;
+		TestEqual(TEXT("TSparseBitSet<uint16, ...<uint32>>::GetMaxNumBits"), BitSet16_32.GetMaxNumBits(), 16u*32u);
+		TestBitSet(this, BitSet16_32);
+
+		TSparseBitSet<uint32, TDynamicSparseBitSetBucketStorage<uint32, 4>> BitSet32_32;
+		TestEqual(TEXT("TSparseBitSet<uint32, ...<uint32>>::GetMaxNumBits"), BitSet32_32.GetMaxNumBits(), 32u*32u);
+		TestBitSet(this, BitSet32_32);
+
+		TSparseBitSet<uint64, TDynamicSparseBitSetBucketStorage<uint32, 4>> BitSet64_32;
+		TestEqual(TEXT("TSparseBitSet<uint64, ...<uint32>>::GetMaxNumBits"), BitSet64_32.GetMaxNumBits(), 64u*32u);
+		TestBitSet(this, BitSet64_32);
+
+		// 64 bit buckets
+		TSparseBitSet<uint8, TDynamicSparseBitSetBucketStorage<uint64, 4>> BitSet8_64;
+		TestEqual(TEXT("TSparseBitSet<uint8, ...<uint64>>::GetMaxNumBits"), BitSet8_64.GetMaxNumBits(), 8u*64u);
+		TestBitSet(this, BitSet8_64);
+
+		TSparseBitSet<uint16, TDynamicSparseBitSetBucketStorage<uint64, 4>> BitSet16_64;
+		TestEqual(TEXT("TSparseBitSet<uint16, ...<uint64>>::GetMaxNumBits"), BitSet16_64.GetMaxNumBits(), 16u*64u);
+		TestBitSet(this, BitSet16_64);
+
+		TSparseBitSet<uint32, TDynamicSparseBitSetBucketStorage<uint64, 4>> BitSet32_64;
+		TestEqual(TEXT("TSparseBitSet<uint32, ...<uint64>>::GetMaxNumBits"), BitSet32_64.GetMaxNumBits(), 32u*64u);
+		TestBitSet(this, BitSet32_64);
+
+		TSparseBitSet<uint64, TDynamicSparseBitSetBucketStorage<uint64, 4>> BitSet64_64;
+		TestEqual(TEXT("TSparseBitSet<uint64, ...<uint64>>::GetMaxNumBits"), BitSet64_64.GetMaxNumBits(), 64u*64u);
+		TestBitSet(this, BitSet64_64);
+
+
+	}
+	return true;
+}
+
+#endif // WITH_DEV_AUTOMATION_TESTS
diff --git a/Engine/Source/Runtime/MovieScene/Public/EntitySystem/IMovieScenePropertyComponentHandler.h b/Engine/Source/Runtime/MovieScene/Public/EntitySystem/IMovieScenePropertyComponentHandler.h
--- a/Engine/Source/Runtime/MovieScene/Public/EntitySystem/IMovieScenePropertyComponentHandler.h
+++ b/Engine/Source/Runtime/MovieScene/Public/EntitySystem/IMovieScenePropertyComponentHandler.h
@@ -24,6 +24,8 @@ struct FValueDecompositionParams;
 struct FPreAnimatedStateExtension;
 struct FPropertyCompositeDefinition;
 
+class IEntitySystemScheduler;
+
 /** Type-erased view of a component. Used for passing typed data through the IPropertyComponentHandler interface */
 struct FPropertyComponentView
 {
@@ -156,6 +158,7 @@ struct IPropertyComponentHandler
 	 * @param Subsequents      Subsequents to add any dispatched tasks to
 	 * @param Linker           The linker that owns the entity manager to dispatch tasks for
 	 */
+	virtual void ScheduleSetterTasks(const FPropertyDefinition& Definition, TArrayView<const FPropertyCompositeDefinition> Composites, const FPropertyStats& Stats, IEntitySystemScheduler* TaskScheduler, UMovieSceneEntitySystemLinker* Linker) = 0;
 	virtual void DispatchSetterTasks(const FPropertyDefinition& Definition, TArrayView<const FPropertyCompositeDefinition> Composites, const FPropertyStats& Stats, FSystemTaskPrerequisites& InPrerequisites, FSystemSubsequentTasks& Subsequents, UMovieSceneEntitySystemLinker* Linker) = 0;
 
 	/**
diff --git a/Engine/Source/Runtime/MovieScene/Public/EntitySystem/MovieSceneComponentAccessors.h b/Engine/Source/Runtime/MovieScene/Public/EntitySystem/MovieSceneComponentAccessors.h
--- a/Engine/Source/Runtime/MovieScene/Public/EntitySystem/MovieSceneComponentAccessors.h
+++ b/Engine/Source/Runtime/MovieScene/Public/EntitySystem/MovieSceneComponentAccessors.h
@@ -10,23 +10,33 @@
 #include "EntitySystem/MovieSceneEntityManager.h"
 #include "EntitySystem/MovieSceneEntitySystemTypes.h"
 #include "EntitySystem/MovieSceneSystemTaskDependencies.h"
+#include "EntitySystem/RelativePtr.h"
 #include "HAL/Platform.h"
 #include "Misc/AssertionMacros.h"
 #include "MovieSceneEntityIDs.h"
 #include "Templates/Tuple.h"
 
+
 #include <initializer_list>
 
 namespace UE
 {
 namespace MovieScene
 {
-class FEntityManager;
 
+class FEntityManager;
+using FPreLockedDataPtr = TRelativePtr<void, uint32>;
 
 struct FComponentAccess
 {
 	FComponentTypeID ComponentType;
+
+	FORCEINLINE void PreLockComponentData(const FEntityAllocation* Allocation, FPreLockedDataPtr* PrelockedComponentData) const
+	{
+		// Always relative to the component data which is a separate allocation
+		const FComponentHeader& Header = Allocation->GetComponentHeaderChecked(ComponentType);
+		PrelockedComponentData->Reset(Allocation->GetComponentDataAddress(), Header.Components);
+	}
 };
 struct FReadAccess : FComponentAccess
 {
@@ -41,6 +51,16 @@ struct FWriteAccess : FComponentAccess
 struct FOptionalComponentAccess
 {
 	FComponentTypeID ComponentType;
+
+	FORCEINLINE void PreLockComponentData(const FEntityAllocation* Allocation, FPreLockedDataPtr* PrelockedComponentData) const
+	{
+		const FComponentHeader* Header = Allocation->FindComponentHeader(ComponentType);
+		if (Header)
+		{
+			// Always relative to the component data which is a separate allocation
+			PrelockedComponentData->Reset(Allocation->GetComponentDataAddress(), Header->Components);
+		}
+	}
 };
 struct FOptionalReadAccess : FOptionalComponentAccess
 {
@@ -55,6 +75,19 @@ struct FOptionalWriteAccess : FOptionalComponentAccess
 struct FEntityIDAccess
 {
 	using AccessType = const FMovieSceneEntityID;
+	static constexpr int32 PreLockedDataNum = 1;
+
+
+	FORCEINLINE void PreLockComponentData(const FEntityAllocation* Allocation, FPreLockedDataPtr* PrelockedComponentData) const
+	{
+		// Must be relative to the base allocation ptr
+		PrelockedComponentData->Reset(Allocation, Allocation->GetRawEntityIDs());
+	}
+	FORCEINLINE TRead<FMovieSceneEntityID> ResolvePreLockedComponentData(const FEntityAllocation* Allocation, const FPreLockedDataPtr* Ptr, FEntityAllocationWriteContext WriteContext) const
+	{
+		// Resolve using the same base ptr as PreLockComponentData
+		return TRead<FMovieSceneEntityID>(Ptr->Resolve<FMovieSceneEntityID>(Allocation));
+	}
 
 	FORCEINLINE TRead<FMovieSceneEntityID> LockComponentData(const FEntityAllocation* Allocation, FEntityAllocationWriteContext WriteContext) const
 	{
@@ -68,11 +101,17 @@ template<typename T>
 struct TReadAccess : FReadAccess
 {
 	using AccessType = const T;
+	static constexpr int32 PreLockedDataNum = 1;
 
 	TReadAccess(FComponentTypeID InComponentTypeID)
 		: FReadAccess{ InComponentTypeID }
 	{}
 
+	FORCEINLINE TRead<T> ResolvePreLockedComponentData(const FEntityAllocation* Allocation, const FPreLockedDataPtr* Ptr, FEntityAllocationWriteContext WriteContext) const
+	{
+		// Resolve using the same base ptr as PreLockComponentData
+		return TRead<T>(Ptr->Resolve<const T>(Allocation->GetComponentDataAddress()));
+	}
 	FORCEINLINE TComponentLock<TRead<T>> LockComponentData(const FEntityAllocation* Allocation, FEntityAllocationWriteContext WriteContext) const
 	{
 		return Allocation->ReadComponents(ComponentType.ReinterpretCast<T>());
@@ -82,23 +121,75 @@ struct TReadAccess : FReadAccess
 
 struct FErasedReadAccess : FReadAccess
 {
+	static constexpr int32 PreLockedDataNum = 1;
+
 	FErasedReadAccess(FComponentTypeID InComponentTypeID)
 		: FReadAccess{ InComponentTypeID }
 	{}
 
+	FORCEINLINE void PreLockComponentData(const FEntityAllocation* Allocation, FPreLockedDataPtr* PrelockedComponentData) const
+	{
+		// Erased reads always pass the header since they need the size information
+		const FComponentHeader& Header = Allocation->GetComponentHeaderChecked(ComponentType);
+		PrelockedComponentData->Reset(Allocation, &Header);
+	}
+	FORCEINLINE FReadErased ResolvePreLockedComponentData(const FEntityAllocation* Allocation, const FPreLockedDataPtr* Ptr, FEntityAllocationWriteContext WriteContext) const
+	{
+		// Resolve using the same base ptr as PreLockComponentData
+		return FReadErased(Ptr->Resolve<const FComponentHeader>(Allocation));
+	}
 	FORCEINLINE FComponentReader LockComponentData(const FEntityAllocation* Allocation, FEntityAllocationWriteContext WriteContext) const
 	{
 		return Allocation->ReadComponentsErased(ComponentType);
 	}
 };
+struct FErasedOptionalReadAccess : FReadAccess
+{
+	static constexpr int32 PreLockedDataNum = 1;
 
+	FErasedOptionalReadAccess(FComponentTypeID InComponentTypeID)
+		: FReadAccess{ InComponentTypeID }
+	{}
+
+	FORCEINLINE void PreLockComponentData(const FEntityAllocation* Allocation, FPreLockedDataPtr* PrelockedComponentData) const
+	{
+		const FComponentHeader* Header = Allocation->FindComponentHeader(ComponentType);
+		if (Header)
+		{
+			// Erased reads always pass the header since they need the size information
+			PrelockedComponentData->Reset(Allocation, Header);
+		}
+	}
+	FORCEINLINE FReadErasedOptional ResolvePreLockedComponentData(const FEntityAllocation* Allocation, const FPreLockedDataPtr* Ptr, FEntityAllocationWriteContext WriteContext) const
+	{
+		// Resolve using the same base ptr as PreLockComponentData
+		return FReadErasedOptional(Ptr->Resolve<const FComponentHeader>(Allocation));
+	}
+	FORCEINLINE FOptionalComponentReader LockComponentData(const FEntityAllocation* Allocation, FEntityAllocationWriteContext WriteContext) const
+	{
+		return Allocation->TryReadComponentsErased(ComponentType);
+	}
+};
 
 struct FErasedWriteAccess : FWriteAccess
 {
+	static constexpr int32 PreLockedDataNum = 1;
+
 	FErasedWriteAccess(FComponentTypeID InComponentTypeID)
 		: FWriteAccess{ InComponentTypeID }
 	{}
 
+	FORCEINLINE void PreLockComponentData(const FEntityAllocation* Allocation, FPreLockedDataPtr* PrelockedComponentData) const
+	{
+		const FComponentHeader& Header = Allocation->GetComponentHeaderChecked(ComponentType);
+		// Erased writes always pass the header since they need the size information
+		PrelockedComponentData->Reset(Allocation, &Header);
+	}
+	FORCEINLINE FWriteErased ResolvePreLockedComponentData(const FEntityAllocation* Allocation, const FPreLockedDataPtr* Ptr, FEntityAllocationWriteContext WriteContext) const
+	{
+		// Resolve using the same base ptr as PreLockComponentData
+		return FWriteErased(Ptr->Resolve<FComponentHeader>(Allocation));
+	}
 	FORCEINLINE FComponentWriter LockComponentData(const FEntityAllocation* Allocation, FEntityAllocationWriteContext WriteContext) const
 	{
 		return Allocation->WriteComponentsErased(ComponentType, WriteContext);
@@ -112,11 +203,17 @@ template<typename T>
 struct TWriteAccess : FWriteAccess
 {
 	using AccessType = T;
+	static constexpr int32 PreLockedDataNum = 1;
 
 	TWriteAccess(FComponentTypeID InComponentTypeID)
 		: FWriteAccess{ InComponentTypeID }
 	{}
 
+	FORCEINLINE TWrite<T> ResolvePreLockedComponentData(const FEntityAllocation* Allocation, const FPreLockedDataPtr* Ptr, FEntityAllocationWriteContext WriteContext) const
+	{
+		// Resolve using the same base ptr as PreLockComponentData
+		return TWrite<T>(Ptr->Resolve<T>(Allocation->GetComponentDataAddress()));
+	}
 	FORCEINLINE TComponentLock<TWrite<T>> LockComponentData(const FEntityAllocation* Allocation, FEntityAllocationWriteContext WriteContext) const
 	{
 		return Allocation->WriteComponents(ComponentType.ReinterpretCast<T>(), WriteContext);
@@ -130,11 +227,17 @@ template<typename T>
 struct TOptionalReadAccess : FOptionalReadAccess
 {
 	using AccessType = const T;
+	static constexpr int32 PreLockedDataNum = 1;
 
 	TOptionalReadAccess(FComponentTypeID InComponentTypeID)
 		: FOptionalReadAccess{InComponentTypeID}
 	{}
 
+	FORCEINLINE TReadOptional<T> ResolvePreLockedComponentData(const FEntityAllocation* Allocation, const FPreLockedDataPtr* Ptr, FEntityAllocationWriteContext WriteContext) const
+	{
+		// Resolve using the same base ptr as PreLockComponentData
+		return TReadOptional<T>(Ptr->Resolve<T>(Allocation->GetComponentDataAddress()));
+	}
 	FORCEINLINE TComponentLock<TReadOptional<T>> LockComponentData(const FEntityAllocation* Allocation, FEntityAllocationWriteContext WriteContext) const
 	{
 		return Allocation->TryReadComponents(ComponentType.ReinterpretCast<T>());
@@ -148,11 +251,17 @@ template<typename T>
 struct TOptionalWriteAccess : FOptionalWriteAccess
 {
 	using AccessType = T;
+	static constexpr int32 PreLockedDataNum = 1;
 
 	TOptionalWriteAccess(FComponentTypeID InComponentTypeID)
 		: FOptionalWriteAccess{ InComponentTypeID }
 	{}
 
+	FORCEINLINE TWriteOptional<T> ResolvePreLockedComponentData(const FEntityAllocation* Allocation, const FPreLockedDataPtr* Ptr, FEntityAllocationWriteContext WriteContext) const
+	{
+		// Resolve using the same base ptr as PreLockComponentData
+		return TWriteOptional<T>(Ptr->Resolve<T>(Allocation->GetComponentDataAddress()));
+	}
 	FORCEINLINE TComponentLock<TWriteOptional<T>> LockComponentData(const FEntityAllocation* Allocation, FEntityAllocationWriteContext WriteContext) const
 	{
 		return Allocation->TryWriteComponents(ComponentType.ReinterpretCast<T>(), WriteContext);
@@ -160,23 +269,57 @@ struct TOptionalWriteAccess : FOptionalWriteAccess
 };
 
 
+template<typename...>
+struct TUnpackMultiComponentData;
 
+template<typename ...T, int ...Indices>
+struct TUnpackMultiComponentData<TIntegerSequence<int, Indices...>, T...>
+{
+	static TMultiComponentData<TReadOptional<T>...> ResolvePreLockedComponentData(const FEntityAllocation* Allocation, const FOptionalReadAccess* ComponentTypes, const FPreLockedDataPtr* PrelockedComponentData, FEntityAllocationWriteContext WriteContext)
+	{
+		return TMultiComponentData<TReadOptional<T>...>(
+			TOptionalReadAccess<T>(ComponentTypes[Indices].ComponentType).ResolvePreLockedComponentData(Allocation, &PrelockedComponentData[Indices], WriteContext)...
+			);
+	}
+	static TMultiComponentLock<TReadOptional<T>...> LockComponentData(const FEntityAllocation* Allocation, const FOptionalReadAccess* ComponentTypes, FEntityAllocationWriteContext WriteContext)
+	{
+		return TMultiComponentLock<TReadOptional<T>...>(
+			TOptionalReadAccess<T>(ComponentTypes[Indices].ComponentType).LockComponentData(Allocation, WriteContext)...
+			);
+	}
+};
 
 template<typename... T>
 struct TReadOneOfAccessor
 {
 	using AccessType = TMultiComponentLock<TReadOptional<T>...>;
+	static constexpr int32 PreLockedDataNum = sizeof...(T);
 
 	TReadOneOfAccessor(TComponentTypeID<T>... InComponentTypeIDs)
 		: ComponentTypes{ InComponentTypeIDs... }
 	{}
 
-	FORCEINLINE TMultiComponentLock<TReadOptional<T>...> LockComponentData(const FEntityAllocation* Allocation, FEntityAllocationWriteContext WriteContext) const
+	void PreLockComponentData(const FEntityAllocation* Allocation, FPreLockedDataPtr* PrelockedComponentData) const
+	{
+		for (int32 Index = 0; Index < sizeof...(T); ++Index)
+		{
+			ComponentTypes[Index].PreLockComponentData(Allocation, &PrelockedComponentData[Index]);
+		}
+	}
+
+	TMultiComponentData<TReadOptional<T>...> ResolvePreLockedComponentData(const FEntityAllocation* Allocation, const FPreLockedDataPtr* PrelockedComponentData, FEntityAllocationWriteContext WriteContext) const
 	{
-		return TransformTuple(ComponentTypes, [Allocation, WriteContext](auto In){ return In.LockComponentData(Allocation, WriteContext); });
+		return TUnpackMultiComponentData<TMakeIntegerSequence<int, sizeof...(T)>, T...>
+			::ResolvePreLockedComponentData(Allocation, ComponentTypes, PrelockedComponentData, WriteContext);
 	}
 
-	TTuple< TOptionalReadAccess<T>... > ComponentTypes;
+	TMultiComponentLock<TReadOptional<T>...> LockComponentData(const FEntityAllocation* Allocation, FEntityAllocationWriteContext WriteContext) const
+	{
+		return TUnpackMultiComponentData<TMakeIntegerSequence<int, sizeof...(T)>, T...>
+			::LockComponentData(Allocation, ComponentTypes, WriteContext);
+	}
+
+	FOptionalReadAccess ComponentTypes[sizeof...(T)];
 };
 
 
@@ -186,17 +329,33 @@ template<typename... T>
 struct TReadOneOrMoreOfAccessor
 {
 	using AccessType = TMultiComponentLock<TReadOptional<T>...>;
+	static constexpr int32 PreLockedDataNum = sizeof...(T);
 
 	TReadOneOrMoreOfAccessor(TComponentTypeID<T>... InComponentTypeIDs)
 		: ComponentTypes{ InComponentTypeIDs... }
 	{}
 
-	FORCEINLINE TMultiComponentLock<TReadOptional<T>...> LockComponentData(const FEntityAllocation* Allocation, FEntityAllocationWriteContext WriteContext) const
+	void PreLockComponentData(const FEntityAllocation* Allocation, FPreLockedDataPtr* PrelockedComponentData) const
+	{
+		for (int32 Index = 0; Index < sizeof...(T); ++Index)
+		{
+			ComponentTypes[Index].PreLockComponentData(Allocation, &PrelockedComponentData[Index]);
+		}
+	}
+
+	TMultiComponentData<TReadOptional<T>...> ResolvePreLockedComponentData(const FEntityAllocation* Allocation, const FPreLockedDataPtr* PrelockedComponentData, FEntityAllocationWriteContext WriteContext) const
 	{
-		return TransformTuple(ComponentTypes, [Allocation, WriteContext](auto In){ return In.LockComponentData(Allocation, WriteContext); });
+		return TUnpackMultiComponentData<TMakeIntegerSequence<int, sizeof...(T)>, T...>
+			::ResolvePreLockedComponentData(Allocation, ComponentTypes, PrelockedComponentData, WriteContext);
 	}
 
-	TTuple< TOptionalReadAccess<T>... > ComponentTypes;
+	TMultiComponentLock<TReadOptional<T>...> LockComponentData(const FEntityAllocation* Allocation, FEntityAllocationWriteContext WriteContext) const
+	{
+		return TUnpackMultiComponentData<TMakeIntegerSequence<int, sizeof...(T)>, T...>
+			::LockComponentData(Allocation, ComponentTypes, WriteContext);
+	}
+
+	FOptionalReadAccess ComponentTypes[sizeof...(T)];
 };
 
 
@@ -215,7 +374,14 @@ template<typename... T>
 void AddAccessorToFilter(const TReadOneOfAccessor<T...>* In, FEntityComponentFilter* OutFilter)
 {
 	FComponentMask Mask;
-	VisitTupleElements([&Mask](FOptionalReadAccess Composite){ if (Composite.ComponentType) { Mask.Set(Composite.ComponentType); } }, In->ComponentTypes);
+	for (int32 Index = 0; Index < sizeof...(T); ++Index)
+	{
+		FComponentTypeID Component = In->ComponentTypes[Index].ComponentType;
+		if (Component)
+		{
+			Mask.Set(Component);
+		}
+	}
 
 	check(Mask.NumComponents() != 0);
 	OutFilter->Complex(Mask, EComplexFilterMode::OneOf);
@@ -224,7 +390,14 @@ template<typename... T>
 void AddAccessorToFilter(const TReadOneOrMoreOfAccessor<T...>* In, FEntityComponentFilter* OutFilter)
 {
 	FComponentMask Mask;
-	VisitTupleElements([&Mask](FOptionalReadAccess Composite) { if (Composite.ComponentType) { Mask.Set(Composite.ComponentType); } }, In->ComponentTypes);
+	for (int32 Index = 0; Index < sizeof...(T); ++Index)
+	{
+		FComponentTypeID Component = In->ComponentTypes[Index].ComponentType;
+		if (Component)
+		{
+			Mask.Set(Component);
+		}
+	}
 
 	check(Mask.NumComponents() != 0);
 	OutFilter->Complex(Mask, EComplexFilterMode::OneOrMoreOf);
@@ -250,22 +423,18 @@ inline void PopulatePrerequisites(const FOptionalComponentAccess* In, const FSys
 template<typename... T>
 void PopulatePrerequisites(const TReadOneOfAccessor<T...>* In, const FSystemTaskPrerequisites& InPrerequisites, FGraphEventArray* OutGatheredPrereqs)
 {
-	VisitTupleElements(
-		[&InPrerequisites, OutGatheredPrereqs](FOptionalReadAccess Composite)
-		{
-			PopulatePrerequisites(&Composite, InPrerequisites, OutGatheredPrereqs);
-		}
-	, In->ComponentTypes);
+	for (int32 Index = 0; Index < sizeof...(T); ++Index)
+	{
+		PopulatePrerequisites(&In->ComponentTypes[Index], InPrerequisites, OutGatheredPrereqs);
+	}
 }
 template<typename... T>
 void PopulatePrerequisites(const TReadOneOrMoreOfAccessor<T...>* In, const FSystemTaskPrerequisites& InPrerequisites, FGraphEventArray* OutGatheredPrereqs)
 {
-	VisitTupleElements(
-		[&InPrerequisites, OutGatheredPrereqs](FOptionalReadAccess Composite)
-		{
-			PopulatePrerequisites(&Composite, InPrerequisites, OutGatheredPrereqs);
-		}
-	, In->ComponentTypes);
+	for (int32 Index = 0; Index < sizeof...(T); ++Index)
+	{
+		PopulatePrerequisites(&In->ComponentTypes[Index], InPrerequisites, OutGatheredPrereqs);
+	}
 }
 
 
@@ -286,6 +455,56 @@ inline void PopulateSubsequents(const void* In, const FGraphEventRef& InEvent, F
 {
 }
 
+
+
+inline void PopulateReadWriteDependencies(const FEntityIDAccess*, FComponentMask& OutReadDependencies, FComponentMask& OutWriteDependencies)
+{
+}
+inline void PopulateReadWriteDependencies(const FReadAccess* In, FComponentMask& OutReadDependencies, FComponentMask& OutWriteDependencies)
+{
+	checkSlow(In->ComponentType);
+	OutReadDependencies.Set(In->ComponentType);
+}
+inline void PopulateReadWriteDependencies(const FOptionalReadAccess* In, FComponentMask& OutReadDependencies, FComponentMask& OutWriteDependencies)
+{
+	if (In->ComponentType)
+	{
+		OutReadDependencies.Set(In->ComponentType);
+	}
+}
+inline void PopulateReadWriteDependencies(const FWriteAccess* In, FComponentMask& OutReadDependencies, FComponentMask& OutWriteDependencies)
+{
+	checkSlow(In->ComponentType);
+	OutReadDependencies.Set(In->ComponentType);
+	OutWriteDependencies.Set(In->ComponentType);
+}
+inline void PopulateReadWriteDependencies(const FOptionalWriteAccess* In, FComponentMask& OutReadDependencies, FComponentMask& OutWriteDependencies)
+{
+	if (In->ComponentType)
+	{
+		OutReadDependencies.Set(In->ComponentType);
+		OutWriteDependencies.Set(In->ComponentType);
+	}
+}
+template<typename... T>
+void PopulateReadWriteDependencies(const TReadOneOfAccessor<T...>* In, FComponentMask& OutReadDependencies, FComponentMask& OutWriteDependencies)
+{
+	for (int32 Index = 0; Index < sizeof...(T); ++Index)
+	{
+		PopulateReadWriteDependencies(&In->ComponentTypes[Index], OutReadDependencies, OutWriteDependencies);
+	}
+}
+template<typename... T>
+void PopulateReadWriteDependencies(const TReadOneOrMoreOfAccessor<T...>* In, FComponentMask& OutReadDependencies, FComponentMask& OutWriteDependencies)
+{
+	for (int32 Index = 0; Index < sizeof...(T); ++Index)
+	{
+		PopulateReadWriteDependencies(&In->ComponentTypes[Index], OutReadDependencies, OutWriteDependencies);
+	}
+}
+
+
+
 inline bool HasBeenWrittenToSince(const FEntityIDAccess* In, FEntityAllocation* Allocation, uint64 InSystemSerial)
 {
 	return Allocation->HasStructureChangedSince(InSystemSerial);
@@ -306,14 +525,20 @@ template<typename... T>
 bool HasBeenWrittenToSince(const TReadOneOfAccessor<T...>* In, FEntityAllocation* Allocation, uint64 InSystemSerial)
 {
 	bool bAnyWrittenTo = false;
-	VisitTupleElements([Allocation, &bAnyWrittenTo, InSystemSerial](FOptionalReadAccess It){ bAnyWrittenTo |= HasBeenWrittenToSince(&It, Allocation, InSystemSerial); }, In->ComponentTypes);
+	for (int32 Index = 0; Index < sizeof...(T); ++Index)
+	{
+		bAnyWrittenTo |= HasBeenWrittenToSince(&In->ComponentTypes[Index], Allocation, InSystemSerial);
+	}
 	return bAnyWrittenTo;
 }
 template<typename... T>
 bool HasBeenWrittenToSince(const TReadOneOrMoreOfAccessor<T...>* In, FEntityAllocation* Allocation, uint64 InSystemSerial)
 {
 	bool bAnyWrittenTo = false;
-	VisitTupleElements([Allocation, &bAnyWrittenTo, InSystemSerial](FOptionalReadAccess It) { bAnyWrittenTo |= HasBeenWrittenToSince(&It, Allocation, InSystemSerial); }, In->ComponentTypes);
+	for (int32 Index = 0; Index < sizeof...(T); ++Index)
+	{
+		bAnyWrittenTo |= HasBeenWrittenToSince(&In->ComponentTypes[Index], Allocation, InSystemSerial);
+	}
 	return bAnyWrittenTo;
 }
 
@@ -335,14 +560,20 @@ template<typename... T>
 inline bool IsAccessorValid(const TReadOneOfAccessor<T...>* In)
 {
 	bool bValid = false;
-	VisitTupleElements([&bValid](FOptionalReadAccess It){ bValid |= IsAccessorValid(&It); }, In->ComponentTypes);
+	for (int32 Index = 0; Index < sizeof...(T); ++Index)
+	{
+		bValid |= IsAccessorValid(&In->ComponentTypes[Index]);
+	}
 	return bValid;
 }
 template<typename... T>
 inline bool IsAccessorValid(const TReadOneOrMoreOfAccessor<T...>* In)
 {
 	bool bValid = false;
-	VisitTupleElements([&bValid](FOptionalReadAccess It) { bValid |= IsAccessorValid(&It); }, In->ComponentTypes);
+	for (int32 Index = 0; Index < sizeof...(T); ++Index)
+	{
+		bValid |= IsAccessorValid(&In->ComponentTypes[Index]);
+	}
 	return bValid;
 }
 
@@ -364,14 +595,20 @@ template<typename... T>
 inline bool HasAccessorWork(const FEntityManager* EntityManager, const TReadOneOfAccessor<T...>* In)
 {
 	bool bAnyWork = false;
-	VisitTupleElements([&bAnyWork, EntityManager](FOptionalReadAccess It){ bAnyWork |= HasAccessorWork(EntityManager , &It); }, In->ComponentTypes);
+	for (int32 Index = 0; Index < sizeof...(T); ++Index)
+	{
+		bAnyWork |= HasAccessorWork(EntityManager, &In->ComponentTypes[Index]);
+	}
 	return bAnyWork;
 }
 template<typename... T>
 inline bool HasAccessorWork(const FEntityManager* EntityManager, const TReadOneOrMoreOfAccessor<T...>* In)
 {
 	bool bAnyWork = false;
-	VisitTupleElements([&bAnyWork, EntityManager](FOptionalReadAccess It) { bAnyWork |= HasAccessorWork(EntityManager , &It); }, In->ComponentTypes);
+	for (int32 Index = 0; Index < sizeof...(T); ++Index)
+	{
+		bAnyWork |= HasAccessorWork(EntityManager, &In->ComponentTypes[Index]);
+	}
 	return bAnyWork;
 }
 
@@ -388,10 +625,11 @@ inline bool HasAccessorWork(const FEntityManager* EntityManager, const TReadOneO
 	void AccessorToString(const TReadOneOfAccessor<T...>* In, FEntityManager* EntityManager, FString& OutString)
 	{
 		TArray<FString> Strings;
-		VisitTupleElements([&Strings, EntityManager](FOptionalReadAccess ReadOptional)
+
+		for (int32 Index = 0; Index < sizeof...(T); ++Index)
 		{
-			OneOfAccessorToString(&ReadOptional, EntityManager, Strings.Emplace_GetRef());
-		}, In->ComponentTypes);
+			OneOfAccessorToString(&In->ComponentTypes[Index], EntityManager, Strings.Emplace_GetRef());
+		}
 
 		OutString += FString::Printf(TEXT("\n\tRead One Of: [ %s ]"), *FString::Join(Strings, TEXT(",")));
 	}
@@ -400,10 +638,11 @@ inline bool HasAccessorWork(const FEntityManager* EntityManager, const TReadOneO
 	void AccessorToString(const TReadOneOrMoreOfAccessor<T...>* In, FEntityManager* EntityManager, FString& OutString)
 	{
 		TArray<FString> Strings;
-		VisitTupleElements([&Strings, EntityManager](FOptionalReadAccess ReadOptional)
-			{
-				OneOfAccessorToString(&ReadOptional, EntityManager, Strings.Emplace_GetRef());
-			}, In->ComponentTypes);
+
+		for (int32 Index = 0; Index < sizeof...(T); ++Index)
+		{
+			OneOfAccessorToString(&In->ComponentTypes[Index], EntityManager, Strings.Emplace_GetRef());
+		}
 
 		OutString += FString::Printf(TEXT("\n\tRead One Or More Of: [ %s ]"), *FString::Join(Strings, TEXT(",")));
 	}
diff --git a/Engine/Source/Runtime/MovieScene/Public/EntitySystem/MovieSceneComponentPtr.h b/Engine/Source/Runtime/MovieScene/Public/EntitySystem/MovieSceneComponentPtr.h
--- a/Engine/Source/Runtime/MovieScene/Public/EntitySystem/MovieSceneComponentPtr.h
+++ b/Engine/Source/Runtime/MovieScene/Public/EntitySystem/MovieSceneComponentPtr.h
@@ -84,6 +84,11 @@ struct TComponentPtr
 		return ComponentPtr[Index];
 	}
 
+	ValueType ComponentAtIndex(int32 Index) const
+	{
+		return (*this)[Index];
+	}
+
 	FORCEINLINE operator T*() const
 	{
 		return ComponentPtr;
@@ -169,25 +174,29 @@ struct FReadErasedOptional : FReadErased
 			ComponentPtr = InHeader->Components + ComponentOffset*InHeader->Sizeof;
 			Sizeof = InHeader->Sizeof;
 		}
+		else
+		{
+			ComponentPtr = nullptr;
+			Sizeof = 0;
+		}
+	}
+
+	const void* ComponentAtIndex(int32 Index) const
+	{
+		return this->ComponentPtr ? FReadErased::ComponentAtIndex(Index) : nullptr;
 	}
 };
 
 template<typename T>
 struct TRead : TComponentPtr<const T>
 {
-	using ValueType = typename TCallTraits<T>::ParamType;
-
 	TRead() = default;
-	explicit TRead(const T* ComponentPtr)
-		: TComponentPtr<const T>(ComponentPtr)
+	explicit TRead(const T* ComponentPtr, int32 ComponentOffset = 0)
+		: TComponentPtr<const T>(ComponentPtr + ComponentOffset)
 	{}
 	explicit TRead(const FComponentHeader* InHeader, int32 ComponentOffset = 0)
 		: TComponentPtr<const T>(reinterpret_cast<const T*>(InHeader->Components) + ComponentOffset)
 	{}
-	ValueType ComponentAtIndex(int32 Index) const
-	{
-		return (*this)[Index];
-	}
 };
 using FReadEntityIDs = TRead<FMovieSceneEntityID>;
 
@@ -195,6 +204,13 @@ template<typename T>
 struct TReadOptional : TComponentPtr<const T>
 {
 	TReadOptional() = default;
+	explicit TReadOptional(const T* InData, int32 ComponentOffset = 0)
+	{
+		if (InData)
+		{
+			this->ComponentPtr = InData + ComponentOffset;
+		}
+	}
 	explicit TReadOptional(const FComponentHeader* InHeader, int32 ComponentOffset = 0)
 	{
 		if (InHeader)
@@ -202,6 +218,11 @@ struct TReadOptional : TComponentPtr<const T>
 			this->ComponentPtr = reinterpret_cast<const T*>(InHeader->Components) + ComponentOffset;
 		}
 	}
+
+	const T* ComponentAtIndex(int32 Index) const
+	{
+		return this->ComponentPtr ? &this->ComponentPtr[Index] : nullptr;
+	}
 };
 
 
@@ -227,6 +248,10 @@ struct FWriteErased
 	{
 		return ComponentPtr != nullptr;
 	}
+	void* ComponentAtIndex(int32 Index) const
+	{
+		return (*this)[Index];
+	}
 
 protected:
 	uint8* ComponentPtr;
@@ -243,6 +268,15 @@ struct FWriteErasedOptional : FWriteErased
 			ComponentPtr = InHeader->Components + ComponentOffset*InHeader->Sizeof;
 			Sizeof = InHeader->Sizeof;
 		}
+		else
+		{
+			ComponentPtr = nullptr;
+			Sizeof = 0;
+		}
+	}
+	void* ComponentAtIndex(int32 Index) const
+	{
+		return ComponentPtr ? (*this)[Index] : nullptr;
 	}
 };
 
@@ -250,6 +284,9 @@ template<typename T>
 struct TWrite : TComponentPtr<T>
 {
 	TWrite() = default;
+	explicit TWrite(T* InData, int32 ComponentOffset = 0)
+		: TComponentPtr<T>(InData + ComponentOffset)
+	{}
 	explicit TWrite(const FComponentHeader* InHeader, int32 ComponentOffset = 0)
 		: TComponentPtr<T>(reinterpret_cast<T*>(InHeader->Components) + ComponentOffset)
 	{}
@@ -259,6 +296,13 @@ template<typename T>
 struct TWriteOptional : TComponentPtr<T>
 {
 	TWriteOptional() = default;
+	explicit TWriteOptional(T* InData, int32 ComponentOffset = 0)
+	{
+		if (InData)
+		{
+			this->ComponentPtr = InData + ComponentOffset;
+		}
+	}
 	explicit TWriteOptional(const FComponentHeader* InHeader, int32 ComponentOffset = 0)
 	{
 		if (InHeader)
@@ -266,6 +310,11 @@ struct TWriteOptional : TComponentPtr<T>
 			this->ComponentPtr = reinterpret_cast<T*>(InHeader->Components) + ComponentOffset;
 		}
 	}
+
+	T* ComponentAtIndex(int32 Index) const
+	{
+		return this->ComponentPtr ? &this->ComponentPtr[Index] : nullptr;
+	}
 };
 
 
@@ -302,11 +351,6 @@ struct TComponentLock<FReadErasedOptional> : TComponentLockMixin<FScopedHeaderRe
 	{
 		return this->ComponentPtr != nullptr;
 	}
-
-	const void* ComponentAtIndex(int32 Index) const
-	{
-		return this->ComponentPtr != nullptr ? (*this)[Index] : nullptr;
-	}
 };
 template<>
 struct TComponentLock<FWriteErased> : TComponentLockMixin<FScopedHeaderWriteLock>, FWriteErased
@@ -315,11 +359,6 @@ struct TComponentLock<FWriteErased> : TComponentLockMixin<FScopedHeaderWriteLock
 		: TComponentLockMixin<FScopedHeaderWriteLock>(InHeader, InLockMode, InWriteContext)
 		, FWriteErased(InHeader, ComponentOffset)
 	{}
-
-	void* ComponentAtIndex(int32 Index) const
-	{
-		return (*this)[Index];
-	}
 };
 template<>
 struct TComponentLock<FWriteErasedOptional> : TComponentLockMixin<FScopedHeaderWriteLock>, FWriteErasedOptional
@@ -333,16 +372,6 @@ struct TComponentLock<FWriteErasedOptional> : TComponentLockMixin<FScopedHeaderW
 			*static_cast<FWriteErasedOptional*>(this) = FWriteErasedOptional(InHeader, ComponentOffset);
 		}
 	}
-
-	explicit operator bool() const
-	{
-		return this->ComponentPtr != nullptr;
-	}
-
-	void* ComponentAtIndex(int32 Index) const
-	{
-		return this->ComponentPtr != nullptr ? (*this)[Index] : nullptr;
-	}
 };
 template<typename T>
 struct TComponentLock<TRead<T>> : TComponentLockMixin<FScopedHeaderReadLock>, TRead<T>
@@ -365,16 +394,6 @@ struct TComponentLock<TReadOptional<T>> : TComponentLockMixin<FScopedHeaderReadL
 			*static_cast<TReadOptional<T>*>(this) = TReadOptional<T>(InHeader, ComponentOffset);
 		}
 	}
-
-	explicit operator bool() const
-	{
-		return this->ComponentPtr != nullptr;
-	}
-
-	const T* ComponentAtIndex(int32 Index) const
-	{
-		return this->ComponentPtr != nullptr ? this->AsPtr() + Index : nullptr;
-	}
 };
 template<typename T>
 struct TComponentLock<TWrite<T>> : TComponentLockMixin<FScopedHeaderWriteLock>, TWrite<T>
@@ -383,11 +402,6 @@ struct TComponentLock<TWrite<T>> : TComponentLockMixin<FScopedHeaderWriteLock>,
 		: TComponentLockMixin<FScopedHeaderWriteLock>(InHeader, InLockMode, InWriteContext)
 		, TWrite<T>(InHeader, ComponentOffset)
 	{}
-
-	T& ComponentAtIndex(int32 Index) const
-	{
-		return (*this)[Index];
-	}
 };
 
 template<typename T>
@@ -402,16 +416,6 @@ struct TComponentLock<TWriteOptional<T>> : TComponentLockMixin<FScopedHeaderWrit
 			*static_cast<TWriteOptional<T>*>(this) = TWriteOptional<T>(InHeader, ComponentOffset);
 		}
 	}
-
-	explicit operator bool() const
-	{
-		return this->ComponentPtr != nullptr;
-	}
-
-	T* ComponentAtIndex(int32 Index) const
-	{
-		return this->ComponentPtr != nullptr ? this->AsPtr() + Index : nullptr;
-	}
 };
 
 
@@ -430,6 +434,10 @@ template<typename... T> using TMultiComponentLock = TTuple<TComponentLock<T>...>
 template<typename... T>
 struct TMultiComponentData
 {
+	TMultiComponentData(T... InData)
+		: Data(InData...)
+	{}
+
 	TMultiComponentData(const TMultiComponentLock<T...>& InAggregateLock)
 	{
 		auto Init = [](const auto& InLock, auto& OutData)
diff --git a/Engine/Source/Runtime/MovieScene/Public/EntitySystem/MovieSceneEntityBuilder.h b/Engine/Source/Runtime/MovieScene/Public/EntitySystem/MovieSceneEntityBuilder.h
--- a/Engine/Source/Runtime/MovieScene/Public/EntitySystem/MovieSceneEntityBuilder.h
+++ b/Engine/Source/Runtime/MovieScene/Public/EntitySystem/MovieSceneEntityBuilder.h
@@ -503,19 +503,11 @@ struct TAdd : FAdd
 		const FComponentHeader& Header = Allocation->GetComponentHeaderChecked(ComponentTypeID);
 		check(!Header.IsTag());
 
-		if (Allocation->GetCurrentLockMode() != EComponentHeaderLockMode::LockFree)
-		{
-			Header.ReadWriteLock.WriteLock();
-		}
+		FScopedHeaderWriteLock WriteLock(&Header, Allocation->GetCurrentLockMode(), FEntityAllocationWriteContext::NewAllocation());
 
 		T* ComponentPtr = static_cast<T*>(Header.GetValuePtr(ComponentOffset));
 		*ComponentPtr = MoveTemp(Payload.GetValue());
 		Payload.Reset();
-
-		if (Allocation->GetCurrentLockMode() != EComponentHeaderLockMode::LockFree)
-		{
-			Header.ReadWriteLock.WriteUnlock();
-		}
 	}
 };
 
@@ -538,19 +530,11 @@ struct TAddConditional : FAddConditional
 			const FComponentHeader& Header = Allocation->GetComponentHeaderChecked(ComponentTypeID);
 			check(!Header.IsTag());
 
-			if (Allocation->GetCurrentLockMode() != EComponentHeaderLockMode::LockFree)
-			{
-				Header.ReadWriteLock.WriteLock();
-			}
+			FScopedHeaderWriteLock WriteLock(&Header, Allocation->GetCurrentLockMode(), FEntityAllocationWriteContext::NewAllocation());
 
 			T* ComponentPtr = static_cast<T*>(Header.GetValuePtr(ComponentOffset));
 			*ComponentPtr = MoveTemp(Payload.GetValue());
 			Payload.Reset();
-
-			if (Allocation->GetCurrentLockMode() != EComponentHeaderLockMode::LockFree)
-			{
-				Header.ReadWriteLock.WriteUnlock();
-			}
 		}
 	}
 };
diff --git a/Engine/Source/Runtime/MovieScene/Public/EntitySystem/MovieSceneEntityManager.h b/Engine/Source/Runtime/MovieScene/Public/EntitySystem/MovieSceneEntityManager.h
--- a/Engine/Source/Runtime/MovieScene/Public/EntitySystem/MovieSceneEntityManager.h
+++ b/Engine/Source/Runtime/MovieScene/Public/EntitySystem/MovieSceneEntityManager.h
@@ -253,18 +253,11 @@ public:
 		if (Entry.Data.Allocation != nullptr)
 		{
 			const FComponentHeader& Header = Entry.Data.Allocation->GetComponentHeaderChecked(ComponentTypeID);
-			if (Entry.Data.Allocation->GetCurrentLockMode() != EComponentHeaderLockMode::LockFree)
-			{
-				Header.ReadWriteLock.WriteLock();
-			}
+
+			FScopedHeaderWriteLock WriteLock(&Header, Entry.Data.Allocation->GetCurrentLockMode(), FEntityAllocationWriteContext::NewAllocation());
 
 			T* Component = reinterpret_cast<T*>(Header.GetValuePtr(Entry.Data.ComponentOffset));
 			*Component = Forward<ValueType>(InValue);
-
-			if (Entry.Data.Allocation->GetCurrentLockMode() != EComponentHeaderLockMode::LockFree)
-			{
-				Header.ReadWriteLock.WriteUnlock();
-			}
 		}
 	}
 
diff --git a/Engine/Source/Runtime/MovieScene/Public/EntitySystem/MovieSceneEntitySystem.h b/Engine/Source/Runtime/MovieScene/Public/EntitySystem/MovieSceneEntitySystem.h
--- a/Engine/Source/Runtime/MovieScene/Public/EntitySystem/MovieSceneEntitySystem.h
+++ b/Engine/Source/Runtime/MovieScene/Public/EntitySystem/MovieSceneEntitySystem.h
@@ -31,6 +31,7 @@ namespace MovieScene
 
 	struct FSystemSubsequentTasks;
 	struct FSystemTaskPrerequisites;
+	class IEntitySystemScheduler;
 
 	enum class EEntitySystemCategory : uint32
 	{
@@ -184,6 +185,9 @@ public:
 	/** Called when the system is added to a linker */
 	void Link(UMovieSceneEntitySystemLinker* InLinker);
 
+	/** Called to schedule work */
+	void SchedulePersistentTasks(UE::MovieScene::IEntitySystemScheduler* Scheduler);
+
 	/** Called when the system should run its logic */
 	void Run(FSystemTaskPrerequisites& InPrerequisites, FSystemSubsequentTasks& Subsequents);
 
@@ -216,6 +220,8 @@ private:
 
 	virtual void OnLink() {}
 
+	virtual void OnSchedulePersistentTasks(UE::MovieScene::IEntitySystemScheduler* TaskScheduler) { }
+
 	virtual void OnRun(FSystemTaskPrerequisites& InPrerequisites, FSystemSubsequentTasks& Subsequents) { }
 
 	virtual void OnUnlink() {}
diff --git a/Engine/Source/Runtime/MovieScene/Public/EntitySystem/MovieSceneEntitySystemGraphs.h b/Engine/Source/Runtime/MovieScene/Public/EntitySystem/MovieSceneEntitySystemGraphs.h
--- a/Engine/Source/Runtime/MovieScene/Public/EntitySystem/MovieSceneEntitySystemGraphs.h
+++ b/Engine/Source/Runtime/MovieScene/Public/EntitySystem/MovieSceneEntitySystemGraphs.h
@@ -39,6 +39,9 @@ namespace MovieScene
 	struct FSystemSubsequentTasks;
 	struct FSystemTaskPrerequisites;
 
+	class FEntityManager;
+	class FEntitySystemScheduler;
+
 
 } // namespace MovieScene
 } // namespace UE
@@ -87,6 +90,15 @@ struct MOVIESCENE_API FMovieSceneEntitySystemGraph
 
 	GENERATED_BODY()
 
+	FMovieSceneEntitySystemGraph();
+	~FMovieSceneEntitySystemGraph();
+
+	FMovieSceneEntitySystemGraph(const FMovieSceneEntitySystemGraph&) = delete;
+	void operator=(const FMovieSceneEntitySystemGraph&) = delete;
+
+	FMovieSceneEntitySystemGraph(FMovieSceneEntitySystemGraph&&);
+	FMovieSceneEntitySystemGraph& operator=(FMovieSceneEntitySystemGraph&&);
+
 	void AddReference(UMovieSceneEntitySystem* FromReference, UMovieSceneEntitySystem* ToReference);
 
 	void RemoveReference(UMovieSceneEntitySystem* FromReference, UMovieSceneEntitySystem* ToReference);
@@ -151,10 +163,15 @@ struct MOVIESCENE_API FMovieSceneEntitySystemGraph
 
 	void Shutdown();
 
+	int32 NumInPhase(UE::MovieScene::ESystemPhase Phase) const;
+
 	void ExecutePhase(UE::MovieScene::ESystemPhase Phase, UMovieSceneEntitySystemLinker* Linker, FGraphEventArray& OutTasks);
 
 	void IteratePhase(UE::MovieScene::ESystemPhase Phase, TFunctionRef<void(UMovieSceneEntitySystem*)> InIter);
 
+	void ReconstructTaskSchedule(UE::MovieScene::FEntityManager* EntityManager);
+	void ScheduleTasks(UE::MovieScene::FEntityManager* EntityManager);
+
 	TArray<UMovieSceneEntitySystem*> GetSystems() const;
 
 	template<typename SystemType>
@@ -184,17 +201,31 @@ private:
 
 	TArray<uint16, TInlineAllocator<4>>  SpawnPhase;
 	TArray<uint16, TInlineAllocator<8>>  InstantiationPhase;
+	TArray<uint16, TInlineAllocator<16>> SchedulingPhase;
 	TArray<uint16, TInlineAllocator<16>> EvaluationPhase;
 	TArray<uint16, TInlineAllocator<2>>  FinalizationPhase;
 
+	TUniquePtr<UE::MovieScene::FEntitySystemScheduler> TaskScheduler;
+
 	UPROPERTY()
 	FMovieSceneEntitySystemGraphNodes Nodes;
 
 	TMap<uint16, uint16> GlobalToLocalNodeIDs;
 
 	UE::MovieScene::FDirectedGraph ReferenceGraph;
+	uint64 SchedulerSerialNumber = 0;
 
 	uint32 SerialNumber = 0;
 	uint32 PreviousSerialNumber = 0;
 	uint32 ReentrancyGuard = 0;
 };
+
+
+template<>
+struct TStructOpsTypeTraits<FMovieSceneEntitySystemGraph> : public TStructOpsTypeTraitsBase2<FMovieSceneEntitySystemGraph>
+{
+	enum
+	{
+		WithCopy = false
+	};
+};
\ No newline at end of file
