Sequencer: Added custom, persistent scheduler for dispatching ECS tasks

diff --git a/Engine/Source/Runtime/MovieScene/Public/EntitySystem/MovieSceneEntitySystemTask.h b/Engine/Source/Runtime/MovieScene/Public/EntitySystem/MovieSceneEntitySystemTask.h
--- a/Engine/Source/Runtime/MovieScene/Public/EntitySystem/MovieSceneEntitySystemTask.h
+++ b/Engine/Source/Runtime/MovieScene/Public/EntitySystem/MovieSceneEntitySystemTask.h
@@ -3,6 +3,7 @@
 #pragma once
 
 #include "Async/TaskGraphInterfaces.h"
+#include "Misc/GeneratedTypeName.h"
 #include "MovieSceneEntityIDs.h"
 #include "EntitySystem/MovieSceneEntityManager.h"
 #include "EntitySystem/EntityAllocationIterator.h"
@@ -11,6 +12,7 @@
 #include "EntitySystem/MovieSceneEntityManager.h"
 #include "EntitySystem/MovieSceneSystemTaskDependencies.h"
 #include "EntitySystem/MovieSceneComponentPtr.h"
+#include "EntitySystem/IMovieSceneTaskScheduler.h"
 
 #include "Templates/AndOrNot.h"
 
@@ -31,26 +33,33 @@ template<typename> struct TOptionalWriteAccess;
 template<typename...> struct TReadOneOfAccessor;
 template<typename...> struct TReadOneOrMoreOfAccessor;
 
+template<typename...> struct TPrelockedDataOffsets;
+
 template<typename...> struct TFilteredEntityTask;
 template<typename...> struct TEntityTaskComponents;
 template<typename...> struct TEntityTaskComponentsImpl;
 
 template<typename, typename...> struct TEntityTask;
 template<typename, typename...> struct TEntityTaskBase;
+template<typename, typename...> struct TScheduledEntityTask;
 template<typename, typename...> struct TEntityAllocationTask;
 template<typename, typename...> struct TEntityAllocationTaskBase;
 template<typename> struct TUnstructuredTask;
 
 
-template<typename TaskImpl, int32 NumComponents, bool AutoExpandAccessors>
+template<int32 NumComponents, bool AutoExpandAccessors>
 struct TEntityTaskCaller;
 
 struct FCommonEntityTaskParams
 {
-	/** A stat ID for the task */
-	TStatId StatId;
+	FCommonEntityTaskParams()
+		: TaskParams(nullptr)
+	{}
 
-	/** The thread that this task wants to run on */
+	/** Task parameters */
+	FTaskParams TaskParams;
+
+	/** (deprecated) The thread that this task wants to run on */
 	ENamedThreads::Type DesiredThread = ENamedThreads::AnyHiPriThreadHiPriTask;
 
 	/** Useful for debugging to break the debugger when this task is run */
@@ -117,6 +126,45 @@ struct FEntityIterationResult
 	bool Value = true;
 };
 
+static FTaskID SchedulePreTask(IEntitySystemScheduler* InScheduler, const FTaskParams& Params, TSharedPtr<ITaskContext> Context, void*, ...)
+{
+	return FTaskID::None();
+}
+
+static FTaskID SchedulePostTask(IEntitySystemScheduler* InScheduler, const FTaskParams& Params, TSharedPtr<ITaskContext> Context, void*, ...)
+{
+	return FTaskID::None();
+}
+
+template <typename TaskType, typename TaskImplType>
+static FTaskID SchedulePreTask(IEntitySystemScheduler* InScheduler, const FTaskParams& Params, TSharedPtr<TaskType> Context, TaskImplType* Unused, decltype(&TaskImplType::PreTask)* = 0)
+{
+	struct FCallPreTask
+	{
+		static void Run(const ITaskContext* Context, FEntityAllocationWriteContext WriteContext)
+		{
+			static_cast<const TaskType*>(Context)->PreTask();
+		}
+	};
+
+	TaskFunctionPtr Function(TInPlaceType<UnboundTaskFunctionPtr>(), &FCallPreTask::Run);
+	return InScheduler->AddTask(Params, Context, Function);
+}
+template <typename TaskType, typename TaskImplType>
+static FTaskID SchedulePostTask(IEntitySystemScheduler* InScheduler, const FTaskParams& Params, TSharedPtr<TaskType> Context, TaskImplType* Unused, decltype(&TaskImplType::PostTask)* = 0)
+{
+	struct FCallPostTask
+	{
+		static void Run(const ITaskContext* Context, FEntityAllocationWriteContext WriteContext)
+		{
+			static_cast<const TaskType*>(Context)->PostTask();
+		}
+	};
+
+	TaskFunctionPtr Function(TInPlaceType<UnboundTaskFunctionPtr>(), &FCallPostTask::Run);
+	return InScheduler->AddTask(Params, Context, Function);
+}
+
 /**
  * Defines the accessors for each desired component of an entity task
  */
@@ -226,6 +274,7 @@ struct TEntityTaskComponents : TEntityTaskComponentsImpl<TMakeIntegerSequence<in
 	TEntityTaskComponents< T... >& SetDesiredThread(ENamedThreads::Type InDesiredThread)
 	{
 		this->CommonParams.DesiredThread = InDesiredThread;
+		this->CommonParams.TaskParams.bForceGameThread = (InDesiredThread == ENamedThreads::GameThread || InDesiredThread == ENamedThreads::GameThread_Local);
 		return *this;
 	}
 
@@ -234,7 +283,16 @@ struct TEntityTaskComponents : TEntityTaskComponentsImpl<TMakeIntegerSequence<in
 	 */
 	TEntityTaskComponents< T... >& SetStat(TStatId InStatId)
 	{
-		this->CommonParams.StatId = InStatId;
+		this->CommonParams.TaskParams.StatId = InStatId;
+		return *this;
+	}
+
+	/**
+	 * Assign the scheduled task parameters for this task
+	 */
+	TEntityTaskComponents< T... >& SetParams(const FTaskParams& InOtherParams)
+	{
+		this->CommonParams.TaskParams = InOtherParams;
 		return *this;
 	}
 
@@ -432,6 +490,78 @@ struct TEntityTaskComponents : TEntityTaskComponentsImpl<TMakeIntegerSequence<in
 		}
 	}
 
+	template<typename TaskImpl, typename... TaskConstructionArgs>
+	FTaskID Fork_PerEntity(FEntityManager* EntityManager, IEntitySystemScheduler* InScheduler, TaskConstructionArgs&&... InArgs) const
+	{
+		FTaskParams FinalParams = this->CommonParams.TaskParams;
+		FinalParams.bSerialTasks = false;
+
+#if !(UE_BUILD_SHIPPING || UE_BUILD_TEST)
+		FinalParams.DebugName = GetGeneratedTypeName<TaskImpl>();
+#endif
+
+		return ScheduleImpl<TaskImpl>(
+			EntityManager,
+			InScheduler,
+			FinalParams,
+			TaskFunctionPtr(TInPlaceType<PreLockedAllocationItemFunctionPtr>(), TScheduledEntityTask<TaskImpl, T...>::ScheduledRun_PerEntity),
+			Forward<TaskConstructionArgs>(InArgs)...);
+	}
+
+	template<typename TaskImpl, typename... TaskConstructionArgs>
+	FTaskID Fork_PerAllocation(FEntityManager* EntityManager, IEntitySystemScheduler* InScheduler, TaskConstructionArgs&&... InArgs) const
+	{
+		FTaskParams FinalParams = this->CommonParams.TaskParams;
+		FinalParams.bSerialTasks = false;
+
+#if !(UE_BUILD_SHIPPING || UE_BUILD_TEST)
+		FinalParams.DebugName = GetGeneratedTypeName<TaskImpl>();
+#endif
+
+		return ScheduleImpl<TaskImpl>(
+			EntityManager,
+			InScheduler,
+			FinalParams,
+			TaskFunctionPtr(TInPlaceType<PreLockedAllocationItemFunctionPtr>(), TScheduledEntityTask<TaskImpl, T...>::ScheduledRun_PerAllocation),
+			Forward<TaskConstructionArgs>(InArgs)...);
+	}
+
+	template<typename TaskImpl, typename... TaskConstructionArgs>
+	FTaskID Schedule_PerEntity(FEntityManager* EntityManager, IEntitySystemScheduler* InScheduler, TaskConstructionArgs&&... InArgs) const
+	{
+		FTaskParams FinalParams = this->CommonParams.TaskParams;
+		FinalParams.bSerialTasks = true;
+
+#if !(UE_BUILD_SHIPPING || UE_BUILD_TEST)
+		FinalParams.DebugName = GetGeneratedTypeName<TaskImpl>();
+#endif
+
+		return ScheduleImpl<TaskImpl>(
+			EntityManager,
+			InScheduler,
+			FinalParams,
+			TaskFunctionPtr(TInPlaceType<PreLockedAllocationItemFunctionPtr>(), TScheduledEntityTask<TaskImpl, T...>::ScheduledRun_PerEntity),
+			Forward<TaskConstructionArgs>(InArgs)...);
+	}
+
+	template<typename TaskImpl, typename... TaskConstructionArgs>
+	FTaskID Schedule_PerAllocation(FEntityManager* EntityManager, IEntitySystemScheduler* InScheduler, TaskConstructionArgs&&... InArgs) const
+	{
+		FTaskParams FinalParams = this->CommonParams.TaskParams;
+		FinalParams.bSerialTasks = false;
+
+#if !(UE_BUILD_SHIPPING || UE_BUILD_TEST)
+		FinalParams.DebugName = GetGeneratedTypeName<TaskImpl>();
+#endif
+
+		return ScheduleImpl<TaskImpl>(
+			EntityManager,
+			InScheduler,
+			FinalParams,
+			TaskFunctionPtr(TInPlaceType<PreLockedAllocationItemFunctionPtr>(), TScheduledEntityTask<TaskImpl, T...>::ScheduledRun_PerAllocation),
+			Forward<TaskConstructionArgs>(InArgs)...);
+	}
+
 	template<typename TaskImpl>
 	void RunInline_PerEntity(FEntityManager* EntityManager, TaskImpl& Task) const
 	{
@@ -452,6 +582,48 @@ public:
 	explicit TEntityTaskComponents(const FCommonEntityTaskParams& InCommonParams, ConstructionTypes&&... InTypes)
 		: Super(InCommonParams, Forward<ConstructionTypes>(InTypes)...)
 	{}
+
+private:
+
+	template<typename TaskImpl, typename... TaskConstructionArgs>
+	FTaskID ScheduleImpl(FEntityManager* EntityManager, IEntitySystemScheduler* InScheduler, const FTaskParams& TaskParams, TaskFunctionPtr InFunction, TaskConstructionArgs&&... InArgs) const
+	{
+		if (!this->IsValid())
+		{
+			return FTaskID::None();
+		}
+
+		if (!this->HasAnyWork(EntityManager))
+		{
+			return FTaskID::None();
+		}
+
+		using TaskType = TScheduledEntityTask<TaskImpl, T...>;
+		TSharedPtr<TaskType> SharedTask = MakeShared<TaskType>(*this, Forward<TaskConstructionArgs>(InArgs)...);
+
+		FComponentMask ReadDependencies, WriteDependencies;
+		this->PopulateReadWriteDependencies(ReadDependencies, WriteDependencies);
+
+		FEntityComponentFilter Filter;
+		this->PopulateFilter(&Filter);
+
+		auto PrelockComponentData = [this](FEntityAllocationIteratorItem Allocation, TArray<FPreLockedDataPtr>& OutComponentHeaders)
+		{
+			this->PreLockComponentHeaders(Allocation, OutComponentHeaders);
+		};
+		FTaskID TaskID = InScheduler->CreateForkedAllocationTask(TaskParams, SharedTask, InFunction, PrelockComponentData, Filter, ReadDependencies, WriteDependencies);
+
+		if (TaskID)
+		{
+			FTaskID PreTask  = SchedulePreTask(InScheduler, TaskParams, SharedTask, (TaskImpl*)0);
+			FTaskID PostTask = SchedulePostTask(InScheduler, TaskParams, SharedTask, (TaskImpl*)0);
+
+			InScheduler->AddChildFront(TaskID, PreTask);
+			InScheduler->AddChildBack(TaskID, PostTask);
+		}
+
+		return TaskID;
+	}
 };
 
 
@@ -534,6 +706,17 @@ struct TEntityTaskComponentsImpl<TIntegerSequence<int, Indices...>, T...>
 		return TEntityTaskComponents<T..., FErasedReadAccess >(CommonParams, Accessors.template Get<Indices>()..., FErasedReadAccess(ComponentType));
 	}
 
+	/**
+	 * Optionally read the type-erased value of a component. Passed to the task as a const void*
+	 * @note Supplying an invalid ComponentType will be handled gracefully, but will result in no task being dispatched.
+	 *
+	 * @param ComponentType   A valid component type to read.
+	 */
+	TEntityTaskComponents<T..., FErasedOptionalReadAccess> ReadErasedOptional(FComponentTypeID ComponentType) const
+	{
+		return TEntityTaskComponents<T..., FErasedOptionalReadAccess >(CommonParams, Accessors.template Get<Indices>()..., FErasedOptionalReadAccess(ComponentType));
+	}
+
 	/**
 	 * Optionally read the value of a component. ComponentType may be invalid, and the component may or may not exist for some/all of the entities in the resulting task
 	 * @note Always passed to the task as a const T* pointer which must be checked for null
@@ -651,37 +834,18 @@ struct TEntityTaskComponentsImpl<TIntegerSequence<int, Indices...>, T...>
 		(void)Temp;
 	}
 
-	/**
-	 * Perform a direct iteration of the specified allocation using this task, inline on the current thread
-	 * @note: This is highly unsafe as it circumvents all the thread-safety mechanisms that protect component data
-	 *
-	 * @param Allocation  The allocation to iterate
-	 * @return An iterator that defines the full range of entities in the allocation
-	 */
-	TEntityRange<typename T::AccessType...> IterateAllocation(const FEntityAllocation* Allocation, FEntityAllocationWriteContext WriteContext) const
+	void PreLockComponentHeaders(const FEntityAllocation* Allocation, TArray<FPreLockedDataPtr>& OutComponentHeaders) const
 	{
-		checkf(IsValid(), TEXT("Attempting to use a component pack with an invalid component type."));
-
-		auto LockedComponentData = MakeTuple( Accessors.template Get<Indices>().LockComponentData(Allocation, WriteContext)... );
-
-		// WARNING: This is highly unsafe as it circumvents all the thread-safety mechanisms that protect component data
-		return TEntityRange<typename T::AccessType...>(Allocation->Num(), LockedComponentData.template Get<Indices>().AsPtr()... );
+		constexpr TPrelockedDataOffsets<T...> PrelockedDataOffsets;
+		int32 StartIndex = OutComponentHeaders.Num();
+		OutComponentHeaders.AddDefaulted(PrelockedDataOffsets.Num);
+		( Accessors.template Get<Indices>().PreLockComponentData(Allocation, &OutComponentHeaders[StartIndex + PrelockedDataOffsets.StartOffset[Indices]]), ... );
 	}
 
-	/**
-	 * Perform a direct iteration of the specified entity range using this task, inline on the current thread
-	 * @note: This is highly unsafe as it circumvents all the thread-safety mechanisms that protect component data
-	 *
-	 * @param Allocation  The allocation to iterate
-	 * @return An iterator that defines the full range of entities in the allocation
-	 */
-	TEntityRange<typename T::AccessType...> IterateRange(const FEntityRange& EntityRange, FEntityAllocationWriteContext WriteContext) const
+	void PopulateReadWriteDependencies(FComponentMask& OutReadDependencies, FComponentMask& OutWriteDependencies) const
 	{
-		check(EntityRange.ComponentStartOffset >= 0 && EntityRange.ComponentStartOffset + EntityRange.Num <= EntityRange.Allocation->Num());
-
-		TEntityRange<typename T::AccessType...> Result = IterateAllocation(EntityRange.Allocation, WriteContext);
-		Result.Slice(EntityRange.ComponentStartOffset, EntityRange.Num);
-		return Result;
+		int Temp[] = { (UE::MovieScene::PopulateReadWriteDependencies(&Accessors.template Get<Indices>(), OutReadDependencies, OutWriteDependencies), 0)..., 0 };
+		(void)Temp;
 	}
 
 	/**
@@ -953,6 +1117,7 @@ struct TFilteredEntityTask
 	TFilteredEntityTask< T... >& SetDesiredThread(ENamedThreads::Type InDesiredThread)
 	{
 		Components.CommonParams.DesiredThread = InDesiredThread;
+		Components.CommonParams.TaskParams.bForceGameThread = (InDesiredThread == ENamedThreads::GameThread || InDesiredThread == ENamedThreads::GameThread_Local);
 		return *this;
 	}
 
@@ -961,7 +1126,16 @@ struct TFilteredEntityTask
 	 */
 	TFilteredEntityTask< T... >& SetStat(TStatId InStatId)
 	{
-		Components.CommonParams.StatId = InStatId;
+		Components.CommonParams.TaskParams.StatId = InStatId;
+		return *this;
+	}
+
+	/**
+	 * Assign the scheduled task parameters for this task
+	 */
+	TEntityTaskComponents< T... >& SetParams(const FTaskParams& InOtherParams)
+	{
+		this->CommonParams.TaskParams = InOtherParams;
 		return *this;
 	}
 
@@ -978,7 +1152,7 @@ struct TFilteredEntityTask
 	 */
 	TStatId GetStatId() const
 	{
-		return Components.CommonParams.StatId;
+		return Components.CommonParams.TaskParams.StatId;
 	}
 
 	/**
@@ -1005,6 +1179,46 @@ struct TFilteredEntityTask
 		return Components;
 	}
 
+	TFilteredEntityTask<T...>& AddDynamicReadDependency(std::initializer_list<FComponentTypeID> InReadDependencies)
+	{
+		DynamicReadMask.SetAll(InReadDependencies);
+		return *this;
+	}
+	template<typename ComponentType>
+	TFilteredEntityTask<T...>& AddDynamicReadDependency(TArrayView<ComponentType> InReadDependencies)
+	{
+		for (FComponentTypeID Component : InReadDependencies)
+		{
+			DynamicReadMask.Set(Component);
+		}
+		return *this;
+	}
+	TFilteredEntityTask<T...>& AddDynamicReadDependency(const FComponentMask& InDynamicReadDependency)
+	{
+		DynamicReadMask.CombineWithBitwiseOR(InDynamicReadDependency, EBitwiseOperatorFlags::MaxSize);
+		return *this;
+	}
+
+	TFilteredEntityTask<T...>& AddDynamicWriteDependency(std::initializer_list<FComponentTypeID> InWriteDependencies)
+	{
+		DynamicWriteMask.SetAll(InWriteDependencies);
+		return *this;
+	}
+	template<typename ComponentType>
+	TFilteredEntityTask<T...>& AddDynamicWriteDependency(TArrayView<ComponentType> InWriteDependencies)
+	{
+		for (FComponentTypeID Component : InWriteDependencies)
+		{
+			DynamicWriteMask.Set(Component);
+		}
+		return *this;
+	}
+	TFilteredEntityTask<T...>& AddDynamicWriteDependency(const FComponentMask& InDynamicWriteDependency)
+	{
+		DynamicWriteMask.CombineWithBitwiseOR(InDynamicWriteDependency, EBitwiseOperatorFlags::MaxSize);
+		return *this;
+	}
+
 	/**
 	 * Dispatch a task for every entity that matches the filters and component types. Must be explicitly instantiated with the task type to dispatch. Construction arguments are deduced.
 	 * Tasks must implement a ForEachEntity function that matches this task's component accessor types.
@@ -1152,6 +1366,78 @@ struct TFilteredEntityTask
 		}
 	}
 
+	template<typename TaskImpl, typename... TaskConstructionArgs>
+	FTaskID Fork_PerEntity(FEntityManager* EntityManager, IEntitySystemScheduler* InScheduler, TaskConstructionArgs&&... InArgs) const
+	{
+		FTaskParams FinalParams = this->Components.CommonParams.TaskParams;
+		FinalParams.bSerialTasks = false;
+
+#if !(UE_BUILD_SHIPPING || UE_BUILD_TEST)
+		FinalParams.DebugName = GetGeneratedTypeName<TaskImpl>();
+#endif
+
+		return ScheduleImpl<TaskImpl>(
+			EntityManager,
+			InScheduler,
+			FinalParams,
+			TaskFunctionPtr(TInPlaceType<PreLockedAllocationItemFunctionPtr>(), TScheduledEntityTask<TaskImpl, T...>::ScheduledRun_PerEntity),
+			Forward<TaskConstructionArgs>(InArgs)...);
+	}
+
+	template<typename TaskImpl, typename... TaskConstructionArgs>
+	FTaskID Fork_PerAllocation(FEntityManager* EntityManager, IEntitySystemScheduler* InScheduler, TaskConstructionArgs&&... InArgs) const
+	{
+		FTaskParams FinalParams = this->Components.CommonParams.TaskParams;
+		FinalParams.bSerialTasks = false;
+
+#if !(UE_BUILD_SHIPPING || UE_BUILD_TEST)
+		FinalParams.DebugName = GetGeneratedTypeName<TaskImpl>();
+#endif
+
+		return ScheduleImpl<TaskImpl>(
+			EntityManager,
+			InScheduler,
+			FinalParams,
+			TaskFunctionPtr(TInPlaceType<PreLockedAllocationItemFunctionPtr>(), TScheduledEntityTask<TaskImpl, T...>::ScheduledRun_PerAllocation),
+			Forward<TaskConstructionArgs>(InArgs)...);
+	}
+
+	template<typename TaskImpl, typename... TaskConstructionArgs>
+	FTaskID Schedule_PerEntity(FEntityManager* EntityManager, IEntitySystemScheduler* InScheduler, TaskConstructionArgs&&... InArgs) const
+	{
+		FTaskParams FinalParams = this->Components.CommonParams.TaskParams;
+		FinalParams.bSerialTasks = true;
+
+#if !(UE_BUILD_SHIPPING || UE_BUILD_TEST)
+		FinalParams.DebugName = GetGeneratedTypeName<TaskImpl>();
+#endif
+
+		return ScheduleImpl<TaskImpl>(
+			EntityManager,
+			InScheduler,
+			FinalParams,
+			TaskFunctionPtr(TInPlaceType<PreLockedAllocationItemFunctionPtr>(), TScheduledEntityTask<TaskImpl, T...>::ScheduledRun_PerEntity),
+			Forward<TaskConstructionArgs>(InArgs)...);
+	}
+
+	template<typename TaskImpl, typename... TaskConstructionArgs>
+	FTaskID Schedule_PerAllocation(FEntityManager* EntityManager, IEntitySystemScheduler* InScheduler, TaskConstructionArgs&&... InArgs) const
+	{
+		FTaskParams FinalParams = this->Components.CommonParams.TaskParams;
+		FinalParams.bSerialTasks = true;
+
+#if !(UE_BUILD_SHIPPING || UE_BUILD_TEST)
+		FinalParams.DebugName = GetGeneratedTypeName<TaskImpl>();
+#endif
+
+		return ScheduleImpl<TaskImpl>(
+			EntityManager,
+			InScheduler,
+			FinalParams,
+			TaskFunctionPtr(TInPlaceType<PreLockedAllocationItemFunctionPtr>(), TScheduledEntityTask<TaskImpl, T...>::ScheduledRun_PerAllocation),
+			Forward<TaskConstructionArgs>(InArgs)...);
+	}
+
 	template<typename TaskImpl>
 	void RunInline_PerEntity(FEntityManager* EntityManager, TaskImpl Task) const
 	{
@@ -1185,11 +1471,51 @@ struct TFilteredEntityTask
 		Components.Iterate_PerAllocationImpl(EntityManager, Filter, Forward<Callback>(InCallback));
 	}
 
+private:
+
+	template<typename TaskImpl, typename... TaskConstructionArgs>
+	FTaskID ScheduleImpl(FEntityManager* EntityManager, IEntitySystemScheduler* InScheduler, const FTaskParams& TaskParams, TaskFunctionPtr InFunction, TaskConstructionArgs&&... InArgs) const
+	{
+		if (!Components.IsValid())
+		{
+			return FTaskID::None();
+		}
+
+		if (!Components.HasAnyWork(EntityManager))
+		{
+			return FTaskID::None();
+		}
+
+		using TaskType = TScheduledEntityTask<TaskImpl, T...>;
+		TSharedPtr<TaskType> SharedTask = MakeShared<TaskType>(Components, Forward<TaskConstructionArgs>(InArgs)...);
+
+		FComponentMask ReadDependencies = DynamicReadMask, WriteDependencies = DynamicWriteMask;
+		Components.PopulateReadWriteDependencies(ReadDependencies, WriteDependencies);
+
+		auto PrelockComponentData = [this](FEntityAllocationIteratorItem Allocation, TArray<FPreLockedDataPtr>& OutComponentHeaders)
+		{
+			this->Components.PreLockComponentHeaders(Allocation, OutComponentHeaders);
+		};
+
+		FTaskID TaskID = InScheduler->CreateForkedAllocationTask(TaskParams, SharedTask, InFunction, PrelockComponentData, Filter, ReadDependencies, WriteDependencies);
+		if (TaskID)
+		{
+			FTaskID PreTask  = SchedulePreTask(InScheduler, TaskParams, SharedTask, (TaskImpl*)0);
+			FTaskID PostTask = SchedulePostTask(InScheduler, TaskParams, SharedTask, (TaskImpl*)0);
+			InScheduler->AddChildFront(TaskID, PreTask);
+			InScheduler->AddChildBack(TaskID, PostTask);
+		}
+
+		return TaskID;
+	}
+
 private:
 
 	TEntityTaskComponents<T...> Components;
 
 	FEntityComponentFilter Filter;
+	FComponentMask DynamicReadMask;
+	FComponentMask DynamicWriteMask;
 };
 
 template<typename TaskImpl, typename... ComponentTypes>
@@ -1250,18 +1576,16 @@ protected:
 		InTask->PostTask();
 	}
 
-	using Caller = TEntityTaskCaller< TaskImpl, sizeof...(ComponentTypes), TEntityTaskTraits<TaskImpl>::AutoExpandAccessors >;
+	using Caller = TEntityTaskCaller<sizeof...(ComponentTypes), TEntityTaskTraits<TaskImpl>::AutoExpandAccessors >;
 
 	TFilteredEntityTask<ComponentTypes...> FilteredTask;
 	FEntityManager* EntityManager;
 	FEntityAllocationWriteContext WriteContext;
 };
 
-
 template<typename TaskImpl, typename... ComponentTypes>
 struct TEntityTask : TEntityTaskBase<TaskImpl, ComponentTypes...>
 {
-
 	template<typename... ArgTypes>
 	explicit TEntityTask(FEntityManager* InEntityManager, const TEntityTaskComponents<ComponentTypes...>& InComponents, ArgTypes&&... InArgs)
 		: TEntityTaskBase<TaskImpl, ComponentTypes...>(InEntityManager, InComponents)
@@ -1299,6 +1623,56 @@ private:
 	TaskImpl TaskImplInstance;
 };
 
+template<typename TaskImpl, typename... ComponentTypes>
+struct TScheduledEntityTask : ITaskContext
+{
+	template<typename... ArgTypes>
+	explicit TScheduledEntityTask(const TEntityTaskComponents<ComponentTypes...>& InComponents, ArgTypes&&... InArgs)
+		: TaskImplInstance{ Forward<ArgTypes>(InArgs)... }
+		, Components(InComponents)
+	{}
+
+	void PreTask() const
+	{
+		// This const_cast is ok because we know nothing else can touch this task in PreTask
+		PreTaskImpl(const_cast<TaskImpl*>(&TaskImplInstance));
+	}
+	void PostTask() const
+	{
+		// This const_cast is ok because we know nothing else can touch this task in PostTask
+		PostTaskImpl(const_cast<TaskImpl*>(&TaskImplInstance));
+	}
+
+	static void ScheduledRun_PerEntity(FEntityAllocationIteratorItem Item, TArrayView<const FPreLockedDataPtr> PreLockedData, const ITaskContext* Context, FEntityAllocationWriteContext WriteContext)
+	{
+		const TScheduledEntityTask<TaskImpl, ComponentTypes...>* This = static_cast<const TScheduledEntityTask<TaskImpl, ComponentTypes...>*>(Context);
+		Caller::ForEachEntityImpl(This->TaskImplInstance, Item, PreLockedData, WriteContext, This->Components);
+	}
+	static void ScheduledRun_PerAllocation(FEntityAllocationIteratorItem Item, TArrayView<const FPreLockedDataPtr> PreLockedData, const ITaskContext* Context, FEntityAllocationWriteContext WriteContext)
+	{
+		const TScheduledEntityTask<TaskImpl, ComponentTypes...>* This = static_cast<const TScheduledEntityTask<TaskImpl, ComponentTypes...>*>(Context);
+		Caller::ForEachAllocationImpl(This->TaskImplInstance, Item, PreLockedData, WriteContext, This->Components);
+	}
+private:
+
+	static void PreTaskImpl(void*, ...){}
+	template <typename T> static void PreTaskImpl(T* InTask, decltype(&T::PreTask)* = 0)
+	{
+		InTask->PreTask();
+	}
+
+	static void PostTaskImpl(void*, ...){}
+	template <typename T> static void PostTaskImpl(T* InTask, decltype(&T::PostTask)* = 0)
+	{
+		InTask->PostTask();
+	}
+
+	using Caller = TEntityTaskCaller<sizeof...(ComponentTypes), TEntityTaskTraits<TaskImpl>::AutoExpandAccessors >;
+
+	TaskImpl TaskImplInstance;
+	TEntityTaskComponents<ComponentTypes...> Components;
+};
+
 template<typename TaskImpl, typename... ComponentTypes>
 struct TEntityAllocationTaskBase
 {
@@ -1359,7 +1733,7 @@ private:
 
 protected:
 
-	using Caller = TEntityTaskCaller< TaskImpl, sizeof...(ComponentTypes), TEntityTaskTraits<TaskImpl>::AutoExpandAccessors >;
+	using Caller = TEntityTaskCaller<sizeof...(ComponentTypes), TEntityTaskTraits<TaskImpl>::AutoExpandAccessors >;
 
 	TFilteredEntityTask<ComponentTypes...> ComponentFilter;
 	FEntityManager* EntityManager;
@@ -1417,7 +1791,7 @@ struct TUnstructuredTask
 
 	TStatId GetStatId() const
 	{
-		return CommonParams.StatId;
+		return CommonParams.TaskParams.StatId;
 	}
 
 	ENamedThreads::Type GetDesiredThread() const
@@ -1456,11 +1830,47 @@ private:
 
 template<typename...> struct TEntityTaskCaller_AutoExpansion;
 
-template<typename TaskImpl, int... Indices>
-struct TEntityTaskCaller_AutoExpansion<TaskImpl, TIntegerSequence<int, Indices...>>
+
+template<typename ...AccessorTypes>
+struct TPrelockedDataOffsets
 {
-	template<typename... AccessorTypes>
-	FORCEINLINE static void ForEachEntityImpl(TaskImpl& TaskImplInstance, const FEntityAllocation* Allocation, FEntityAllocationWriteContext WriteContext, const TEntityTaskComponents<AccessorTypes...>& Components)
+	static constexpr int32 Num = (... + AccessorTypes::PreLockedDataNum);
+
+	int32 StartOffset[sizeof...(AccessorTypes)]={};
+
+	constexpr TPrelockedDataOffsets()
+	{
+		int Index = 0, DataIndex = 0;
+		(..., Assign(Index, DataIndex, AccessorTypes::PreLockedDataNum));
+	}
+	constexpr void Assign(int& Index, int& DataIndex, int DataSize)
+	{
+		StartOffset[Index] = DataIndex;
+		DataIndex += DataSize;
+		++Index;
+	}
+};
+
+template<int... Indices>
+struct TEntityTaskCaller_AutoExpansion<TIntegerSequence<int, Indices...>>
+{
+	template<typename TaskImpl, typename... AccessorTypes>
+	static void ForEachEntityImpl(TaskImpl& TaskImplInstance, const FEntityAllocation* Allocation, TArrayView<const FPreLockedDataPtr> PreLockedData, FEntityAllocationWriteContext WriteContext, const TEntityTaskComponents<AccessorTypes...>& Components)
+	{
+		FEntityIterationResult Result;
+
+		constexpr TPrelockedDataOffsets<AccessorTypes...> PrelockedDataOffsets;
+		auto ResolvedComponentData = MakeTuple( Components.template GetAccessor<Indices>().ResolvePreLockedComponentData(Allocation, &PreLockedData[PrelockedDataOffsets.StartOffset[Indices]], WriteContext)... );
+
+		const int32 Num = Allocation->Num();
+		for (int32 ComponentOffset = 0; ComponentOffset < Num && Result.Value; ++ComponentOffset )
+		{
+			Result = (TaskImplInstance.ForEachEntity(ResolvedComponentData.template Get<Indices>().ComponentAtIndex(ComponentOffset)... ), Result);
+		}
+	}
+
+	template<typename TaskImpl, typename... AccessorTypes>
+	static void ForEachEntityImpl(TaskImpl& TaskImplInstance, const FEntityAllocation* Allocation, FEntityAllocationWriteContext WriteContext, const TEntityTaskComponents<AccessorTypes...>& Components)
 	{
 		FEntityIterationResult Result;
 
@@ -1474,28 +1884,43 @@ struct TEntityTaskCaller_AutoExpansion<TaskImpl, TIntegerSequence<int, Indices..
 		}
 	}
 
-	template<typename... AccessorTypes>
-	FORCEINLINE static void ForEachAllocationImpl(TaskImpl& TaskImplInstance, FEntityAllocationIteratorItem Item, FEntityAllocationWriteContext WriteContext, const TEntityTaskComponents<AccessorTypes...>& Components)
+	template<typename TaskImpl, typename... AccessorTypes>
+	static void ForEachAllocationImpl(TaskImpl& TaskImplInstance, FEntityAllocationIteratorItem Item, TArrayView<const FPreLockedDataPtr> PreLockedData, FEntityAllocationWriteContext WriteContext, const TEntityTaskComponents<AccessorTypes...>& Components)
+	{
+		constexpr TPrelockedDataOffsets<AccessorTypes...> PrelockedDataOffsets;
+		const FEntityAllocation* Allocation = Item.GetAllocation();
+		TaskImplInstance.ForEachAllocation(Item, Components.template GetAccessor<Indices>().ResolvePreLockedComponentData(Allocation, &PreLockedData[PrelockedDataOffsets.StartOffset[Indices]], WriteContext)...);
+	}
+
+	template<typename TaskImpl, typename... AccessorTypes>
+	static void ForEachAllocationImpl(TaskImpl& TaskImplInstance, FEntityAllocationIteratorItem Item, FEntityAllocationWriteContext WriteContext, const TEntityTaskComponents<AccessorTypes...>& Components)
 	{
 		auto LockedComponentData = MakeTuple( Components.template GetAccessor<Indices>().LockComponentData(Item.GetAllocation(), WriteContext)... );
 		TaskImplInstance.ForEachAllocation(Item, LockedComponentData.template Get<Indices>()...);
 	}
 };
 
-template<typename TaskImpl, int NumComponents>
-struct TEntityTaskCaller<TaskImpl, NumComponents, true> : TEntityTaskCaller_AutoExpansion<TaskImpl, TMakeIntegerSequence<int, NumComponents>>
+template<int NumComponents>
+struct TEntityTaskCaller<NumComponents, true> : TEntityTaskCaller_AutoExpansion<TMakeIntegerSequence<int, NumComponents>>
 {
 };
 
-template<typename TaskImpl, int32 NumComponents>
-struct TEntityTaskCaller<TaskImpl, NumComponents, false>
+template<int32 NumComponents>
+struct TEntityTaskCaller<NumComponents, false>
 {
-	FORCEINLINE static void ForEachEntityImpl(...)
+	template<typename TaskImpl>
+	FORCEINLINE static void ForEachEntityImpl(TaskImpl& TaskImplInstance, ...)
 	{
 		static_assert(!std::is_same_v<TaskImpl, TaskImpl>, "non-expanded entity iteration is not supported");
 	}
 
-	template<typename... AccessorTypes>
+	template<typename TaskImpl, typename... AccessorTypes>
+	FORCEINLINE static void ForEachAllocationImpl(TaskImpl& TaskImplInstance, FEntityAllocationIteratorItem Item, TArrayView<const FPreLockedDataPtr> PreLockedData, FEntityAllocationWriteContext WriteContext, const TEntityTaskComponents<AccessorTypes...>& Components)
+	{
+		TaskImplInstance.ForEachAllocation(Item, Components, PreLockedData);
+	}
+
+	template<typename TaskImpl, typename... AccessorTypes>
 	FORCEINLINE static void ForEachAllocationImpl(TaskImpl& TaskImplInstance, FEntityAllocationIteratorItem Item, const TEntityTaskComponents<AccessorTypes...>& Components)
 	{
 		TaskImplInstance.ForEachAllocation(Item, Components);
diff --git a/Engine/Source/Runtime/MovieScene/Public/EntitySystem/MovieSceneEntitySystemTypes.h b/Engine/Source/Runtime/MovieScene/Public/EntitySystem/MovieSceneEntitySystemTypes.h
--- a/Engine/Source/Runtime/MovieScene/Public/EntitySystem/MovieSceneEntitySystemTypes.h
+++ b/Engine/Source/Runtime/MovieScene/Public/EntitySystem/MovieSceneEntitySystemTypes.h
@@ -16,6 +16,7 @@
 #include "Stats/Stats2.h"
 #include "Templates/UnrealTemplate.h"
 
+#include <atomic>
 #include <initializer_list>
 
 class UMovieSceneSequence;
@@ -48,17 +49,23 @@ enum class ESystemPhase : uint8
 	/** Null phase which indicates that the system never runs, but still exists in the reference graph */
 	None = 0,
 
+	/**  */
+	Import = 1 << 0,
+
 	/** Expensive: Phase that is run before instantiation any time any boundary is crossed in the sequence. Used to spawn new objects and trigger pre/post-spawn events. */
-	Spawn = 1 << 0,
+	Spawn = 1 << 1,
 
 	/** Expensive: Houses any system that needs to instantiate global entities into the linker, or make meaningful changes to entity structures.. */
-	Instantiation = 1 << 1,
+	Instantiation = 1 << 2,
+
+	/**  */
+	Scheduling = 1 << 3,
 
 	/** Fast, distributed: Houses the majority of evaluation systems that compute animation data. Entity manager is locked down for the duration of this phase. */
-	Evaluation = 1 << 2,
+	Evaluation = 1 << 4,
 
 	/** Finalization phase for enything that wants to run after everything else. */
-	Finalization = 1 << 3,
+	Finalization = 1 << 5,
 };
 ENUM_CLASS_FLAGS(ESystemPhase);
 
@@ -426,6 +433,18 @@ struct FEntityAllocationWriteContext
 		return FEntityAllocationWriteContext();
 	}
 
+	FEntityAllocationWriteContext Add(FEntityAllocationWriteContext InOther) const
+	{
+		FEntityAllocationWriteContext Context;
+		Context.SystemSerial = SystemSerial + InOther.SystemSerial;
+		return Context;
+	}
+	FEntityAllocationWriteContext Subtract(FEntityAllocationWriteContext InOther) const
+	{
+		FEntityAllocationWriteContext Context;
+		Context.SystemSerial = SystemSerial - InOther.SystemSerial;
+		return Context;
+	}
 	uint64 GetSystemSerial() const
 	{
 		return SystemSerial;
@@ -454,6 +473,8 @@ private:
 
 public:
 
+	mutable std::atomic<int32> ScheduledAccessCount;
+
 	uint8 Sizeof;
 	FComponentTypeID ComponentType;
 
@@ -516,6 +537,7 @@ struct FScopedHeaderReadLock
 
 private:
 	const FComponentHeader* Header;
+	EComponentHeaderLockMode LockMode;
 };
 
 
@@ -690,6 +712,15 @@ struct FEntityAllocation
 	}
 
 
+	/**
+	 * Retrieve the address of this allocation's component data. Only to be used for construction of TRelativePtrs.
+	 */
+	const void* GetComponentDataAddress() const
+	{
+		return ComponentData;
+	}
+
+
 	/**
 	 * Get the unique identifier for this allocation. This identifier is unique to the specific allocation and entity manager, but is not globally unique.
 	 * Typically used for caching component data on a per-allocation basis
@@ -914,6 +945,17 @@ struct FEntityAllocationProxy
 	/** Gets the entity allocation component mask */
 	MOVIESCENE_API const FComponentMask& GetAllocationType() const;
 
+	static FEntityAllocationProxy MakeInstance(const FEntityManager* InManager, int32 InAllocationIndex)
+	{
+		return FEntityAllocationProxy(InManager, InAllocationIndex);
+	}
+
+	/** Return this allocation's index within the entity manager */
+	int32 GetAllocationIndex() const
+	{
+		return AllocationIndex;
+	}
+
 	/** Implicit cast to an entity allocation */
 	operator const FEntityAllocation*() const
 	{
@@ -1013,19 +1055,19 @@ inline FScopedHeaderReadLock::FScopedHeaderReadLock()
 	: Header(nullptr)
 {}
 inline FScopedHeaderReadLock::FScopedHeaderReadLock(const FComponentHeader* InHeader, EComponentHeaderLockMode InLockMode)
+	: Header(InHeader)
+	, LockMode(InLockMode)
 {
 	if (InLockMode == EComponentHeaderLockMode::Mutex)
 	{
-		Header = InHeader;
 		InHeader->ReadWriteLock.ReadLock();
 	}
-	else
-	{
-		Header = nullptr;
-	}
+
+	InHeader->ScheduledAccessCount.fetch_add(1, std::memory_order_relaxed);
 }
 inline FScopedHeaderReadLock::FScopedHeaderReadLock(FScopedHeaderReadLock&& RHS)
 	: Header(RHS.Header)
+	, LockMode(RHS.LockMode)
 {
 	RHS.Header = nullptr;
 }
@@ -1033,10 +1075,15 @@ inline FScopedHeaderReadLock& FScopedHeaderReadLock::operator=(FScopedHeaderRead
 {
 	if (Header)
 	{
-		Header->ReadWriteLock.ReadUnlock();
+		Header->ScheduledAccessCount.fetch_add(1, std::memory_order_relaxed);
+		if (LockMode == EComponentHeaderLockMode::Mutex)
+		{
+			Header->ReadWriteLock.ReadUnlock();
+		}
 	}
 
 	Header = RHS.Header;
+	LockMode = RHS.LockMode;
 	RHS.Header = nullptr;
 	return *this;
 }
@@ -1045,7 +1092,11 @@ inline FScopedHeaderReadLock::~FScopedHeaderReadLock()
 {
 	if (Header)
 	{
-		Header->ReadWriteLock.ReadUnlock();
+		Header->ScheduledAccessCount.fetch_sub(1, std::memory_order_relaxed);
+		if (LockMode == EComponentHeaderLockMode::Mutex)
+		{
+			Header->ReadWriteLock.ReadUnlock();
+		}
 	}
 }
 
@@ -1064,6 +1115,9 @@ inline FScopedHeaderWriteLock::FScopedHeaderWriteLock(const FComponentHeader* In
 	{
 		InHeader->ReadWriteLock.WriteLock();
 	}
+
+	const int32 PreviousAccessCount = InHeader->ScheduledAccessCount.fetch_add(1, std::memory_order_relaxed);
+	checkf(PreviousAccessCount == 0, TEXT("Component header is still in use when it is being opened for write!"));
 }
 
 inline FScopedHeaderWriteLock::FScopedHeaderWriteLock(FScopedHeaderWriteLock&& RHS)
@@ -1078,6 +1132,7 @@ inline FScopedHeaderWriteLock& FScopedHeaderWriteLock::operator=(FScopedHeaderWr
 	if (Header)
 	{
 		Header->PostWriteComponents(WriteContext);
+		Header->ScheduledAccessCount.fetch_sub(1, std::memory_order_relaxed);
 
 		if (LockMode == EComponentHeaderLockMode::Mutex)
 		{
@@ -1097,6 +1152,7 @@ inline FScopedHeaderWriteLock::~FScopedHeaderWriteLock()
 	if (Header)
 	{
 		Header->PostWriteComponents(WriteContext);
+		Header->ScheduledAccessCount.fetch_sub(1, std::memory_order_relaxed);
 
 		if (LockMode == EComponentHeaderLockMode::Mutex)
 		{
diff --git a/Engine/Source/Runtime/MovieScene/Public/EntitySystem/MovieSceneEvalTimeSystem.h b/Engine/Source/Runtime/MovieScene/Public/EntitySystem/MovieSceneEvalTimeSystem.h
--- a/Engine/Source/Runtime/MovieScene/Public/EntitySystem/MovieSceneEvalTimeSystem.h
+++ b/Engine/Source/Runtime/MovieScene/Public/EntitySystem/MovieSceneEvalTimeSystem.h
@@ -15,24 +15,31 @@
 class UObject;
 struct FFrameTime;
 
+namespace UE::MovieScene
+{
+	struct FEvaluatedTime
+	{
+		FFrameTime FrameTime;
+		double Seconds;
+	};
+}
+
 UCLASS()
 class MOVIESCENE_API UMovieSceneEvalTimeSystem : public UMovieSceneEntitySystem
 {
 public:
+
 	GENERATED_BODY()
 
 	UMovieSceneEvalTimeSystem(const FObjectInitializer& ObjInit);
 
-	virtual void OnRun(FSystemTaskPrerequisites& InPrerequisites, FSystemSubsequentTasks& Subsequents) override;
 	virtual bool IsRelevantImpl(UMovieSceneEntitySystemLinker* InLinker) const override;
+	virtual void OnSchedulePersistentTasks(UE::MovieScene::IEntitySystemScheduler* TaskScheduler) override;
+	virtual void OnRun(FSystemTaskPrerequisites& InPrerequisites, FSystemSubsequentTasks& Subsequents) override;
 
 private:
-	struct FEvaluatedTime
-	{
-		FFrameTime FrameTime;
-		double Seconds;
-	};
-	TArray<FEvaluatedTime> EvaluatedTimes;
+
+	TArray<UE::MovieScene::FEvaluatedTime, TInlineAllocator<16>> EvaluatedTimes;
 
 	UE::MovieScene::FEntityComponentFilter RelevantFilter;
 };
diff --git a/Engine/Source/Runtime/MovieScene/Public/EntitySystem/MovieSceneEvaluationHookSystem.h b/Engine/Source/Runtime/MovieScene/Public/EntitySystem/MovieSceneEvaluationHookSystem.h
--- a/Engine/Source/Runtime/MovieScene/Public/EntitySystem/MovieSceneEvaluationHookSystem.h
+++ b/Engine/Source/Runtime/MovieScene/Public/EntitySystem/MovieSceneEvaluationHookSystem.h
@@ -94,6 +94,7 @@ protected:
 
 private:
 
+	virtual void OnSchedulePersistentTasks(UE::MovieScene::IEntitySystemScheduler* TaskScheduler) override;
 	virtual void OnRun(FSystemTaskPrerequisites& InPrerequisites, FSystemSubsequentTasks& Subsequents) override;
 	virtual bool IsRelevantImpl(UMovieSceneEntitySystemLinker* InLinker) const override;
 
diff --git a/Engine/Source/Runtime/MovieScene/Public/EntitySystem/MovieScenePartialProperties.h b/Engine/Source/Runtime/MovieScene/Public/EntitySystem/MovieScenePartialProperties.h
--- a/Engine/Source/Runtime/MovieScene/Public/EntitySystem/MovieScenePartialProperties.h
+++ b/Engine/Source/Runtime/MovieScene/Public/EntitySystem/MovieScenePartialProperties.h
@@ -57,12 +57,6 @@ struct TSetPartialPropertyValuesImpl<PropertyTraits, TPropertyMetaData<MetaDataT
 	explicit TSetPartialPropertyValuesImpl(ICustomPropertyRegistration* InCustomProperties, TArrayView<const FPropertyCompositeDefinition> InCompositeDefinitions)
 		: CustomProperties(InCustomProperties)
 		, CompositeDefinitions(InCompositeDefinitions)
-	{}
-
-	/**
-	 * Run before this task executes any logic over entities and components
-	 */
-	void PreTask()
 	{
 		if (CustomProperties)
 		{
@@ -70,16 +64,16 @@ struct TSetPartialPropertyValuesImpl<PropertyTraits, TPropertyMetaData<MetaDataT
 		}
 	}
 
-	void ForEachAllocation(const FEntityAllocation* Allocation, TRead<UObject*> BoundObjectComponents, FThreeWayAccessor PropertyBindingComponents, TRead<MetaDataTypes>... InMetaData, TReadOptional<CompositeTypes>... InCompositeComponents);
-	void ForEachAllocation(const FEntityAllocation* Allocation, TRead<UObject*> BoundObjectComponents, FTwoWayAccessor PropertyBindingComponents, TRead<MetaDataTypes>... InMetaData, TReadOptional<CompositeTypes>... InCompositeComponents);
+	void ForEachAllocation(const FEntityAllocation* Allocation, TRead<UObject*> BoundObjectComponents, FThreeWayAccessor PropertyBindingComponents, TRead<MetaDataTypes>... InMetaData, TReadOptional<CompositeTypes>... InCompositeComponents) const;
+	void ForEachAllocation(const FEntityAllocation* Allocation, TRead<UObject*> BoundObjectComponents, FTwoWayAccessor PropertyBindingComponents, TRead<MetaDataTypes>... InMetaData, TReadOptional<CompositeTypes>... InCompositeComponents) const;
 
 private:
 
-	void ForEachCustom(const FEntityAllocation* Allocation, UObject* const* Objects, const FCustomPropertyIndex* Custom, const MetaDataTypes*... InMetaData, const CompositeTypes*... InCompositeComponents);
+	void ForEachCustom(const FEntityAllocation* Allocation, UObject* const* Objects, const FCustomPropertyIndex* Custom, const MetaDataTypes*... InMetaData, const CompositeTypes*... InCompositeComponents) const;
 
-	void ForEachFast(const FEntityAllocation* Allocation, UObject* const* Objects, const uint16* Fast, const MetaDataTypes*... InMetaData, const CompositeTypes*... InCompositeComponents);
+	void ForEachFast(const FEntityAllocation* Allocation, UObject* const* Objects, const uint16* Fast, const MetaDataTypes*... InMetaData, const CompositeTypes*... InCompositeComponents) const;
 
-	void ForEachSlow(const FEntityAllocation* Allocation, UObject* const* Objects, const TSharedPtr<FTrackInstancePropertyBindings>* Slow, const MetaDataTypes*... InMetaData, const CompositeTypes*... InCompositeComponents);
+	void ForEachSlow(const FEntityAllocation* Allocation, UObject* const* Objects, const TSharedPtr<FTrackInstancePropertyBindings>* Slow, const MetaDataTypes*... InMetaData, const CompositeTypes*... InCompositeComponents) const;
 
 private:
 
diff --git a/Engine/Source/Runtime/MovieScene/Public/EntitySystem/MovieScenePartialProperties.inl b/Engine/Source/Runtime/MovieScene/Public/EntitySystem/MovieScenePartialProperties.inl
--- a/Engine/Source/Runtime/MovieScene/Public/EntitySystem/MovieScenePartialProperties.inl
+++ b/Engine/Source/Runtime/MovieScene/Public/EntitySystem/MovieScenePartialProperties.inl
@@ -13,7 +13,7 @@ namespace MovieScene
 
 
 template<typename PropertyTraits, typename ...MetaDataTypes, int ...CompositeIndices, typename ...CompositeTypes>
-void TSetPartialPropertyValuesImpl<PropertyTraits, TPropertyMetaData<MetaDataTypes...>, TIntegerSequence<int, CompositeIndices...>, CompositeTypes...>::ForEachAllocation(const FEntityAllocation* Allocation, TRead<UObject*> BoundObjectComponents, FThreeWayAccessor PropertyBindingComponents, TRead<MetaDataTypes>... InMetaData, TReadOptional<CompositeTypes>... InCompositeComponents)
+void TSetPartialPropertyValuesImpl<PropertyTraits, TPropertyMetaData<MetaDataTypes...>, TIntegerSequence<int, CompositeIndices...>, CompositeTypes...>::ForEachAllocation(const FEntityAllocation* Allocation, TRead<UObject*> BoundObjectComponents, FThreeWayAccessor PropertyBindingComponents, TRead<MetaDataTypes>... InMetaData, TReadOptional<CompositeTypes>... InCompositeComponents) const
 {
 	// ----------------------------------------------------------------------------------------------------------------------------
 	// For partially animated composites, we first retrieve the current properties for the allocation, then go through and patch in
@@ -37,7 +37,7 @@ void TSetPartialPropertyValuesImpl<PropertyTraits, TPropertyMetaData<MetaDataTyp
 
 
 template<typename PropertyTraits, typename ...MetaDataTypes, int ...CompositeIndices, typename ...CompositeTypes>
-void TSetPartialPropertyValuesImpl<PropertyTraits, TPropertyMetaData<MetaDataTypes...>, TIntegerSequence<int, CompositeIndices...>, CompositeTypes...>::ForEachAllocation(const FEntityAllocation* Allocation, TRead<UObject*> BoundObjectComponents, FTwoWayAccessor PropertyBindingComponents, TRead<MetaDataTypes>... InMetaData, TReadOptional<CompositeTypes>... InCompositeComponents)
+void TSetPartialPropertyValuesImpl<PropertyTraits, TPropertyMetaData<MetaDataTypes...>, TIntegerSequence<int, CompositeIndices...>, CompositeTypes...>::ForEachAllocation(const FEntityAllocation* Allocation, TRead<UObject*> BoundObjectComponents, FTwoWayAccessor PropertyBindingComponents, TRead<MetaDataTypes>... InMetaData, TReadOptional<CompositeTypes>... InCompositeComponents) const
 {
 	// ----------------------------------------------------------------------------------------------------------------------------
 	// For partially animated composites, we first retrieve the current properties for the allocation, then go through and patch in
@@ -56,7 +56,7 @@ void TSetPartialPropertyValuesImpl<PropertyTraits, TPropertyMetaData<MetaDataTyp
 
 
 template<typename PropertyTraits, typename ...MetaDataTypes, int ...CompositeIndices, typename ...CompositeTypes>
-void TSetPartialPropertyValuesImpl<PropertyTraits, TPropertyMetaData<MetaDataTypes...>, TIntegerSequence<int, CompositeIndices...>, CompositeTypes...>::ForEachCustom(const FEntityAllocation* Allocation, UObject* const* Objects, const FCustomPropertyIndex* Custom, const MetaDataTypes*... InMetaData, const CompositeTypes*... InCompositeComponents)
+void TSetPartialPropertyValuesImpl<PropertyTraits, TPropertyMetaData<MetaDataTypes...>, TIntegerSequence<int, CompositeIndices...>, CompositeTypes...>::ForEachCustom(const FEntityAllocation* Allocation, UObject* const* Objects, const FCustomPropertyIndex* Custom, const MetaDataTypes*... InMetaData, const CompositeTypes*... InCompositeComponents) const
 {
 	const int32 Num = Allocation->Num();
 
@@ -75,7 +75,7 @@ void TSetPartialPropertyValuesImpl<PropertyTraits, TPropertyMetaData<MetaDataTyp
 
 
 template<typename PropertyTraits, typename ...MetaDataTypes, int ...CompositeIndices, typename ...CompositeTypes>
-void TSetPartialPropertyValuesImpl<PropertyTraits, TPropertyMetaData<MetaDataTypes...>, TIntegerSequence<int, CompositeIndices...>, CompositeTypes...>::ForEachFast(const FEntityAllocation* Allocation, UObject* const* Objects, const uint16* Fast, const MetaDataTypes*... InMetaData, const CompositeTypes*... InCompositeComponents)
+void TSetPartialPropertyValuesImpl<PropertyTraits, TPropertyMetaData<MetaDataTypes...>, TIntegerSequence<int, CompositeIndices...>, CompositeTypes...>::ForEachFast(const FEntityAllocation* Allocation, UObject* const* Objects, const uint16* Fast, const MetaDataTypes*... InMetaData, const CompositeTypes*... InCompositeComponents) const
 {
 	const int32 Num = Allocation->Num();
 
@@ -95,7 +95,7 @@ void TSetPartialPropertyValuesImpl<PropertyTraits, TPropertyMetaData<MetaDataTyp
 
 
 template<typename PropertyTraits, typename ...MetaDataTypes, int ...CompositeIndices, typename ...CompositeTypes>
-void TSetPartialPropertyValuesImpl<PropertyTraits, TPropertyMetaData<MetaDataTypes...>, TIntegerSequence<int, CompositeIndices...>, CompositeTypes...>::ForEachSlow(const FEntityAllocation* Allocation, UObject* const* Objects, const TSharedPtr<FTrackInstancePropertyBindings>* Slow, const MetaDataTypes*... InMetaData, const CompositeTypes*... InCompositeComponents)
+void TSetPartialPropertyValuesImpl<PropertyTraits, TPropertyMetaData<MetaDataTypes...>, TIntegerSequence<int, CompositeIndices...>, CompositeTypes...>::ForEachSlow(const FEntityAllocation* Allocation, UObject* const* Objects, const TSharedPtr<FTrackInstancePropertyBindings>* Slow, const MetaDataTypes*... InMetaData, const CompositeTypes*... InCompositeComponents) const
 {
 	const int32 Num = Allocation->Num();
 
diff --git a/Engine/Source/Runtime/MovieScene/Public/EntitySystem/MovieScenePropertyComponentHandler.h b/Engine/Source/Runtime/MovieScene/Public/EntitySystem/MovieScenePropertyComponentHandler.h
--- a/Engine/Source/Runtime/MovieScene/Public/EntitySystem/MovieScenePropertyComponentHandler.h
+++ b/Engine/Source/Runtime/MovieScene/Public/EntitySystem/MovieScenePropertyComponentHandler.h
@@ -333,6 +333,42 @@ struct TPropertyComponentHandlerImpl<PropertyTraits, TPropertyMetaData<MetaDataT
 		return Existing;
 	}
 
+	virtual void ScheduleSetterTasks(const FPropertyDefinition& Definition, TArrayView<const FPropertyCompositeDefinition> Composites, const FPropertyStats& Stats, IEntitySystemScheduler* TaskScheduler, UMovieSceneEntitySystemLinker* Linker)
+	{
+		FBuiltInComponentTypes* BuiltInComponents = FBuiltInComponentTypes::Get();
+
+		FEntityTaskBuilder()
+		.Read(BuiltInComponents->BoundObject)
+		.ReadOneOf(BuiltInComponents->CustomPropertyIndex, BuiltInComponents->FastPropertyOffset, BuiltInComponents->SlowProperty)
+		.ReadAllOf(Definition.GetMetaDataComponent<MetaDataTypes>(MetaDataIndices)...)
+		.ReadAllOf(Composites[CompositeIndices].ComponentTypeID.ReinterpretCast<CompositeTypes>()...)
+		.FilterAll({ Definition.PropertyType })
+		.SetStat(Definition.StatID)
+		.SetDesiredThread(Linker->EntityManager.GetGatherThread())
+		.template Fork_PerAllocation<CompleteSetterTask>(&Linker->EntityManager, TaskScheduler, Definition.CustomPropertyRegistration);
+
+		if (Stats.NumPartialProperties > 0)
+		{
+			FComponentMask CompletePropertyMask;
+			for (const FPropertyCompositeDefinition& Composite : Composites)
+			{
+				CompletePropertyMask.Set(Composite.ComponentTypeID);
+			}
+
+			FEntityTaskBuilder()
+			.Read(BuiltInComponents->BoundObject)
+			.ReadOneOf(BuiltInComponents->CustomPropertyIndex, BuiltInComponents->FastPropertyOffset, BuiltInComponents->SlowProperty)
+			.ReadAllOf(Definition.GetMetaDataComponent<MetaDataTypes>(MetaDataIndices)...)
+			.ReadAnyOf(Composites[CompositeIndices].ComponentTypeID.ReinterpretCast<CompositeTypes>()...)
+			.FilterAny({ CompletePropertyMask })
+			.FilterAll({ Definition.PropertyType })
+			.FilterOut(CompletePropertyMask)
+			.SetStat(Definition.StatID)
+			.SetDesiredThread(Linker->EntityManager.GetGatherThread())
+			.template Fork_PerAllocation<PartialSetterTask>(&Linker->EntityManager, TaskScheduler, Definition.CustomPropertyRegistration, Composites);
+		}
+	}
+
 	virtual void DispatchSetterTasks(const FPropertyDefinition& Definition, TArrayView<const FPropertyCompositeDefinition> Composites, const FPropertyStats& Stats, FSystemTaskPrerequisites& InPrerequisites, FSystemSubsequentTasks& Subsequents, UMovieSceneEntitySystemLinker* Linker)
 	{
 		FBuiltInComponentTypes* BuiltInComponents = FBuiltInComponentTypes::Get();
diff --git a/Engine/Source/Runtime/MovieScene/Public/EntitySystem/MovieScenePropertySystemTypes.h b/Engine/Source/Runtime/MovieScene/Public/EntitySystem/MovieScenePropertySystemTypes.h
--- a/Engine/Source/Runtime/MovieScene/Public/EntitySystem/MovieScenePropertySystemTypes.h
+++ b/Engine/Source/Runtime/MovieScene/Public/EntitySystem/MovieScenePropertySystemTypes.h
@@ -296,13 +296,12 @@ struct TSetPropertyValuesImpl<PropertyTraits, TPropertyMetaData<MetaDataTypes...
 
 	explicit TSetPropertyValuesImpl(ICustomPropertyRegistration* InCustomProperties)
 		: CustomProperties(InCustomProperties)
-	{}
-
-	/**
-	 * Run before this task executes any logic over entities and components
-	 */
-	void PreTask();
-
+	{
+		if (CustomProperties)
+		{
+			CustomAccessors = CustomProperties->GetAccessors();
+		}
+	}
 
 	/**
 	 * Task callback that applies a value to an object property via a custom native setter function
@@ -314,7 +313,7 @@ struct TSetPropertyValuesImpl<PropertyTraits, TPropertyMetaData<MetaDataTypes...
 	 *     .Read( TComponentTypeID<PropertyType>(...) )
 	 *     .Dispatch_PerEntity<TSetPropertyValues<PropertyType>>(...);
 	 */
-	void ForEachEntity(UObject* InObject, FCustomPropertyIndex CustomIndex, typename TCallTraits<MetaDataTypes>::ParamType... MetaData, InParamType ValueToSet);
+	void ForEachEntity(UObject* InObject, FCustomPropertyIndex CustomIndex, typename TCallTraits<MetaDataTypes>::ParamType... MetaData, InParamType ValueToSet) const;
 
 	/**
 	 * Task callback that applies a value to an object property via a fast pointer offset
@@ -356,7 +355,7 @@ public:
 	 *     .Read(      TComponentTypeID<PropertyType>(...) )
 	 *     .Dispatch_PerAllocation<TSetPropertyValues<PropertyType>>(...);
 	 */
-	void ForEachAllocation(const FEntityAllocation* Allocation, TRead<UObject*> BoundObjectComponents, FThreeWayAccessor ResolvedPropertyComponents, TRead<MetaDataTypes>... MetaDataComponents, TRead<StorageType> PropertyValueComponents);
+	void ForEachAllocation(const FEntityAllocation* Allocation, TRead<UObject*> BoundObjectComponents, FThreeWayAccessor ResolvedPropertyComponents, TRead<MetaDataTypes>... MetaDataComponents, TRead<StorageType> PropertyValueComponents) const;
 
 
 	/**
@@ -369,7 +368,7 @@ public:
 	 *     .Read(      TComponentTypeID<PropertyType>(...) )
 	 *     .Dispatch_PerAllocation<TSetPropertyValues<PropertyType>>(...);
 	 */
-	void ForEachAllocation(const FEntityAllocation* Allocation, TRead<UObject*> BoundObjectComponents, FTwoWayAccessor ResolvedPropertyComponents, TRead<MetaDataTypes>... MetaDataComponents, TRead<StorageType> PropertyValueComponents);
+	void ForEachAllocation(const FEntityAllocation* Allocation, TRead<UObject*> BoundObjectComponents, FTwoWayAccessor ResolvedPropertyComponents, TRead<MetaDataTypes>... MetaDataComponents, TRead<StorageType> PropertyValueComponents) const;
 
 private:
 
@@ -429,9 +428,12 @@ struct TGetPropertyValuesImpl<PropertyTraits, TPropertyMetaData<MetaDataTypes...
 
 	explicit TGetPropertyValuesImpl(ICustomPropertyRegistration* InCustomProperties)
 		: CustomProperties(InCustomProperties)
-	{}
-
-	void PreTask();
+	{
+		if (CustomProperties)
+		{
+			CustomAccessors = CustomProperties->GetAccessors();
+		}
+	}
 
 	/**
 	 * Task callback that retrieves the object's current value via a custom native setter function, and writes it to the specified output variable
@@ -443,7 +445,7 @@ struct TGetPropertyValuesImpl<PropertyTraits, TPropertyMetaData<MetaDataTypes...
 	 *     .Write( TComponentTypeID<StorageType>(...) )
 	 *     .Dispatch_PerEntity<TGetPropertyValues<PropertyType, StorageType>>(...);
 	 */
-	void ForEachEntity(UObject* InObject, FCustomPropertyIndex CustomPropertyIndex, typename TCallTraits<MetaDataTypes>::ParamType... MetaData, StorageType& OutValue);
+	void ForEachEntity(UObject* InObject, FCustomPropertyIndex CustomPropertyIndex, typename TCallTraits<MetaDataTypes>::ParamType... MetaData, StorageType& OutValue) const;
 
 	/**
 	 * Task callback that retrieves the object's current value via a fast pointer offset, and writes it to the specified output variable
@@ -456,7 +458,7 @@ struct TGetPropertyValuesImpl<PropertyTraits, TPropertyMetaData<MetaDataTypes...
 	 *     .Read( TComponentTypeID<StorageType>(...) )
 	 *     .Dispatch_PerEntity<TGetPropertyValues<PropertyType, StorageType>>(...);
 	 */
-	void ForEachEntity(UObject* InObject, uint16 PropertyOffset, typename TCallTraits<MetaDataTypes>::ParamType... MetaData, StorageType& OutValue);
+	void ForEachEntity(UObject* InObject, uint16 PropertyOffset, typename TCallTraits<MetaDataTypes>::ParamType... MetaData, StorageType& OutValue) const;
 
 	/**
 	 * Task callback that retrieves the object's current value via a slow (legacy) track instance binding, and writes it to the specified output variable
@@ -469,7 +471,7 @@ struct TGetPropertyValuesImpl<PropertyTraits, TPropertyMetaData<MetaDataTypes...
 	 *     .Read( TComponentTypeID<StorageType>(...) )
 	 *     .Dispatch_PerEntity<TGetPropertyValues<PropertyType, StorageType>>(...);
 	 */
-	void ForEachEntity(UObject* InObject, const TSharedPtr<FTrackInstancePropertyBindings>& PropertyBindings, typename TCallTraits<MetaDataTypes>::ParamType... MetaData, StorageType& OutValue);
+	void ForEachEntity(UObject* InObject, const TSharedPtr<FTrackInstancePropertyBindings>& PropertyBindings, typename TCallTraits<MetaDataTypes>::ParamType... MetaData, StorageType& OutValue) const;
 
 public:
 
@@ -488,7 +490,7 @@ public:
 	 *     .Write(     TComponentTypeID<StorageType>(...) )
 	 *     .Dispatch_PerAllocation<TGetPropertyValues<PropertyType, StorageType>>(...);
 	 */
-	void ForEachAllocation(const FEntityAllocation* Allocation, TRead<UObject*> BoundObjectComponents, FThreeWayAccessor ResolvedPropertyComponents, TRead<MetaDataTypes>... MetaDataComponents, TWrite<StorageType> OutValueComponents);
+	void ForEachAllocation(const FEntityAllocation* Allocation, TRead<UObject*> BoundObjectComponents, FThreeWayAccessor ResolvedPropertyComponents, TRead<MetaDataTypes>... MetaDataComponents, TWrite<StorageType> OutValueComponents) const;
 
 
 	/**
@@ -502,7 +504,7 @@ public:
 	 *     .Write(     TComponentTypeID<StorageType>(...) )
 	 *     .Dispatch_PerAllocation<TGetPropertyValues<PropertyType ,StorageType>>(...);
 	 */
-	void ForEachAllocation(const FEntityAllocation* Allocation, TRead<UObject*> BoundObjectComponents, FTwoWayAccessor ResolvedPropertyComponents, TRead<MetaDataTypes>... MetaDataComponents, TWrite<StorageType> OutValueComponents);
+	void ForEachAllocation(const FEntityAllocation* Allocation, TRead<UObject*> BoundObjectComponents, FTwoWayAccessor ResolvedPropertyComponents, TRead<MetaDataTypes>... MetaDataComponents, TWrite<StorageType> OutValueComponents) const;
 
 private:
 
@@ -561,12 +563,12 @@ struct TSetCompositePropertyValuesImpl<PropertyTraits, TPropertyMetaData<MetaDat
 
 	explicit TSetCompositePropertyValuesImpl(ICustomPropertyRegistration* InCustomProperties)
 		: CustomProperties(InCustomProperties)
-	{}
-
-	/**
-	 * Run before this task executes any logic over entities and components
-	 */
-	void PreTask();
+	{
+		if (CustomProperties)
+		{
+			CustomAccessors = CustomProperties->GetAccessors();
+		}
+	}
 
 	/**
 	 * Task callback that applies a value to an object property via a custom native setter function
@@ -580,7 +582,7 @@ struct TSetCompositePropertyValuesImpl<PropertyTraits, TPropertyMetaData<MetaDat
 	 *     .Read( TComponentTypeID<CompositeType[N-1]>(...) )
 	 *     .Dispatch_PerEntity<TSetCompositePropertyValues<PropertyType, CompositeTypes...>>(...);
 	 */
-	void ForEachEntity(UObject* InObject, FCustomPropertyIndex CustomPropertyIndex, typename TCallTraits<MetaDataTypes>::ParamType... MetaData, typename TCallTraits<CompositeTypes>::ParamType... CompositeResults);
+	void ForEachEntity(UObject* InObject, FCustomPropertyIndex CustomPropertyIndex, typename TCallTraits<MetaDataTypes>::ParamType... MetaData, typename TCallTraits<CompositeTypes>::ParamType... CompositeResults) const;
 
 
 	/**
@@ -595,7 +597,7 @@ struct TSetCompositePropertyValuesImpl<PropertyTraits, TPropertyMetaData<MetaDat
 	 *     .Read( TComponentTypeID<CompositeType[N-1]>(...) )
 	 *     .Dispatch_PerEntity<TSetCompositePropertyValues<PropertyType, CompositeTypes...>>(...);
 	 */
-	void ForEachEntity(UObject* InObject, uint16 PropertyOffset, typename TCallTraits<MetaDataTypes>::ParamType... MetaData, typename TCallTraits<CompositeTypes>::ParamType... CompositeResults);
+	void ForEachEntity(UObject* InObject, uint16 PropertyOffset, typename TCallTraits<MetaDataTypes>::ParamType... MetaData, typename TCallTraits<CompositeTypes>::ParamType... CompositeResults) const;
 
 
 	/**
@@ -610,7 +612,7 @@ struct TSetCompositePropertyValuesImpl<PropertyTraits, TPropertyMetaData<MetaDat
 	 *     .Read( TComponentTypeID<CompositeType[N-1]>(...) )
 	 *     .Dispatch_PerEntity<TSetCompositePropertyValues<PropertyType, CompositeTypes...>>(...);
 	 */
-	void ForEachEntity(UObject* InObject, const TSharedPtr<FTrackInstancePropertyBindings>& PropertyBindings, typename TCallTraits<MetaDataTypes>::ParamType... MetaData, typename TCallTraits<CompositeTypes>::ParamType... CompositeResults);
+	void ForEachEntity(UObject* InObject, const TSharedPtr<FTrackInstancePropertyBindings>& PropertyBindings, typename TCallTraits<MetaDataTypes>::ParamType... MetaData, typename TCallTraits<CompositeTypes>::ParamType... CompositeResults) const;
 
 public:
 
@@ -630,7 +632,7 @@ public:
 	 *     .Read( TComponentTypeID<CompositeType[N-1]>(...) )
 	 *     .Dispatch_PerAllocation<TSetCompositePropertyValues<PropertyType, CompositeTypes...>>(...);
 	 */
-	void ForEachAllocation(const FEntityAllocation* Allocation, TRead<UObject*> BoundObjectComponents, FThreeWayAccessor ResolvedPropertyComponents, TRead<MetaDataTypes>... InMetaData, TRead<CompositeTypes>... VariadicComponents);
+	void ForEachAllocation(const FEntityAllocation* Allocation, TRead<UObject*> BoundObjectComponents, FThreeWayAccessor ResolvedPropertyComponents, TRead<MetaDataTypes>... InMetaData, TRead<CompositeTypes>... VariadicComponents) const;
 
 
 	/**
@@ -645,7 +647,7 @@ public:
 	 *     .Read( TComponentTypeID<CompositeType[N-1]>(...) )
 	 *     .Dispatch_PerAllocation<TSetCompositePropertyValues<PropertyType, CompositeTypes...>>(...);
 	 */
-	void ForEachAllocation(const FEntityAllocation* Allocation, TRead<UObject*> BoundObjectComponents, FTwoWayAccessor ResolvedPropertyComponents, TRead<MetaDataTypes>... InMetaData, TRead<CompositeTypes>... VariadicComponents);
+	void ForEachAllocation(const FEntityAllocation* Allocation, TRead<UObject*> BoundObjectComponents, FTwoWayAccessor ResolvedPropertyComponents, TRead<MetaDataTypes>... InMetaData, TRead<CompositeTypes>... VariadicComponents) const;
 
 private:
 
diff --git a/Engine/Source/Runtime/MovieScene/Public/EntitySystem/MovieScenePropertySystemTypes.inl b/Engine/Source/Runtime/MovieScene/Public/EntitySystem/MovieScenePropertySystemTypes.inl
--- a/Engine/Source/Runtime/MovieScene/Public/EntitySystem/MovieScenePropertySystemTypes.inl
+++ b/Engine/Source/Runtime/MovieScene/Public/EntitySystem/MovieScenePropertySystemTypes.inl
@@ -22,18 +22,8 @@ namespace UE
 namespace MovieScene
 {
 
-
 template<typename PropertyTraits, typename ...MetaDataTypes>
-void TSetPropertyValuesImpl<PropertyTraits, TPropertyMetaData<MetaDataTypes...>>::PreTask()
-{
-	if (CustomProperties)
-	{
-		CustomAccessors = CustomProperties->GetAccessors();
-	}
-}
-
-template<typename PropertyTraits, typename ...MetaDataTypes>
-void TSetPropertyValuesImpl<PropertyTraits, TPropertyMetaData<MetaDataTypes...>>::ForEachEntity(UObject* InObject, FCustomPropertyIndex CustomIndex, typename TCallTraits<MetaDataTypes>::ParamType... MetaData, InParamType ValueToSet)
+void TSetPropertyValuesImpl<PropertyTraits, TPropertyMetaData<MetaDataTypes...>>::ForEachEntity(UObject* InObject, FCustomPropertyIndex CustomIndex, typename TCallTraits<MetaDataTypes>::ParamType... MetaData, InParamType ValueToSet) const
 {
 	PropertyTraits::SetObjectPropertyValue(InObject, MetaData..., CustomAccessors[CustomIndex.Value], ValueToSet);
 }
@@ -55,7 +45,7 @@ void TSetPropertyValuesImpl<PropertyTraits, TPropertyMetaData<MetaDataTypes...>>
 }
 
 template<typename PropertyTraits, typename ...MetaDataTypes>
-void TSetPropertyValuesImpl<PropertyTraits, TPropertyMetaData<MetaDataTypes...>>::ForEachAllocation(const FEntityAllocation* Allocation, TRead<UObject*> BoundObjectComponents, FThreeWayAccessor ResolvedPropertyComponents, TRead<MetaDataTypes>... MetaDataComponents, TRead<StorageType> PropertyValueComponents)
+void TSetPropertyValuesImpl<PropertyTraits, TPropertyMetaData<MetaDataTypes...>>::ForEachAllocation(const FEntityAllocation* Allocation, TRead<UObject*> BoundObjectComponents, FThreeWayAccessor ResolvedPropertyComponents, TRead<MetaDataTypes>... MetaDataComponents, TRead<StorageType> PropertyValueComponents) const
 {
 	const int32 Num = Allocation->Num();
 	if (const FCustomPropertyIndex* Custom = ResolvedPropertyComponents.template Get<0>())
@@ -82,7 +72,7 @@ void TSetPropertyValuesImpl<PropertyTraits, TPropertyMetaData<MetaDataTypes...>>
 }
 
 template<typename PropertyTraits, typename ...MetaDataTypes>
-void TSetPropertyValuesImpl<PropertyTraits, TPropertyMetaData<MetaDataTypes...>>::ForEachAllocation(const FEntityAllocation* Allocation, TRead<UObject*> BoundObjectComponents, FTwoWayAccessor ResolvedPropertyComponents, TRead<MetaDataTypes>... MetaDataComponents, TRead<StorageType> PropertyValueComponents)
+void TSetPropertyValuesImpl<PropertyTraits, TPropertyMetaData<MetaDataTypes...>>::ForEachAllocation(const FEntityAllocation* Allocation, TRead<UObject*> BoundObjectComponents, FTwoWayAccessor ResolvedPropertyComponents, TRead<MetaDataTypes>... MetaDataComponents, TRead<StorageType> PropertyValueComponents) const
 {
 	const int32 Num = Allocation->Num();
 	if (const uint16* Fast = ResolvedPropertyComponents.template Get<0>())
@@ -101,24 +91,14 @@ void TSetPropertyValuesImpl<PropertyTraits, TPropertyMetaData<MetaDataTypes...>>
 	}
 }
 
-
 template<typename PropertyTraits, typename ...MetaDataTypes>
-void TGetPropertyValuesImpl<PropertyTraits, TPropertyMetaData<MetaDataTypes...>>::PreTask()
-{
-	if (CustomProperties)
-	{
-		CustomAccessors = CustomProperties->GetAccessors();
-	}
-}
-
-template<typename PropertyTraits, typename ...MetaDataTypes>
-void TGetPropertyValuesImpl<PropertyTraits, TPropertyMetaData<MetaDataTypes...>>::ForEachEntity(UObject* InObject, FCustomPropertyIndex CustomPropertyIndex, typename TCallTraits<MetaDataTypes>::ParamType... MetaData, StorageType& OutValue)
+void TGetPropertyValuesImpl<PropertyTraits, TPropertyMetaData<MetaDataTypes...>>::ForEachEntity(UObject* InObject, FCustomPropertyIndex CustomPropertyIndex, typename TCallTraits<MetaDataTypes>::ParamType... MetaData, StorageType& OutValue) const
 {
 	PropertyTraits::GetObjectPropertyValue(InObject, MetaData..., CustomAccessors[CustomPropertyIndex.Value], OutValue);
 }
 
 template<typename PropertyTraits, typename ...MetaDataTypes>
-void TGetPropertyValuesImpl<PropertyTraits, TPropertyMetaData<MetaDataTypes...>>::ForEachEntity(UObject* InObject, uint16 PropertyOffset, typename TCallTraits<MetaDataTypes>::ParamType... MetaData, StorageType& OutValue)
+void TGetPropertyValuesImpl<PropertyTraits, TPropertyMetaData<MetaDataTypes...>>::ForEachEntity(UObject* InObject, uint16 PropertyOffset, typename TCallTraits<MetaDataTypes>::ParamType... MetaData, StorageType& OutValue) const
 {
 	// Would really like to avoid branching here, but if we encounter this data the options are either handle it gracefully, stomp a vtable, or report a fatal error.
 	if (ensureAlwaysMsgf(PropertyOffset != 0, TEXT("Invalid property offset specified (ptr+%d bytes) for property on object %s. This would otherwise overwrite the object's vfptr."), PropertyOffset, *InObject->GetName()))
@@ -128,13 +108,13 @@ void TGetPropertyValuesImpl<PropertyTraits, TPropertyMetaData<MetaDataTypes...>>
 }
 
 template<typename PropertyTraits, typename ...MetaDataTypes>
-void TGetPropertyValuesImpl<PropertyTraits, TPropertyMetaData<MetaDataTypes...>>::ForEachEntity(UObject* InObject, const TSharedPtr<FTrackInstancePropertyBindings>& PropertyBindings, typename TCallTraits<MetaDataTypes>::ParamType... MetaData, StorageType& OutValue)
+void TGetPropertyValuesImpl<PropertyTraits, TPropertyMetaData<MetaDataTypes...>>::ForEachEntity(UObject* InObject, const TSharedPtr<FTrackInstancePropertyBindings>& PropertyBindings, typename TCallTraits<MetaDataTypes>::ParamType... MetaData, StorageType& OutValue) const
 {
 	PropertyTraits::GetObjectPropertyValue(InObject, MetaData..., PropertyBindings.Get(), OutValue);
 }
 
 template<typename PropertyTraits, typename ...MetaDataTypes>
-void TGetPropertyValuesImpl<PropertyTraits, TPropertyMetaData<MetaDataTypes...>>::ForEachAllocation(const FEntityAllocation* Allocation, TRead<UObject*> BoundObjectComponents, FThreeWayAccessor ResolvedPropertyComponents, TRead<MetaDataTypes>... MetaData, TWrite<StorageType> OutValueComponents)
+void TGetPropertyValuesImpl<PropertyTraits, TPropertyMetaData<MetaDataTypes...>>::ForEachAllocation(const FEntityAllocation* Allocation, TRead<UObject*> BoundObjectComponents, FThreeWayAccessor ResolvedPropertyComponents, TRead<MetaDataTypes>... MetaData, TWrite<StorageType> OutValueComponents) const
 {
 	const int32 Num = Allocation->Num();
 	if (const FCustomPropertyIndex* Custom = ResolvedPropertyComponents.template Get<0>())
@@ -161,7 +141,7 @@ void TGetPropertyValuesImpl<PropertyTraits, TPropertyMetaData<MetaDataTypes...>>
 }
 
 template<typename PropertyTraits, typename ...MetaDataTypes>
-void TGetPropertyValuesImpl<PropertyTraits, TPropertyMetaData<MetaDataTypes...>>::ForEachAllocation(const FEntityAllocation* Allocation, TRead<UObject*> BoundObjectComponents, FTwoWayAccessor ResolvedPropertyComponents, TRead<MetaDataTypes>... MetaData, TWrite<StorageType> OutValueComponents)
+void TGetPropertyValuesImpl<PropertyTraits, TPropertyMetaData<MetaDataTypes...>>::ForEachAllocation(const FEntityAllocation* Allocation, TRead<UObject*> BoundObjectComponents, FTwoWayAccessor ResolvedPropertyComponents, TRead<MetaDataTypes>... MetaData, TWrite<StorageType> OutValueComponents) const
 {
 	const int32 Num = Allocation->Num();
 	if (const uint16* Fast = ResolvedPropertyComponents.template Get<0>())
@@ -182,24 +162,14 @@ void TGetPropertyValuesImpl<PropertyTraits, TPropertyMetaData<MetaDataTypes...>>
 
 
 template<typename PropertyTraits, typename... MetaDataTypes, typename... CompositeTypes>
-
-void TSetCompositePropertyValuesImpl<PropertyTraits, TPropertyMetaData<MetaDataTypes...>, CompositeTypes...>::PreTask()
-{
-	if (CustomProperties)
-	{
-		CustomAccessors = CustomProperties->GetAccessors();
-	}
-}
-
-template<typename PropertyTraits, typename... MetaDataTypes, typename... CompositeTypes>
-void TSetCompositePropertyValuesImpl<PropertyTraits, TPropertyMetaData<MetaDataTypes...>, CompositeTypes...>::ForEachEntity(UObject* InObject, FCustomPropertyIndex CustomPropertyIndex, typename TCallTraits<MetaDataTypes>::ParamType... MetaData, typename TCallTraits<CompositeTypes>::ParamType... CompositeResults)
+void TSetCompositePropertyValuesImpl<PropertyTraits, TPropertyMetaData<MetaDataTypes...>, CompositeTypes...>::ForEachEntity(UObject* InObject, FCustomPropertyIndex CustomPropertyIndex, typename TCallTraits<MetaDataTypes>::ParamType... MetaData, typename TCallTraits<CompositeTypes>::ParamType... CompositeResults) const
 {
 	StorageType Result = PropertyTraits::CombineComposites(MetaData..., CompositeResults...);
 	PropertyTraits::SetObjectPropertyValue(InObject, MetaData..., CustomAccessors[CustomPropertyIndex.Value], Result);
 }
 
 template<typename PropertyTraits, typename... MetaDataTypes, typename... CompositeTypes>
-void TSetCompositePropertyValuesImpl<PropertyTraits, TPropertyMetaData<MetaDataTypes...>, CompositeTypes...>::ForEachEntity(UObject* InObject, uint16 PropertyOffset, typename TCallTraits<MetaDataTypes>::ParamType... MetaData, typename TCallTraits<CompositeTypes>::ParamType... CompositeResults)
+void TSetCompositePropertyValuesImpl<PropertyTraits, TPropertyMetaData<MetaDataTypes...>, CompositeTypes...>::ForEachEntity(UObject* InObject, uint16 PropertyOffset, typename TCallTraits<MetaDataTypes>::ParamType... MetaData, typename TCallTraits<CompositeTypes>::ParamType... CompositeResults) const
 {
 	// Would really like to avoid branching here, but if we encounter this data the options are either handle it gracefully, stomp a vtable, or report a fatal error.
 	if (ensureAlwaysMsgf(PropertyOffset != 0, TEXT("Invalid property offset specified (ptr+%d bytes) for property on object %s. This would otherwise overwrite the object's vfptr."), PropertyOffset, *InObject->GetName()))
@@ -210,14 +180,14 @@ void TSetCompositePropertyValuesImpl<PropertyTraits, TPropertyMetaData<MetaDataT
 }
 
 template<typename PropertyTraits, typename... MetaDataTypes, typename... CompositeTypes>
-void TSetCompositePropertyValuesImpl<PropertyTraits, TPropertyMetaData<MetaDataTypes...>, CompositeTypes...>::ForEachEntity(UObject* InObject, const TSharedPtr<FTrackInstancePropertyBindings>& PropertyBindings, typename TCallTraits<MetaDataTypes>::ParamType... MetaData, typename TCallTraits<CompositeTypes>::ParamType... CompositeResults)
+void TSetCompositePropertyValuesImpl<PropertyTraits, TPropertyMetaData<MetaDataTypes...>, CompositeTypes...>::ForEachEntity(UObject* InObject, const TSharedPtr<FTrackInstancePropertyBindings>& PropertyBindings, typename TCallTraits<MetaDataTypes>::ParamType... MetaData, typename TCallTraits<CompositeTypes>::ParamType... CompositeResults) const
 {
 	StorageType Result = PropertyTraits::CombineComposites(MetaData..., CompositeResults...);
 	PropertyTraits::SetObjectPropertyValue(InObject, MetaData..., PropertyBindings.Get(), Result);
 }
 
 template<typename PropertyTraits, typename... MetaDataTypes, typename... CompositeTypes>
-void TSetCompositePropertyValuesImpl<PropertyTraits, TPropertyMetaData<MetaDataTypes...>, CompositeTypes...>::ForEachAllocation(const FEntityAllocation* Allocation, TRead<UObject*> BoundObjectComponents, FThreeWayAccessor ResolvedPropertyComponents, TRead<MetaDataTypes>... InMetaData, TRead<CompositeTypes>... VariadicComponents)
+void TSetCompositePropertyValuesImpl<PropertyTraits, TPropertyMetaData<MetaDataTypes...>, CompositeTypes...>::ForEachAllocation(const FEntityAllocation* Allocation, TRead<UObject*> BoundObjectComponents, FThreeWayAccessor ResolvedPropertyComponents, TRead<MetaDataTypes>... InMetaData, TRead<CompositeTypes>... VariadicComponents) const
 {
 	const int32 Num = Allocation->Num();
 	if (const FCustomPropertyIndex* Custom = ResolvedPropertyComponents.template Get<0>())
@@ -244,7 +214,7 @@ void TSetCompositePropertyValuesImpl<PropertyTraits, TPropertyMetaData<MetaDataT
 }
 
 template<typename PropertyTraits, typename... MetaDataTypes, typename... CompositeTypes>
-void TSetCompositePropertyValuesImpl<PropertyTraits, TPropertyMetaData<MetaDataTypes...>, CompositeTypes...>::ForEachAllocation(const FEntityAllocation* Allocation, TRead<UObject*> BoundObjectComponents, FTwoWayAccessor ResolvedPropertyComponents, TRead<MetaDataTypes>... InMetaData, TRead<CompositeTypes>... VariadicComponents)
+void TSetCompositePropertyValuesImpl<PropertyTraits, TPropertyMetaData<MetaDataTypes...>, CompositeTypes...>::ForEachAllocation(const FEntityAllocation* Allocation, TRead<UObject*> BoundObjectComponents, FTwoWayAccessor ResolvedPropertyComponents, TRead<MetaDataTypes>... InMetaData, TRead<CompositeTypes>... VariadicComponents) const
 {
 	const int32 Num = Allocation->Num();
 	if (const uint16* Fast = ResolvedPropertyComponents.template Get<0>())
diff --git a/Engine/Source/Runtime/MovieScene/Public/EntitySystem/MovieSceneSequenceInstance.h b/Engine/Source/Runtime/MovieScene/Public/EntitySystem/MovieSceneSequenceInstance.h
--- a/Engine/Source/Runtime/MovieScene/Public/EntitySystem/MovieSceneSequenceInstance.h
+++ b/Engine/Source/Runtime/MovieScene/Public/EntitySystem/MovieSceneSequenceInstance.h
@@ -40,6 +40,7 @@ enum class ESequenceInstanceUpdateFlags : uint8
 	NeedsDissection     = 1u << 0,
 	NeedsPreEvaluation  = 1u << 1,
 	NeedsPostEvaluation = 1u << 2,
+	HasLegacyTemplates  = 1u << 3,
 };
 ENUM_CLASS_FLAGS(ESequenceInstanceUpdateFlags);
 
diff --git a/Engine/Source/Runtime/MovieSceneTracks/Private/Systems/BoolChannelEvaluatorSystem.cpp b/Engine/Source/Runtime/MovieSceneTracks/Private/Systems/BoolChannelEvaluatorSystem.cpp
--- a/Engine/Source/Runtime/MovieSceneTracks/Private/Systems/BoolChannelEvaluatorSystem.cpp
+++ b/Engine/Source/Runtime/MovieSceneTracks/Private/Systems/BoolChannelEvaluatorSystem.cpp
@@ -12,6 +12,20 @@
 
 DECLARE_CYCLE_STAT(TEXT("MovieScene: Evaluate bool channels"), MovieSceneEval_EvaluateBoolChannelTask, STATGROUP_MovieSceneECS);
 
+namespace UE::MovieScene
+{
+	struct FEvaluateBoolChannels
+	{
+		static void ForEachEntity(FSourceBoolChannel BoolChannel, FFrameTime FrameTime, bool& OutResult)
+		{
+			if (!BoolChannel.Source->Evaluate(FrameTime, OutResult))
+			{
+				OutResult = false;
+			}
+		}
+	};
+} // namespace UE::MovieScene
+
 UBoolChannelEvaluatorSystem::UBoolChannelEvaluatorSystem(const FObjectInitializer& ObjInit)
 	: Super(ObjInit)
 {
@@ -19,6 +33,7 @@ UBoolChannelEvaluatorSystem::UBoolChannelEvaluatorSystem(const FObjectInitialize
 
 	SystemCategories = EEntitySystemCategory::ChannelEvaluators;
 	RelevantComponent = FBuiltInComponentTypes::Get()->BoolChannel;
+	Phase = ESystemPhase::Scheduling;
 
 	if (HasAnyFlags(RF_ClassDefaultObject))
 	{
@@ -26,22 +41,27 @@ UBoolChannelEvaluatorSystem::UBoolChannelEvaluatorSystem(const FObjectInitialize
 	}
 }
 
-void UBoolChannelEvaluatorSystem::OnRun(FSystemTaskPrerequisites& InPrerequisites, FSystemSubsequentTasks& Subsequents)
+void UBoolChannelEvaluatorSystem::OnSchedulePersistentTasks(UE::MovieScene::IEntitySystemScheduler* TaskScheduler)
 {
 	using namespace UE::MovieScene;
 
 	FBuiltInComponentTypes* BuiltInComponents = FBuiltInComponentTypes::Get();
 
-	struct FEvaluateBoolChannels
-	{
-		static void ForEachEntity(FSourceBoolChannel BoolChannel, FFrameTime FrameTime, bool& OutResult)
-		{
-			if (!BoolChannel.Source->Evaluate(FrameTime, OutResult))
-			{
-				OutResult = false;
-			}
-		}
-	};
+	FEntityTaskBuilder()
+	.Read(BuiltInComponents->BoolChannel)
+	.Read(BuiltInComponents->EvalTime)
+	.Write(BuiltInComponents->BoolResult)
+	.FilterNone({ BuiltInComponents->Tags.Ignored })
+	.SetStat(GET_STATID(MovieSceneEval_EvaluateBoolChannelTask))
+	.Fork_PerEntity<FEvaluateBoolChannels>(&Linker->EntityManager, TaskScheduler);
+}
+
+void UBoolChannelEvaluatorSystem::OnRun(FSystemTaskPrerequisites& InPrerequisites, FSystemSubsequentTasks& Subsequents)
+{
+	// Backwards compat Evaluation phase
+	using namespace UE::MovieScene;
+
+	FBuiltInComponentTypes* BuiltInComponents = FBuiltInComponentTypes::Get();
 
 	FEntityTaskBuilder()
 	.Read(BuiltInComponents->BoolChannel)
diff --git a/Engine/Source/Runtime/MovieSceneTracks/Private/Systems/ByteChannelEvaluatorSystem.cpp b/Engine/Source/Runtime/MovieSceneTracks/Private/Systems/ByteChannelEvaluatorSystem.cpp
--- a/Engine/Source/Runtime/MovieSceneTracks/Private/Systems/ByteChannelEvaluatorSystem.cpp
+++ b/Engine/Source/Runtime/MovieSceneTracks/Private/Systems/ByteChannelEvaluatorSystem.cpp
@@ -21,7 +21,7 @@ namespace MovieScene
 // Do we need to optimize for this case using something like the code below, while pessimizing the common (non-multi-bind) codepath??
 struct FEvaluateByteChannels
 {
-	void ForEachEntity(FSourceByteChannel ByteChannel, FFrameTime FrameTime, uint8& OutResult)
+	static void ForEachEntity(FSourceByteChannel ByteChannel, FFrameTime FrameTime, uint8& OutResult)
 	{
 		if (!ByteChannel.Source->Evaluate(FrameTime, OutResult))
 		{
@@ -44,6 +44,7 @@ UByteChannelEvaluatorSystem::UByteChannelEvaluatorSystem(const FObjectInitialize
 
 	const FBuiltInComponentTypes* BuiltInComponents = FBuiltInComponentTypes::Get();
 	RelevantComponent = BuiltInComponents->ByteChannel;
+	Phase = ESystemPhase::Scheduling;
 
 	if (HasAnyFlags(RF_ClassDefaultObject))
 	{
@@ -53,6 +54,22 @@ UByteChannelEvaluatorSystem::UByteChannelEvaluatorSystem(const FObjectInitialize
 	}
 }
 
+void UByteChannelEvaluatorSystem::OnSchedulePersistentTasks(UE::MovieScene::IEntitySystemScheduler* TaskScheduler)
+{
+	using namespace UE::MovieScene;
+
+	const FBuiltInComponentTypes* BuiltInComponents = FBuiltInComponentTypes::Get();
+
+	// Evaluate byte channels per instance and write the evaluated value into the output
+	FEntityTaskBuilder()
+	.Read(BuiltInComponents->ByteChannel)
+	.Read(BuiltInComponents->EvalTime)
+	.Write(BuiltInComponents->ByteResult)
+	.FilterNone({ BuiltInComponents->Tags.Ignored })
+	.SetStat(GET_STATID(MovieSceneEval_EvaluateByteChannelTask))
+	.Fork_PerEntity<FEvaluateByteChannels>(&Linker->EntityManager, TaskScheduler);
+}
+
 void UByteChannelEvaluatorSystem::OnRun(FSystemTaskPrerequisites& InPrerequisites, FSystemSubsequentTasks& Subsequents)
 {
 	using namespace UE::MovieScene;
diff --git a/Engine/Source/Runtime/MovieSceneTracks/Private/Systems/DoubleChannelEvaluatorSystem.cpp b/Engine/Source/Runtime/MovieSceneTracks/Private/Systems/DoubleChannelEvaluatorSystem.cpp
--- a/Engine/Source/Runtime/MovieSceneTracks/Private/Systems/DoubleChannelEvaluatorSystem.cpp
+++ b/Engine/Source/Runtime/MovieSceneTracks/Private/Systems/DoubleChannelEvaluatorSystem.cpp
@@ -40,7 +40,7 @@ TArray<FDoubleChannelTypeAssociation, TInlineAllocator<4>> GDoubleChannelTypeAss
 /** Entity-component task that evaluates using the non-cached codepath for testing parity */
 struct FEvaluateDoubleChannels_Uncached
 {
-	void ForEachEntity(FSourceDoubleChannel DoubleChannel, FFrameTime FrameTime, double& OutResult)
+	void ForEachEntity(FSourceDoubleChannel DoubleChannel, FFrameTime FrameTime, double& OutResult) const
 	{
 		if (!DoubleChannel.Source->Evaluate(FrameTime, OutResult))
 		{
@@ -52,7 +52,7 @@ struct FEvaluateDoubleChannels_Uncached
 /** Entity-component task that evaluates using a cached interpolation if possible */
 struct FEvaluateDoubleChannels_Cached
 {
-	void ForEachEntity(FSourceDoubleChannel DoubleChannel, FFrameTime FrameTime, Interpolation::FCachedInterpolation& Cache, double& OutResult)
+	void ForEachEntity(FSourceDoubleChannel DoubleChannel, FFrameTime FrameTime, Interpolation::FCachedInterpolation& Cache, double& OutResult) const
 	{
 		if (!Cache.IsCacheValidForTime(FrameTime.GetFrame()))
 		{
@@ -123,7 +123,7 @@ UDoubleChannelEvaluatorSystem::UDoubleChannelEvaluatorSystem(const FObjectInitia
 	using namespace UE::MovieScene;
 
 	SystemCategories = EEntitySystemCategory::ChannelEvaluators;
-	Phase = ESystemPhase::Instantiation | ESystemPhase::Evaluation;
+	Phase = ESystemPhase::Instantiation | ESystemPhase::Scheduling;
 
 	if (HasAnyFlags(RF_ClassDefaultObject))
 	{
@@ -168,6 +168,39 @@ bool UDoubleChannelEvaluatorSystem::IsRelevantImpl(UMovieSceneEntitySystemLinker
 	return false;
 }
 
+void UDoubleChannelEvaluatorSystem::OnSchedulePersistentTasks(UE::MovieScene::IEntitySystemScheduler* TaskScheduler)
+{
+	using namespace UE::MovieScene;
+
+	FBuiltInComponentTypes* BuiltInComponents = FBuiltInComponentTypes::Get();
+
+	for (const FDoubleChannelTypeAssociation& ChannelType : GDoubleChannelTypeAssociations)
+	{
+		if (GEnableCachedChannelEvaluation)
+		{
+			// Evaluate double channels per instance and write the evaluated value into the output
+			FEntityTaskBuilder()
+			.Read(ChannelType.ChannelType)
+			.Read(BuiltInComponents->EvalTime)
+			.Write(ChannelType.CachedInterpolationType)
+			.Write(ChannelType.ResultType)
+			.FilterNone({ BuiltInComponents->Tags.Ignored })
+			.SetStat(GET_STATID(MovieSceneEval_EvaluateDoubleChannelTask))
+			.Fork_PerEntity<FEvaluateDoubleChannels_Cached>(&Linker->EntityManager, TaskScheduler);
+		}
+		else
+		{
+			FEntityTaskBuilder()
+			.Read(ChannelType.ChannelType)
+			.Read(BuiltInComponents->EvalTime)
+			.Write(ChannelType.ResultType)
+			.FilterNone({ BuiltInComponents->Tags.Ignored })
+			.SetStat(GET_STATID(MovieSceneEval_EvaluateDoubleChannelTask))
+			.Fork_PerEntity<FEvaluateDoubleChannels_Uncached>(&Linker->EntityManager, TaskScheduler);
+		}
+	}
+}
+
 void UDoubleChannelEvaluatorSystem::OnRun(FSystemTaskPrerequisites& InPrerequisites, FSystemSubsequentTasks& Subsequents)
 {
 	using namespace UE::MovieScene;
diff --git a/Engine/Source/Runtime/MovieSceneTracks/Private/Systems/DoublePerlinNoiseChannelEvaluatorSystem.cpp b/Engine/Source/Runtime/MovieSceneTracks/Private/Systems/DoublePerlinNoiseChannelEvaluatorSystem.cpp
--- a/Engine/Source/Runtime/MovieSceneTracks/Private/Systems/DoublePerlinNoiseChannelEvaluatorSystem.cpp
+++ b/Engine/Source/Runtime/MovieSceneTracks/Private/Systems/DoublePerlinNoiseChannelEvaluatorSystem.cpp
@@ -22,7 +22,7 @@ namespace UE
 	{
 		struct FEvaluateDoublePerlinNoiseChannels
 		{
-			void ForEachEntity(double EvalSeconds, const FPerlinNoiseParams& PerlinNoiseParams, double& OutResult)
+			static void ForEachEntity(double EvalSeconds, const FPerlinNoiseParams& PerlinNoiseParams, double& OutResult)
 			{
 				OutResult = FMovieSceneDoublePerlinNoiseChannel::Evaluate(PerlinNoiseParams, EvalSeconds);
 			}
@@ -36,7 +36,7 @@ UDoublePerlinNoiseChannelEvaluatorSystem::UDoublePerlinNoiseChannelEvaluatorSyst
 	using namespace UE::MovieScene;
 
 	SystemCategories = EEntitySystemCategory::ChannelEvaluators;
-	Phase = ESystemPhase::Instantiation | ESystemPhase::Evaluation;
+	Phase = ESystemPhase::Instantiation | ESystemPhase::Scheduling;
 
 	const FBuiltInComponentTypes* BuiltInComponents = FBuiltInComponentTypes::Get();
 	const FMovieSceneTracksComponentTypes* TrackComponents = FMovieSceneTracksComponentTypes::Get();
@@ -57,6 +57,25 @@ UDoublePerlinNoiseChannelEvaluatorSystem::UDoublePerlinNoiseChannelEvaluatorSyst
 	}
 }
 
+void UDoublePerlinNoiseChannelEvaluatorSystem::OnSchedulePersistentTasks(UE::MovieScene::IEntitySystemScheduler* TaskScheduler)
+{
+	using namespace UE::MovieScene;
+
+	const FBuiltInComponentTypes* BuiltInComponents = FBuiltInComponentTypes::Get();
+	const FMovieSceneTracksComponentTypes* TrackComponents = FMovieSceneTracksComponentTypes::Get();
+
+	for (int32 i = 0; i < UE_ARRAY_COUNT(BuiltInComponents->DoubleResult); ++i)
+	{
+		FEntityTaskBuilder()
+		.Read(BuiltInComponents->EvalSeconds)
+		.Read(TrackComponents->DoublePerlinNoiseChannel)
+		.Write(BuiltInComponents->DoubleResult[i])
+		.FilterNone({ BuiltInComponents->Tags.Ignored })
+		.SetStat(GET_STATID(MovieSceneEval_EvaluateDoublePerlinNoiseChannelTask))
+		.Fork_PerEntity<FEvaluateDoublePerlinNoiseChannels>(&Linker->EntityManager, TaskScheduler);
+	}
+}
+
 void UDoublePerlinNoiseChannelEvaluatorSystem::OnRun(FSystemTaskPrerequisites& InPrerequisites, FSystemSubsequentTasks& Subsequents)
 {
 	using namespace UE::MovieScene;
diff --git a/Engine/Source/Runtime/MovieSceneTracks/Private/Systems/FloatChannelEvaluatorSystem.cpp b/Engine/Source/Runtime/MovieSceneTracks/Private/Systems/FloatChannelEvaluatorSystem.cpp
--- a/Engine/Source/Runtime/MovieSceneTracks/Private/Systems/FloatChannelEvaluatorSystem.cpp
+++ b/Engine/Source/Runtime/MovieSceneTracks/Private/Systems/FloatChannelEvaluatorSystem.cpp
@@ -40,7 +40,7 @@ TArray<FFloatChannelTypeAssociation, TInlineAllocator<4>> GFloatChannelTypeAssoc
 /** Entity-component task that evaluates using the non-cached codepath for testing parity */
 struct FEvaluateFloatChannels_Uncached
 {
-	void ForEachEntity(FSourceFloatChannel FloatChannel, FFrameTime FrameTime, double& OutResult)
+	static void ForEachEntity(FSourceFloatChannel FloatChannel, FFrameTime FrameTime, double& OutResult)
 	{
 		float Result;
 		if (FloatChannel.Source->Evaluate(FrameTime, Result))
@@ -57,7 +57,7 @@ struct FEvaluateFloatChannels_Uncached
 /** Entity-component task that evaluates using a cached interpolation if possible */
 struct FEvaluateFloatChannels_Cached
 {
-	void ForEachEntity(FSourceFloatChannel FloatChannel, FFrameTime FrameTime, Interpolation::FCachedInterpolation& Cache, double& OutResult)
+	static void ForEachEntity(FSourceFloatChannel FloatChannel, FFrameTime FrameTime, Interpolation::FCachedInterpolation& Cache, double& OutResult)
 	{
 		if (!Cache.IsCacheValidForTime(FrameTime.GetFrame()))
 		{
@@ -133,7 +133,7 @@ UFloatChannelEvaluatorSystem::UFloatChannelEvaluatorSystem(const FObjectInitiali
 	using namespace UE::MovieScene;
 
 	SystemCategories = EEntitySystemCategory::ChannelEvaluators;
-	Phase = ESystemPhase::Instantiation | ESystemPhase::Evaluation;
+	Phase = ESystemPhase::Instantiation | ESystemPhase::Scheduling;
 
 	if (HasAnyFlags(RF_ClassDefaultObject))
 	{
@@ -179,6 +179,41 @@ bool UFloatChannelEvaluatorSystem::IsRelevantImpl(UMovieSceneEntitySystemLinker*
 	return false;
 }
 
+void UFloatChannelEvaluatorSystem::OnSchedulePersistentTasks(UE::MovieScene::IEntitySystemScheduler* TaskScheduler)
+{
+	using namespace UE::MovieScene;
+
+	FBuiltInComponentTypes* BuiltInComponents = FBuiltInComponentTypes::Get();
+	if (GEnableCachedChannelEvaluation)
+	{
+		for (const FFloatChannelTypeAssociation& ChannelType : GFloatChannelTypeAssociations)
+		{
+			// Evaluate float channels per instance and write the evaluated value into the output
+			FEntityTaskBuilder()
+			.Read(ChannelType.ChannelType)
+			.Read(BuiltInComponents->EvalTime)
+			.Write(ChannelType.CachedInterpolationType)
+			.Write(ChannelType.ResultType)
+			.FilterNone({ BuiltInComponents->Tags.Ignored })
+			.SetStat(GET_STATID(MovieSceneEval_EvaluateFloatChannelTask))
+			.Fork_PerEntity<FEvaluateFloatChannels_Cached>(&Linker->EntityManager, TaskScheduler);
+		}
+	}
+	else
+	{
+		for (const FFloatChannelTypeAssociation& ChannelType : GFloatChannelTypeAssociations)
+		{
+			FEntityTaskBuilder()
+			.Read(ChannelType.ChannelType)
+			.Read(BuiltInComponents->EvalTime)
+			.Write(ChannelType.ResultType)
+			.FilterNone({ BuiltInComponents->Tags.Ignored })
+			.SetStat(GET_STATID(MovieSceneEval_EvaluateFloatChannelTask))
+			.Fork_PerEntity<FEvaluateFloatChannels_Uncached>(&Linker->EntityManager, TaskScheduler);
+		}
+	}
+}
+
 void UFloatChannelEvaluatorSystem::OnRun(FSystemTaskPrerequisites& InPrerequisites, FSystemSubsequentTasks& Subsequents)
 {
 	using namespace UE::MovieScene;
diff --git a/Engine/Source/Runtime/MovieSceneTracks/Private/Systems/FloatPerlinNoiseChannelEvaluatorSystem.cpp b/Engine/Source/Runtime/MovieSceneTracks/Private/Systems/FloatPerlinNoiseChannelEvaluatorSystem.cpp
--- a/Engine/Source/Runtime/MovieSceneTracks/Private/Systems/FloatPerlinNoiseChannelEvaluatorSystem.cpp
+++ b/Engine/Source/Runtime/MovieSceneTracks/Private/Systems/FloatPerlinNoiseChannelEvaluatorSystem.cpp
@@ -22,7 +22,7 @@ namespace UE
 	{
 		struct FEvaluateFloatPerlinNoiseChannels
 		{
-			void ForEachEntity(double EvalSeconds, const FPerlinNoiseParams& PerlinNoiseParams, double& OutResult)
+			static void ForEachEntity(double EvalSeconds, const FPerlinNoiseParams& PerlinNoiseParams, double& OutResult)
 			{
 				const float NoiseResult = FMovieSceneFloatPerlinNoiseChannel::Evaluate(PerlinNoiseParams, EvalSeconds);
 				OutResult = (double)NoiseResult;
@@ -37,7 +37,7 @@ UFloatPerlinNoiseChannelEvaluatorSystem::UFloatPerlinNoiseChannelEvaluatorSystem
 	using namespace UE::MovieScene;
 
 	SystemCategories = EEntitySystemCategory::ChannelEvaluators;
-	Phase = ESystemPhase::Instantiation | ESystemPhase::Evaluation;
+	Phase = ESystemPhase::Instantiation | ESystemPhase::Scheduling;
 
 	const FBuiltInComponentTypes* BuiltInComponents = FBuiltInComponentTypes::Get();
 	const FMovieSceneTracksComponentTypes* TrackComponents = FMovieSceneTracksComponentTypes::Get();
@@ -58,6 +58,25 @@ UFloatPerlinNoiseChannelEvaluatorSystem::UFloatPerlinNoiseChannelEvaluatorSystem
 	}
 }
 
+void UFloatPerlinNoiseChannelEvaluatorSystem::OnSchedulePersistentTasks(UE::MovieScene::IEntitySystemScheduler* TaskScheduler)
+{
+	using namespace UE::MovieScene;
+
+	const FBuiltInComponentTypes* BuiltInComponents = FBuiltInComponentTypes::Get();
+	const FMovieSceneTracksComponentTypes* TrackComponents = FMovieSceneTracksComponentTypes::Get();
+
+	for (int32 i = 0; i < UE_ARRAY_COUNT(BuiltInComponents->DoubleResult); ++i)
+	{
+		FEntityTaskBuilder()
+		.Read(BuiltInComponents->EvalSeconds)
+		.Read(TrackComponents->FloatPerlinNoiseChannel)
+		.Write(BuiltInComponents->DoubleResult[i])
+		.FilterNone({ BuiltInComponents->Tags.Ignored })
+		.SetStat(GET_STATID(MovieSceneEval_EvaluateFloatPerlinNoiseChannelTask))
+		.Fork_PerEntity<FEvaluateFloatPerlinNoiseChannels>(&Linker->EntityManager, TaskScheduler);
+	}
+}
+
 void UFloatPerlinNoiseChannelEvaluatorSystem::OnRun(FSystemTaskPrerequisites& InPrerequisites, FSystemSubsequentTasks& Subsequents)
 {
 	using namespace UE::MovieScene;
diff --git a/Engine/Source/Runtime/MovieSceneTracks/Private/Systems/IntegerChannelEvaluatorSystem.cpp b/Engine/Source/Runtime/MovieSceneTracks/Private/Systems/IntegerChannelEvaluatorSystem.cpp
--- a/Engine/Source/Runtime/MovieSceneTracks/Private/Systems/IntegerChannelEvaluatorSystem.cpp
+++ b/Engine/Source/Runtime/MovieSceneTracks/Private/Systems/IntegerChannelEvaluatorSystem.cpp
@@ -21,7 +21,7 @@ namespace MovieScene
 // Do we need to optimize for this case using something like the code below, while pessimizing the common (non-multi-bind) codepath??
 struct FEvaluateIntegerChannels
 {
-	void ForEachEntity(FSourceIntegerChannel IntegerChannel, FFrameTime FrameTime, int32& OutResult)
+	static void ForEachEntity(FSourceIntegerChannel IntegerChannel, FFrameTime FrameTime, int32& OutResult)
 	{
 		if (!IntegerChannel.Source->Evaluate(FrameTime, OutResult))
 		{
@@ -44,6 +44,7 @@ UIntegerChannelEvaluatorSystem::UIntegerChannelEvaluatorSystem(const FObjectInit
 
 	const FBuiltInComponentTypes* BuiltInComponents = FBuiltInComponentTypes::Get();
 	RelevantComponent = BuiltInComponents->IntegerChannel;
+	Phase = ESystemPhase::Scheduling;
 
 	if (HasAnyFlags(RF_ClassDefaultObject))
 	{
@@ -53,6 +54,22 @@ UIntegerChannelEvaluatorSystem::UIntegerChannelEvaluatorSystem(const FObjectInit
 	}
 }
 
+void UIntegerChannelEvaluatorSystem::OnSchedulePersistentTasks(UE::MovieScene::IEntitySystemScheduler* TaskScheduler)
+{
+	using namespace UE::MovieScene;
+
+	const FBuiltInComponentTypes* BuiltInComponents = FBuiltInComponentTypes::Get();
+
+	// Evaluate integer channels per instance and write the evaluated value into the output
+	FEntityTaskBuilder()
+	.Read(BuiltInComponents->IntegerChannel)
+	.Read(BuiltInComponents->EvalTime)
+	.Write(BuiltInComponents->IntegerResult)
+	.FilterNone({ BuiltInComponents->Tags.Ignored })
+	.SetStat(GET_STATID(MovieSceneEval_EvaluateIntegerChannelTask))
+	.Fork_PerEntity<FEvaluateIntegerChannels>(&Linker->EntityManager, TaskScheduler);
+}
+
 void UIntegerChannelEvaluatorSystem::OnRun(FSystemTaskPrerequisites& InPrerequisites, FSystemSubsequentTasks& Subsequents)
 {
 	using namespace UE::MovieScene;
diff --git a/Engine/Source/Runtime/MovieSceneTracks/Private/Systems/MovieSceneFadeSystem.cpp b/Engine/Source/Runtime/MovieSceneTracks/Private/Systems/MovieSceneFadeSystem.cpp
--- a/Engine/Source/Runtime/MovieSceneTracks/Private/Systems/MovieSceneFadeSystem.cpp
+++ b/Engine/Source/Runtime/MovieSceneTracks/Private/Systems/MovieSceneFadeSystem.cpp
@@ -112,7 +112,7 @@ struct FEvaluateFade
 		, PreAnimatedStorage(InPreAnimatedStorage)
 	{}
 
-	void ForEachAllocation(const FEntityAllocation* Allocation, TRead<FMovieSceneEntityID> EntityIDs, TRead<FRootInstanceHandle> RootInstanceHandles, TRead<FFadeComponentData> FadeComponents, TRead<double> FadeAmounts)
+	void ForEachAllocation(const FEntityAllocation* Allocation, TRead<FMovieSceneEntityID> EntityIDs, TRead<FRootInstanceHandle> RootInstanceHandles, TRead<FFadeComponentData> FadeComponents, TRead<double> FadeAmounts) const
 	{
 		static const TMovieSceneAnimTypeID<FEvaluateFade> AnimTypeID;
 
@@ -145,6 +145,7 @@ UMovieSceneFadeSystem::UMovieSceneFadeSystem(const FObjectInitializer& ObjInit)
 	using namespace UE::MovieScene;
 
 	RelevantComponent = FMovieSceneTracksComponentTypes::Get()->Fade;
+	Phase = ESystemPhase::Scheduling;
 
 	if (HasAnyFlags(RF_ClassDefaultObject))
 	{
@@ -160,6 +161,22 @@ void UMovieSceneFadeSystem::OnLink()
 	PreAnimatedStorage = Linker->PreAnimatedState.GetOrCreateStorage<FPreAnimatedFadeStateStorage>();
 }
 
+void UMovieSceneFadeSystem::OnSchedulePersistentTasks(UE::MovieScene::IEntitySystemScheduler* TaskScheduler)
+{
+	using namespace UE::MovieScene;
+
+	const FBuiltInComponentTypes* BuiltInComponents = FBuiltInComponentTypes::Get();
+	const FMovieSceneTracksComponentTypes* TrackComponents = FMovieSceneTracksComponentTypes::Get();
+
+	FEntityTaskBuilder()
+	.ReadEntityIDs()
+	.Read(BuiltInComponents->RootInstanceHandle)
+	.Read(TrackComponents->Fade)
+	.Read(BuiltInComponents->DoubleResult[0])
+	.SetDesiredThread(Linker->EntityManager.GetGatherThread())
+	.Schedule_PerAllocation<FEvaluateFade>(& Linker->EntityManager, TaskScheduler, Linker->GetInstanceRegistry(), PreAnimatedStorage);
+}
+
 void UMovieSceneFadeSystem::OnRun(FSystemTaskPrerequisites& InPrerequisites, FSystemSubsequentTasks& Subsequents)
 {
 	using namespace UE::MovieScene;
diff --git a/Engine/Source/Runtime/MovieSceneTracks/Private/Systems/MovieSceneFadeSystem.h b/Engine/Source/Runtime/MovieSceneTracks/Private/Systems/MovieSceneFadeSystem.h
--- a/Engine/Source/Runtime/MovieSceneTracks/Private/Systems/MovieSceneFadeSystem.h
+++ b/Engine/Source/Runtime/MovieSceneTracks/Private/Systems/MovieSceneFadeSystem.h
@@ -21,6 +21,7 @@ public:
 	UMovieSceneFadeSystem(const FObjectInitializer& ObjInit);
 
 	virtual void OnLink() override;
+	virtual void OnSchedulePersistentTasks(UE::MovieScene::IEntitySystemScheduler* TaskScheduler) override;
 	virtual void OnRun(FSystemTaskPrerequisites& InPrerequisites, FSystemSubsequentTasks& Subsequents) override;
 
 private:
diff --git a/Engine/Source/Runtime/MovieSceneTracks/Private/Systems/MovieSceneMaterialParameterSystem.cpp b/Engine/Source/Runtime/MovieSceneTracks/Private/Systems/MovieSceneMaterialParameterSystem.cpp
--- a/Engine/Source/Runtime/MovieSceneTracks/Private/Systems/MovieSceneMaterialParameterSystem.cpp
+++ b/Engine/Source/Runtime/MovieSceneTracks/Private/Systems/MovieSceneMaterialParameterSystem.cpp
@@ -436,7 +436,7 @@ UMovieSceneMaterialParameterSystem::UMovieSceneMaterialParameterSystem(const FOb
 	FMovieSceneTracksComponentTypes* TracksComponents = FMovieSceneTracksComponentTypes::Get();
 
 	RelevantComponent = TracksComponents->BoundMaterial;
-	Phase = ESystemPhase::Instantiation | ESystemPhase::Evaluation;
+	Phase = ESystemPhase::Instantiation | ESystemPhase::Scheduling;
 
 	if (HasAnyFlags(RF_ClassDefaultObject))
 	{
@@ -566,6 +566,31 @@ void UMovieSceneMaterialParameterSystem::OnInstantiation()
 	}
 }
 
+void UMovieSceneMaterialParameterSystem::OnSchedulePersistentTasks(UE::MovieScene::IEntitySystemScheduler* TaskScheduler)
+{
+	using namespace UE::MovieScene;
+
+	FBuiltInComponentTypes* BuiltInComponents = FBuiltInComponentTypes::Get();
+	FMovieSceneTracksComponentTypes* TracksComponents = FMovieSceneTracksComponentTypes::Get();
+
+	FEntityTaskBuilder()
+	.Read(TracksComponents->BoundMaterial)
+	.Read(TracksComponents->ScalarParameterName)
+	.Read(BuiltInComponents->DoubleResult[0])
+	.FilterNone({ BuiltInComponents->BlendChannelInput })
+	.SetDesiredThread(Linker->EntityManager.GetDispatchThread())
+	.Fork_PerEntity<FApplyScalarParameters>(&Linker->EntityManager, TaskScheduler);
+
+	// Vectors and colors use the same API
+	FEntityTaskBuilder()
+	.Read(TracksComponents->BoundMaterial)
+	.ReadOneOrMoreOf(TracksComponents->VectorParameterName, TracksComponents->ColorParameterName)
+	.ReadOneOrMoreOf(BuiltInComponents->DoubleResult[0], BuiltInComponents->DoubleResult[1], BuiltInComponents->DoubleResult[2], BuiltInComponents->DoubleResult[3])
+	.FilterNone({ BuiltInComponents->BlendChannelInput })
+	.SetDesiredThread(Linker->EntityManager.GetDispatchThread())
+	.Fork_PerAllocation<FApplyVectorParameters>(&Linker->EntityManager, TaskScheduler);
+}
+
 void UMovieSceneMaterialParameterSystem::OnEvaluation(FSystemTaskPrerequisites& InPrerequisites, FSystemSubsequentTasks& Subsequents)
 {
 	using namespace UE::MovieScene;
