Sequencer: Speculative fix for ASAN error

diff --git a/Engine/Source/Runtime/MovieScene/Public/EntitySystem/MovieSceneEntitySystemTask.h b/Engine/Source/Runtime/MovieScene/Public/EntitySystem/MovieSceneEntitySystemTask.h
--- a/Engine/Source/Runtime/MovieScene/Public/EntitySystem/MovieSceneEntitySystemTask.h
+++ b/Engine/Source/Runtime/MovieScene/Public/EntitySystem/MovieSceneEntitySystemTask.h
@@ -126,18 +126,19 @@ struct FEntityIterationResult
 	bool Value = true;
 };
 
-static FTaskID SchedulePreTask(IEntitySystemScheduler* InScheduler, const FTaskParams& Params, TSharedPtr<ITaskContext> Context, void*, ...)
+template <typename TaskType>
+static FTaskID SchedulePreTask(IEntitySystemScheduler*, const FTaskParams&, const TSharedPtr<TaskType>&, void*)
 {
 	return FTaskID::None();
 }
-
-static FTaskID SchedulePostTask(IEntitySystemScheduler* InScheduler, const FTaskParams& Params, TSharedPtr<ITaskContext> Context, void*, ...)
+template <typename TaskType>
+static FTaskID SchedulePostTask(IEntitySystemScheduler*, const FTaskParams&, const TSharedPtr<TaskType>&, void*)
 {
 	return FTaskID::None();
 }
 
 template <typename TaskType, typename TaskImplType>
-static FTaskID SchedulePreTask(IEntitySystemScheduler* InScheduler, const FTaskParams& Params, TSharedPtr<TaskType> Context, TaskImplType* Unused, decltype(&TaskImplType::PreTask)* = 0)
+static FTaskID SchedulePreTask(IEntitySystemScheduler* InScheduler, const FTaskParams& Params, const TSharedPtr<TaskType>& Context, TaskImplType* Unused, decltype(&TaskImplType::PreTask)* = 0)
 {
 	struct FCallPreTask
 	{
@@ -151,7 +152,7 @@ static FTaskID SchedulePreTask(IEntitySystemScheduler* InScheduler, const FTaskP
 	return InScheduler->AddTask(Params, Context, Function);
 }
 template <typename TaskType, typename TaskImplType>
-static FTaskID SchedulePostTask(IEntitySystemScheduler* InScheduler, const FTaskParams& Params, TSharedPtr<TaskType> Context, TaskImplType* Unused, decltype(&TaskImplType::PostTask)* = 0)
+static FTaskID SchedulePostTask(IEntitySystemScheduler* InScheduler, const FTaskParams& Params, const TSharedPtr<TaskType>& Context, TaskImplType* Unused, decltype(&TaskImplType::PostTask)* = 0)
 {
 	struct FCallPostTask
 	{
