MovieScene: Fixed issue where RunLatentActions was invoked while the player/runner was still evaluating for players with budgeted evaluation. This caused latent actions to be evaluated out of order where the first latent action would be requeued to the end of the list.

diff --git a/Engine/Source/Runtime/MovieScene/Private/MovieSceneSequenceTickManager.cpp b/Engine/Source/Runtime/MovieScene/Private/MovieSceneSequenceTickManager.cpp
--- a/Engine/Source/Runtime/MovieScene/Private/MovieSceneSequenceTickManager.cpp
+++ b/Engine/Source/Runtime/MovieScene/Private/MovieSceneSequenceTickManager.cpp
@@ -1,6 +1,7 @@
 // Copyright Epic Games, Inc. All Rights Reserved.
 
 #include "MovieSceneSequenceTickManager.h"
+#include "MovieSceneSequencePlayer.h"
 #include "Engine/World.h"
 #include "EntitySystem/MovieSceneEntitySystemLinker.h"
 #include "MovieSceneSequenceTickManagerClient.h"
@@ -552,9 +553,19 @@ void FMovieSceneLatentActionManager::RunLatentActions(TFunctionRef<void()> Flush
 			UObject* BoundObject = Delegate.GetUObject();
 			if (ensure(BoundObject) && !ExecutedDelegateOwners.Contains(BoundObject))
 			{
-				Delegate.ExecuteIfBound();
+				UMovieSceneSequencePlayer* Player = Cast<UMovieSceneSequencePlayer>(BoundObject);
+				if (Player && Player->IsEvaluating())
+				{
+					// If our player is still evaluating, defer all latent actions for this
+					// sequence player to the next pass.
+					++Index;
+				}
+				else
+				{
+					Delegate.ExecuteIfBound();
+                	LatentActions.RemoveAt(Index);
+				}
 				ExecutedDelegateOwners.Add(BoundObject);
-				LatentActions.RemoveAt(Index);
 			}
 			else
 			{
