Sequencer: When using current type to set the inpterpolation type don't do if it's user or broken, this will set a flat tangent, instead use an auto type.
#trivial

diff --git a/Engine/Source/Runtime/MovieScene/Private/Channels/MovieSceneCurveChannelImpl.cpp b/Engine/Source/Runtime/MovieScene/Private/Channels/MovieSceneCurveChannelImpl.cpp
--- a/Engine/Source/Runtime/MovieScene/Private/Channels/MovieSceneCurveChannelImpl.cpp
+++ b/Engine/Source/Runtime/MovieScene/Private/Channels/MovieSceneCurveChannelImpl.cpp
@@ -1259,11 +1259,20 @@ EMovieSceneKeyInterpolation TMovieSceneCurveChannelImpl<ChannelType>::GetInterpo
 			case RCTM_Auto:
 				return EMovieSceneKeyInterpolation::Auto;
 				break;
-			case RCTM_Break:
-				return EMovieSceneKeyInterpolation::Break;
+			case RCTM_SmartAuto:
+				return EMovieSceneKeyInterpolation::SmartAuto;
 				break;
+			case RCTM_Break:
 			case RCTM_User:
-				return EMovieSceneKeyInterpolation::User;
+				if (DefaultInterpolationMode == EMovieSceneKeyInterpolation::Auto ||
+					DefaultInterpolationMode == EMovieSceneKeyInterpolation::SmartAuto)
+				{
+					return DefaultInterpolationMode;
+				}
+				else
+				{
+					return EMovieSceneKeyInterpolation::SmartAuto;
+				}
 				break;
 			}
 			break;
