Sequencer: Fixed construction of Set<PropertyName>() string.

FName::ToString resets the provided string before appending

diff --git a/Engine/Source/Runtime/MovieScene/Private/EntitySystem/MovieScenePropertyRegistry.cpp b/Engine/Source/Runtime/MovieScene/Private/EntitySystem/MovieScenePropertyRegistry.cpp
--- a/Engine/Source/Runtime/MovieScene/Private/EntitySystem/MovieScenePropertyRegistry.cpp
+++ b/Engine/Source/Runtime/MovieScene/Private/EntitySystem/MovieScenePropertyRegistry.cpp
@@ -124,7 +124,7 @@ TOptional<uint16> ComputeFastPropertyPtrOffset(UClass* ObjectClass, const FMovie
 	//        This is a little better now we use a string builder and an FNAME_Find, but it's still not ideal
 	TStringBuilder<128> SetterName;
 	SetterName.Append(TEXT("Set"));
-	PropertyBinding.PropertyName.ToString(SetterName);
+	PropertyBinding.PropertyName.AppendString(SetterName);
 
 	FName SetterFunctionName(SetterName.ToString(), FNAME_Find);
 	if (SetterFunctionName.IsNone() || ObjectClass->FindFunctionByName(SetterFunctionName) == nullptr)
