Sequencer: Fixed add track menus not showing up for properties that have custom accessors

diff --git a/Engine/Source/Editor/Sequencer/Private/SequencerObjectChangeListener.cpp b/Engine/Source/Editor/Sequencer/Private/SequencerObjectChangeListener.cpp
--- a/Engine/Source/Editor/Sequencer/Private/SequencerObjectChangeListener.cpp
+++ b/Engine/Source/Editor/Sequencer/Private/SequencerObjectChangeListener.cpp
@@ -9,6 +9,7 @@
 #include "PropertyPermissionList.h"
 #include "IPropertyChangeListener.h"
 #include "MovieSceneSequence.h"
+#include "EntitySystem/MovieScenePropertySystemTypes.h"
 #include "ScopedTransaction.h"
 
 DEFINE_LOG_CATEGORY_STATIC(LogSequencerTools, Log, All);
@@ -297,6 +298,18 @@ bool FSequencerObjectChangeListener::CanKeyProperty_Internal(FCanKeyPropertyPara
 
 				{
 					FAnimatedPropertyKey PropertyKey = FAnimatedPropertyKey::FromProperty(Property);
+
+					// If there is a custom accessor for this specific property path, it is animatable (as long as there is a supported track editor registered for the property type)
+					if (UE::MovieScene::GlobalCustomAccessorExists(CanKeyPropertyParams.ObjectClass, InOutPropertyPath.ToString(TEXT("."))))
+					{
+						if (const FOnAnimatablePropertyChanged* DelegatePtr = PropertyChangedEventMap.Find(PropertyKey))
+						{
+							InOutProperty = Property;
+							InOutDelegate = *DelegatePtr;
+							return true;
+						}
+					}
+					// Otherwise we check for magic named functions using the default logic
 					const FOnAnimatablePropertyChanged* DelegatePtr = FindPropertySetter(*PropertyContainer, PropertyKey, *Property);
 					if (DelegatePtr != nullptr)
 					{
diff --git a/Engine/Source/Runtime/MovieScene/Private/EntitySystem/MovieScenePropertySystemTypes.cpp b/Engine/Source/Runtime/MovieScene/Private/EntitySystem/MovieScenePropertySystemTypes.cpp
new file mode 100644
--- /dev/null
+++ b/Engine/Source/Runtime/MovieScene/Private/EntitySystem/MovieScenePropertySystemTypes.cpp
@@ -0,0 +1,45 @@
+// Copyright Epic Games, Inc. All Rights Reserved.
+
+#include "EntitySystem/MovieScenePropertySystemTypes.h"
+#include "UObject/ObjectKey.h"
+#include "UObject/NameTypes.h"
+#include "UObject/Class.h"
+#include "Templates/Tuple.h"
+
+namespace UE::MovieScene
+{
+
+#if WITH_EDITOR
+
+TSet<TTuple<FObjectKey, FName>> GRegisteredCustomAccessors;
+
+void AddGlobalCustomAccessor(const UClass* ClassType, FName PropertyPath)
+{
+	GRegisteredCustomAccessors.Add(TTuple<FObjectKey, FName>(ClassType, PropertyPath));
+}
+void RemoveGlobalCustomAccessor(const UClass* ClassType, FName PropertyPath)
+{
+	GRegisteredCustomAccessors.Remove(TTuple<FObjectKey, FName>(ClassType, PropertyPath));
+}
+bool GlobalCustomAccessorExists(const UClass* ClassType, TStringView<WIDECHAR> PropertyPath)
+{
+	FName PropertyPathName(PropertyPath.Len(), PropertyPath.GetData(), FNAME_Find);
+	if (PropertyPathName != NAME_None)
+	{
+		return GRegisteredCustomAccessors.Contains(TTuple<FObjectKey, FName>(ClassType, PropertyPath));
+	}
+	return false;
+}
+bool GlobalCustomAccessorExists(const UClass* ClassType, TStringView<ANSICHAR> PropertyPath)
+{
+	FName PropertyPathName(PropertyPath.Len(), PropertyPath.GetData(), FNAME_Find);
+	if (PropertyPathName != NAME_None)
+	{
+		return GRegisteredCustomAccessors.Contains(TTuple<FObjectKey, FName>(ClassType, PropertyPath));
+	}
+	return false;
+}
+
+#endif // WITH_EDITOR
+
+} // namespace namespace UE::MovieScene
diff --git a/Engine/Source/Runtime/MovieScene/Public/EntitySystem/MovieScenePropertySystemTypes.h b/Engine/Source/Runtime/MovieScene/Public/EntitySystem/MovieScenePropertySystemTypes.h
--- a/Engine/Source/Runtime/MovieScene/Public/EntitySystem/MovieScenePropertySystemTypes.h
+++ b/Engine/Source/Runtime/MovieScene/Public/EntitySystem/MovieScenePropertySystemTypes.h
@@ -8,6 +8,7 @@
 #include "UObject/NameTypes.h"
 #include "Templates/SharedPointer.h"
 #include "Templates/UnrealTypeTraits.h"
+#include "Containers/StringView.h"
 #include "EntitySystem/MovieSceneComponentAccessors.h"
 
 class UClass;
@@ -217,6 +218,13 @@ struct ICustomPropertyRegistration
 	virtual FCustomAccessorView GetAccessors() const = 0;
 };
 
+#if WITH_EDITOR
+MOVIESCENE_API void AddGlobalCustomAccessor(const UClass* ClassType, FName PropertyPath);
+MOVIESCENE_API void RemoveGlobalCustomAccessor(const UClass* ClassType, FName PropertyPath);
+MOVIESCENE_API bool GlobalCustomAccessorExists(const UClass* ClassType, TStringView<WIDECHAR> PropertyPath);
+MOVIESCENE_API bool GlobalCustomAccessorExists(const UClass* ClassType, TStringView<ANSICHAR> PropertyPath);
+#endif // WITH_EDITOR
+
 /** Generally static collection of accessors for a given type of property */
 template<typename PropertyTraits, int InlineSize = 8>
 struct TCustomPropertyRegistration : ICustomPropertyRegistration
@@ -232,22 +240,43 @@ struct TCustomPropertyRegistration : ICustomPropertyRegistration
 	void Add(UClass* ClassType, FName PropertyName, GetterFunc Getter, SetterFunc Setter)
 	{
 		CustomAccessors.Add(TCustomPropertyAccessor<PropertyTraits>{ ClassType, PropertyName, { Getter, Setter } });
+#if WITH_EDITOR
+		AddGlobalCustomAccessor(ClassType, PropertyName);
+#endif
 	}
 
 	void Remove(UClass* ClassType, FName PropertyName)
 	{
-		CustomAccessors.RemoveAll([=](const TCustomPropertyAccessor<PropertyTraits>& Accessor)
-				{
-					return Accessor.Class == ClassType && Accessor.PropertyPath == PropertyName;
-				});
+		for (int32 Index = CustomAccessors.Num()-1; Index >= 0; --Index)
+		{
+			TCustomPropertyAccessor<PropertyTraits>& Accessor = CustomAccessors[Index];
+			if (Accessor.Class == ClassType && Accessor.PropertyPath == PropertyName)
+			{
+#if WITH_EDITOR
+				RemoveGlobalCustomAccessor(Accessor.Class, Accessor.PropertyPath);
+#endif
+				// Null out the entry rather than remove it because we don't want to invalidate any cached array indices
+				Accessor.Class = nullptr;
+				Accessor.PropertyPath = NAME_None;
+			}
+		}
 	}
 
 	void RemoveAll(UClass* ClassType)
 	{
-		CustomAccessors.RemoveAll([=](const TCustomPropertyAccessor<PropertyTraits>& Accessor)
-				{
-					return Accessor.Class == ClassType;
-				});
+		for (int32 Index = CustomAccessors.Num()-1; Index >= 0; --Index)
+		{
+			TCustomPropertyAccessor<PropertyTraits>& Accessor = CustomAccessors[Index];
+			if (Accessor.Class == ClassType)
+			{
+#if WITH_EDITOR
+				RemoveGlobalCustomAccessor(Accessor.Class, Accessor.PropertyPath);
+#endif
+				// Null out the entry rather than remove it because we don't want to invalidate any cached array indices
+				Accessor.Class = nullptr;
+				Accessor.PropertyPath = NAME_None;
+			}
+		}
 	}
 
 private:
