Sequencer: Changed all raw UObject ptr usage to FObjectComponent, TWeakObjectPtr or FObjectKey

CL25489627 made a blanket change to all AddReferencedObjects functions which actually removed all of sequencer's exposure to the reference graph because all its raw UObject* component types now ended up using the default catch all reference collection function.

In response (and in light of the recent deprecation of raw ptr interop with the reference graph), I've converted everything to either FObjectKey, TWeakObjectPtr or FObjectComponent depending on useage and expected reference semantics.

TObjectPtr<> remains exposed to the reference graph when used as component data (in the same way raw pointers used to be).

Existing uses of the overlapping entity tracker that used raw UObjects now use FObjectKey, and still respond to GC passes by removing null ptrs (through a hack that adds a stubbed out overload of AddReferencedObjectForComponent for object key types that enables the instantiation of TOverlappingEntityTracker_WithGarbage.

Component types and overlapping entity trackers instantiated with UObject* will no longer compile due to static asserts added for protection against this in future.

diff --git a/Engine/Source/Runtime/MovieScene/Private/EntitySystem/BuiltInComponentTypes.cpp b/Engine/Source/Runtime/MovieScene/Private/EntitySystem/BuiltInComponentTypes.cpp
--- a/Engine/Source/Runtime/MovieScene/Private/EntitySystem/BuiltInComponentTypes.cpp
+++ b/Engine/Source/Runtime/MovieScene/Private/EntitySystem/BuiltInComponentTypes.cpp
@@ -38,10 +38,6 @@ void AddReferencedObjectForComponent(FReferenceCollector& ReferenceCollector, FO
 		ReferenceCollector.AddReferencedObject(ComponentData->ObjectPtr);
 	}
 }
-void AddReferencedObjectForBoundObject(FReferenceCollector& ReferenceCollector, void* ComponentData, int32 Num)
-{
-	// Intentionally hidden from the reference graph
-}
 
 static bool GMovieSceneBuiltInComponentTypesDestroyed = false;
 static TUniquePtr<FBuiltInComponentTypes> GMovieSceneBuiltInComponentTypes;
@@ -77,7 +73,8 @@ FBuiltInComponentTypes::FBuiltInComponentTypes()
 	ComponentRegistry->NewComponentType(&EvalSeconds,           TEXT("Eval Seconds"));
 
 	ComponentRegistry->NewComponentType(&BoundObjectKey,        TEXT("Bound Object Key"));
-	ComponentRegistry->NewComponentType(&BoundObject,           TEXT("Bound Object"), FNewComponentTypeParams(&AddReferencedObjectForBoundObject, EComponentTypeFlags::None));
+	// Intentionally hidden from the reference graph because they are always accompanied by a BoundObjectKey which is used for garbage collection
+	ComponentRegistry->NewComponentTypeNoAddReferencedObjects(&BoundObject, TEXT("Bound Object"));
 
 	ComponentRegistry->NewComponentType(&PropertyBinding,         TEXT("Property Binding"), EComponentTypeFlags::CopyToOutput);
 	ComponentRegistry->NewComponentType(&GenericObjectBinding,    TEXT("Generic Object Binding ID"));
diff --git a/Engine/Source/Runtime/MovieScene/Public/EntitySystem/BuiltInComponentTypes.h b/Engine/Source/Runtime/MovieScene/Public/EntitySystem/BuiltInComponentTypes.h
--- a/Engine/Source/Runtime/MovieScene/Public/EntitySystem/BuiltInComponentTypes.h
+++ b/Engine/Source/Runtime/MovieScene/Public/EntitySystem/BuiltInComponentTypes.h
@@ -450,7 +450,7 @@ public:
 	TComponentTypeID<FString> StringResult;
 
 	// The result of an evaluated FMovieSceneObjectPathChannel
-	TComponentTypeID<UObject*> ObjectResult;
+	TComponentTypeID<FObjectComponent> ObjectResult;
 
 public:
 
diff --git a/Engine/Source/Runtime/MovieScene/Public/EntitySystem/MovieSceneComponentRegistry.h b/Engine/Source/Runtime/MovieScene/Public/EntitySystem/MovieSceneComponentRegistry.h
--- a/Engine/Source/Runtime/MovieScene/Public/EntitySystem/MovieSceneComponentRegistry.h
+++ b/Engine/Source/Runtime/MovieScene/Public/EntitySystem/MovieSceneComponentRegistry.h
@@ -79,12 +79,24 @@ public:
 	template<typename T>
 	TComponentTypeID<T> NewComponentType(const TCHAR* const DebugName, const FNewComponentTypeParams& Params = FNewComponentTypeParams());
 
+	/**
+	 * Same as NewComponentType but specifically does not expose the component type to the reference graph. Use with caution!
+	 */
+	template<typename T>
+	TComponentTypeID<T> NewComponentTypeNoAddReferencedObjects(const TCHAR* const DebugName, const FNewComponentTypeParams& Params = FNewComponentTypeParams());
+
 	template<typename T>
 	void NewComponentType(TComponentTypeID<T>* Ref, const TCHAR* const DebugName, const FNewComponentTypeParams& Params = FNewComponentTypeParams())
 	{
 		*Ref = NewComponentType<T>(DebugName, Params);
 	}
 
+	template<typename T>
+	void NewComponentTypeNoAddReferencedObjects(TComponentTypeID<T>* Ref, const TCHAR* const DebugName, const FNewComponentTypeParams& Params = FNewComponentTypeParams())
+	{
+		*Ref = NewComponentTypeNoAddReferencedObjects<T>(DebugName, Params);
+	}
+
 	template<typename PropertyTraits>
 	void NewPropertyType(TPropertyComponents<PropertyTraits>& OutComponents, const TCHAR* DebugName)
 	{
@@ -159,6 +171,9 @@ private:
 
 	FComponentTypeID NewComponentTypeInternal(FComponentTypeInfo&& TypeInfo);
 
+	template<typename T>
+	FComponentTypeInfo MakeComponentTypeInfoWithoutComponentOps(const TCHAR* const DebugName, const FNewComponentTypeParams& Params);
+
 private:
 
 	TSparseArray<FComponentTypeInfo> ComponentTypes;
diff --git a/Engine/Source/Runtime/MovieScene/Public/EntitySystem/MovieSceneComponentTypeInfo.h b/Engine/Source/Runtime/MovieScene/Public/EntitySystem/MovieSceneComponentTypeInfo.h
--- a/Engine/Source/Runtime/MovieScene/Public/EntitySystem/MovieSceneComponentTypeInfo.h
+++ b/Engine/Source/Runtime/MovieScene/Public/EntitySystem/MovieSceneComponentTypeInfo.h
@@ -26,14 +26,31 @@ using FComponentReferenceCollectionPtr = void (*)(FReferenceCollector&, void*, i
  */
 inline FNotImplemented* AddReferencedObjectForComponent(...) { return nullptr; }
 
+/** Instantiation for raw pointers */
 template<typename T>
-typename TEnableIf<TPointerIsConvertibleFromTo<T, UObject>::Value>::Type AddReferencedObjectForComponent(FReferenceCollector* ReferenceCollector, TObjectPtr<T>* Component)
+typename TEnableIf<TPointerIsConvertibleFromTo<T, UObject>::Value>::Type AddReferencedObjectForComponent(FReferenceCollector* ReferenceCollector, T** Component)
+{
+	// Ideally this would be a compile-time check buyt 
+	constexpr bool bTypeDependentFalse = !std::is_same_v<T, T>;
+	static_assert(bTypeDependentFalse, "Raw object pointers are no longer supported. Please use TObjectPtr<T> instead.");
+}
+
+template<typename T>
+void AddReferencedObjectForComponent(FReferenceCollector* ReferenceCollector, TObjectPtr<T>* Component)
 {
 	ReferenceCollector->AddReferencedObject(*Component);
 }
 
+// Hack to enable garbage collection path for object keys
+inline void AddReferencedObjectForComponent(FReferenceCollector* ReferenceCollector, FObjectKey* Component)
+{}
+
 template<typename T>
-typename TEnableIf< std::is_same_v<decltype(T::StaticStruct()), decltype(T::StaticStruct())> >::Type AddReferencedObjectForComponent(FReferenceCollector* ReferenceCollector, T* Component)
+void AddReferencedObjectForComponent(FReferenceCollector* ReferenceCollector, TObjectKey<T>* Component)
+{}
+
+template<typename T, typename U = decltype(T::StaticStruct())>
+void AddReferencedObjectForComponent(FReferenceCollector* ReferenceCollector, T* Component)
 {
 	for (TPropertyValueIterator<const FObjectProperty> It(T::StaticStruct(), Component); It; ++It)
 	{
@@ -184,6 +201,15 @@ struct FComponentTypeInfo
 		ComplexComponentOps = MakeUnique<TComplexComponentOps<T>>();
 	}
 
+	/**
+	 * Define complex component operations for this type of component without implementing AddReferencedObjects - use with caution!
+	 */
+	template<typename T>
+	void MakeComplexComponentOpsNoAddReferencedObjects()
+	{
+		ComplexComponentOps = MakeUnique<TComplexComponentOpsBase<T>>();
+	}
+
 	/**
 	 * Define complex component with a specific reference collection callback
 	 */
@@ -234,6 +260,8 @@ private:
 				++TypedSrc;
 			}
 		}
+		virtual void AddReferencedObjects(FReferenceCollector& ReferenceCollector, void* ComponentStart, int32 Num)
+		{}
 	};
 	template<typename T>
 	struct TComplexComponentOps : TComplexComponentOpsBase<T>
diff --git a/Engine/Source/Runtime/MovieScene/Public/EntitySystem/MovieSceneEntityFactoryTemplates.h b/Engine/Source/Runtime/MovieScene/Public/EntitySystem/MovieSceneEntityFactoryTemplates.h
--- a/Engine/Source/Runtime/MovieScene/Public/EntitySystem/MovieSceneEntityFactoryTemplates.h
+++ b/Engine/Source/Runtime/MovieScene/Public/EntitySystem/MovieSceneEntityFactoryTemplates.h
@@ -207,7 +207,7 @@ inline void FEntityFactories::DefineChildComponent(TComponentTypeID<ParentCompon
 }
 
 template<typename T>
-TComponentTypeID<T> FComponentRegistry::NewComponentType(const TCHAR* const DebugName, const FNewComponentTypeParams& Params)
+FComponentTypeInfo FComponentRegistry::MakeComponentTypeInfoWithoutComponentOps(const TCHAR* const DebugName, const FNewComponentTypeParams& Params)
 {
 	static const uint32 ComponentTypeSize = sizeof(T);
 	static_assert(ComponentTypeSize < TNumericLimits<decltype(FComponentTypeInfo::Sizeof)>::Max(), "Type too large to be used as component data");
@@ -217,23 +217,33 @@ TComponentTypeID<T> FComponentRegistry::NewComponentType(const TCHAR* const Debu
 
 	FComponentTypeInfo NewTypeInfo;
 
-	NewTypeInfo.Sizeof                     = ComponentTypeSize;
-	NewTypeInfo.Alignment                  = Alignment;
-	NewTypeInfo.bIsZeroConstructType       = TIsZeroConstructType<T>::Value;
-	NewTypeInfo.bIsTriviallyDestructable   = TIsTriviallyDestructible<T>::Value;
+	NewTypeInfo.Sizeof = ComponentTypeSize;
+	NewTypeInfo.Alignment = Alignment;
+	NewTypeInfo.bIsZeroConstructType = TIsZeroConstructType<T>::Value;
+	NewTypeInfo.bIsTriviallyDestructable = TIsTriviallyDestructible<T>::Value;
 	NewTypeInfo.bIsTriviallyCopyAssignable = TIsTriviallyCopyAssignable<T>::Value;
-	NewTypeInfo.bIsPreserved               = EnumHasAnyFlags(Params.Flags, EComponentTypeFlags::Preserved);
-	NewTypeInfo.bIsCopiedToOutput          = EnumHasAnyFlags(Params.Flags, EComponentTypeFlags::CopyToOutput);
-	NewTypeInfo.bIsMigratedToOutput        = EnumHasAnyFlags(Params.Flags, EComponentTypeFlags::MigrateToOutput);
-	NewTypeInfo.bHasReferencedObjects      = Params.ReferenceCollectionCallback != nullptr || THasAddReferencedObjectForComponent<T>::Value;
+	NewTypeInfo.bIsPreserved = EnumHasAnyFlags(Params.Flags, EComponentTypeFlags::Preserved);
+	NewTypeInfo.bIsCopiedToOutput = EnumHasAnyFlags(Params.Flags, EComponentTypeFlags::CopyToOutput);
+	NewTypeInfo.bIsMigratedToOutput = EnumHasAnyFlags(Params.Flags, EComponentTypeFlags::MigrateToOutput);
+	NewTypeInfo.bHasReferencedObjects = false;
 
 #if UE_MOVIESCENE_ENTITY_DEBUG
-	NewTypeInfo.DebugInfo                = MakeUnique<FComponentTypeDebugInfo>();
-	NewTypeInfo.DebugInfo->DebugName     = DebugName;
+	NewTypeInfo.DebugInfo = MakeUnique<FComponentTypeDebugInfo>();
+	NewTypeInfo.DebugInfo->DebugName = DebugName;
 	NewTypeInfo.DebugInfo->DebugTypeName = GetGeneratedTypeName<T>();
-	NewTypeInfo.DebugInfo->Type          = TComponentDebugType<T>::Type;
+	NewTypeInfo.DebugInfo->Type = TComponentDebugType<T>::Type;
 #endif
 
+	return NewTypeInfo;
+}
+
+template<typename T>
+TComponentTypeID<T> FComponentRegistry::NewComponentType(const TCHAR* const DebugName, const FNewComponentTypeParams& Params)
+{
+	FComponentTypeInfo NewTypeInfo = FComponentRegistry::MakeComponentTypeInfoWithoutComponentOps<T>(DebugName, Params);
+
+	NewTypeInfo.bHasReferencedObjects = Params.ReferenceCollectionCallback != nullptr || THasAddReferencedObjectForComponent<T>::Value;
+
 	if (Params.ReferenceCollectionCallback)
 	{
 		NewTypeInfo.MakeComplexComponentOps<T>(Params.ReferenceCollectionCallback);
@@ -254,5 +264,27 @@ TComponentTypeID<T> FComponentRegistry::NewComponentType(const TCHAR* const Debu
 	return TypedTypeID;
 }
 
+template<typename T>
+TComponentTypeID<T> FComponentRegistry::NewComponentTypeNoAddReferencedObjects(const TCHAR* const DebugName, const FNewComponentTypeParams& Params)
+{
+	FComponentTypeInfo NewTypeInfo = FComponentRegistry::MakeComponentTypeInfoWithoutComponentOps<T>(DebugName, Params);
+
+	NewTypeInfo.bHasReferencedObjects = false;
+	if (!NewTypeInfo.bIsZeroConstructType || !NewTypeInfo.bIsTriviallyDestructable || !NewTypeInfo.bIsTriviallyCopyAssignable)
+	{
+		NewTypeInfo.MakeComplexComponentOpsNoAddReferencedObjects<T>();
+	}
+
+	FComponentTypeID    NewTypeID = NewComponentTypeInternal(MoveTemp(NewTypeInfo));
+	TComponentTypeID<T> TypedTypeID = NewTypeID.ReinterpretCast<T>();
+
+	if (EnumHasAnyFlags(Params.Flags, EComponentTypeFlags::CopyToChildren))
+	{
+		Factories.DefineChildComponent(TDuplicateChildEntityInitializer<T>(TypedTypeID));
+	}
+
+	return TypedTypeID;
+}
+
 }	// using namespace MovieScene
 }	// using namespace UE
diff --git a/Engine/Source/Runtime/MovieScene/Public/EntitySystem/MovieSceneOverlappingEntityTracker.h b/Engine/Source/Runtime/MovieScene/Public/EntitySystem/MovieSceneOverlappingEntityTracker.h
--- a/Engine/Source/Runtime/MovieScene/Public/EntitySystem/MovieSceneOverlappingEntityTracker.h
+++ b/Engine/Source/Runtime/MovieScene/Public/EntitySystem/MovieSceneOverlappingEntityTracker.h
@@ -34,17 +34,30 @@ struct FGarbageTraits
 	{
 		return false;
 	}
-	FORCEINLINE static bool IsGarbage(UObject* InObject)
+	template<typename T>
+	FORCEINLINE static std::enable_if_t<TPointerIsConvertibleFromTo<T, const UObject>::Value, bool> IsGarbage(T* InObject)
+	{
+		constexpr bool bTypeDependentFalse = !std::is_same_v<T, T>;
+		static_assert(bTypeDependentFalse, "Raw object pointers are no longer supported. Please use TObjectPtr<T> instead.");
+	}
+	FORCEINLINE static bool IsGarbage(FObjectKey InObject)
+	{
+		return FBuiltInComponentTypes::IsBoundObjectGarbage(InObject.ResolveObjectPtr());
+	}
+	template<typename T>
+	FORCEINLINE static bool IsGarbage(TObjectPtr<T>& InObject)
 	{
 		return FBuiltInComponentTypes::IsBoundObjectGarbage(InObject);
 	}
-	FORCEINLINE static constexpr void AddReferencedObjects(FReferenceCollector& ReferenceCollector, ...)
-	{}
+
+	
 	template<typename T>
-	FORCEINLINE static typename TEnableIf<TPointerIsConvertibleFromTo<T, volatile const UObject>::Value>::Type
-		AddReferencedObjects(FReferenceCollector& ReferenceCollector, T*& InObjectPtr)
+	FORCEINLINE static void AddReferencedObjects(FReferenceCollector& ReferenceCollector, T* In)
 	{
-		ReferenceCollector.AddReferencedObject(InObjectPtr);
+		if constexpr (THasAddReferencedObjectForComponent<T>::Value)
+		{
+			AddReferencedObjectForComponent(&ReferenceCollector, In);
+		}
 	}
 };
 
@@ -54,8 +67,6 @@ template<typename... T> struct TOverlappingEntityInput;
 template<typename... T, int... Indices>
 struct TGarbageTraitsImpl<TIntegerSequence<int, Indices...>, T...>
 {
-	static constexpr bool bCanBeGarbage = (TPointerIsConvertibleFromTo<typename TDecay<typename TRemovePointer<T>::Type>::Type, UObject>::Value || ...);
-
 	static bool IsGarbage(TOverlappingEntityInput<T...>& InParam)
 	{
 		return (FGarbageTraits::IsGarbage(InParam.Key.template Get<Indices>()) || ...);
@@ -63,7 +74,7 @@ struct TGarbageTraitsImpl<TIntegerSequence<int, Indices...>, T...>
 
 	static void AddReferencedObjects(FReferenceCollector& ReferenceCollector, TOverlappingEntityInput<T...>& InParam)
 	{
-		(FGarbageTraits::AddReferencedObjects(ReferenceCollector, InParam.Key.template Get<Indices>()), ...);
+		(FGarbageTraits::AddReferencedObjects(ReferenceCollector, &InParam.Key.template Get<Indices>()), ...);
 	}
 
 	template<typename CallbackType>
@@ -111,7 +122,6 @@ struct TOverlappingEntityInput
  * Templated utility class that assists in tracking the state of many -> one data relationships in an FEntityManager.
  * InputKeyTypes defines the component type(s) which defines the key that determines whether an entity animates the same output.
  * OutputType defines the user-specfied data to be associated with the multiple inputs (ie, its output)
- * NOTE: Where any of InputKeyTypes is a UObject*, AddReferencedObjects and CleanupGarbage must be called
  */
 template<typename OutputType, typename... InputKeyTypes>
 struct TOverlappingEntityTrackerImpl
@@ -129,6 +139,16 @@ struct TOverlappingEntityTrackerImpl
 	 * Only entities tagged as NeedsLink or NeedsUnlink are iterated, invalidating their outputs
 	 */
 	void Update(UMovieSceneEntitySystemLinker* Linker, TComponentTypeID<InputKeyTypes>... InKeyComponents, const FEntityComponentFilter& InFilter)
+	{
+		UpdateFromComponents(Linker, InFilter, InKeyComponents...);
+	}
+
+	/**
+	 * Update this tracker by iterating any entity that contains InKeyComponent, and matches the additional optional filter
+	 * Only entities tagged as NeedsLink or NeedsUnlink are iterated, invalidating their outputs
+	 */
+	template<typename ...ComponentTypes>
+	void UpdateFromComponents(UMovieSceneEntitySystemLinker* Linker, const FEntityComponentFilter& InFilter, TComponentTypeID<ComponentTypes>... InKeyComponents)
 	{
 		check(bIsInitialized);
 
@@ -145,14 +165,15 @@ struct TOverlappingEntityTrackerImpl
 		.ReadAllOf(InKeyComponents...)
 		.CombineFilter(InFilter)
 		.FilterAll({ BuiltInComponents->Tags.NeedsLink })
-		.Iterate_PerAllocation(&Linker->EntityManager, [this](const FEntityAllocation* Allocation, TRead<InputKeyTypes>... ReadKeys){ this->VisitLinkedAllocation(Allocation, ReadKeys...); });
+		.Iterate_PerAllocation(&Linker->EntityManager, [this](const FEntityAllocation* Allocation, TRead<ComponentTypes>... ReadKeys){ this->VisitLinkedAllocation(Allocation, ReadKeys...); });
 	}
 
 	/**
 	 * Update this tracker by iterating any entity that contains InKeyComponent, and matches the additional optional filter
 	 * Only entities tagged as NeedsUnlink are iterated, invalidating their outputs
 	 */
-	void UpdateUnlinkedOnly(UMovieSceneEntitySystemLinker* Linker, TComponentTypeID<InputKeyTypes>... InKeyComponent, const FEntityComponentFilter& InFilter)
+	template<typename ...ComponentTypes>
+	void UpdateUnlinkedOnly(UMovieSceneEntitySystemLinker* Linker, TComponentTypeID<ComponentTypes>... InKeyComponent, const FEntityComponentFilter& InFilter)
 	{
 		check(bIsInitialized);
 
@@ -168,14 +189,16 @@ struct TOverlappingEntityTrackerImpl
 	/**
 	 * Update this tracker by (re)linking the specified allocation
 	 */
-	void VisitActiveAllocation(const FEntityAllocation* Allocation, TComponentPtr<const InputKeyTypes>... ReadKeys)
+	template<typename ...ComponentTypes>
+	void VisitActiveAllocation(const FEntityAllocation* Allocation, TComponentPtr<const ComponentTypes>... ReadKeys)
 	{
 		VisitActiveAllocationImpl(Allocation, ReadKeys...);
 	}
 	/**
 	 * Update this tracker by (re)linking the specified allocation
 	 */
-	void VisitLinkedAllocation(const FEntityAllocation* Allocation, TComponentPtr<const InputKeyTypes>... ReadKeys)
+	template<typename ...ComponentTypes>
+	void VisitLinkedAllocation(const FEntityAllocation* Allocation, TComponentPtr<const ComponentTypes>... ReadKeys)
 	{
 		VisitActiveAllocationImpl(Allocation, ReadKeys...);
 	}
@@ -354,7 +377,8 @@ struct TOverlappingEntityTrackerImpl
 
 protected:
 
-	void VisitActiveAllocationImpl(const FEntityAllocation* Allocation, TComponentPtr<const InputKeyTypes>... Keys)
+	template<typename ...ComponentTypes>
+	void VisitActiveAllocationImpl(const FEntityAllocation* Allocation, TComponentPtr<const ComponentTypes>... Keys)
 	{
 		check(bIsInitialized);
 
diff --git a/Engine/Source/Runtime/MovieSceneTracks/Private/Sections/MovieScenePrimitiveMaterialSection.cpp b/Engine/Source/Runtime/MovieSceneTracks/Private/Sections/MovieScenePrimitiveMaterialSection.cpp
--- a/Engine/Source/Runtime/MovieSceneTracks/Private/Sections/MovieScenePrimitiveMaterialSection.cpp
+++ b/Engine/Source/Runtime/MovieSceneTracks/Private/Sections/MovieScenePrimitiveMaterialSection.cpp
@@ -118,7 +118,7 @@ void UMovieScenePrimitiveMaterialSection::ImportEntityImpl(UMovieSceneEntitySyst
 	OutImportedEntity->AddBuilder(
 		FEntityBuilder()
 		.AddConditional(BuiltInComponentTypes->GenericObjectBinding, ObjectBindingID, ObjectBindingID.IsValid())
-		.Add(BuiltInComponentTypes->ObjectResult, ValueToImport.Get())
+		.Add(BuiltInComponentTypes->ObjectResult, FObjectComponent::Strong(ValueToImport.Get()))
 		.Add(TracksComponentTypes->ComponentMaterialIndex, MaterialIndex)
 	);
 }
diff --git a/Engine/Source/Runtime/MovieSceneTracks/Private/Systems/MovieSceneComponentAttachmentSystem.cpp b/Engine/Source/Runtime/MovieSceneTracks/Private/Systems/MovieSceneComponentAttachmentSystem.cpp
--- a/Engine/Source/Runtime/MovieSceneTracks/Private/Systems/MovieSceneComponentAttachmentSystem.cpp
+++ b/Engine/Source/Runtime/MovieSceneTracks/Private/Systems/MovieSceneComponentAttachmentSystem.cpp
@@ -33,7 +33,7 @@ struct FInitializeAttachParentsTask
 {
 	FInstanceRegistry* InstanceRegistry;
 
-	void ForEachEntity(UE::MovieScene::FInstanceHandle InstanceHandle, const FMovieSceneObjectBindingID& BindingID, const UE::MovieScene::FAttachmentComponent& AttachComponent, USceneComponent*& OutAttachedParent)
+	void ForEachEntity(UE::MovieScene::FInstanceHandle InstanceHandle, const FMovieSceneObjectBindingID& BindingID, const UE::MovieScene::FAttachmentComponent& AttachComponent, TWeakObjectPtr<USceneComponent>& OutAttachedParent)
 	{
 		const FSequenceInstance& TargetInstance = InstanceRegistry->GetInstance(InstanceHandle);
 
@@ -63,9 +63,9 @@ struct FAttachmentHandler
 		TrackComponents = FMovieSceneTracksComponentTypes::Get();
 	}
 
-	void InitializeOutput(UObject* Object, TArrayView<const FMovieSceneEntityID> Inputs, FPreAnimAttachment* Output, FEntityOutputAggregate Aggregate)
+	void InitializeOutput(FObjectKey Object, TArrayView<const FMovieSceneEntityID> Inputs, FPreAnimAttachment* Output, FEntityOutputAggregate Aggregate)
 	{
-		USceneComponent* AttachChild = CastChecked<USceneComponent>(Object);
+		USceneComponent* AttachChild = CastChecked<USceneComponent>(Object.ResolveObjectPtr());
 
 		Output->OldAttachParent = AttachChild->GetAttachParent();
 		Output->OldAttachSocket = AttachChild->GetAttachSocketName();
@@ -73,17 +73,17 @@ struct FAttachmentHandler
 		UpdateOutput(Object, Inputs, Output, Aggregate);
 	}
 
-	void UpdateOutput(UObject* Object, TArrayView<const FMovieSceneEntityID> Inputs, FPreAnimAttachment* Output, FEntityOutputAggregate Aggregate)
+	void UpdateOutput(FObjectKey Object, TArrayView<const FMovieSceneEntityID> Inputs, FPreAnimAttachment* Output, FEntityOutputAggregate Aggregate)
 	{
-		USceneComponent* AttachChild = CastChecked<USceneComponent>(Object);
+		USceneComponent* AttachChild = CastChecked<USceneComponent>(Object.ResolveObjectPtr());
 
 		for (FMovieSceneEntityID Entity : Inputs)
 		{
-			TOptionalComponentReader<USceneComponent*>     AttachParentComponent = EntityManager->ReadComponent(Entity, TrackComponents->AttachParent);
-			TOptionalComponentReader<FAttachmentComponent> AttachmentComponent   = EntityManager->ReadComponent(Entity, TrackComponents->AttachComponent);
+			TOptionalComponentReader<TWeakObjectPtr<USceneComponent>> AttachParentComponent = EntityManager->ReadComponent(Entity, TrackComponents->AttachParent);
+			TOptionalComponentReader<FAttachmentComponent>            AttachmentComponent   = EntityManager->ReadComponent(Entity, TrackComponents->AttachComponent);
 			if (AttachParentComponent && AttachmentComponent)
 			{
-				if (USceneComponent* AttachParent = *AttachParentComponent)
+				if (USceneComponent* AttachParent = AttachParentComponent->Get())
 				{
 					Output->DetachParams = AttachmentComponent->DetachParams;
 					AttachmentComponent->AttachParams.ApplyAttach(AttachChild, AttachParent, AttachmentComponent->Destination.SocketName);
@@ -95,9 +95,10 @@ struct FAttachmentHandler
 		}
 	}
 
-	void DestroyOutput(UObject* Object, FPreAnimAttachment* Output, FEntityOutputAggregate Aggregate)
+	void DestroyOutput(FObjectKey ObjectKey, FPreAnimAttachment* Output, FEntityOutputAggregate Aggregate)
 	{
-		if (Aggregate.bNeedsRestoration)
+		UObject* Object = ObjectKey.ResolveObjectPtr();
+		if (Aggregate.bNeedsRestoration && Object != nullptr)
 		{
 			USceneComponent* AttachChild = CastChecked<USceneComponent>(Object);
 			AttachmentSystem->AddPendingDetach(AttachChild, *Output);
@@ -228,7 +229,7 @@ void UMovieSceneComponentAttachmentSystem::OnRun(FSystemTaskPrerequisites& InPre
 	FEntityComponentFilter Filter;
 	Filter.All({ TrackComponents->AttachComponent });
 
-	AttachmentTracker.Update(Linker, FBuiltInComponentTypes::Get()->BoundObject, Filter);
+	AttachmentTracker.UpdateFromComponents(Linker, Filter, FBuiltInComponentTypes::Get()->BoundObject);
 	AttachmentTracker.ProcessInvalidatedOutputs(Linker, FAttachmentHandler(this));
 }
 
diff --git a/Engine/Source/Runtime/MovieSceneTracks/Private/Systems/MovieSceneComponentMobilitySystem.cpp b/Engine/Source/Runtime/MovieSceneTracks/Private/Systems/MovieSceneComponentMobilitySystem.cpp
--- a/Engine/Source/Runtime/MovieSceneTracks/Private/Systems/MovieSceneComponentMobilitySystem.cpp
+++ b/Engine/Source/Runtime/MovieSceneTracks/Private/Systems/MovieSceneComponentMobilitySystem.cpp
@@ -37,24 +37,24 @@ struct FMobilityCacheHandler
 		: System(InSystem)
 	{}
 
-	void InitializeOutput(UObject* Object, TArrayView<const FMovieSceneEntityID> Inputs, EComponentMobility::Type* OutMobility, FEntityOutputAggregate Aggregate)
+	void InitializeOutput(FObjectKey Object, TArrayView<const FMovieSceneEntityID> Inputs, EComponentMobility::Type* OutMobility, FEntityOutputAggregate Aggregate)
 	{
-		if (USceneComponent* SceneComponent = Cast<USceneComponent>(Object))
+		if (USceneComponent* SceneComponent = Cast<USceneComponent>(Object.ResolveObjectPtr()))
 		{
 			*OutMobility = SceneComponent->Mobility;
 			SceneComponent->SetMobility(EComponentMobility::Movable);
 		}
 	}
 
-	static void UpdateOutput(UObject* Object, TArrayView<const FMovieSceneEntityID> Inputs, EComponentMobility::Type* OutMobility, FEntityOutputAggregate Aggregate)
+	static void UpdateOutput(FObjectKey Object, TArrayView<const FMovieSceneEntityID> Inputs, EComponentMobility::Type* OutMobility, FEntityOutputAggregate Aggregate)
 	{
-		if (USceneComponent* SceneComponent = Cast<USceneComponent>(Object))
+		if (USceneComponent* SceneComponent = Cast<USceneComponent>(Object.ResolveObjectPtr()))
 		{
 			SceneComponent->SetMobility(EComponentMobility::Movable);
 		}
 	}
 
-	void DestroyOutput(UObject* Object, EComponentMobility::Type* Output, FEntityOutputAggregate Aggregate)
+	void DestroyOutput(FObjectKey Object, EComponentMobility::Type* Output, FEntityOutputAggregate Aggregate)
 	{
 	}
 };
@@ -164,7 +164,7 @@ void UMovieSceneComponentMobilitySystem::OnRun(FSystemTaskPrerequisites& InPrere
 	using namespace UE::MovieScene;
 
 	// Update the mobility tracker, caching preanimated mobilities and assigning everything as moveable that needs it
-	MobilityTracker.Update(Linker, FBuiltInComponentTypes::Get()->BoundObject, Filter);
+	MobilityTracker.UpdateFromComponents(Linker, Filter, FBuiltInComponentTypes::Get()->BoundObject);
 	MobilityTracker.ProcessInvalidatedOutputs(Linker, FMobilityCacheHandler(this));
 }
 
diff --git a/Engine/Source/Runtime/MovieSceneTracks/Private/Systems/MovieSceneMaterialParameterCollectionSystem.cpp b/Engine/Source/Runtime/MovieSceneTracks/Private/Systems/MovieSceneMaterialParameterCollectionSystem.cpp
--- a/Engine/Source/Runtime/MovieSceneTracks/Private/Systems/MovieSceneMaterialParameterCollectionSystem.cpp
+++ b/Engine/Source/Runtime/MovieSceneTracks/Private/Systems/MovieSceneMaterialParameterCollectionSystem.cpp
@@ -75,7 +75,7 @@ void UMovieSceneMaterialParameterCollectionSystem::OnRun(FSystemTaskPrerequisite
 		virtual void InitializeAllocation(FEntityAllocation* Allocation, const FComponentMask& AllocationType) const
 		{
 			TComponentWriter<FObjectComponent> OutBoundMaterials = Allocation->WriteComponents(TracksComponents->BoundMaterial, FEntityAllocationWriteContext::NewAllocation());
-			TComponentReader<UMaterialParameterCollection*> MPCs = Allocation->ReadComponents(TracksComponents->MPC);
+			TComponentReader<TWeakObjectPtr<UMaterialParameterCollection>> MPCs = Allocation->ReadComponents(TracksComponents->MPC);
 			TComponentReader<FInstanceHandle> InstanceHandles = Allocation->ReadComponents(BuiltInComponents->InstanceHandle);
 
 			TOptionalComponentReader<FName> ScalarParameterNames = Allocation->TryReadComponents(TracksComponents->ScalarParameterName);
@@ -87,7 +87,7 @@ void UMovieSceneMaterialParameterCollectionSystem::OnRun(FSystemTaskPrerequisite
 			{
 				OutBoundMaterials[Index] = FObjectComponent::Null();
 
-				UMaterialParameterCollection* Collection = MPCs[Index];
+				UMaterialParameterCollection* Collection = MPCs[Index].Get();
 				IMovieScenePlayer* Player = InstanceRegistry->GetInstance(InstanceHandles[Index]).GetPlayer();
 				UObject* WorldContextObject = Player->GetPlaybackContext();
 				UWorld* World = GEngine->GetWorldFromContextObject(WorldContextObject, EGetWorldErrorMode::LogAndReturnNull);
diff --git a/Engine/Source/Runtime/MovieSceneTracks/Private/Systems/ObjectPathChannelEvaluatorSystem.cpp b/Engine/Source/Runtime/MovieSceneTracks/Private/Systems/ObjectPathChannelEvaluatorSystem.cpp
--- a/Engine/Source/Runtime/MovieSceneTracks/Private/Systems/ObjectPathChannelEvaluatorSystem.cpp
+++ b/Engine/Source/Runtime/MovieSceneTracks/Private/Systems/ObjectPathChannelEvaluatorSystem.cpp
@@ -15,11 +15,17 @@ namespace UE::MovieScene
 
 	struct FEvaluateObjectPathChannels
 	{
-		static void ForEachEntity(FSourceObjectPathChannel ObjectPathChannel, FFrameTime FrameTime, UObject*& OutResult)
+		static void ForEachEntity(FSourceObjectPathChannel ObjectPathChannel, FFrameTime FrameTime, FObjectComponent& OutResult)
 		{
-			if (!ObjectPathChannel.Source->Evaluate(FrameTime, OutResult))
+			UObject* ObjectResult = nullptr;
+			if (!ObjectPathChannel.Source->Evaluate(FrameTime, ObjectResult))
 			{
-				OutResult = nullptr;
+				OutResult = FObjectComponent::Null();
+			}
+			else
+			{
+				// Strong reference to keep the object alive - it is probably an asset of some sort that we wouldn't want to keep reloading
+				OutResult = FObjectComponent::Strong(ObjectResult);
 			}
 		}
 	};
diff --git a/Engine/Source/Runtime/MovieSceneTracks/Public/MovieSceneTracksComponentTypes.h b/Engine/Source/Runtime/MovieSceneTracks/Public/MovieSceneTracksComponentTypes.h
--- a/Engine/Source/Runtime/MovieSceneTracks/Public/MovieSceneTracksComponentTypes.h
+++ b/Engine/Source/Runtime/MovieSceneTracks/Public/MovieSceneTracksComponentTypes.h
@@ -548,7 +548,7 @@ struct MOVIESCENETRACKS_API FMovieSceneTracksComponentTypes
 
 	TComponentTypeID<FConstraintComponentData> ConstraintChannel;
 
-	TComponentTypeID<USceneComponent*> AttachParent;
+	TComponentTypeID<TWeakObjectPtr<USceneComponent>> AttachParent;
 	TComponentTypeID<FAttachmentComponent> AttachComponent;
 	TComponentTypeID<FMovieSceneObjectBindingID> AttachParentBinding;
 	TComponentTypeID<FPerlinNoiseParams> FloatPerlinNoiseChannel;
@@ -566,7 +566,7 @@ struct MOVIESCENETRACKS_API FMovieSceneTracksComponentTypes
 	TComponentTypeID<FName> TransformParameterName;
 
 	TComponentTypeID<FObjectComponent> BoundMaterial;
-	TComponentTypeID<UMaterialParameterCollection*> MPC;
+	TComponentTypeID<TWeakObjectPtr<UMaterialParameterCollection>> MPC;
 
 	TComponentTypeID<FFadeComponentData> Fade;
 
diff --git a/Engine/Source/Runtime/MovieSceneTracks/Public/Systems/MovieSceneComponentAttachmentSystem.h b/Engine/Source/Runtime/MovieSceneTracks/Public/Systems/MovieSceneComponentAttachmentSystem.h
--- a/Engine/Source/Runtime/MovieSceneTracks/Public/Systems/MovieSceneComponentAttachmentSystem.h
+++ b/Engine/Source/Runtime/MovieSceneTracks/Public/Systems/MovieSceneComponentAttachmentSystem.h
@@ -68,7 +68,7 @@ private:
 	virtual void SavePreAnimatedState(const FPreAnimationParameters& InParameters) override;
 	virtual void RestorePreAnimatedState(const FPreAnimationParameters& InParameters) override;
 
-	UE::MovieScene::TOverlappingEntityTracker<UE::MovieScene::FPreAnimAttachment, UObject*> AttachmentTracker;
+	UE::MovieScene::TOverlappingEntityTracker<UE::MovieScene::FPreAnimAttachment, FObjectKey> AttachmentTracker;
 
 	TArray<TTuple<USceneComponent*, UE::MovieScene::FPreAnimAttachment>> PendingAttachmentsToRestore;
 };
diff --git a/Engine/Source/Runtime/MovieSceneTracks/Public/Systems/MovieSceneComponentMobilitySystem.h b/Engine/Source/Runtime/MovieSceneTracks/Public/Systems/MovieSceneComponentMobilitySystem.h
--- a/Engine/Source/Runtime/MovieSceneTracks/Public/Systems/MovieSceneComponentMobilitySystem.h
+++ b/Engine/Source/Runtime/MovieSceneTracks/Public/Systems/MovieSceneComponentMobilitySystem.h
@@ -36,7 +36,7 @@ private:
 
 private:
 
-	UE::MovieScene::TOverlappingEntityTracker<EComponentMobility::Type, UObject*> MobilityTracker;
+	UE::MovieScene::TOverlappingEntityTracker<EComponentMobility::Type, FObjectKey> MobilityTracker;
 
 	UE::MovieScene::FEntityComponentFilter Filter;
 
diff --git a/Engine/Source/Runtime/MovieSceneTracks/Public/Systems/MovieSceneMaterialSystem.h b/Engine/Source/Runtime/MovieSceneTracks/Public/Systems/MovieSceneMaterialSystem.h
--- a/Engine/Source/Runtime/MovieSceneTracks/Public/Systems/MovieSceneMaterialSystem.h
+++ b/Engine/Source/Runtime/MovieSceneTracks/Public/Systems/MovieSceneMaterialSystem.h
@@ -113,10 +113,10 @@ protected:
 template<typename AccessorType, typename... RequiredComponents>
 struct TApplyMaterialSwitchers
 {
-	static void ForEachEntity(typename TCallTraits<RequiredComponents>::ParamType... Inputs, UObject* ObjectResult)
+	static void ForEachEntity(typename TCallTraits<RequiredComponents>::ParamType... Inputs, const FObjectComponent& ObjectResult)
 	{
 		// ObjectResult must be a material
-		UMaterialInterface* NewMaterial = Cast<UMaterialInterface>(ObjectResult);
+		UMaterialInterface* NewMaterial = Cast<UMaterialInterface>(ObjectResult.GetObject());
 
 		AccessorType Accessor(Inputs...);
 		if (!Accessor)
