Sequencer: Ensure that system serial numbers are incremented before possible changes to the entity-component data handling garbage

diff --git a/Engine/Source/Runtime/MovieScene/Private/EntitySystem/MovieSceneEntitySystemLinker.cpp b/Engine/Source/Runtime/MovieScene/Private/EntitySystem/MovieSceneEntitySystemLinker.cpp
--- a/Engine/Source/Runtime/MovieScene/Private/EntitySystem/MovieSceneEntitySystemLinker.cpp
+++ b/Engine/Source/Runtime/MovieScene/Private/EntitySystem/MovieSceneEntitySystemLinker.cpp
@@ -367,6 +367,9 @@ void UMovieSceneEntitySystemLinker::HandlePostGarbageCollection()
 {
 	using namespace UE::MovieScene;
 
+	// Increment the system serial number to ensure that any structural mutation that occurs in this function does so under a unique serial
+	EntityManager.IncrementSystemSerial();
+
 	// All the instance registry to unlink garbage first
 	InstanceRegistry->TagGarbage();
 
