Sequencer: Ensure that system serials are incremented correctly when running scheduled tasks

diff --git a/Engine/Source/Runtime/MovieScene/Private/EntitySystem/MovieSceneTaskScheduler.cpp b/Engine/Source/Runtime/MovieScene/Private/EntitySystem/MovieSceneTaskScheduler.cpp
--- a/Engine/Source/Runtime/MovieScene/Private/EntitySystem/MovieSceneTaskScheduler.cpp
+++ b/Engine/Source/Runtime/MovieScene/Private/EntitySystem/MovieSceneTaskScheduler.cpp
@@ -469,12 +469,11 @@ void FEntitySystemScheduler::ExecuteTasks()
 		return;
 	}
 
-	WriteContextBase = FEntityAllocationWriteContext(*EntityManager);
-
 	const int32 PreviousNumRemaining = NumTasksRemaining.exchange(Tasks.Num());
 	check(PreviousNumRemaining == 0);
 
 	ThreadingModel = EntityManager->GetThreadingModel();
+	WriteContextBase = FEntityAllocationWriteContext(*EntityManager);
 
 	// Condition 1: No threading
 	//              Initiate all tasks immediately. Their subsequents will be triggered inline
@@ -487,6 +486,7 @@ void FEntitySystemScheduler::ExecuteTasks()
 		}
 
 		check(NumTasksRemaining.load() == 0);
+		EntityManager->IncrementSystemSerial(SystemSerialIncrement);
 		return;
 	}
 
@@ -539,6 +539,7 @@ void FEntitySystemScheduler::ExecuteTasks()
 	}
 
 	check(NumTasksRemaining.load() == 0);
+	EntityManager->IncrementSystemSerial(SystemSerialIncrement);
 }
 
 void FEntitySystemScheduler::CompleteTask(const FScheduledTask* Task, FTaskExecutionFlags InFlags) const
@@ -640,6 +641,7 @@ bool FEntitySystemScheduler::HasAnyTasksToPropagateDownstream() const
 void FEntitySystemScheduler::BeginConstruction()
 {
 	WriteContextBase = FEntityAllocationWriteContext(*EntityManager);
+	SystemSerialIncrement = EntityManager->GetSystemSerial();
 
 	Tasks.Reset();
 	PreLockedComponentData.Reset();
@@ -666,6 +668,18 @@ void FEntitySystemScheduler::EndSystem(uint16 NodeID)
 
 void FEntitySystemScheduler::EndConstruction()
 {
+	// SystemSerialIncrement is currently the system serial number from BeginConstruction
+	// Make this a diff by subtracting the current system serial so we can increment by it each time we run our tasks
+	const uint64 CurrentSerial = EntityManager->GetSystemSerial();
+	if (SystemSerialIncrement <= CurrentSerial)
+	{
+		SystemSerialIncrement = CurrentSerial - SystemSerialIncrement;
+	}
+	else
+	{
+		SystemSerialIncrement = 0;
+	}
+
 	Tasks.Shrink();
 	PreLockedComponentData.Shrink();
 
diff --git a/Engine/Source/Runtime/MovieScene/Private/EntitySystem/MovieSceneTaskScheduler.h b/Engine/Source/Runtime/MovieScene/Private/EntitySystem/MovieSceneTaskScheduler.h
--- a/Engine/Source/Runtime/MovieScene/Private/EntitySystem/MovieSceneTaskScheduler.h
+++ b/Engine/Source/Runtime/MovieScene/Private/EntitySystem/MovieSceneTaskScheduler.h
@@ -355,6 +355,7 @@ private:
 	FEvent* GameThreadSignal = nullptr;
 	mutable TLockFreePointerListFIFO<FScheduledTask, PLATFORM_CACHE_LINE_SIZE> GameThreadTaskList;
 	FEntityAllocationWriteContext WriteContextBase = FEntityAllocationWriteContext::NewAllocation();
+	uint32 SystemSerialIncrement = 0;
 	EEntityThreadingModel ThreadingModel = EEntityThreadingModel::NoThreading;
 };
 
diff --git a/Engine/Source/Runtime/MovieScene/Public/EntitySystem/MovieSceneEntityManager.h b/Engine/Source/Runtime/MovieScene/Public/EntitySystem/MovieSceneEntityManager.h
--- a/Engine/Source/Runtime/MovieScene/Public/EntitySystem/MovieSceneEntityManager.h
+++ b/Engine/Source/Runtime/MovieScene/Public/EntitySystem/MovieSceneEntityManager.h
@@ -819,9 +819,9 @@ public:
 	/**
 	 * Increment the current serial number for systems observing this manager. Should be called after any system is run
 	 */
-	void IncrementSystemSerial()
+	void IncrementSystemSerial(uint64 IncAmount = 1)
 	{
-		++SystemSerialNumber;
+		SystemSerialNumber += IncAmount;
 	}
 
 	/**
