Sequencer: Fix a bug where when overriding a spawnable binding during the sequencer, the spawned actor will get left behind, creating duplicate characters. Now when a binding is invalidated, and we detect an overridden binding, we also check to see if a spawnable is registered, and if so, destroy that spawnable.

diff --git a/Engine/Source/Runtime/MovieScene/Private/EntitySystem/MovieSceneSpawnablesSystem.cpp b/Engine/Source/Runtime/MovieScene/Private/EntitySystem/MovieSceneSpawnablesSystem.cpp
--- a/Engine/Source/Runtime/MovieScene/Private/EntitySystem/MovieSceneSpawnablesSystem.cpp
+++ b/Engine/Source/Runtime/MovieScene/Private/EntitySystem/MovieSceneSpawnablesSystem.cpp
@@ -118,10 +118,36 @@ void UMovieSceneSpawnablesSystem::OnRun(FSystemTaskPrerequisites& InPrerequisite
 		}
 	}
 
+	// Used below.
+	TArray<TTuple<FGuid, FMovieSceneSequenceID, FInstanceHandle>> DestroyedObjects;
+	auto DestroyOldSpawnables = [&DestroyedObjects, InstanceRegistry](FInstanceHandle InstanceHandle, const FGuid& SpawnableObjectID)
+	{
+		SCOPE_CYCLE_COUNTER(MovieSceneEval_DestroySpawnables)
+
+		if (ensure(InstanceRegistry->IsHandleValid(InstanceHandle)))
+		{
+			const FSequenceInstance& Instance = InstanceRegistry->GetInstance(InstanceHandle);
+			IMovieScenePlayer* Player = Instance.GetPlayer();
+
+			// If the sequence instance has finished and it is a sub sequence, we do not destroy the spawnable
+			// if it is owned by the root sequence or externally. These will get destroyed or forgotten by the player when it ends
+			if (Instance.HasFinished() && Instance.IsSubSequence())
+			{
+				const UMovieSceneSequence* Sequence = Player->State.FindSequence(Instance.GetSequenceID());
+				FMovieSceneSpawnable* Spawnable = Sequence ? Sequence->GetMovieScene()->FindSpawnable(SpawnableObjectID) : nullptr;
+				if (!Spawnable || Spawnable->GetSpawnOwnership() != ESpawnOwnership::InnerSequence)
+				{
+					return;
+				}
+			}
+
+			DestroyedObjects.Emplace(SpawnableObjectID, Instance.GetSequenceID(), InstanceHandle);
+		}
+	};
 
 	// ----------------------------------------------------------------------------------------------------------------------------------------
 	// Step 1 - iterate all pending spawnables and spawn their objects if necessary
-	auto SpawnNewObjects = [InstanceRegistry](FInstanceHandle InstanceHandle, const FGuid& SpawnableBindingID)
+	auto SpawnNewObjects = [&DestroyOldSpawnables, InstanceRegistry](FInstanceHandle InstanceHandle, const FGuid& SpawnableBindingID)
 	{
 		SCOPE_CYCLE_COUNTER(MovieSceneEval_SpawnSpawnables)
 
@@ -138,11 +164,7 @@ void UMovieSceneSpawnablesSystem::OnRun(FSystemTaskPrerequisites& InPrerequisite
 			return;
 		}
 
-		// Check if we already have a spawned object in the sapwn register - if we have we use that
-		if (UObject* ExistingSpawnedObject = Player->GetSpawnRegister().FindSpawnedObject(SpawnableBindingID, SequenceID).Get())
-		{
-			return;
-		}
+		UObject* ExistingSpawnedObject = Player->GetSpawnRegister().FindSpawnedObject(SpawnableBindingID, SequenceID).Get();
 
 		// Check whether the binding is overridden - if it is we cannot spawn a new object
 		if (const IMovieScenePlaybackClient* PlaybackClient = Player->GetPlaybackClient())
@@ -151,11 +173,23 @@ void UMovieSceneSpawnablesSystem::OnRun(FSystemTaskPrerequisites& InPrerequisite
 			bool bUseDefaultBinding = PlaybackClient->RetrieveBindingOverrides(SpawnableBindingID, SequenceID, FoundObjects);
 			if (!bUseDefaultBinding)
 			{
+				// If the binding has been overridden but we have an existing spawned object, then the binding is new and we need to destroy the spawned object.
+				if (ExistingSpawnedObject)
+				{
+					DestroyOldSpawnables(InstanceHandle, SpawnableBindingID);
+				}
+
 				// This spawnable is overridden so don't try and spawn anything
 				return;
 			}
 		}
 
+		// Check if we already have a spawned object in the spawn register - if we have we use that
+		if (ExistingSpawnedObject)
+		{
+			return;
+		}
+
 		// At this point we've decided that we need to spawn a whole new object
 		const UMovieSceneSequence* Sequence = Player->State.FindSequence(SequenceID);
 		if (!Sequence)
@@ -184,31 +218,6 @@ void UMovieSceneSpawnablesSystem::OnRun(FSystemTaskPrerequisites& InPrerequisite
 	// Step 2 - destroy any spawnable objects that are no longer relevant
 	//          NOTE: We gather the objects into an array because destroying an object can potentially cause a garbage collection to run (ie
 	//                if the spawnable is a level instance), and that could assert because we are currently iterating the ECS
-	TArray<TTuple<FGuid, FMovieSceneSequenceID, FInstanceHandle>> DestroyedObjects;
-	auto DestroyOldSpawnables = [&DestroyedObjects, InstanceRegistry](FInstanceHandle InstanceHandle, const FGuid& SpawnableObjectID)
-	{
-		SCOPE_CYCLE_COUNTER(MovieSceneEval_DestroySpawnables)
-
-		if (ensure(InstanceRegistry->IsHandleValid(InstanceHandle)))
-		{
-			const FSequenceInstance& Instance = InstanceRegistry->GetInstance(InstanceHandle);
-			IMovieScenePlayer* Player = Instance.GetPlayer();
-
-			// If the sequence instance has finished and it is a sub sequence, we do not destroy the spawnable
-			// if it is owned by the root sequence or externally. These will get destroyed or forgotten by the player when it ends
-			if (Instance.HasFinished() && Instance.IsSubSequence())
-			{
-				const UMovieSceneSequence* Sequence = Player->State.FindSequence(Instance.GetSequenceID());
-				FMovieSceneSpawnable* Spawnable = Sequence ? Sequence->GetMovieScene()->FindSpawnable(SpawnableObjectID) : nullptr;
-				if (!Spawnable || Spawnable->GetSpawnOwnership() != ESpawnOwnership::InnerSequence)
-				{
-					return;
-				}
-			}
-
-			DestroyedObjects.Emplace(SpawnableObjectID, Instance.GetSequenceID(), InstanceHandle);
-		}
-	};
 
 	FEntityTaskBuilder()
 	.Read(BuiltInComponents->InstanceHandle)
