Fixing potential unchecked pointer crashes.

diff --git a/Engine/Source/Runtime/MovieScene/Private/Evaluation/MovieSceneEvaluationTemplateInstance.cpp b/Engine/Source/Runtime/MovieScene/Private/Evaluation/MovieSceneEvaluationTemplateInstance.cpp
--- a/Engine/Source/Runtime/MovieScene/Private/Evaluation/MovieSceneEvaluationTemplateInstance.cpp
+++ b/Engine/Source/Runtime/MovieScene/Private/Evaluation/MovieSceneEvaluationTemplateInstance.cpp
@@ -142,9 +142,12 @@ void FMovieSceneRootEvaluationTemplateInstance::Initialize(UMovieSceneSequence&
 
 		Player.State.PersistentEntityData.Reset();
 		Player.State.PersistentSharedData.Reset();
-
+		
 		EntitySystemLinker = ConstructEntityLinker(Player);
-		RootInstanceHandle = EntitySystemLinker->GetInstanceRegistry()->AllocateRootInstance(&Player);
+		if (EntitySystemLinker != nullptr && EntitySystemLinker->GetInstanceRegistry())
+		{
+			RootInstanceHandle = EntitySystemLinker->GetInstanceRegistry()->AllocateRootInstance(&Player);
+		}
 
 		TSharedPtr<FMovieSceneEntitySystemRunner> Runner = WeakRunner.Pin();
 		if (ensure(Runner) && EntitySystemLinker)
