Sequencer: don't unregister global custom accessors on shutdown

diff --git a/Engine/Source/Runtime/MovieScene/Private/EntitySystem/MovieScenePropertySystemTypes.cpp b/Engine/Source/Runtime/MovieScene/Private/EntitySystem/MovieScenePropertySystemTypes.cpp
--- a/Engine/Source/Runtime/MovieScene/Private/EntitySystem/MovieScenePropertySystemTypes.cpp
+++ b/Engine/Source/Runtime/MovieScene/Private/EntitySystem/MovieScenePropertySystemTypes.cpp
@@ -19,7 +19,11 @@ void AddGlobalCustomAccessor(const UClass* ClassType, FName PropertyPath)
 }
 void RemoveGlobalCustomAccessor(const UClass* ClassType, FName PropertyPath)
 {
-	GRegisteredCustomAccessors.Remove(TTuple<FObjectKey, FName>(ClassType, PropertyPath));
+	if (UObjectInitialized())
+	{
+		GRegisteredCustomAccessors.Remove(TTuple<FObjectKey, FName>(ClassType, PropertyPath));
+	}
+	// else: during engine shutdown, we can't make object keys anymore
 }
 bool GlobalCustomAccessorExists(const UClass* ClassType, TStringView<WIDECHAR> PropertyPath)
 {
