[Engine]

diff --git a/Engine/Source/Runtime/MovieScene/Public/Channels/MovieSceneIntegerChannel.h b/Engine/Source/Runtime/MovieScene/Public/Channels/MovieSceneIntegerChannel.h
--- a/Engine/Source/Runtime/MovieScene/Public/Channels/MovieSceneIntegerChannel.h
+++ b/Engine/Source/Runtime/MovieScene/Public/Channels/MovieSceneIntegerChannel.h
@@ -28,7 +28,7 @@ struct FPropertyTag;
 
 
 USTRUCT()
-struct MOVIESCENE_API FMovieSceneIntegerChannel : public FMovieSceneChannel
+struct FMovieSceneIntegerChannel : public FMovieSceneChannel
 {
 	GENERATED_BODY()
 
@@ -39,7 +39,7 @@ struct MOVIESCENE_API FMovieSceneIntegerChannel : public FMovieSceneChannel
 	/**
 	 * Serialize this type from another
 	 */
-	bool SerializeFromMismatchedTag(const FPropertyTag& Tag, FStructuredArchive::FSlot Slot);
+	MOVIESCENE_API bool SerializeFromMismatchedTag(const FPropertyTag& Tag, FStructuredArchive::FSlot Slot);
 
 	/**
 	 * Access a mutable interface for this channel's data
@@ -92,7 +92,7 @@ struct MOVIESCENE_API FMovieSceneIntegerChannel : public FMovieSceneChannel
 	 * @param OutValue   A value to receive the result
 	 * @return true if the channel was evaluated successfully, false otherwise
 	 */
-	bool Evaluate(FFrameTime InTime, int32& OutValue) const;
+	MOVIESCENE_API bool Evaluate(FFrameTime InTime, int32& OutValue) const;
 
 	/**
 	 * Set the channel's times and values to the requested values
@@ -131,19 +131,19 @@ struct MOVIESCENE_API FMovieSceneIntegerChannel : public FMovieSceneChannel
 public:
 
 	// ~ FMovieSceneChannel Interface
-	virtual void GetKeys(const TRange<FFrameNumber>& WithinRange, TArray<FFrameNumber>* OutKeyTimes, TArray<FKeyHandle>* OutKeyHandles) override;
-	virtual void GetKeyTimes(TArrayView<const FKeyHandle> InHandles, TArrayView<FFrameNumber> OutKeyTimes) override;
-	virtual void SetKeyTimes(TArrayView<const FKeyHandle> InHandles, TArrayView<const FFrameNumber> InKeyTimes) override;
-	virtual void DuplicateKeys(TArrayView<const FKeyHandle> InHandles, TArrayView<FKeyHandle> OutNewHandles) override;
-	virtual void DeleteKeys(TArrayView<const FKeyHandle> InHandles) override;
-	virtual void DeleteKeysFrom(FFrameNumber InTime, bool bDeleteKeysBefore) override;
-	virtual void ChangeFrameResolution(FFrameRate SourceRate, FFrameRate DestinationRate) override;
-	virtual TRange<FFrameNumber> ComputeEffectiveRange() const override;
-	virtual int32 GetNumKeys() const override;
-	virtual void Reset() override;
-	virtual void Offset(FFrameNumber DeltaPosition) override;
-	virtual void Optimize(const FKeyDataOptimizationParams& InParameters) override;
-	virtual void ClearDefault() override;
+	MOVIESCENE_API virtual void GetKeys(const TRange<FFrameNumber>& WithinRange, TArray<FFrameNumber>* OutKeyTimes, TArray<FKeyHandle>* OutKeyHandles) override;
+	MOVIESCENE_API virtual void GetKeyTimes(TArrayView<const FKeyHandle> InHandles, TArrayView<FFrameNumber> OutKeyTimes) override;
+	MOVIESCENE_API virtual void SetKeyTimes(TArrayView<const FKeyHandle> InHandles, TArrayView<const FFrameNumber> InKeyTimes) override;
+	MOVIESCENE_API virtual void DuplicateKeys(TArrayView<const FKeyHandle> InHandles, TArrayView<FKeyHandle> OutNewHandles) override;
+	MOVIESCENE_API virtual void DeleteKeys(TArrayView<const FKeyHandle> InHandles) override;
+	MOVIESCENE_API virtual void DeleteKeysFrom(FFrameNumber InTime, bool bDeleteKeysBefore) override;
+	MOVIESCENE_API virtual void ChangeFrameResolution(FFrameRate SourceRate, FFrameRate DestinationRate) override;
+	MOVIESCENE_API virtual TRange<FFrameNumber> ComputeEffectiveRange() const override;
+	MOVIESCENE_API virtual int32 GetNumKeys() const override;
+	MOVIESCENE_API virtual void Reset() override;
+	MOVIESCENE_API virtual void Offset(FFrameNumber DeltaPosition) override;
+	MOVIESCENE_API virtual void Optimize(const FKeyDataOptimizationParams& InParameters) override;
+	MOVIESCENE_API virtual void ClearDefault() override;
 
 public:
 
diff --git a/Engine/Source/Runtime/MovieScene/Public/Channels/MovieSceneInterpolation.h b/Engine/Source/Runtime/MovieScene/Public/Channels/MovieSceneInterpolation.h
--- a/Engine/Source/Runtime/MovieScene/Public/Channels/MovieSceneInterpolation.h
+++ b/Engine/Source/Runtime/MovieScene/Public/Channels/MovieSceneInterpolation.h
@@ -73,7 +73,7 @@ struct FLinearInterpolation
 /**
  * A cubic bezier interpolation between 2 control points with tangents, represented as 4 control points on a Bezier curve
  */
-struct MOVIESCENE_API FCubicInterpolation
+struct FCubicInterpolation
 {
 	/** The delta value between the two control points in the time-domain */
 	double DX;
@@ -82,7 +82,7 @@ struct MOVIESCENE_API FCubicInterpolation
 	/** The origin time of the first control point */
 	FFrameNumber Origin;
 
-	FCubicInterpolation(
+	MOVIESCENE_API FCubicInterpolation(
 		FFrameNumber InOrigin,
 		double InDX,
 		double InStartValue,
@@ -93,7 +93,7 @@ struct MOVIESCENE_API FCubicInterpolation
 	/**
 	 * Evaluate the expression
 	 */
-	double Evaluate(FFrameTime InTime) const;
+	MOVIESCENE_API double Evaluate(FFrameTime InTime) const;
 };
 
 /**
@@ -136,23 +136,23 @@ struct FWeightedCubicInterpolation
 /**
  * Simple 1 dimensional range based off a FFrameNumber to define the range within which a cached interpolation is valid
  */
-struct MOVIESCENE_API FCachedInterpolationRange
+struct FCachedInterpolationRange
 {
 	/** Make an empty range */
-	static FCachedInterpolationRange Empty();
+	static MOVIESCENE_API FCachedInterpolationRange Empty();
 	/** Make finite range from InStart to InEnd, not including the end frame */
-	static FCachedInterpolationRange Finite(FFrameNumber InStart, FFrameNumber InEnd);
+	static MOVIESCENE_API FCachedInterpolationRange Finite(FFrameNumber InStart, FFrameNumber InEnd);
 	/** Make an infinite range */
-	static FCachedInterpolationRange Infinite();
+	static MOVIESCENE_API FCachedInterpolationRange Infinite();
 
 	/** Make a range that only contains the specified time */
-	static FCachedInterpolationRange Only(FFrameNumber InTime);
+	static MOVIESCENE_API FCachedInterpolationRange Only(FFrameNumber InTime);
 	/** Make a range that covers all times from (and including) the specified start */
-	static FCachedInterpolationRange From(FFrameNumber InStart);
+	static MOVIESCENE_API FCachedInterpolationRange From(FFrameNumber InStart);
 	/** Make a range that covers all times up to (but not including) the specified end */
-	static FCachedInterpolationRange Until(FFrameNumber InEnd);
+	static MOVIESCENE_API FCachedInterpolationRange Until(FFrameNumber InEnd);
 
-	bool Contains(FFrameNumber FrameNumber) const;
+	MOVIESCENE_API bool Contains(FFrameNumber FrameNumber) const;
 
 	/** Inclusive start frame */
 	FFrameNumber Start;
@@ -164,24 +164,24 @@ struct MOVIESCENE_API FCachedInterpolationRange
  * Variant structure that wraps an interpolation and the range within which it is valid.
  * ~96 bytes
  */
-struct MOVIESCENE_API FCachedInterpolation
+struct FCachedInterpolation
 {
 	/** Default construction to an invalid state. Calling Evaluate will always return false */
-	FCachedInterpolation();
+	MOVIESCENE_API FCachedInterpolation();
 	/** Construction as a constant value */
-	FCachedInterpolation(const FCachedInterpolationRange& InRange, const FConstantValue& Constant);
+	MOVIESCENE_API FCachedInterpolation(const FCachedInterpolationRange& InRange, const FConstantValue& Constant);
 	/** Construction as a linear interpolation */
-	FCachedInterpolation(const FCachedInterpolationRange& InRange, const FLinearInterpolation& Linear);
+	MOVIESCENE_API FCachedInterpolation(const FCachedInterpolationRange& InRange, const FLinearInterpolation& Linear);
 	/** Construction as a cubic interpolation */
-	FCachedInterpolation(const FCachedInterpolationRange& InRange, const FCubicInterpolation& Cubic);
+	MOVIESCENE_API FCachedInterpolation(const FCachedInterpolationRange& InRange, const FCubicInterpolation& Cubic);
 	/** Construction as a weighted cubic interpolation */
-	FCachedInterpolation(const FCachedInterpolationRange& InRange, const FWeightedCubicInterpolation& WeightedCubic);
+	MOVIESCENE_API FCachedInterpolation(const FCachedInterpolationRange& InRange, const FWeightedCubicInterpolation& WeightedCubic);
 
 	/**
 	 * Check whether this cache is still valid for the specified frame number.
 	 * @return true if this interpolation is relevant to the frame, or false if it should be re-generated
 	 */
-	bool IsCacheValidForTime(FFrameNumber FrameNumber) const;
+	MOVIESCENE_API bool IsCacheValidForTime(FFrameNumber FrameNumber) const;
 
 	/**
 	 * Evaluate this interpolation for the specified frame time
@@ -190,7 +190,7 @@ struct MOVIESCENE_API FCachedInterpolation
 	 * @param OutResult  Reference to recieve the resulting value. Only written to if this function returns true.
 	 * @return true if this interpolation evaluated successfully (and the result written to OutResult), false otherwise.
 	 */
-	bool Evaluate(FFrameTime FrameTime, double& OutResult) const;
+	MOVIESCENE_API bool Evaluate(FFrameTime FrameTime, double& OutResult) const;
 
 private:
 
@@ -205,4 +205,4 @@ private:
 	FCachedInterpolationRange Range;
 };
 
-} // namespace UE::MovieScene
\ No newline at end of file
+} // namespace UE::MovieScene
diff --git a/Engine/Source/Runtime/MovieScene/Public/Channels/MovieSceneObjectPathChannel.h b/Engine/Source/Runtime/MovieScene/Public/Channels/MovieSceneObjectPathChannel.h
--- a/Engine/Source/Runtime/MovieScene/Public/Channels/MovieSceneObjectPathChannel.h
+++ b/Engine/Source/Runtime/MovieScene/Public/Channels/MovieSceneObjectPathChannel.h
@@ -12,7 +12,7 @@
  * Key value type for object path channels that stores references to objects as both a hard and soft reference, to ensure compatability with both sub objects and async loading
  */
 USTRUCT()
-struct MOVIESCENE_API FMovieSceneObjectPathChannelKeyValue
+struct FMovieSceneObjectPathChannelKeyValue
 {
 	GENERATED_BODY()
 
@@ -22,19 +22,19 @@ struct MOVIESCENE_API FMovieSceneObjectPathChannelKeyValue
 	{}
 
 	/** Construction from an object pointer */
-	FMovieSceneObjectPathChannelKeyValue(UObject* InObject);
+	MOVIESCENE_API FMovieSceneObjectPathChannelKeyValue(UObject* InObject);
 
 	/**
 	 * Assignment from a raw object pointer
 	 */
-	FMovieSceneObjectPathChannelKeyValue& operator=(UObject* NewObject);
+	MOVIESCENE_API FMovieSceneObjectPathChannelKeyValue& operator=(UObject* NewObject);
 
 public:
 
 	/**
 	 * Legacy conversion from a TSoftObjectPtr<>
 	 */
-	bool SerializeFromMismatchedTag(FPropertyTag const& Tag, FStructuredArchive::FSlot Slot);
+	MOVIESCENE_API bool SerializeFromMismatchedTag(FPropertyTag const& Tag, FStructuredArchive::FSlot Slot);
 	
 	/**
 	 * Access the soft object pointer that this key should load
@@ -47,7 +47,7 @@ public:
 	/**
 	 * Attempt to find this object either by returning the internally kept raw pointer, or by resolving (but not loading) the soft object path
 	 */
-	UObject* Get() const;
+	MOVIESCENE_API UObject* Get() const;
 
 private:
 
@@ -67,7 +67,7 @@ struct TStructOpsTypeTraits<FMovieSceneObjectPathChannelKeyValue> : public TStru
 };
 
 USTRUCT()
-struct MOVIESCENE_API FMovieSceneObjectPathChannel : public FMovieSceneChannel
+struct FMovieSceneObjectPathChannel : public FMovieSceneChannel
 {
 	GENERATED_BODY()
 
@@ -113,24 +113,24 @@ struct MOVIESCENE_API FMovieSceneObjectPathChannel : public FMovieSceneChannel
 	 * @param OutValue   A value to receive the result
 	 * @return true if the channel was evaluated successfully, false otherwise
 	 */
-	bool Evaluate(FFrameTime InTime, UObject*& OutValue) const;
+	MOVIESCENE_API bool Evaluate(FFrameTime InTime, UObject*& OutValue) const;
 
 public:
 
 	//~ FMovieSceneChannel Interface
-	virtual void GetKeys(const TRange<FFrameNumber>& WithinRange, TArray<FFrameNumber>* OutKeyTimes, TArray<FKeyHandle>* OutKeyHandles) override;
-	virtual void GetKeyTimes(TArrayView<const FKeyHandle> InHandles, TArrayView<FFrameNumber> OutKeyTimes) override;
-	virtual void SetKeyTimes(TArrayView<const FKeyHandle> InHandles, TArrayView<const FFrameNumber> InKeyTimes) override;
-	virtual void DuplicateKeys(TArrayView<const FKeyHandle> InHandles, TArrayView<FKeyHandle> OutNewHandles) override;
-	virtual void DeleteKeys(TArrayView<const FKeyHandle> InHandles) override;
-	virtual void DeleteKeysFrom(FFrameNumber InTime, bool bDeleteKeysBefore) override;
-	virtual void ChangeFrameResolution(FFrameRate SourceRate, FFrameRate DestinationRate) override;
-	virtual TRange<FFrameNumber> ComputeEffectiveRange() const override;
-	virtual int32 GetNumKeys() const override;
-	virtual void Reset() override;
-	virtual void Offset(FFrameNumber DeltaPosition) override;
-	virtual void Optimize(const FKeyDataOptimizationParams& InParameters) override;
-	virtual void ClearDefault() override;
+	MOVIESCENE_API virtual void GetKeys(const TRange<FFrameNumber>& WithinRange, TArray<FFrameNumber>* OutKeyTimes, TArray<FKeyHandle>* OutKeyHandles) override;
+	MOVIESCENE_API virtual void GetKeyTimes(TArrayView<const FKeyHandle> InHandles, TArrayView<FFrameNumber> OutKeyTimes) override;
+	MOVIESCENE_API virtual void SetKeyTimes(TArrayView<const FKeyHandle> InHandles, TArrayView<const FFrameNumber> InKeyTimes) override;
+	MOVIESCENE_API virtual void DuplicateKeys(TArrayView<const FKeyHandle> InHandles, TArrayView<FKeyHandle> OutNewHandles) override;
+	MOVIESCENE_API virtual void DeleteKeys(TArrayView<const FKeyHandle> InHandles) override;
+	MOVIESCENE_API virtual void DeleteKeysFrom(FFrameNumber InTime, bool bDeleteKeysBefore) override;
+	MOVIESCENE_API virtual void ChangeFrameResolution(FFrameRate SourceRate, FFrameRate DestinationRate) override;
+	MOVIESCENE_API virtual TRange<FFrameNumber> ComputeEffectiveRange() const override;
+	MOVIESCENE_API virtual int32 GetNumKeys() const override;
+	MOVIESCENE_API virtual void Reset() override;
+	MOVIESCENE_API virtual void Offset(FFrameNumber DeltaPosition) override;
+	MOVIESCENE_API virtual void Optimize(const FKeyDataOptimizationParams& InParameters) override;
+	MOVIESCENE_API virtual void ClearDefault() override;
 
 public:
 
@@ -190,4 +190,4 @@ struct TMovieSceneChannelTraits<FMovieSceneObjectPathChannel> : TMovieSceneChann
 	typedef TMovieSceneExternalValue<UObject*> ExtendedEditorDataType;
 
 #endif
-};
\ No newline at end of file
+};
diff --git a/Engine/Source/Runtime/MovieScene/Public/Channels/MovieSceneSectionChannelOverrideRegistry.h b/Engine/Source/Runtime/MovieScene/Public/Channels/MovieSceneSectionChannelOverrideRegistry.h
--- a/Engine/Source/Runtime/MovieScene/Public/Channels/MovieSceneSectionChannelOverrideRegistry.h
+++ b/Engine/Source/Runtime/MovieScene/Public/Channels/MovieSceneSectionChannelOverrideRegistry.h
@@ -39,40 +39,40 @@ namespace MovieScene
 /**
 * This object contains a map of actual channel overrides, where each override is a channel identifier and a channel container.
 */
-UCLASS()
-class MOVIESCENE_API UMovieSceneSectionChannelOverrideRegistry : public UObject
+UCLASS(MinimalAPI)
+class UMovieSceneSectionChannelOverrideRegistry : public UObject
 {
 	GENERATED_BODY()
 
 public:
-	UMovieSceneSectionChannelOverrideRegistry();
+	MOVIESCENE_API UMovieSceneSectionChannelOverrideRegistry();
 
 	/** 
 	* Add channel to the registry 
 	* @param ChannelName		The name of the channel to override
 	* @param Data				The container that owns a overriden channel instanse
 	*/
-	void AddChannel(FName ChannelName, UMovieSceneChannelOverrideContainer* ChannelContainer);
+	MOVIESCENE_API void AddChannel(FName ChannelName, UMovieSceneChannelOverrideContainer* ChannelContainer);
 
 	/**
 	* Returns if the channel is overriden
 	* @param ChannelName	Name of the channel
 	* @return	Whether this channel is overriden
 	*/
-	bool ContainsChannel(FName ChannelName) const;
+	MOVIESCENE_API bool ContainsChannel(FName ChannelName) const;
 
 	/**
 	 * Returns the number of override channels in this registry
 	 * @return	The number of override channels
 	 */
-	int32 NumChannels() const;
+	MOVIESCENE_API int32 NumChannels() const;
 
 	/**
 	 * Returns the channel override for a given name, or nullptr if not found
 	 * @param ChannelName	Name of the channel
 	 * @return	The channel override container
 	 */
-	UMovieSceneChannelOverrideContainer* GetChannel(FName ChannelName) const;
+	MOVIESCENE_API UMovieSceneChannelOverrideContainer* GetChannel(FName ChannelName) const;
 
 	/**
 	 * Get all channel containers of a given type.
@@ -94,29 +94,29 @@ public:
 	* Removes a channel from the registry
 	* @param ChannelName		The name of the channel whose override to remove
 	*/
-	void RemoveChannel(FName ChannelName);
+	MOVIESCENE_API void RemoveChannel(FName ChannelName);
 
 	/**
 	* Forward ImportEntityImpl calls to an overriden channel
 	*/
-	void ImportEntityImpl(const UE::MovieScene::FChannelOverrideEntityImportParams& OverrideParams, const UE::MovieScene::FEntityImportParams& ImportParams, UE::MovieScene::FImportedEntity* OutImportedEntity);
+	MOVIESCENE_API void ImportEntityImpl(const UE::MovieScene::FChannelOverrideEntityImportParams& OverrideParams, const UE::MovieScene::FEntityImportParams& ImportParams, UE::MovieScene::FImportedEntity* OutImportedEntity);
 
 	/**
 	* Called when overridden channels should populate evaluation field
 	*/
-	void PopulateEvaluationFieldImpl(const TRange<FFrameNumber>& EffectiveRange, const FMovieSceneEvaluationFieldEntityMetaData& InMetaData, FMovieSceneEntityComponentFieldBuilder* OutFieldBuilder, UMovieSceneSection& OwnerSection);
+	MOVIESCENE_API void PopulateEvaluationFieldImpl(const TRange<FFrameNumber>& EffectiveRange, const FMovieSceneEvaluationFieldEntityMetaData& InMetaData, FMovieSceneEntityComponentFieldBuilder* OutFieldBuilder, UMovieSceneSection& OwnerSection);
 
 #if WITH_EDITOR
 
 	/**
 	 * Called by the owning section after it has been created for a paste operation.
 	 */
-	void OnPostPaste();
+	MOVIESCENE_API void OnPostPaste();
 
 	/**
 	 * Called when an undo/redo operation has affected this override registry.
 	 */
-	virtual void PostEditUndo() override;
+	MOVIESCENE_API virtual void PostEditUndo() override;
 
 #endif
 
diff --git a/Engine/Source/Runtime/MovieScene/Public/Compilation/IMovieSceneDeterminismSource.h b/Engine/Source/Runtime/MovieScene/Public/Compilation/IMovieSceneDeterminismSource.h
--- a/Engine/Source/Runtime/MovieScene/Public/Compilation/IMovieSceneDeterminismSource.h
+++ b/Engine/Source/Runtime/MovieScene/Public/Compilation/IMovieSceneDeterminismSource.h
@@ -30,8 +30,8 @@ struct FMovieSceneDeterminismData
 };
 
 
-UINTERFACE()
-class MOVIESCENE_API UMovieSceneDeterminismSource : public UInterface
+UINTERFACE(MinimalAPI)
+class UMovieSceneDeterminismSource : public UInterface
 {
 public:
 	GENERATED_BODY()
diff --git a/Engine/Source/Runtime/MovieScene/Public/Compilation/IMovieSceneTrackTemplateProducer.h b/Engine/Source/Runtime/MovieScene/Public/Compilation/IMovieSceneTrackTemplateProducer.h
--- a/Engine/Source/Runtime/MovieScene/Public/Compilation/IMovieSceneTrackTemplateProducer.h
+++ b/Engine/Source/Runtime/MovieScene/Public/Compilation/IMovieSceneTrackTemplateProducer.h
@@ -55,14 +55,14 @@ enum class EMovieSceneCompileResult : uint8
 	Unimplemented
 };
 
-UINTERFACE()
-class MOVIESCENE_API UMovieSceneTrackTemplateProducer : public UInterface
+UINTERFACE(MinimalAPI)
+class UMovieSceneTrackTemplateProducer : public UInterface
 {
 public:
 	GENERATED_BODY()
 };
 
-class MOVIESCENE_API IMovieSceneTrackTemplateProducer
+class IMovieSceneTrackTemplateProducer
 {
 public:
 	GENERATED_BODY()
@@ -77,12 +77,12 @@ public:
 	 *
 	 * @param Args 			Compilation arguments
 	 */
-	virtual void GenerateTemplate(const FMovieSceneTrackCompilerArgs& Args) const;
+	MOVIESCENE_API virtual void GenerateTemplate(const FMovieSceneTrackCompilerArgs& Args) const;
 
 	/**
 	 * Get a raw compiled copy of this track with no additional shared tracks or compiler parameters
 	 */
-	FMovieSceneEvaluationTrack GenerateTrackTemplate(UMovieSceneTrack* SourceTrack) const;
+	MOVIESCENE_API FMovieSceneEvaluationTrack GenerateTrackTemplate(UMovieSceneTrack* SourceTrack) const;
 
 
 protected:
@@ -121,6 +121,6 @@ protected:
 	 * @param Args 			Compilation arguments
 	 * @return Compilation result
 	 */
-	EMovieSceneCompileResult Compile(FMovieSceneEvaluationTrack& Track, const FMovieSceneTrackCompilerArgs& Args) const;
+	MOVIESCENE_API EMovieSceneCompileResult Compile(FMovieSceneEvaluationTrack& Track, const FMovieSceneTrackCompilerArgs& Args) const;
 
-};
\ No newline at end of file
+};
diff --git a/Engine/Source/Runtime/MovieScene/Public/Compilation/MovieSceneCompiledDataManager.h b/Engine/Source/Runtime/MovieScene/Public/Compilation/MovieSceneCompiledDataManager.h
--- a/Engine/Source/Runtime/MovieScene/Public/Compilation/MovieSceneCompiledDataManager.h
+++ b/Engine/Source/Runtime/MovieScene/Public/Compilation/MovieSceneCompiledDataManager.h
@@ -211,40 +211,40 @@ struct FMovieSceneCompiledDataEntry
 	FMovieSceneCompiledSequenceFlagStruct CompiledFlags;
 };
 
-UCLASS()
-class MOVIESCENE_API UMovieSceneCompiledDataManager
+UCLASS(MinimalAPI)
+class UMovieSceneCompiledDataManager
 	: public UObject
 {
 public:
 	GENERATED_BODY()
 
-	UMovieSceneCompiledDataManager();
+	MOVIESCENE_API UMovieSceneCompiledDataManager();
 
 #if WITH_EDITOR
-	static UMovieSceneCompiledDataManager* GetPrecompiledData(EMovieSceneServerClientMask EmulatedMask = EMovieSceneServerClientMask::All);
+	static MOVIESCENE_API UMovieSceneCompiledDataManager* GetPrecompiledData(EMovieSceneServerClientMask EmulatedMask = EMovieSceneServerClientMask::All);
 #else
-	static UMovieSceneCompiledDataManager* GetPrecompiledData();
+	static MOVIESCENE_API UMovieSceneCompiledDataManager* GetPrecompiledData();
 #endif
 
-	UMovieSceneCompiledData* MakeCompiledData(UMovieSceneSequence* Sequence) const;
+	MOVIESCENE_API UMovieSceneCompiledData* MakeCompiledData(UMovieSceneSequence* Sequence) const;
 
-	void SetEmulatedNetworkMask(EMovieSceneServerClientMask NewMask);
+	MOVIESCENE_API void SetEmulatedNetworkMask(EMovieSceneServerClientMask NewMask);
 
-	void Reset(UMovieSceneSequence* Sequence);
+	MOVIESCENE_API void Reset(UMovieSceneSequence* Sequence);
 
-	void DestroyAllData();
+	MOVIESCENE_API void DestroyAllData();
 
-	FMovieSceneCompiledDataID GetDataID(UMovieSceneSequence* Sequence);
+	MOVIESCENE_API FMovieSceneCompiledDataID GetDataID(UMovieSceneSequence* Sequence);
 
-	FMovieSceneCompiledDataID FindDataID(UMovieSceneSequence* Sequence) const;
+	MOVIESCENE_API FMovieSceneCompiledDataID FindDataID(UMovieSceneSequence* Sequence) const;
 
-	void DestroyTemplate(FMovieSceneCompiledDataID DataID);
+	MOVIESCENE_API void DestroyTemplate(FMovieSceneCompiledDataID DataID);
 
-	bool IsDirty(UMovieSceneSequence* Sequence) const;
+	MOVIESCENE_API bool IsDirty(UMovieSceneSequence* Sequence) const;
 
-	bool IsDirty(FMovieSceneCompiledDataID CompiledDataID) const;
+	MOVIESCENE_API bool IsDirty(FMovieSceneCompiledDataID CompiledDataID) const;
 
-	bool IsDirty(const FMovieSceneCompiledDataEntry& Entry) const;
+	MOVIESCENE_API bool IsDirty(const FMovieSceneCompiledDataEntry& Entry) const;
 
 	/**
 	 * Return a reference to a compiled data entry.
@@ -256,7 +256,7 @@ public:
 		return CompiledDataEntries[DataID.Value];
 	}
 
-	FMovieSceneCompiledDataID GetSubDataID(FMovieSceneCompiledDataID DataID, FMovieSceneSequenceID SubSequenceID);
+	MOVIESCENE_API FMovieSceneCompiledDataID GetSubDataID(FMovieSceneCompiledDataID DataID, FMovieSceneSequenceID SubSequenceID);
 
 	uint32 GetReallocationVersion() const
 	{
@@ -273,50 +273,50 @@ public:
 	const FMovieSceneEvaluationField&      GetTrackTemplateFieldChecked(FMovieSceneCompiledDataID DataID) const { return TrackTemplateFields.FindChecked(DataID.Value);   }
 	const FMovieSceneEntityComponentField& GetEntityComponentFieldChecked(FMovieSceneCompiledDataID DataID) const { return EntityComponentFields.FindChecked(DataID.Value); }
 
-	void Compile(FMovieSceneCompiledDataID DataID);
+	MOVIESCENE_API void Compile(FMovieSceneCompiledDataID DataID);
 
-	FMovieSceneCompiledDataID Compile(UMovieSceneSequence* Sequence);
+	MOVIESCENE_API FMovieSceneCompiledDataID Compile(UMovieSceneSequence* Sequence);
 
-	void Compile(FMovieSceneCompiledDataID DataID, UMovieSceneSequence* Sequence);
+	MOVIESCENE_API void Compile(FMovieSceneCompiledDataID DataID, UMovieSceneSequence* Sequence);
 
-	static bool CompileHierarchy(UMovieSceneSequence* Sequence, FMovieSceneSequenceHierarchy* InOutHierarchy, EMovieSceneServerClientMask NetworkMask);
+	static MOVIESCENE_API bool CompileHierarchy(UMovieSceneSequence* Sequence, FMovieSceneSequenceHierarchy* InOutHierarchy, EMovieSceneServerClientMask NetworkMask);
 
-	void CopyCompiledData(UMovieSceneSequence* Sequence);
-	void LoadCompiledData(UMovieSceneSequence* Sequence);
+	MOVIESCENE_API void CopyCompiledData(UMovieSceneSequence* Sequence);
+	MOVIESCENE_API void LoadCompiledData(UMovieSceneSequence* Sequence);
 
 private:
 
-	void Gather(const FMovieSceneCompiledDataEntry& Entry, UMovieSceneSequence* Sequence, const FTrackGatherParameters& Params, FMovieSceneGatheredCompilerData* OutCompilerData) const;
+	MOVIESCENE_API void Gather(const FMovieSceneCompiledDataEntry& Entry, UMovieSceneSequence* Sequence, const FTrackGatherParameters& Params, FMovieSceneGatheredCompilerData* OutCompilerData) const;
 
-	void CompileTrack(FMovieSceneCompiledDataEntry* OutEntry, const FMovieSceneBinding* ObjectBinding, UMovieSceneTrack* Track, const FTrackGatherParameters& Params, TSet<FGuid>* OutCompiledSignatures, FMovieSceneGatheredCompilerData* OutCompilerData);
+	MOVIESCENE_API void CompileTrack(FMovieSceneCompiledDataEntry* OutEntry, const FMovieSceneBinding* ObjectBinding, UMovieSceneTrack* Track, const FTrackGatherParameters& Params, TSet<FGuid>* OutCompiledSignatures, FMovieSceneGatheredCompilerData* OutCompilerData);
 
-	void GatherTrack(const FMovieSceneBinding* ObjectBinding, UMovieSceneTrack* Track, const FTrackGatherParameters& Params, const FMovieSceneEvaluationTemplate* TrackTemplate, FMovieSceneGatheredCompilerData* OutCompilerData) const;
+	MOVIESCENE_API void GatherTrack(const FMovieSceneBinding* ObjectBinding, UMovieSceneTrack* Track, const FTrackGatherParameters& Params, const FMovieSceneEvaluationTemplate* TrackTemplate, FMovieSceneGatheredCompilerData* OutCompilerData) const;
 
-	void CompileSubSequences(const FMovieSceneSequenceHierarchy& Hierarchy, const FTrackGatherParameters& Params, FMovieSceneGatheredCompilerData* OutCompilerData);
+	MOVIESCENE_API void CompileSubSequences(const FMovieSceneSequenceHierarchy& Hierarchy, const FTrackGatherParameters& Params, FMovieSceneGatheredCompilerData* OutCompilerData);
 
-	static bool CompileHierarchy(UMovieSceneSequence* Sequence, const FGatherParameters& Params, FMovieSceneSequenceHierarchy* InOutHierarchy);
+	static MOVIESCENE_API bool CompileHierarchy(UMovieSceneSequence* Sequence, const FGatherParameters& Params, FMovieSceneSequenceHierarchy* InOutHierarchy);
 
-	static bool CompileHierarchyImpl(UMovieSceneSequence* Sequence, const FGatherParameters& Params, const FMovieSceneEvaluationOperand& Operand, UE::MovieScene::FSubSequencePath* RootPath, FMovieSceneSequenceHierarchy* InOutHierarchy);
+	static MOVIESCENE_API bool CompileHierarchyImpl(UMovieSceneSequence* Sequence, const FGatherParameters& Params, const FMovieSceneEvaluationOperand& Operand, UE::MovieScene::FSubSequencePath* RootPath, FMovieSceneSequenceHierarchy* InOutHierarchy);
 
-	static bool GenerateSubSequenceData(UMovieSceneSequence* SubSequence, const FGatherParameters& Params, const FMovieSceneEvaluationOperand& Operand, UE::MovieScene::FSubSequencePath* RootPath, FMovieSceneSequenceHierarchy* InOutHierarchy);
-	static bool GenerateSubSequenceData(UMovieSceneSubTrack* SubTrack, const FGatherParameters& Params, const FMovieSceneEvaluationOperand& Operand, UE::MovieScene::FSubSequencePath* RootPath, FMovieSceneSequenceHierarchy* InOutHierarchy);
+	static MOVIESCENE_API bool GenerateSubSequenceData(UMovieSceneSequence* SubSequence, const FGatherParameters& Params, const FMovieSceneEvaluationOperand& Operand, UE::MovieScene::FSubSequencePath* RootPath, FMovieSceneSequenceHierarchy* InOutHierarchy);
+	static MOVIESCENE_API bool GenerateSubSequenceData(UMovieSceneSubTrack* SubTrack, const FGatherParameters& Params, const FMovieSceneEvaluationOperand& Operand, UE::MovieScene::FSubSequencePath* RootPath, FMovieSceneSequenceHierarchy* InOutHierarchy);
 
-	static void PopulateSubSequenceTree(UMovieSceneSequence* SubSequence, const FGatherParameters& Params, UE::MovieScene::FSubSequencePath* RootPath, FMovieSceneSequenceHierarchy* InOutHierarchy);
-	static void PopulateSubSequenceTree(UMovieSceneSubTrack* SubTrack, const FGatherParameters& Params, UE::MovieScene::FSubSequencePath* RootPath, FMovieSceneSequenceHierarchy* InOutHierarchy);
+	static MOVIESCENE_API void PopulateSubSequenceTree(UMovieSceneSequence* SubSequence, const FGatherParameters& Params, UE::MovieScene::FSubSequencePath* RootPath, FMovieSceneSequenceHierarchy* InOutHierarchy);
+	static MOVIESCENE_API void PopulateSubSequenceTree(UMovieSceneSubTrack* SubTrack, const FGatherParameters& Params, UE::MovieScene::FSubSequencePath* RootPath, FMovieSceneSequenceHierarchy* InOutHierarchy);
 
-	static TOptional<FFrameNumber> GetLoopingSubSectionEndTime(const UMovieSceneSequence* InRootSequence, const UMovieSceneSubSection* SubSection, const FGatherParameters& Params);
+	static MOVIESCENE_API TOptional<FFrameNumber> GetLoopingSubSectionEndTime(const UMovieSceneSequence* InRootSequence, const UMovieSceneSubSection* SubSection, const FGatherParameters& Params);
 
-	void CompileTrackTemplateField(FMovieSceneCompiledDataEntry* OutEntry, const FMovieSceneSequenceHierarchy& Hierarchy, FMovieSceneGatheredCompilerData* InCompilerData);
+	MOVIESCENE_API void CompileTrackTemplateField(FMovieSceneCompiledDataEntry* OutEntry, const FMovieSceneSequenceHierarchy& Hierarchy, FMovieSceneGatheredCompilerData* InCompilerData);
 
-	void PopulateEvaluationGroup(const TArray<FCompileOnTheFlyData>& SortedCompileData, FMovieSceneEvaluationGroup* OutGroup);
+	MOVIESCENE_API void PopulateEvaluationGroup(const TArray<FCompileOnTheFlyData>& SortedCompileData, FMovieSceneEvaluationGroup* OutGroup);
 
-	void PopulateMetaData(const FMovieSceneSequenceHierarchy& RootHierarchy, const TArray<FCompileOnTheFlyData>& SortedCompileData, TMovieSceneEvaluationTreeDataIterator<FMovieSceneSubSequenceTreeEntry> SubSequences, FMovieSceneEvaluationMetaData* OutMetaData);
+	MOVIESCENE_API void PopulateMetaData(const FMovieSceneSequenceHierarchy& RootHierarchy, const TArray<FCompileOnTheFlyData>& SortedCompileData, TMovieSceneEvaluationTreeDataIterator<FMovieSceneSubSequenceTreeEntry> SubSequences, FMovieSceneEvaluationMetaData* OutMetaData);
 
-	void ProcessTrack(FMovieSceneCompiledDataEntry* OutEntry, const FMovieSceneBinding* ObjectBinding, UMovieSceneTrack* Track, const FTrackGatherParameters& Params, FMovieSceneGatheredCompilerData* OutCompilerData, TSet<FGuid>* OutCompiledSignatures);
+	MOVIESCENE_API void ProcessTrack(FMovieSceneCompiledDataEntry* OutEntry, const FMovieSceneBinding* ObjectBinding, UMovieSceneTrack* Track, const FTrackGatherParameters& Params, FMovieSceneGatheredCompilerData* OutCompilerData, TSet<FGuid>* OutCompiledSignatures);
 
-	void ProcessSubTrack(FMovieSceneCompiledDataEntry* OutEntry, UMovieSceneSubTrack* SubTrack, const FGuid& ObjectBindingId, const FTrackGatherParameters& Params, FMovieSceneGatheredCompilerData* OutCompilerData);
+	MOVIESCENE_API void ProcessSubTrack(FMovieSceneCompiledDataEntry* OutEntry, UMovieSceneSubTrack* SubTrack, const FGuid& ObjectBindingId, const FTrackGatherParameters& Params, FMovieSceneGatheredCompilerData* OutCompilerData);
 
-	void DestroyData(FMovieSceneCompiledDataID DataID);
+	MOVIESCENE_API void DestroyData(FMovieSceneCompiledDataID DataID);
 
 private:
 
@@ -326,7 +326,7 @@ private:
 		return &CompiledDataEntries[DataID.Value];
 	}
 
-	void ConsoleVariableSink();
+	MOVIESCENE_API void ConsoleVariableSink();
 
 private:
 
diff --git a/Engine/Source/Runtime/MovieScene/Public/EntitySystem/BuiltInComponentTypes.h b/Engine/Source/Runtime/MovieScene/Public/EntitySystem/BuiltInComponentTypes.h
--- a/Engine/Source/Runtime/MovieScene/Public/EntitySystem/BuiltInComponentTypes.h
+++ b/Engine/Source/Runtime/MovieScene/Public/EntitySystem/BuiltInComponentTypes.h
@@ -329,9 +329,9 @@ private:
 /**
  * Pre-defined built in component types
  */
-struct MOVIESCENE_API FBuiltInComponentTypes
+struct FBuiltInComponentTypes
 {
-	~FBuiltInComponentTypes();
+	MOVIESCENE_API ~FBuiltInComponentTypes();
 
 public:
 
@@ -553,20 +553,20 @@ public:
 
 public:
 
-	static void Destroy();
+	static MOVIESCENE_API void Destroy();
 
-	static FBuiltInComponentTypes* Get();
+	static MOVIESCENE_API FBuiltInComponentTypes* Get();
 
 	FORCEINLINE static bool IsBoundObjectGarbage(UObject* InObject)
 	{
 		return InObject == nullptr || !IsValidChecked(InObject) || InObject->IsUnreachable();
 	}
 
-	FComponentTypeID GetBaseValueComponentType(const FComponentTypeID& InResultComponentType);
+	MOVIESCENE_API FComponentTypeID GetBaseValueComponentType(const FComponentTypeID& InResultComponentType);
 
 private:
 
-	FBuiltInComponentTypes();
+	MOVIESCENE_API FBuiltInComponentTypes();
 
 	TMap<FComponentTypeID, FComponentTypeID> ResultToBase;
 };
diff --git a/Engine/Source/Runtime/MovieScene/Public/EntitySystem/EntityAllocationIterator.h b/Engine/Source/Runtime/MovieScene/Public/EntitySystem/EntityAllocationIterator.h
--- a/Engine/Source/Runtime/MovieScene/Public/EntitySystem/EntityAllocationIterator.h
+++ b/Engine/Source/Runtime/MovieScene/Public/EntitySystem/EntityAllocationIterator.h
@@ -19,12 +19,12 @@ using FEntityAllocationIteratorItem = FEntityAllocationProxy;
 /**
  * Object that iterates all entity allocations that match a specific filter
  */
-struct MOVIESCENE_API FEntityAllocationIterator
+struct FEntityAllocationIterator
 {
 	/**
 	 * End iterator constructor
 	 */
-	explicit FEntityAllocationIterator(const FEntityManager* InManager);
+	MOVIESCENE_API explicit FEntityAllocationIterator(const FEntityManager* InManager);
 
 	/**
 	 * Construction from the entity manager to iterate, and a filter
@@ -32,33 +32,33 @@ struct MOVIESCENE_API FEntityAllocationIterator
 	 * @param InManager        The entity manager to iterate
 	 * @param InFilter         Filter that defines the components to match. Copied into this object.
 	 */
-	explicit FEntityAllocationIterator(const FEntityManager* InManager, const FEntityComponentFilter* InFilter);
+	MOVIESCENE_API explicit FEntityAllocationIterator(const FEntityManager* InManager, const FEntityComponentFilter* InFilter);
 
-	FEntityAllocationIterator(FEntityAllocationIterator&&);
-	FEntityAllocationIterator& operator=(FEntityAllocationIterator&&);
+	MOVIESCENE_API FEntityAllocationIterator(FEntityAllocationIterator&&);
+	MOVIESCENE_API FEntityAllocationIterator& operator=(FEntityAllocationIterator&&);
 
 	/**
 	 * Destructor
 	 */
-	~FEntityAllocationIterator();
+	MOVIESCENE_API ~FEntityAllocationIterator();
 
 
 	/**
 	 * Retrieve the entity allocation that this iterator represents. Only valid if this iterator != end()
 	 */
-	FEntityAllocationIteratorItem operator*() const;
+	MOVIESCENE_API FEntityAllocationIteratorItem operator*() const;
 
 
 	/**
 	 * Increment this iterator to the next matching allocation. Only valid if this iterator != end()
 	 */
-	FEntityAllocationIterator& operator++();
+	MOVIESCENE_API FEntityAllocationIterator& operator++();
 
 
 	/**
 	 * Test whether this iterator is valid (ie not at the end of the iteration)
 	 */
-	bool operator!=(const FEntityAllocationIterator& Other) const;
+	MOVIESCENE_API bool operator!=(const FEntityAllocationIterator& Other) const;
 
 private:
 
diff --git a/Engine/Source/Runtime/MovieScene/Public/EntitySystem/IMovieSceneBlenderSystemSupport.h b/Engine/Source/Runtime/MovieScene/Public/EntitySystem/IMovieSceneBlenderSystemSupport.h
--- a/Engine/Source/Runtime/MovieScene/Public/EntitySystem/IMovieSceneBlenderSystemSupport.h
+++ b/Engine/Source/Runtime/MovieScene/Public/EntitySystem/IMovieSceneBlenderSystemSupport.h
@@ -11,8 +11,8 @@
 
 class UMovieSceneBlenderSystem;
 
-UINTERFACE()
-class MOVIESCENE_API UMovieSceneBlenderSystemSupport : public UInterface
+UINTERFACE(MinimalAPI)
+class UMovieSceneBlenderSystemSupport : public UInterface
 {
 public:
 	GENERATED_BODY()
@@ -22,7 +22,7 @@ public:
 /**
  * Interface that can be added to UMovieSceneTracks to enable user-selection of blender systems
  */
-class MOVIESCENE_API IMovieSceneBlenderSystemSupport
+class IMovieSceneBlenderSystemSupport
 {
 public:
 	GENERATED_BODY()
@@ -30,4 +30,4 @@ public:
 	virtual TSubclassOf<UMovieSceneBlenderSystem> GetBlenderSystem() const = 0;
 	virtual void SetBlenderSystem(TSubclassOf<UMovieSceneBlenderSystem> BlenderSystemClass) = 0;
 	virtual void GetSupportedBlenderSystems(TArray<TSubclassOf<UMovieSceneBlenderSystem>>& OutSystemClasses) const = 0;
-};
\ No newline at end of file
+};
diff --git a/Engine/Source/Runtime/MovieScene/Public/EntitySystem/IMovieSceneEntityProvider.h b/Engine/Source/Runtime/MovieScene/Public/EntitySystem/IMovieSceneEntityProvider.h
--- a/Engine/Source/Runtime/MovieScene/Public/EntitySystem/IMovieSceneEntityProvider.h
+++ b/Engine/Source/Runtime/MovieScene/Public/EntitySystem/IMovieSceneEntityProvider.h
@@ -108,8 +108,8 @@ struct FEntityImportParams
 } // namespace UE
 
 
-UINTERFACE()
-class MOVIESCENE_API UMovieSceneEntityProvider : public UInterface
+UINTERFACE(MinimalAPI)
+class UMovieSceneEntityProvider : public UInterface
 {
 public:
 	GENERATED_BODY()
@@ -118,7 +118,7 @@ public:
 /**
  * Interface to be added to UMovieSceneSection types when they contain entity data
  */
-class MOVIESCENE_API IMovieSceneEntityProvider
+class IMovieSceneEntityProvider
 {
 public:
 
@@ -138,8 +138,8 @@ public:
 	}
 
 
-	void ImportEntity(UMovieSceneEntitySystemLinker* EntityLinker, const FEntityImportParams& Params, FImportedEntity* OutImportedEntity);
-	void InterrogateEntity(UMovieSceneEntitySystemLinker* EntityLinker, const FEntityImportParams& Params, FImportedEntity* OutImportedEntity);
+	MOVIESCENE_API void ImportEntity(UMovieSceneEntitySystemLinker* EntityLinker, const FEntityImportParams& Params, FImportedEntity* OutImportedEntity);
+	MOVIESCENE_API void InterrogateEntity(UMovieSceneEntitySystemLinker* EntityLinker, const FEntityImportParams& Params, FImportedEntity* OutImportedEntity);
 
 private:
 
diff --git a/Engine/Source/Runtime/MovieScene/Public/EntitySystem/MovieSceneBlenderSystem.h b/Engine/Source/Runtime/MovieScene/Public/EntitySystem/MovieSceneBlenderSystem.h
--- a/Engine/Source/Runtime/MovieScene/Public/EntitySystem/MovieSceneBlenderSystem.h
+++ b/Engine/Source/Runtime/MovieScene/Public/EntitySystem/MovieSceneBlenderSystem.h
@@ -40,8 +40,8 @@ class UObject;
  * into the index of its blend channel input component. A final combination pass walks over blend channel outputs
  * a writes the results into the result component from the accumulation buffer.
  */
-UCLASS(Abstract)
-class MOVIESCENE_API UMovieSceneBlenderSystem : public UMovieSceneEntitySystem
+UCLASS(Abstract, MinimalAPI)
+class UMovieSceneBlenderSystem : public UMovieSceneEntitySystem
 {
 public:
 
@@ -53,7 +53,7 @@ public:
 	/**
 	 * Get a blender system class, given a blender system ID.
 	 */
-	static TSubclassOf<UMovieSceneBlenderSystem> GetBlenderSystemClass(FMovieSceneBlenderSystemID InSystemID);
+	static MOVIESCENE_API TSubclassOf<UMovieSceneBlenderSystem> GetBlenderSystemClass(FMovieSceneBlenderSystemID InSystemID);
 
 	/**
 	 * Get the unique ID for a given blender system type.
@@ -69,18 +69,18 @@ public:
 	/**
 	 * Get the unique ID for this blender system's underlying type.
 	 */
-	FMovieSceneBlenderSystemID GetBlenderSystemID() const;
+	MOVIESCENE_API FMovieSceneBlenderSystemID GetBlenderSystemID() const;
 
 	/**
 	 * Allocate a new blend channel.
 	 * @note Must be released when it is no longer needed in order to prevent leaking channels.
 	 */
-	FMovieSceneBlendChannelID AllocateBlendChannel();
+	MOVIESCENE_API FMovieSceneBlendChannelID AllocateBlendChannel();
 
 	/**
 	 * Release a previously allocated blend channel.
 	 */
-	void ReleaseBlendChannel(FMovieSceneBlendChannelID BlendChannelID);
+	MOVIESCENE_API void ReleaseBlendChannel(FMovieSceneBlendChannelID BlendChannelID);
 
 	/**
 	 * Retrieve the selection priority to use when choosing a blender system. Highest wins.
@@ -100,11 +100,11 @@ public:
 
 protected:
 
-	UMovieSceneBlenderSystem(const FObjectInitializer& ObjInit);
+	MOVIESCENE_API UMovieSceneBlenderSystem(const FObjectInitializer& ObjInit);
 
-	virtual bool IsRelevantImpl(UMovieSceneEntitySystemLinker* InLinker) const override;
+	MOVIESCENE_API virtual bool IsRelevantImpl(UMovieSceneEntitySystemLinker* InLinker) const override;
 
-	void CompactBlendChannels();
+	MOVIESCENE_API void CompactBlendChannels();
 
 	/** Bit array specifying currently allocated blend channels */
 	TBitArray<> AllocatedBlendChannels;
diff --git a/Engine/Source/Runtime/MovieScene/Public/EntitySystem/MovieSceneBoundObjectInstantiator.h b/Engine/Source/Runtime/MovieScene/Public/EntitySystem/MovieSceneBoundObjectInstantiator.h
--- a/Engine/Source/Runtime/MovieScene/Public/EntitySystem/MovieSceneBoundObjectInstantiator.h
+++ b/Engine/Source/Runtime/MovieScene/Public/EntitySystem/MovieSceneBoundObjectInstantiator.h
@@ -12,8 +12,8 @@
 class UObject;
 
 
-UCLASS()
-class MOVIESCENE_API UMovieSceneGenericBoundObjectInstantiator : public UMovieSceneEntityInstantiatorSystem
+UCLASS(MinimalAPI)
+class UMovieSceneGenericBoundObjectInstantiator : public UMovieSceneEntityInstantiatorSystem
 {
 	GENERATED_BODY()
 
@@ -21,6 +21,6 @@ class MOVIESCENE_API UMovieSceneGenericBoundObjectInstantiator : public UMovieSc
 
 private:
 
-	virtual void OnRun(FSystemTaskPrerequisites& InPrerequisites, FSystemSubsequentTasks& Subsequents) override final;
+	MOVIESCENE_API virtual void OnRun(FSystemTaskPrerequisites& InPrerequisites, FSystemSubsequentTasks& Subsequents) override final;
 };
 
diff --git a/Engine/Source/Runtime/MovieScene/Public/EntitySystem/MovieSceneBoundSceneComponentInstantiator.h b/Engine/Source/Runtime/MovieScene/Public/EntitySystem/MovieSceneBoundSceneComponentInstantiator.h
--- a/Engine/Source/Runtime/MovieScene/Public/EntitySystem/MovieSceneBoundSceneComponentInstantiator.h
+++ b/Engine/Source/Runtime/MovieScene/Public/EntitySystem/MovieSceneBoundSceneComponentInstantiator.h
@@ -12,8 +12,8 @@
 
 class UObject;
 
-UCLASS()
-class MOVIESCENE_API UMovieSceneBoundSceneComponentInstantiator : public UMovieSceneEntityInstantiatorSystem
+UCLASS(MinimalAPI)
+class UMovieSceneBoundSceneComponentInstantiator : public UMovieSceneEntityInstantiatorSystem
 {
 	GENERATED_BODY()
 
@@ -21,6 +21,6 @@ class MOVIESCENE_API UMovieSceneBoundSceneComponentInstantiator : public UMovieS
 
 private:
 
-	virtual void OnRun(FSystemTaskPrerequisites& InPrerequisites, FSystemSubsequentTasks& Subsequents) override final;
+	MOVIESCENE_API virtual void OnRun(FSystemTaskPrerequisites& InPrerequisites, FSystemSubsequentTasks& Subsequents) override final;
 };
 
diff --git a/Engine/Source/Runtime/MovieScene/Public/EntitySystem/MovieSceneCachedEntityFilterResult.h b/Engine/Source/Runtime/MovieScene/Public/EntitySystem/MovieSceneCachedEntityFilterResult.h
--- a/Engine/Source/Runtime/MovieScene/Public/EntitySystem/MovieSceneCachedEntityFilterResult.h
+++ b/Engine/Source/Runtime/MovieScene/Public/EntitySystem/MovieSceneCachedEntityFilterResult.h
@@ -64,13 +64,13 @@ private:
 /**
  * Simple cached filter results that stores whether its filter passes or not
  */
-struct MOVIESCENE_API FCachedEntityFilterResult_Match
+struct FCachedEntityFilterResult_Match
 {
 	FEntityComponentFilter Filter;
 
-	bool Matches(const FEntityManager& InEntityManager);
+	MOVIESCENE_API bool Matches(const FEntityManager& InEntityManager);
 
-	void Invalidate();
+	MOVIESCENE_API void Invalidate();
 
 private:
 
@@ -84,13 +84,13 @@ private:
 /**
  * Cached filter result that caches pointers to allocations matching the specified filter
  */
-struct MOVIESCENE_API FCachedEntityFilterResult_Allocations
+struct FCachedEntityFilterResult_Allocations
 {
 	FEntityComponentFilter Filter;
 
-	TArrayView<FEntityAllocation* const> GetMatchingAllocations(const FEntityManager& InEntityManager);
+	MOVIESCENE_API TArrayView<FEntityAllocation* const> GetMatchingAllocations(const FEntityManager& InEntityManager);
 
-	void Invalidate();
+	MOVIESCENE_API void Invalidate();
 
 private:
 
diff --git a/Engine/Source/Runtime/MovieScene/Public/EntitySystem/MovieSceneComponentRegistry.h b/Engine/Source/Runtime/MovieScene/Public/EntitySystem/MovieSceneComponentRegistry.h
--- a/Engine/Source/Runtime/MovieScene/Public/EntitySystem/MovieSceneComponentRegistry.h
+++ b/Engine/Source/Runtime/MovieScene/Public/EntitySystem/MovieSceneComponentRegistry.h
@@ -42,7 +42,7 @@ struct FNewComponentTypeParams
 	EComponentTypeFlags Flags;
 };
 
-struct MOVIESCENE_API FComponentRegistry
+struct FComponentRegistry
 {
 public:
 	FEntityFactories Factories;
@@ -65,7 +65,7 @@ public:
 	 * @param Flags          Flags relating to the new component type
 	 * @return A new component type identifier for the tag
 	 */
-	FComponentTypeID NewTag(const TCHAR* const DebugName, EComponentTypeFlags Flags = EComponentTypeFlags::None);
+	MOVIESCENE_API FComponentTypeID NewTag(const TCHAR* const DebugName, EComponentTypeFlags Flags = EComponentTypeFlags::None);
 
 
 	/**
@@ -111,7 +111,7 @@ public:
 #endif
 	}
 
-	const FComponentTypeInfo& GetComponentTypeChecked(FComponentTypeID ComponentTypeID) const;
+	MOVIESCENE_API const FComponentTypeInfo& GetComponentTypeChecked(FComponentTypeID ComponentTypeID) const;
 
 public:
 
@@ -122,7 +122,7 @@ public:
 	 *
 	 * @param ComponentTypeID The component type to destroy
 	 */
-	void DestroyComponentTypeSafe(FComponentTypeID ComponentTypeID);
+	MOVIESCENE_API void DestroyComponentTypeSafe(FComponentTypeID ComponentTypeID);
 
 
 	/**
@@ -131,7 +131,7 @@ public:
 	 *
 	 * @param ComponentTypeID The component type to destroy
 	 */
-	void DestroyComponentUnsafeFast(FComponentTypeID ComponentTypeID);
+	MOVIESCENE_API void DestroyComponentUnsafeFast(FComponentTypeID ComponentTypeID);
 
 public:
 
@@ -169,7 +169,7 @@ public:
 
 private:
 
-	FComponentTypeID NewComponentTypeInternal(FComponentTypeInfo&& TypeInfo);
+	MOVIESCENE_API FComponentTypeID NewComponentTypeInternal(FComponentTypeInfo&& TypeInfo);
 
 	template<typename T>
 	FComponentTypeInfo MakeComponentTypeInfoWithoutComponentOps(const TCHAR* const DebugName, const FNewComponentTypeParams& Params);
@@ -194,4 +194,4 @@ private:
 
 
 } // namespace MovieScene
-} // namespace UE
\ No newline at end of file
+} // namespace UE
diff --git a/Engine/Source/Runtime/MovieScene/Public/EntitySystem/MovieSceneDecompositionQuery.h b/Engine/Source/Runtime/MovieScene/Public/EntitySystem/MovieSceneDecompositionQuery.h
--- a/Engine/Source/Runtime/MovieScene/Public/EntitySystem/MovieSceneDecompositionQuery.h
+++ b/Engine/Source/Runtime/MovieScene/Public/EntitySystem/MovieSceneDecompositionQuery.h
@@ -188,8 +188,8 @@ namespace MovieScene
 } // namespace UE
 
 
-UINTERFACE()
-class MOVIESCENE_API UMovieSceneValueDecomposer : public UInterface
+UINTERFACE(MinimalAPI)
+class UMovieSceneValueDecomposer : public UInterface
 {
 public:
 	GENERATED_BODY()
diff --git a/Engine/Source/Runtime/MovieScene/Public/EntitySystem/MovieSceneEntityInstantiatorSystem.h b/Engine/Source/Runtime/MovieScene/Public/EntitySystem/MovieSceneEntityInstantiatorSystem.h
--- a/Engine/Source/Runtime/MovieScene/Public/EntitySystem/MovieSceneEntityInstantiatorSystem.h
+++ b/Engine/Source/Runtime/MovieScene/Public/EntitySystem/MovieSceneEntityInstantiatorSystem.h
@@ -17,14 +17,14 @@ namespace UE { namespace MovieScene { template <typename T> struct TComponentTyp
 struct FGuid;
 struct FMovieSceneObjectBindingID;
 
-UCLASS(Abstract)
-class MOVIESCENE_API UMovieSceneEntityInstantiatorSystem : public UMovieSceneEntitySystem
+UCLASS(Abstract, MinimalAPI)
+class UMovieSceneEntityInstantiatorSystem : public UMovieSceneEntitySystem
 {
 public:
 	GENERATED_BODY()
 
-	UMovieSceneEntityInstantiatorSystem(const FObjectInitializer& ObjInit);
+	MOVIESCENE_API UMovieSceneEntityInstantiatorSystem(const FObjectInitializer& ObjInit);
 
-	void UnlinkStaleObjectBindings(UE::MovieScene::TComponentTypeID<FGuid> BindingType);
-	void UnlinkStaleObjectBindings(UE::MovieScene::TComponentTypeID<FMovieSceneObjectBindingID> BindingType);
+	MOVIESCENE_API void UnlinkStaleObjectBindings(UE::MovieScene::TComponentTypeID<FGuid> BindingType);
+	MOVIESCENE_API void UnlinkStaleObjectBindings(UE::MovieScene::TComponentTypeID<FMovieSceneObjectBindingID> BindingType);
 };
diff --git a/Engine/Source/Runtime/MovieScene/Public/EntitySystem/MovieSceneEntityLedger.h b/Engine/Source/Runtime/MovieScene/Public/EntitySystem/MovieSceneEntityLedger.h
--- a/Engine/Source/Runtime/MovieScene/Public/EntitySystem/MovieSceneEntityLedger.h
+++ b/Engine/Source/Runtime/MovieScene/Public/EntitySystem/MovieSceneEntityLedger.h
@@ -29,7 +29,7 @@ struct IMovieScenePerEntityMutation;
  * An entity ledger is responsible for tracking the entites that have been imported for the currently relevant frame of a sequence instance.
  * It is responsible for linking and unlinking new and expired entities in the linker's entity manager based on the current set of entities required for evaluation.
  */
-struct MOVIESCENE_API FEntityLedger
+struct FEntityLedger
 {
 	/**
 	 * To be called any time this ledger's instance is to be evaluated with a different set of entities - updates the set of entities that are required for the current evaluation environment
@@ -39,7 +39,7 @@ struct MOVIESCENE_API FEntityLedger
 	 * @param EntityField    Possibly null if NewEntities is empty- an entity field containing structural information about the sequence
 	 * @param NewEntities    A set specifying all the entities required for the next evaluation. Specifying an empty set will unlink all existing entities.
 	 */
-	void UpdateEntities(UMovieSceneEntitySystemLinker* Linker, const FEntityImportSequenceParams& ImportParams, const FMovieSceneEntityComponentField* EntityField, const FMovieSceneEvaluationFieldEntitySet& NewEntities);
+	MOVIESCENE_API void UpdateEntities(UMovieSceneEntitySystemLinker* Linker, const FEntityImportSequenceParams& ImportParams, const FMovieSceneEntityComponentField* EntityField, const FMovieSceneEvaluationFieldEntitySet& NewEntities);
 
 	/**
 	 * Update any one-shot entities for the current frame
@@ -49,34 +49,34 @@ struct MOVIESCENE_API FEntityLedger
 	 * @param EntityField    Possibly null if NewEntities is empty- an entity field containing structural information about the sequence
 	 * @param NewEntities    A set specifying all the entities required for the next evaluation. Specifying an empty set will unlink all existing entities.
 	 */
-	void UpdateOneShotEntities(UMovieSceneEntitySystemLinker* Linker, const FEntityImportSequenceParams& ImportParams, const FMovieSceneEntityComponentField* EntityField, const FMovieSceneEvaluationFieldEntitySet& NewEntities);
+	MOVIESCENE_API void UpdateOneShotEntities(UMovieSceneEntitySystemLinker* Linker, const FEntityImportSequenceParams& ImportParams, const FMovieSceneEntityComponentField* EntityField, const FMovieSceneEvaluationFieldEntitySet& NewEntities);
 
 	/**
 	 * Invalidate any and all entities that are currently being tracked, causing new linker entities to be created on the next evaluation, and ones to become unlinked (preserving any components with the preserve flag)
 	 */
-	void Invalidate();
+	MOVIESCENE_API void Invalidate();
 
 public:
 
 	/**
 	 * Check whether this ledger contains any information at all (ie is tracking any global entities, even if it has not created any linker entities for them)
 	 */
-	bool IsEmpty() const;
+	MOVIESCENE_API bool IsEmpty() const;
 
 	/**
 	 * Check whether the specified entity is being tracked by this ledger at all
 	 */
-	bool HasImportedEntity(const FMovieSceneEvaluationFieldEntityKey& EntityKey) const;
+	MOVIESCENE_API bool HasImportedEntity(const FMovieSceneEvaluationFieldEntityKey& EntityKey) const;
 
 	/**
 	 * Find an imported entity
 	 */
-	FMovieSceneEntityID FindImportedEntity(const FMovieSceneEvaluationFieldEntityKey& EntityKey) const;
+	MOVIESCENE_API FMovieSceneEntityID FindImportedEntity(const FMovieSceneEvaluationFieldEntityKey& EntityKey) const;
 
 	/**
 	 * Find imported entities
 	 */
-	void FindImportedEntities(TWeakObjectPtr<UObject> EntityOwner, TArray<FMovieSceneEntityID>& OutEntityIDs) const;
+	MOVIESCENE_API void FindImportedEntities(TWeakObjectPtr<UObject> EntityOwner, TArray<FMovieSceneEntityID>& OutEntityIDs) const;
 
 	/**
 	 * Indicate that the specified field entity is currently being evaluated
@@ -85,14 +85,14 @@ public:
 	 * @param InstanceHandle A handle to the sequence instance that the entity relates to (relating to Linker->GetInstanceRegistry())
 	 * @param Entity         The field entity that is being imported
 	 */
-	void ImportEntity(UMovieSceneEntitySystemLinker* Linker, const FEntityImportSequenceParams& ImportParams, const FMovieSceneEntityComponentField* EntityField, const FMovieSceneEvaluationFieldEntityQuery& Query);
+	MOVIESCENE_API void ImportEntity(UMovieSceneEntitySystemLinker* Linker, const FEntityImportSequenceParams& ImportParams, const FMovieSceneEntityComponentField* EntityField, const FMovieSceneEvaluationFieldEntityQuery& Query);
 
 	/**
 	 * Unlink all imported linker entities and their children, whilst maintaining the map of imported entities
 	 *
 	 * @param Linker The linker that owns this ledger
 	 */
-	void UnlinkEverything(UMovieSceneEntitySystemLinker* Linker);
+	MOVIESCENE_API void UnlinkEverything(UMovieSceneEntitySystemLinker* Linker);
 
 	/**
 	 * Unlink all imported one-shot linker entities and their children and clear the list of one shots
@@ -100,18 +100,18 @@ public:
 	 * @param Linker The linker that owns this ledger
 	 * @return A mask representing the changes made to the environment
 	 */
-	void UnlinkOneShots(UMovieSceneEntitySystemLinker* Linker);
+	MOVIESCENE_API void UnlinkOneShots(UMovieSceneEntitySystemLinker* Linker);
 
 
 	/**
 	 * Check whether any of the entities in this ledger or their children match the specified filter
 	 */
-	bool Contains(UMovieSceneEntitySystemLinker* Linker, const FEntityComponentFilter& Filter) const;
+	MOVIESCENE_API bool Contains(UMovieSceneEntitySystemLinker* Linker, const FEntityComponentFilter& Filter) const;
 
 	/**
 	 * Mutate all the entities within this ledger by using the specified filter
 	 */
-	void MutateAll(UMovieSceneEntitySystemLinker* Linker, const FEntityComponentFilter& Filter, const IMovieScenePerEntityMutation& Mutation) const;
+	MOVIESCENE_API void MutateAll(UMovieSceneEntitySystemLinker* Linker, const FEntityComponentFilter& Filter, const IMovieScenePerEntityMutation& Mutation) const;
 
 public:
 
@@ -120,14 +120,14 @@ public:
 	 *
 	 * @param Linker         The linker that owns this ledger
 	 */
-	void TagGarbage(UMovieSceneEntitySystemLinker* Linker);
+	MOVIESCENE_API void TagGarbage(UMovieSceneEntitySystemLinker* Linker);
 
 	/**
 	 * Remove linker entity IDs that exist in the specified set since they are no longer valid
 	 *
 	 * @param LinkerEntities Set of entity IDs that have been destroyed
 	 */
-	void CleanupLinkerEntities(const TSet<FMovieSceneEntityID>& LinkerEntities);
+	MOVIESCENE_API void CleanupLinkerEntities(const TSet<FMovieSceneEntityID>& LinkerEntities);
 
 private:
 
diff --git a/Engine/Source/Runtime/MovieScene/Public/EntitySystem/MovieSceneEntityManager.h b/Engine/Source/Runtime/MovieScene/Public/EntitySystem/MovieSceneEntityManager.h
--- a/Engine/Source/Runtime/MovieScene/Public/EntitySystem/MovieSceneEntityManager.h
+++ b/Engine/Source/Runtime/MovieScene/Public/EntitySystem/MovieSceneEntityManager.h
@@ -91,12 +91,12 @@ enum class EEntityThreadingModel : uint8
  * offset from the start of each component array. This enables efficient read/write access into specific component arrays and makes issuing parallel tasks that require component data trivial. See FEntityAllocation for
  * a more detailed explanation of entity component data layout.
  */
-class MOVIESCENE_API FEntityManager : public FUObjectArray::FUObjectDeleteListener
+class FEntityManager : public FUObjectArray::FUObjectDeleteListener
 {
 public:
 
-	FEntityManager();
-	~FEntityManager();
+	MOVIESCENE_API FEntityManager();
+	MOVIESCENE_API ~FEntityManager();
 
 	FEntityManager(const FEntityManager&) = delete;
 	void operator=(const FEntityManager&) = delete;
@@ -121,7 +121,7 @@ public:
 	/**
 	 * Destroy this entity manager and all the entities and components contained within it, resetting it back to its default state
 	 */
-	void Destroy();
+	MOVIESCENE_API void Destroy();
 
 
 	/**
@@ -129,7 +129,7 @@ public:
 	 *
 	 * @return A stable ID that relates to this entity. Will remain valid until the entity is freed
 	 */
-	FMovieSceneEntityID AllocateEntity();
+	MOVIESCENE_API FMovieSceneEntityID AllocateEntity();
 
 
 	/**
@@ -138,7 +138,7 @@ public:
 	 * @param EntityID  A valid entity ID to free
 	 * @return The number of entities released
 	 */
-	void FreeEntity(FMovieSceneEntityID EntityID);
+	MOVIESCENE_API void FreeEntity(FMovieSceneEntityID EntityID);
 
 
 	/**
@@ -148,7 +148,7 @@ public:
 	 * @param OutFreedEntities (Optional) A set to populate with all the entities that were freed (including any children)
 	 * @return The number of entities released
 	 */
-	int32 FreeEntities(const FEntityComponentFilter& Filter, TSet<FMovieSceneEntityID>* OutFreedEntities = nullptr);
+	MOVIESCENE_API int32 FreeEntities(const FEntityComponentFilter& Filter, TSet<FMovieSceneEntityID>* OutFreedEntities = nullptr);
 
 
 	/**
@@ -158,7 +158,7 @@ public:
 	 * @param OutFreedEntities (Optional) A set to populate with all the entities that were freed (including any children)
 	 * @return The number of entities released
 	 */
-	int32 FreeEntities(const FFreeEntityOperation& Operation, TSet<FMovieSceneEntityID>* OutFreedEntities = nullptr);
+	MOVIESCENE_API int32 FreeEntities(const FFreeEntityOperation& Operation, TSet<FMovieSceneEntityID>* OutFreedEntities = nullptr);
 
 
 	/**
@@ -169,7 +169,7 @@ public:
 	 * @param InOutNum               A valid pointer to the desired number to allocate. Is overwritten with the number that were actually allocated.
 	 * @return A structure that points to the first entity that was allocated.
 	 */
-	FEntityDataLocation AllocateContiguousEntities(const FComponentMask& EntityComponentMask, int32* InOutNum);
+	MOVIESCENE_API FEntityDataLocation AllocateContiguousEntities(const FComponentMask& EntityComponentMask, int32* InOutNum);
 
 
 	/**
@@ -178,7 +178,7 @@ public:
 	 * @param EntityComponentMask    Defines the components that should exist on the allocated components. Set bits denote allocated components and tags.
 	 * @return A structure that points to the allocated entity.
 	 */
-	FEntityInfo AllocateEntity(const FComponentMask& EntityComponentMask);
+	MOVIESCENE_API FEntityInfo AllocateEntity(const FComponentMask& EntityComponentMask);
 
 
 	/**
@@ -186,7 +186,7 @@ public:
 	 *
 	 * @return A structure that points to the allocated entity.
 	 */
-	FEntityInfo GetEntity(FMovieSceneEntityID EntityID) const;
+	MOVIESCENE_API FEntityInfo GetEntity(FMovieSceneEntityID EntityID) const;
 
 
 	/**
@@ -194,7 +194,7 @@ public:
 	 *
 	 * @return A handle that points to the allocated entity.
 	 */
-	FEntityHandle GetEntityHandle(FMovieSceneEntityID EntityID);
+	MOVIESCENE_API FEntityHandle GetEntityHandle(FMovieSceneEntityID EntityID);
 
 
 	/**
@@ -214,25 +214,25 @@ public:
 	 *
 	 * @return true if the entity ID is allocated, false otherwise
 	 */
-	bool IsHandleValid(FEntityHandle EntityID) const;
+	MOVIESCENE_API bool IsHandleValid(FEntityHandle EntityID) const;
 
 
 	/**
 	 * Compute and return this entity manager's threading model. Does not change the current cached threading model.
 	 */
-	EEntityThreadingModel ComputeThreadingModel() const;
+	MOVIESCENE_API EEntityThreadingModel ComputeThreadingModel() const;
 
 
 	/**
 	 * Compute and store the current threading model.
 	 */
-	void UpdateThreadingModel();
+	MOVIESCENE_API void UpdateThreadingModel();
 
 
 	/**
 	 * Get this entitiy manager's current threading model based on the last time UpdateThreadingModel was called.
 	 */
-	EEntityThreadingModel GetThreadingModel() const;
+	MOVIESCENE_API EEntityThreadingModel GetThreadingModel() const;
 
 public:
 
@@ -268,8 +268,8 @@ public:
 	 * @param EntityID         The ID of the entity to add the component to
 	 * @param ComponentTypeID  The ID of the component type to add to the entity
 	 */
-	void AddComponent(FMovieSceneEntityID EntityID, FComponentTypeID ComponentTypeID);
-	void AddComponent(FMovieSceneEntityID EntityID, FComponentTypeID ComponentTypeID, EEntityRecursion Recursion);
+	MOVIESCENE_API void AddComponent(FMovieSceneEntityID EntityID, FComponentTypeID ComponentTypeID);
+	MOVIESCENE_API void AddComponent(FMovieSceneEntityID EntityID, FComponentTypeID ComponentTypeID, EEntityRecursion Recursion);
 
 
 	/**
@@ -278,8 +278,8 @@ public:
 	 * @param EntityID             The ID of the entity to add the components to
 	 * @param EntityComponentMask  A mask constituting the components that should be added (set bits indicate components to add)
 	 */
-	void AddComponents(FMovieSceneEntityID EntityID, const FComponentMask& EntityComponentMask);
-	void AddComponents(FMovieSceneEntityID EntityID, const FComponentMask& EntityComponentMask, EEntityRecursion Recursion);
+	MOVIESCENE_API void AddComponents(FMovieSceneEntityID EntityID, const FComponentMask& EntityComponentMask);
+	MOVIESCENE_API void AddComponents(FMovieSceneEntityID EntityID, const FComponentMask& EntityComponentMask, EEntityRecursion Recursion);
 
 
 	/**
@@ -288,8 +288,8 @@ public:
 	 * @param EntityID         The ID of the entity to remove the component from
 	 * @param ComponentTypeID  The ID of the component type to remove from the entity
 	 */
-	void RemoveComponent(FMovieSceneEntityID EntityID, FComponentTypeID ComponentTypeID);
-	void RemoveComponent(FMovieSceneEntityID EntityID, FComponentTypeID ComponentTypeID, EEntityRecursion Recursion);
+	MOVIESCENE_API void RemoveComponent(FMovieSceneEntityID EntityID, FComponentTypeID ComponentTypeID);
+	MOVIESCENE_API void RemoveComponent(FMovieSceneEntityID EntityID, FComponentTypeID ComponentTypeID, EEntityRecursion Recursion);
 
 
 	/**
@@ -298,8 +298,8 @@ public:
 	 * @param EntityID             The ID of the entity to remove the components from
 	 * @param ComponentsToRemove   A mask constituting the components that should be removed (set bits indicate components to add)
 	 */
-	void RemoveComponents(FMovieSceneEntityID EntityID, const FComponentMask& ComponentsToRemove);
-	void RemoveComponents(FMovieSceneEntityID EntityID, const FComponentMask& ComponentsToRemove, EEntityRecursion Recursion);
+	MOVIESCENE_API void RemoveComponents(FMovieSceneEntityID EntityID, const FComponentMask& ComponentsToRemove);
+	MOVIESCENE_API void RemoveComponents(FMovieSceneEntityID EntityID, const FComponentMask& ComponentsToRemove, EEntityRecursion Recursion);
 
 
 	/**
@@ -309,7 +309,7 @@ public:
 	 * @param DstEntityID      The ID of the entity to copy to
 	 * @param ComponentTypeID  The ID of the component type to copy. Must exist on SrcEntityID
 	 */
-	bool CopyComponent(FMovieSceneEntityID SrcEntityID, FMovieSceneEntityID DstEntityID, FComponentTypeID ComponentTypeID);
+	MOVIESCENE_API bool CopyComponent(FMovieSceneEntityID SrcEntityID, FMovieSceneEntityID DstEntityID, FComponentTypeID ComponentTypeID);
 
 
 	/**
@@ -319,7 +319,7 @@ public:
 	 * @param DstEntityID      The ID of the entity to copy to
 	 * @param ComponentsToCopy  A mask constituting the components that should be copied (set bits indicate components to copy)
 	 */
-	void CopyComponents(FMovieSceneEntityID SrcEntityID, FMovieSceneEntityID DstEntityID, const FComponentMask& ComponentsToCopy);
+	MOVIESCENE_API void CopyComponents(FMovieSceneEntityID SrcEntityID, FMovieSceneEntityID DstEntityID, const FComponentMask& ComponentsToCopy);
 
 
 	/**
@@ -328,7 +328,7 @@ public:
 	 * @param EntityID             The ID of the entity to check
 	 * @param ComponentTypeID      The type of the component that is being tested for
 	 */
-	bool HasComponent(FMovieSceneEntityID EntityID, FComponentTypeID ComponentTypeID) const;
+	MOVIESCENE_API bool HasComponent(FMovieSceneEntityID EntityID, FComponentTypeID ComponentTypeID) const;
 
 
 	/**
@@ -337,7 +337,7 @@ public:
 	 * @param InEntity  The ID of the entity
 	 * @return The type mask for this entity, or an empty mask if it has no components
 	 */
-	const FComponentMask& GetEntityType(FMovieSceneEntityID InEntity) const;
+	MOVIESCENE_API const FComponentMask& GetEntityType(FMovieSceneEntityID InEntity) const;
 
 
 	/**
@@ -346,7 +346,7 @@ public:
 	 * @param InEntity   The ID of the entity
 	 * @param InNewMask  The new mask of the entity that defines which components it should have
 	 */
-	void ChangeEntityType(FMovieSceneEntityID InEntity, const FComponentMask& InNewMask);
+	MOVIESCENE_API void ChangeEntityType(FMovieSceneEntityID InEntity, const FComponentMask& InNewMask);
 
 
 	/**
@@ -355,7 +355,7 @@ public:
 	 * @param EntityID             The ID of the entity to remove the components from
 	 * @param EntitiesToKeep       A mask constituting the components that should be kept on the entity
 	 */
-	void FilterComponents(FMovieSceneEntityID EntityID, const FComponentMask& EntitiesToKeep);
+	MOVIESCENE_API void FilterComponents(FMovieSceneEntityID EntityID, const FComponentMask& EntitiesToKeep);
 
 
 	/**
@@ -365,7 +365,7 @@ public:
 	 * @param SourceEntityID       The ID of the entity to copy components from
 	 * @param OptionalMask         (Optional) A mask constituting the components that should be copied
 	 */
-	void CombineComponents(FMovieSceneEntityID DestinationEntityID, FMovieSceneEntityID SourceEntityID, const FComponentMask* OptionalMask = nullptr);
+	MOVIESCENE_API void CombineComponents(FMovieSceneEntityID DestinationEntityID, FMovieSceneEntityID SourceEntityID, const FComponentMask* OptionalMask = nullptr);
 
 
 	/**
@@ -374,7 +374,7 @@ public:
 	 * @param InOther  The ID of the entity to duplicate
 	 * @return An ID to a new entity that has exactly the same combination and value of components as InOther
 	 */
-	FMovieSceneEntityID DuplicateEntity(FMovieSceneEntityID InOther);
+	MOVIESCENE_API FMovieSceneEntityID DuplicateEntity(FMovieSceneEntityID InOther);
 
 
 	/**
@@ -383,7 +383,7 @@ public:
 	 * @param InOutEntity          The ID of the entity to overwrite. Doesn't have to be valid.
 	 * @param InEntityToDuplicate  The ID of the entity to duplicate
 	 */
-	void OverwriteEntityWithDuplicate(FMovieSceneEntityID& InOutEntity, FMovieSceneEntityID InEntityToDuplicate);
+	MOVIESCENE_API void OverwriteEntityWithDuplicate(FMovieSceneEntityID& InOutEntity, FMovieSceneEntityID InEntityToDuplicate);
 
 
 	/**
@@ -401,7 +401,7 @@ public:
 	/**
 	 * Runs all initializers for the specified parent/child allocation
 	 */
-	void InitializeChildAllocation(const FComponentMask& ParentType, const FComponentMask& ChildType, const FEntityAllocation* ParentAllocation, TArrayView<const int32> ParentAllocationOffsets, const FEntityRange& InChildEntityRange);
+	MOVIESCENE_API void InitializeChildAllocation(const FComponentMask& ParentType, const FComponentMask& ChildType, const FEntityAllocation* ParentAllocation, TArrayView<const int32> ParentAllocationOffsets, const FEntityRange& InChildEntityRange);
 
 	/**
 	 * Destroy a previously registered instanced child initializer using its index
@@ -414,12 +414,12 @@ public:
 	/**
 	 * Run through all entities in this entity manager, ensuring that all mutual components exist
 	 */
-	void AddMutualComponents();
+	MOVIESCENE_API void AddMutualComponents();
 
 	/**
 	 * Run through all entities in this entity manager, ensuring that all mutual components exist for any component types that match the specified filter
 	 */
-	void AddMutualComponents(const FEntityComponentFilter& InFilter);
+	MOVIESCENE_API void AddMutualComponents(const FEntityComponentFilter& InFilter);
 
 public:
 
@@ -568,7 +568,7 @@ public:
 	/**
 	 * Goes through all entity data and compacts like-for-like allocations into as few allocations as possible, resulting in the optimal data layout
 	 */
-	void Compact();
+	MOVIESCENE_API void Compact();
 
 
 	/**
@@ -578,7 +578,7 @@ public:
 	 * @param Mutation    Implementation that defines how to mutate the entities that match the filter
 	 * @return The number of entities that were mutated, or 0 if none were matched
 	 */
-	int32 MutateAll(const FEntityComponentFilter& Filter, const IMovieSceneEntityMutation& Mutation, EMutuallyInclusiveComponentType MutualTypes = EMutuallyInclusiveComponentType::Mandatory);
+	MOVIESCENE_API int32 MutateAll(const FEntityComponentFilter& Filter, const IMovieSceneEntityMutation& Mutation, EMutuallyInclusiveComponentType MutualTypes = EMutuallyInclusiveComponentType::Mandatory);
 
 
 	/**
@@ -588,7 +588,7 @@ public:
 	 * @param Mutation    Implementation that defines how to mutate the entities that match the filter
 	 * @return The number of entities that were mutated, or 0 if none were matched
 	 */
-	int32 MutateConditional(const FEntityComponentFilter& Filter, const IMovieSceneConditionalEntityMutation& Mutation, EMutuallyInclusiveComponentType MutualTypes = EMutuallyInclusiveComponentType::Mandatory);
+	MOVIESCENE_API int32 MutateConditional(const FEntityComponentFilter& Filter, const IMovieSceneConditionalEntityMutation& Mutation, EMutuallyInclusiveComponentType MutualTypes = EMutuallyInclusiveComponentType::Mandatory);
 
 
 	/**
@@ -598,7 +598,7 @@ public:
 	 *
 	 * @param EntityID    The ID of the entity to touch.
 	 */
-	void TouchEntity(FMovieSceneEntityID EntityID);
+	MOVIESCENE_API void TouchEntity(FMovieSceneEntityID EntityID);
 
 
 	/**
@@ -607,7 +607,7 @@ public:
 	 * @param ParentID   The ID of the parent
 	 * @param ChildID    The ID of the child
 	 */
-	void AddChild(FMovieSceneEntityID ParentID, FMovieSceneEntityID ChildID);
+	MOVIESCENE_API void AddChild(FMovieSceneEntityID ParentID, FMovieSceneEntityID ChildID);
 
 
 	/**
@@ -617,7 +617,7 @@ public:
 	 * @param ParentID    The ID of the parent
 	 * @param OutChildren (out) Array to populate with child entities.
 	 */
-	void GetImmediateChildren(FMovieSceneEntityID ParentID, TArray<FMovieSceneEntityID>& OutChildren) const;
+	MOVIESCENE_API void GetImmediateChildren(FMovieSceneEntityID ParentID, TArray<FMovieSceneEntityID>& OutChildren) const;
 
 
 	/**
@@ -627,7 +627,7 @@ public:
 	 * @param ParentID    The ID of the parent
 	 * @param OutChildren (out) Array to populate with child entities.
 	 */
-	void GetChildren_ParentFirst(FMovieSceneEntityID ParentID, TArray<FMovieSceneEntityID>& OutChildren) const;
+	MOVIESCENE_API void GetChildren_ParentFirst(FMovieSceneEntityID ParentID, TArray<FMovieSceneEntityID>& OutChildren) const;
 
 
 	/**
@@ -673,13 +673,13 @@ public:
 	 * @param Ar          The archive to save to. Ar.IsCountingMemory() must be true.
 	 * @param EntityID    Identifier for the entity to replace object references within
 	 */
-	void CountMemory(FArchive& Ar, FMovieSceneEntityID EntityID);
+	MOVIESCENE_API void CountMemory(FArchive& Ar, FMovieSceneEntityID EntityID);
 
 
 	/**
 	 * Find a component type hander from its registered GUID
 	 */
-	static IComponentTypeHandler* FindComponentTypeHandler(const FGuid& ComponentGuid);
+	static MOVIESCENE_API IComponentTypeHandler* FindComponentTypeHandler(const FGuid& ComponentGuid);
 
 public:
 
@@ -689,7 +689,7 @@ public:
 	 * @param InFilter (required, non-null) The filter to match allocations against. Filter is copied into the iterator.
 	 * @return An iterator object that walks over all allocations matching the filter.
 	 */
-	FEntityAllocationIteratorProxy Iterate(const FEntityComponentFilter* InFilter) const;
+	MOVIESCENE_API FEntityAllocationIteratorProxy Iterate(const FEntityComponentFilter* InFilter) const;
 
 
 	/**
@@ -698,7 +698,7 @@ public:
 	 * @param InFilter The filter to match allocations against.
 	 * @return True if the entity manager contains any allocations that match the filter, false otherwise
 	 */
-	bool Contains(const FEntityComponentFilter& InFilter) const;
+	MOVIESCENE_API bool Contains(const FEntityComponentFilter& InFilter) const;
 
 
 	/**
@@ -757,13 +757,13 @@ public:
 	 * @param InFilter The filter to match allocations against. Filter is copied into the iterator.
 	 * @param OutMask  The mask to receive the binary OR accumulation of all entities that pass the filter
 	 */
-	void AccumulateMask(const FEntityComponentFilter& InFilter, FComponentMask& OutMask) const;
+	MOVIESCENE_API void AccumulateMask(const FEntityComponentFilter& InFilter, FComponentMask& OutMask) const;
 
 
 	/**
 	 * Retrieve an up-to-date accumulation of all components present on entities in this manager
 	 */
-	const FComponentMask& GetAccumulatedMask() const;
+	MOVIESCENE_API const FComponentMask& GetAccumulatedMask() const;
 
 
 	/**
@@ -772,8 +772,8 @@ public:
 	 * @param InFilter The filter to match allocations against.
 	 * @return True if the entity manager contains any allocations that match the filter, false otherwise
 	 */
-	void EnterIteration() const;
-	void ExitIteration() const;
+	MOVIESCENE_API void EnterIteration() const;
+	MOVIESCENE_API void ExitIteration() const;
 
 	void CheckCanChangeStructure() const
 	{
@@ -863,7 +863,7 @@ public:
 	/**
 	 * Explicitly add referenced objects from an external source. This is not called by default for the global entity manager because each entity owner should do so to avoid cyclic dependencies.
 	 */
-	void AddReferencedObjects(FReferenceCollector& ReferenceCollector);
+	MOVIESCENE_API void AddReferencedObjects(FReferenceCollector& ReferenceCollector);
 
 	
 	/**
@@ -872,7 +872,7 @@ public:
 	 * @param InOutEntity     The entity ID to be reassigned. If it is already valid, the existing entity will be freed
 	 * @param EntityToDiscard The entity containing the components to replace InOutEntity with. This Entity ID will be invalid after this function call.
 	 */
-	void ReplaceEntityID(FMovieSceneEntityID& InOutEntity, FMovieSceneEntityID EntityToDiscard);
+	MOVIESCENE_API void ReplaceEntityID(FMovieSceneEntityID& InOutEntity, FMovieSceneEntityID EntityToDiscard);
 
 	uint32 GetHandleGeneration() const
 	{
@@ -892,39 +892,39 @@ private:
 		Uninitialized,
 		DefaultConstructed
 	};
-	void OnStructureChanged();
+	MOVIESCENE_API void OnStructureChanged();
 
-	FEntityAllocation* CreateEntityAllocation(const FComponentMask& EntityComponentMask, uint16 InitialCapacity, uint16 MaxCapacity, FEntityAllocation* MigrateComponentDataFrom = nullptr);
-	int32 CreateEntityAllocationEntry(const FComponentMask& EntityComponentMask, uint16 InitialCapacity, uint16 MaxCapacity);
+	MOVIESCENE_API FEntityAllocation* CreateEntityAllocation(const FComponentMask& EntityComponentMask, uint16 InitialCapacity, uint16 MaxCapacity, FEntityAllocation* MigrateComponentDataFrom = nullptr);
+	MOVIESCENE_API int32 CreateEntityAllocationEntry(const FComponentMask& EntityComponentMask, uint16 InitialCapacity, uint16 MaxCapacity);
 
-	int32 GetOrCreateAllocationWithSlack(const FComponentMask& EntityComponentMask, int32* InOutDesiredSlack = nullptr);
-	int32 CreateAllocationWithSlack(const FComponentMask& EntityComponentMask, int32* InOutDesiredSlack = nullptr);
-	int32 MigrateEntity(int32 DestIndex, int32 SourceIndex, int32 SourceEntryIndexWithinAllocation);
+	MOVIESCENE_API int32 GetOrCreateAllocationWithSlack(const FComponentMask& EntityComponentMask, int32* InOutDesiredSlack = nullptr);
+	MOVIESCENE_API int32 CreateAllocationWithSlack(const FComponentMask& EntityComponentMask, int32* InOutDesiredSlack = nullptr);
+	MOVIESCENE_API int32 MigrateEntity(int32 DestIndex, int32 SourceIndex, int32 SourceEntryIndexWithinAllocation);
 
-	void CopyComponents(int32 DestAllocationIndex, int32 DestEntityIndex, int32 SourceAllocationIndex, int32 SourceEntityIndex, const FComponentMask* OptionalMask = nullptr);
+	MOVIESCENE_API void CopyComponents(int32 DestAllocationIndex, int32 DestEntityIndex, int32 SourceAllocationIndex, int32 SourceEntityIndex, const FComponentMask* OptionalMask = nullptr);
 
-	int32 AddEntityToAllocation(int32 AllocationIndex, FMovieSceneEntityID ID, EMemoryType MemoryType = EMemoryType::DefaultConstructed);
-	void RemoveEntityFromAllocation(int32 AllocationIndex, int32 SourceEntryIndexWithinAllocation);
+	MOVIESCENE_API int32 AddEntityToAllocation(int32 AllocationIndex, FMovieSceneEntityID ID, EMemoryType MemoryType = EMemoryType::DefaultConstructed);
+	MOVIESCENE_API void RemoveEntityFromAllocation(int32 AllocationIndex, int32 SourceEntryIndexWithinAllocation);
 
-	FEntityAllocation* MigrateAllocation(int32 AllocationIndex, const FComponentMask& NewComponentMask);
+	MOVIESCENE_API FEntityAllocation* MigrateAllocation(int32 AllocationIndex, const FComponentMask& NewComponentMask);
 
-	void CombineAllocations(int32 DestinationIndex, int32 SourceIndex);
+	MOVIESCENE_API void CombineAllocations(int32 DestinationIndex, int32 SourceIndex);
 
 	FEntityAllocation* GetAllocation(int32 AllocationIndex) const
 	{
 		return EntityAllocations[AllocationIndex];
 	}
 
-	int32 ReserveAllocation(int32 AllocationIndex, int32 NumToReserve);
-	FEntityAllocation* GrowAllocation(int32 AllocationIndex, int32 MinNumToGrowBy = 1);
+	MOVIESCENE_API int32 ReserveAllocation(int32 AllocationIndex, int32 NumToReserve);
+	MOVIESCENE_API FEntityAllocation* GrowAllocation(int32 AllocationIndex, int32 MinNumToGrowBy = 1);
 
-	void DestroyAllocation(FEntityAllocation* Allocation, bool bDestructComponentData = true);
+	MOVIESCENE_API void DestroyAllocation(FEntityAllocation* Allocation, bool bDestructComponentData = true);
 
-	virtual void NotifyUObjectDeleted(const UObjectBase* Object, int32 Index) override;
-	virtual void OnUObjectArrayShutdown() override;
-	virtual FString GetReferencerName() const;
+	MOVIESCENE_API virtual void NotifyUObjectDeleted(const UObjectBase* Object, int32 Index) override;
+	MOVIESCENE_API virtual void OnUObjectArrayShutdown() override;
+	MOVIESCENE_API virtual FString GetReferencerName() const;
 
-	void CheckInvariants();
+	MOVIESCENE_API void CheckInvariants();
 
 	friend struct FEntityAllocationProxy;
 	friend struct FEntityAllocationIterator;
diff --git a/Engine/Source/Runtime/MovieScene/Public/EntitySystem/MovieSceneEntitySystem.h b/Engine/Source/Runtime/MovieScene/Public/EntitySystem/MovieSceneEntitySystem.h
--- a/Engine/Source/Runtime/MovieScene/Public/EntitySystem/MovieSceneEntitySystem.h
+++ b/Engine/Source/Runtime/MovieScene/Public/EntitySystem/MovieSceneEntitySystem.h
@@ -61,8 +61,8 @@ namespace MovieScene
 } // namespace UE
 
 
-UCLASS()
-class MOVIESCENE_API UMovieSceneEntitySystem : public UObject
+UCLASS(MinimalAPI)
+class UMovieSceneEntitySystem : public UObject
 {
 public:
 	GENERATED_BODY()
@@ -76,8 +76,8 @@ public:
 	using FSystemTaskPrerequisites = UE::MovieScene::FSystemTaskPrerequisites;
 	using FSystemSubsequentTasks   = UE::MovieScene::FSystemSubsequentTasks;
 
-	UMovieSceneEntitySystem(const FObjectInitializer& ObjInit);
-	~UMovieSceneEntitySystem();
+	MOVIESCENE_API UMovieSceneEntitySystem(const FObjectInitializer& ObjInit);
+	MOVIESCENE_API ~UMovieSceneEntitySystem();
 
 	/**
 	 * Creates a relationship between the two system types that ensures any systems of type UpstreamSystemType always execute before DownstreamSystemType if they are both present
@@ -85,7 +85,7 @@ public:
 	 * @param UpstreamSystemType     The UClass of the system that should always be a prerequisite of DownstreamSystemType (ie, runs first)
 	 * @param DownstreamSystemType   The UClass of the system that should always run after UpstreamSystemType
 	 */
-	static void DefineImplicitPrerequisite(TSubclassOf<UMovieSceneEntitySystem> UpstreamSystemType, TSubclassOf<UMovieSceneEntitySystem> DownstreamSystemType);
+	static MOVIESCENE_API void DefineImplicitPrerequisite(TSubclassOf<UMovieSceneEntitySystem> UpstreamSystemType, TSubclassOf<UMovieSceneEntitySystem> DownstreamSystemType);
 
 	/**
 	 * Informs the dependency graph that the specified class type produces components of the specified type.
@@ -94,7 +94,7 @@ public:
 	 * @param ClassType         The UClass of the system that produces the component type
 	 * @param ComponentType     The type of the component produced by the system
 	 */
-	static void DefineComponentProducer(TSubclassOf<UMovieSceneEntitySystem> ClassType, FComponentTypeID ComponentType);
+	static MOVIESCENE_API void DefineComponentProducer(TSubclassOf<UMovieSceneEntitySystem> ClassType, FComponentTypeID ComponentType);
 
 	/**
 	 * Informs the dependency graph that the specified class type consumes components of the specified type, and as such should always execute after any producers of that component type.
@@ -102,42 +102,42 @@ public:
 	 * @param ClassType         The UClass of the system that consumes the component type
 	 * @param ComponentType     The type of the component consumed by the system
 	 */
-	static void DefineComponentConsumer(TSubclassOf<UMovieSceneEntitySystem> ClassType, FComponentTypeID ComponentType);
+	static MOVIESCENE_API void DefineComponentConsumer(TSubclassOf<UMovieSceneEntitySystem> ClassType, FComponentTypeID ComponentType);
 
 	/**
 	 * Ensure that any systems relevant to the specified linker's entity manager are linked
 	 */
-	static void LinkRelevantSystems(UMovieSceneEntitySystemLinker* InLinker);
+	static MOVIESCENE_API void LinkRelevantSystems(UMovieSceneEntitySystemLinker* InLinker);
 
 	/**
 	 * Link all systems in a given category
 	 */
-	static void LinkCategorySystems(UMovieSceneEntitySystemLinker* InLinker, UE::MovieScene::EEntitySystemCategory InCategory);
+	static MOVIESCENE_API void LinkCategorySystems(UMovieSceneEntitySystemLinker* InLinker, UE::MovieScene::EEntitySystemCategory InCategory);
 
 	/**
 	 * Link all systems that pass the given linker's filter
 	 */
-	static void LinkAllSystems(UMovieSceneEntitySystemLinker* InLinker);
+	static MOVIESCENE_API void LinkAllSystems(UMovieSceneEntitySystemLinker* InLinker);
 
 	/**
 	 * Create a new system category
 	 */
-	static UE::MovieScene::EEntitySystemCategory RegisterCustomSystemCategory();
+	static MOVIESCENE_API UE::MovieScene::EEntitySystemCategory RegisterCustomSystemCategory();
 
 	/**
 	 * Sort the given systems by their flow order, suitable for execution
 	 */
-	static void SortByFlowOrder(TArray<uint16>& InOutGlobalNodeIDs);
+	static MOVIESCENE_API void SortByFlowOrder(TArray<uint16>& InOutGlobalNodeIDs);
 
 	/**
 	 * Get the global IDs of all subsequent systems of the given system
 	 */
-	static void GetSubsequentSystems(uint16 FromGlobalNodeID, TArray<uint16>& OutSubsequentGlobalNodeIDs);
+	static MOVIESCENE_API void GetSubsequentSystems(uint16 FromGlobalNodeID, TArray<uint16>& OutSubsequentGlobalNodeIDs);
 
 	/**
 	 * Prints a graphviz markup for the global system dependency graph
 	 */
-	static void DebugPrintGlobalDependencyGraph(bool bUpdateCache = true);
+	static MOVIESCENE_API void DebugPrintGlobalDependencyGraph(bool bUpdateCache = true);
 
 public:
 
@@ -177,44 +177,44 @@ public:
 	}
 
 	/** Called when the system is removed from the linker */
-	void Unlink();
+	MOVIESCENE_API void Unlink();
 
 	/** Called when the linker is being destroyed */
-	void Abandon();
+	MOVIESCENE_API void Abandon();
 
 	/** Called when the system is added to a linker */
-	void Link(UMovieSceneEntitySystemLinker* InLinker);
+	MOVIESCENE_API void Link(UMovieSceneEntitySystemLinker* InLinker);
 
 	/** Called to schedule work */
-	void SchedulePersistentTasks(UE::MovieScene::IEntitySystemScheduler* Scheduler);
+	MOVIESCENE_API void SchedulePersistentTasks(UE::MovieScene::IEntitySystemScheduler* Scheduler);
 
 	/** Called when the system should run its logic */
-	void Run(FSystemTaskPrerequisites& InPrerequisites, FSystemSubsequentTasks& Subsequents);
+	MOVIESCENE_API void Run(FSystemTaskPrerequisites& InPrerequisites, FSystemSubsequentTasks& Subsequents);
 
 	/** Called to know if the system is still relevant and should be kept around */
-	bool IsRelevant(UMovieSceneEntitySystemLinker* InLinker) const;
+	MOVIESCENE_API bool IsRelevant(UMovieSceneEntitySystemLinker* InLinker) const;
 
-	void ConditionalLinkSystem(UMovieSceneEntitySystemLinker* InLinker) const;
+	MOVIESCENE_API void ConditionalLinkSystem(UMovieSceneEntitySystemLinker* InLinker) const;
 
-	void TagGarbage();
+	MOVIESCENE_API void TagGarbage();
 
-	void CleanTaggedGarbage();
+	MOVIESCENE_API void CleanTaggedGarbage();
 
 	/**
 	 * Enable this system if it is not already.
 	 */
-	void Enable();
+	MOVIESCENE_API void Enable();
 
 	/**
 	 * Disable this system if it is not already.
 	 * Disabled systems will remain in the system graph, and will stay alive as long as they are relevant, but will not be Run.
 	 */
-	void Disable();
+	MOVIESCENE_API void Disable();
 
 protected:
 
-	virtual bool IsReadyForFinishDestroy() override;
-	virtual void FinishDestroy() override;
+	MOVIESCENE_API virtual bool IsReadyForFinishDestroy() override;
+	MOVIESCENE_API virtual void FinishDestroy() override;
 
 private:
 
@@ -226,9 +226,9 @@ private:
 
 	virtual void OnUnlink() {}
 
-	virtual bool IsRelevantImpl(UMovieSceneEntitySystemLinker* InLinker) const;
+	MOVIESCENE_API virtual bool IsRelevantImpl(UMovieSceneEntitySystemLinker* InLinker) const;
 
-	virtual void ConditionalLinkSystemImpl(UMovieSceneEntitySystemLinker* InLinker) const;
+	MOVIESCENE_API virtual void ConditionalLinkSystemImpl(UMovieSceneEntitySystemLinker* InLinker) const;
 
 	virtual void OnTagGarbage() {}
 
diff --git a/Engine/Source/Runtime/MovieScene/Public/EntitySystem/MovieSceneEntitySystemDirectedGraph.h b/Engine/Source/Runtime/MovieScene/Public/EntitySystem/MovieSceneEntitySystemDirectedGraph.h
--- a/Engine/Source/Runtime/MovieScene/Public/EntitySystem/MovieSceneEntitySystemDirectedGraph.h
+++ b/Engine/Source/Runtime/MovieScene/Public/EntitySystem/MovieSceneEntitySystemDirectedGraph.h
@@ -21,7 +21,7 @@ struct FDirectedGraphStringParameters
 /**
  * Directed graph represented as a bitarray for allocated nodes, and edges defined by pairs of integers (from->to).
  */
-struct MOVIESCENE_API FDirectedGraph
+struct FDirectedGraph
 {
 	struct FDirectionalEdge
 	{
@@ -56,13 +56,13 @@ struct MOVIESCENE_API FDirectedGraph
 
 public:
 
-	struct MOVIESCENE_API FDepthFirstSearch
+	struct FDepthFirstSearch
 	{
 		TArray<uint16> PostNodes;
 
-		explicit FDepthFirstSearch(const FDirectedGraph* InGraph);
+		MOVIESCENE_API explicit FDepthFirstSearch(const FDirectedGraph* InGraph);
 
-		void Search(uint16 Node);
+		MOVIESCENE_API void Search(uint16 Node);
 
 		const TBitArray<>& GetVisited() const
 		{
@@ -77,13 +77,13 @@ public:
 		const FDirectedGraph* Graph;
 	};
 
-	struct MOVIESCENE_API FBreadthFirstSearch
+	struct FBreadthFirstSearch
 	{
 		TArray<uint16> Nodes;
 
-		explicit FBreadthFirstSearch(const FDirectedGraph* InGraph);
+		MOVIESCENE_API explicit FBreadthFirstSearch(const FDirectedGraph* InGraph);
 
-		void Search(uint16 Node);
+		MOVIESCENE_API void Search(uint16 Node);
 
 		const TBitArray<>& GetVisited() const
 		{
@@ -97,11 +97,11 @@ public:
 		int32 StackIndex;
 	};
 
-	struct MOVIESCENE_API FDiscoverCyclicEdges
+	struct FDiscoverCyclicEdges
 	{
-		explicit FDiscoverCyclicEdges(const FDirectedGraph* InGraph);
+		MOVIESCENE_API explicit FDiscoverCyclicEdges(const FDirectedGraph* InGraph);
 
-		void Search(uint16 Node);
+		MOVIESCENE_API void Search(uint16 Node);
 
 		bool IsCyclic(const uint16 EdgeIndex) const
 		{
@@ -113,13 +113,13 @@ public:
 			return CyclicEdges;
 		}
 
-		void Search();
-		void SearchFrom(uint16 NodeID);
+		MOVIESCENE_API void Search();
+		MOVIESCENE_API void SearchFrom(uint16 NodeID);
 
 	private:
 
-		void DiscoverCycles(uint16 NodeID, TBitArray<>& VisitedNodes);
-		void TagCyclicChain(uint16 CyclicNodeID);
+		MOVIESCENE_API void DiscoverCycles(uint16 NodeID, TBitArray<>& VisitedNodes);
+		MOVIESCENE_API void TagCyclicChain(uint16 CyclicNodeID);
 
 		TBitArray<> CyclicEdges;
 		TBitArray<> VisitedEdges;
@@ -133,50 +133,50 @@ public:
 		: bHasDanglingEdges(false)
 	{}
 
-	void AllocateNode(uint16 NodeID);
+	MOVIESCENE_API void AllocateNode(uint16 NodeID);
 
-	bool IsNodeAllocated(uint16 NodeID) const;
+	MOVIESCENE_API bool IsNodeAllocated(uint16 NodeID) const;
 
-	void CleanUpDanglingEdges();
+	MOVIESCENE_API void CleanUpDanglingEdges();
 
-	void RemoveNode(uint16 NodeID);
+	MOVIESCENE_API void RemoveNode(uint16 NodeID);
 
 	const TBitArray<>& GetNodeMask() const
 	{
 		return Nodes;
 	}
 
-	bool IsCyclic() const;
+	MOVIESCENE_API bool IsCyclic() const;
 
-	void MakeEdge(uint16 FromNode, uint16 ToNode);
+	MOVIESCENE_API void MakeEdge(uint16 FromNode, uint16 ToNode);
 
-	void DestroyEdge(uint16 FromNode, uint16 ToNode);
+	MOVIESCENE_API void DestroyEdge(uint16 FromNode, uint16 ToNode);
 
-	void DestroyAllEdges();
+	MOVIESCENE_API void DestroyAllEdges();
 
-	TBitArray<> FindEdgeUpstreamNodes() const;
+	MOVIESCENE_API TBitArray<> FindEdgeUpstreamNodes() const;
 
-	TArrayView<const FDirectionalEdge> GetEdges() const;
+	MOVIESCENE_API TArrayView<const FDirectionalEdge> GetEdges() const;
 
-	TArrayView<const FDirectionalEdge> GetEdgesFrom(uint16 InNode) const;
+	MOVIESCENE_API TArrayView<const FDirectionalEdge> GetEdgesFrom(uint16 InNode) const;
 
-	bool HasEdgeFrom(uint16 InNode) const;
+	MOVIESCENE_API bool HasEdgeFrom(uint16 InNode) const;
 
-	bool HasEdgeTo(uint16 InNode) const;
+	MOVIESCENE_API bool HasEdgeTo(uint16 InNode) const;
 
-	FString ToString(const UE::MovieScene::FDirectedGraphStringParameters& Parameters) const;
+	MOVIESCENE_API FString ToString(const UE::MovieScene::FDirectedGraphStringParameters& Parameters) const;
 
-	FString ToString(const UE::MovieScene::FDirectedGraphStringParameters& Parameters, TFunctionRef<void(uint16, FStringBuilderBase&)> EmitLabel) const;
+	MOVIESCENE_API FString ToString(const UE::MovieScene::FDirectedGraphStringParameters& Parameters, TFunctionRef<void(uint16, FStringBuilderBase&)> EmitLabel) const;
 
 private:
 
-	int32 FindEdgeStart(uint16 FromNode) const;
+	MOVIESCENE_API int32 FindEdgeStart(uint16 FromNode) const;
 
-	int32 FindEdgeIndex(const FDirectionalEdge& Edge) const;
+	MOVIESCENE_API int32 FindEdgeIndex(const FDirectionalEdge& Edge) const;
 
-	bool EdgeExists(const FDirectionalEdge& Edge) const;
+	MOVIESCENE_API bool EdgeExists(const FDirectionalEdge& Edge) const;
 
-	bool IsCyclicImpl(uint16 NodeID, TBitArray<>& Visiting) const;
+	MOVIESCENE_API bool IsCyclicImpl(uint16 NodeID, TBitArray<>& Visiting) const;
 
 private:
 
@@ -190,4 +190,4 @@ private:
 };
 
 
-} // namespace UE::MovieScene
\ No newline at end of file
+} // namespace UE::MovieScene
diff --git a/Engine/Source/Runtime/MovieScene/Public/EntitySystem/MovieSceneEntitySystemGraphs.h b/Engine/Source/Runtime/MovieScene/Public/EntitySystem/MovieSceneEntitySystemGraphs.h
--- a/Engine/Source/Runtime/MovieScene/Public/EntitySystem/MovieSceneEntitySystemGraphs.h
+++ b/Engine/Source/Runtime/MovieScene/Public/EntitySystem/MovieSceneEntitySystemGraphs.h
@@ -84,24 +84,24 @@ struct TStructOpsTypeTraits<FMovieSceneEntitySystemGraphNodes> : public TStructO
 
 
 USTRUCT()
-struct MOVIESCENE_API FMovieSceneEntitySystemGraph
+struct FMovieSceneEntitySystemGraph
 {
 	using FDirectionalEdge = UE::MovieScene::FDirectedGraph::FDirectionalEdge;
 
 	GENERATED_BODY()
 
-	FMovieSceneEntitySystemGraph();
-	~FMovieSceneEntitySystemGraph();
+	MOVIESCENE_API FMovieSceneEntitySystemGraph();
+	MOVIESCENE_API ~FMovieSceneEntitySystemGraph();
 
 	FMovieSceneEntitySystemGraph(const FMovieSceneEntitySystemGraph&) = delete;
 	void operator=(const FMovieSceneEntitySystemGraph&) = delete;
 
-	FMovieSceneEntitySystemGraph(FMovieSceneEntitySystemGraph&&);
-	FMovieSceneEntitySystemGraph& operator=(FMovieSceneEntitySystemGraph&&);
+	MOVIESCENE_API FMovieSceneEntitySystemGraph(FMovieSceneEntitySystemGraph&&);
+	MOVIESCENE_API FMovieSceneEntitySystemGraph& operator=(FMovieSceneEntitySystemGraph&&);
 
-	void AddReference(UMovieSceneEntitySystem* FromReference, UMovieSceneEntitySystem* ToReference);
+	MOVIESCENE_API void AddReference(UMovieSceneEntitySystem* FromReference, UMovieSceneEntitySystem* ToReference);
 
-	void RemoveReference(UMovieSceneEntitySystem* FromReference, UMovieSceneEntitySystem* ToReference);
+	MOVIESCENE_API void RemoveReference(UMovieSceneEntitySystem* FromReference, UMovieSceneEntitySystem* ToReference);
 
 	/** Olog(n) time */
 	template<typename Allocator>
@@ -155,24 +155,24 @@ struct MOVIESCENE_API FMovieSceneEntitySystemGraph
 		return ReferenceGraph.HasEdgeFrom(GetGraphID(FromReference));
 	}
 
-	void AddSystem(UMovieSceneEntitySystem* InSystem);
+	MOVIESCENE_API void AddSystem(UMovieSceneEntitySystem* InSystem);
 
-	void RemoveSystem(UMovieSceneEntitySystem* InSystem);
+	MOVIESCENE_API void RemoveSystem(UMovieSceneEntitySystem* InSystem);
 
-	int32 RemoveIrrelevantSystems(UMovieSceneEntitySystemLinker* Linker);
+	MOVIESCENE_API int32 RemoveIrrelevantSystems(UMovieSceneEntitySystemLinker* Linker);
 
-	void Shutdown();
+	MOVIESCENE_API void Shutdown();
 
-	int32 NumInPhase(UE::MovieScene::ESystemPhase Phase) const;
+	MOVIESCENE_API int32 NumInPhase(UE::MovieScene::ESystemPhase Phase) const;
 
-	void ExecutePhase(UE::MovieScene::ESystemPhase Phase, UMovieSceneEntitySystemLinker* Linker, FGraphEventArray& OutTasks);
+	MOVIESCENE_API void ExecutePhase(UE::MovieScene::ESystemPhase Phase, UMovieSceneEntitySystemLinker* Linker, FGraphEventArray& OutTasks);
 
-	void IteratePhase(UE::MovieScene::ESystemPhase Phase, TFunctionRef<void(UMovieSceneEntitySystem*)> InIter);
+	MOVIESCENE_API void IteratePhase(UE::MovieScene::ESystemPhase Phase, TFunctionRef<void(UMovieSceneEntitySystem*)> InIter);
 
-	void ReconstructTaskSchedule(UE::MovieScene::FEntityManager* EntityManager);
-	void ScheduleTasks(UE::MovieScene::FEntityManager* EntityManager);
+	MOVIESCENE_API void ReconstructTaskSchedule(UE::MovieScene::FEntityManager* EntityManager);
+	MOVIESCENE_API void ScheduleTasks(UE::MovieScene::FEntityManager* EntityManager);
 
-	TArray<UMovieSceneEntitySystem*> GetSystems() const;
+	MOVIESCENE_API TArray<UMovieSceneEntitySystem*> GetSystems() const;
 
 	template<typename SystemType>
 	SystemType* FindSystemOfType() const
@@ -180,18 +180,18 @@ struct MOVIESCENE_API FMovieSceneEntitySystemGraph
 		return CastChecked<SystemType>(FindSystemOfType(SystemType::StaticClass()), ECastCheckedType::NullAllowed);
 	}
 
-	UMovieSceneEntitySystem* FindSystemOfType(TSubclassOf<UMovieSceneEntitySystem> InClassType) const;
+	MOVIESCENE_API UMovieSceneEntitySystem* FindSystemOfType(TSubclassOf<UMovieSceneEntitySystem> InClassType) const;
 
-	void DebugPrint() const;
+	MOVIESCENE_API void DebugPrint() const;
 
-	FString ToString() const;
+	MOVIESCENE_API FString ToString() const;
 
 private:
 
 	// Implementation function that means we don't need to #include the entity system
-	static uint16 GetGraphID(const UMovieSceneEntitySystem* InSystem);
+	static MOVIESCENE_API uint16 GetGraphID(const UMovieSceneEntitySystem* InSystem);
 
-	void UpdateCache();
+	MOVIESCENE_API void UpdateCache();
 
 	template<typename ArrayType>
 	void ExecutePhase(UE::MovieScene::ESystemPhase Phase, const ArrayType& RetrieveEntries, UMovieSceneEntitySystemLinker* Linker, FGraphEventArray& OutTasks);
@@ -228,4 +228,4 @@ struct TStructOpsTypeTraits<FMovieSceneEntitySystemGraph> : public TStructOpsTyp
 	{
 		WithCopy = false
 	};
-};
\ No newline at end of file
+};
