[Engine]

diff --git a/Engine/Source/Runtime/MovieScene/Public/EntitySystem/MovieSceneEntitySystemLinker.h b/Engine/Source/Runtime/MovieScene/Public/EntitySystem/MovieSceneEntitySystemLinker.h
--- a/Engine/Source/Runtime/MovieScene/Public/EntitySystem/MovieSceneEntitySystemLinker.h
+++ b/Engine/Source/Runtime/MovieScene/Public/EntitySystem/MovieSceneEntitySystemLinker.h
@@ -60,10 +60,10 @@ namespace MovieScene
 	MOVIESCENE_API EEntitySystemLinkerRole RegisterCustomEntitySystemLinkerRole();
 
 	/** Utility class for filtering systems */
-	struct MOVIESCENE_API FSystemFilter
+	struct FSystemFilter
 	{
 		/** Constructs a default filter that allows all systems */
-		FSystemFilter();
+		MOVIESCENE_API FSystemFilter();
 
 		/** Checks whether the given system class passes all filters */
 		template<typename SystemClass>
@@ -73,23 +73,23 @@ namespace MovieScene
 		}
 
 		/** Checks whether the given system class passes all filters */
-		bool CheckSystem(TSubclassOf<UMovieSceneEntitySystem> InClass) const;
+		MOVIESCENE_API bool CheckSystem(TSubclassOf<UMovieSceneEntitySystem> InClass) const;
 		/** Checks whether the given system passes all filters */
-		bool CheckSystem(const UMovieSceneEntitySystem* InSystem) const;
+		MOVIESCENE_API bool CheckSystem(const UMovieSceneEntitySystem* InSystem) const;
 
 		/** Sets system categories that are allowed */
-		void SetAllowedCategories(EEntitySystemCategory InCategory);
+		MOVIESCENE_API void SetAllowedCategories(EEntitySystemCategory InCategory);
 		/** Add system categories to be allowed */
-		void AllowCategory(EEntitySystemCategory InCategory);
+		MOVIESCENE_API void AllowCategory(EEntitySystemCategory InCategory);
 		/** Sets system categories that are disallowed */
-		void SetDisallowedCategories(EEntitySystemCategory InCategory);
+		MOVIESCENE_API void SetDisallowedCategories(EEntitySystemCategory InCategory);
 		/** Add system categories to be disallowed */
-		void DisallowCategory(EEntitySystemCategory InCategory);
+		MOVIESCENE_API void DisallowCategory(EEntitySystemCategory InCategory);
 
 		/** Specifically allow the given system type */
-		void AllowSystem(TSubclassOf<UMovieSceneEntitySystem> InClass);
+		MOVIESCENE_API void AllowSystem(TSubclassOf<UMovieSceneEntitySystem> InClass);
 		/** Specifically disallow the given system type */
-		void DisallowSystem(TSubclassOf<UMovieSceneEntitySystem> InClass);
+		MOVIESCENE_API void DisallowSystem(TSubclassOf<UMovieSceneEntitySystem> InClass);
 
 	private:
 		UE::MovieScene::EEntitySystemCategory CategoriesAllowed;
@@ -106,8 +106,8 @@ DECLARE_MULTICAST_DELEGATE_TwoParams(FMovieSceneEntitySystemLinkerAROEvent, UMov
 DECLARE_MULTICAST_DELEGATE_TwoParams(FMovieSceneEntitySystemLinkerWorldEvent, UMovieSceneEntitySystemLinker*, UWorld*);
 DECLARE_MULTICAST_DELEGATE_OneParam(FMovieSceneEntitySystemLinkerPostSpawnEvent, UMovieSceneEntitySystemLinker*);
 
-UCLASS()
-class MOVIESCENE_API UMovieSceneEntitySystemLinker
+UCLASS(MinimalAPI)
+class UMovieSceneEntitySystemLinker
 	: public UObject
 {
 public:
@@ -133,15 +133,15 @@ public:
 	UE::MovieScene::FPreAnimatedStateExtension PreAnimatedState;
 
 	/** Constructs a new linker */
-	UMovieSceneEntitySystemLinker(const FObjectInitializer& ObjInit);
+	MOVIESCENE_API UMovieSceneEntitySystemLinker(const FObjectInitializer& ObjInit);
 
 	/** Gets the global component registry */
-	static FComponentRegistry* GetComponents();
+	static MOVIESCENE_API FComponentRegistry* GetComponents();
 
 	/** Finds or creates a named linker */
-	static UMovieSceneEntitySystemLinker* FindOrCreateLinker(UObject* PreferredOuter, UE::MovieScene::EEntitySystemLinkerRole LinkerRole, const TCHAR* Name = TEXT("DefaultMovieSceneEntitySystemLinker"));
+	static MOVIESCENE_API UMovieSceneEntitySystemLinker* FindOrCreateLinker(UObject* PreferredOuter, UE::MovieScene::EEntitySystemLinkerRole LinkerRole, const TCHAR* Name = TEXT("DefaultMovieSceneEntitySystemLinker"));
 	/** Creates a new linker */
-	static UMovieSceneEntitySystemLinker* CreateLinker(UObject* PreferredOuter, UE::MovieScene::EEntitySystemLinkerRole LinkerRole);
+	static MOVIESCENE_API UMovieSceneEntitySystemLinker* CreateLinker(UObject* PreferredOuter, UE::MovieScene::EEntitySystemLinkerRole LinkerRole);
 
 	/** Gets this linker's instance registry */
 	FInstanceRegistry* GetInstanceRegistry()
@@ -176,12 +176,12 @@ public:
 		return CastChecked<SystemType>(FindSystem(SystemType::StaticClass()), ECastCheckedType::NullAllowed);
 	}
 
-	UMovieSceneEntitySystem* LinkSystem(TSubclassOf<UMovieSceneEntitySystem> InClassType);
-	UMovieSceneEntitySystem* LinkSystemIfAllowed(TSubclassOf<UMovieSceneEntitySystem> InClassType);
+	MOVIESCENE_API UMovieSceneEntitySystem* LinkSystem(TSubclassOf<UMovieSceneEntitySystem> InClassType);
+	MOVIESCENE_API UMovieSceneEntitySystem* LinkSystemIfAllowed(TSubclassOf<UMovieSceneEntitySystem> InClassType);
 
-	UMovieSceneEntitySystem* FindSystem(TSubclassOf<UMovieSceneEntitySystem> Class) const;
+	MOVIESCENE_API UMovieSceneEntitySystem* FindSystem(TSubclassOf<UMovieSceneEntitySystem> Class) const;
 
-	static void AddReferencedObjects(UObject* InThis, FReferenceCollector& Collector);
+	static MOVIESCENE_API void AddReferencedObjects(UObject* InThis, FReferenceCollector& Collector);
 
 	/**
 	 * Gets this linker's system filter.
@@ -210,7 +210,7 @@ public:
 	/**
 	 * Completely reset this linker back to its default state, abandoning all systems and destroying all entities
 	 */
-	void Reset();
+	MOVIESCENE_API void Reset();
 
 public:
 
@@ -311,44 +311,44 @@ public:
 
 	// Internal API
 	
-	void SystemLinked(UMovieSceneEntitySystem* InSystem);
-	void SystemUnlinked(UMovieSceneEntitySystem* InSystem);
+	MOVIESCENE_API void SystemLinked(UMovieSceneEntitySystem* InSystem);
+	MOVIESCENE_API void SystemUnlinked(UMovieSceneEntitySystem* InSystem);
 
-	bool HasLinkedSystem(const uint16 GlobalDependencyGraphID);
+	MOVIESCENE_API bool HasLinkedSystem(const uint16 GlobalDependencyGraphID);
 
-	void LinkRelevantSystems();
-	void UnlinkIrrelevantSystems();
-	void AutoLinkRelevantSystems();
-	void AutoUnlinkIrrelevantSystems();
+	MOVIESCENE_API void LinkRelevantSystems();
+	MOVIESCENE_API void UnlinkIrrelevantSystems();
+	MOVIESCENE_API void AutoLinkRelevantSystems();
+	MOVIESCENE_API void AutoUnlinkIrrelevantSystems();
 
-	bool HasStructureChangedSinceLastRun() const;
+	MOVIESCENE_API bool HasStructureChangedSinceLastRun() const;
 
-	void InvalidateObjectBinding(const FGuid& ObjectBindingID, FInstanceHandle InstanceHandle);
-	void CleanupInvalidBoundObjects();
+	MOVIESCENE_API void InvalidateObjectBinding(const FGuid& ObjectBindingID, FInstanceHandle InstanceHandle);
+	MOVIESCENE_API void CleanupInvalidBoundObjects();
 
-	bool StartEvaluation(FMovieSceneEntitySystemRunner& InRunner);
-	FMovieSceneEntitySystemRunner* GetActiveRunner() const;
-	void PostInstantation(FMovieSceneEntitySystemRunner& InRunner);
-	void EndEvaluation(FMovieSceneEntitySystemRunner& InRunner);
+	MOVIESCENE_API bool StartEvaluation(FMovieSceneEntitySystemRunner& InRunner);
+	MOVIESCENE_API FMovieSceneEntitySystemRunner* GetActiveRunner() const;
+	MOVIESCENE_API void PostInstantation(FMovieSceneEntitySystemRunner& InRunner);
+	MOVIESCENE_API void EndEvaluation(FMovieSceneEntitySystemRunner& InRunner);
 
-	void ResetActiveRunners();
+	MOVIESCENE_API void ResetActiveRunners();
 
 private:
 
-	UMovieSceneEntitySystem* LinkSystemImpl(TSubclassOf<UMovieSceneEntitySystem> InClassType);
+	MOVIESCENE_API UMovieSceneEntitySystem* LinkSystemImpl(TSubclassOf<UMovieSceneEntitySystem> InClassType);
 
-	void HandlePreGarbageCollection();
-	void HandlePostGarbageCollection();
+	MOVIESCENE_API void HandlePreGarbageCollection();
+	MOVIESCENE_API void HandlePostGarbageCollection();
 
-	void TagInvalidBoundObjects();
-	void CleanGarbage();
+	MOVIESCENE_API void TagInvalidBoundObjects();
+	MOVIESCENE_API void CleanGarbage();
 
-	void OnWorldCleanup(UWorld* InWorld, bool bSessionEnded, bool bCleanupResources);
-	void OnObjectsReplaced(const TMap<UObject*, UObject*>& ReplacementMap);
+	MOVIESCENE_API void OnWorldCleanup(UWorld* InWorld, bool bSessionEnded, bool bCleanupResources);
+	MOVIESCENE_API void OnObjectsReplaced(const TMap<UObject*, UObject*>& ReplacementMap);
 
-	virtual void BeginDestroy() override;
+	MOVIESCENE_API virtual void BeginDestroy() override;
 
-	static UE::MovieScene::FEntitySystemLinkerExtensionID RegisterExtension();
+	static MOVIESCENE_API UE::MovieScene::FEntitySystemLinkerExtensionID RegisterExtension();
 
 private:
 
diff --git a/Engine/Source/Runtime/MovieScene/Public/EntitySystem/MovieSceneEntitySystemRunner.h b/Engine/Source/Runtime/MovieScene/Public/EntitySystem/MovieSceneEntitySystemRunner.h
--- a/Engine/Source/Runtime/MovieScene/Public/EntitySystem/MovieSceneEntitySystemRunner.h
+++ b/Engine/Source/Runtime/MovieScene/Public/EntitySystem/MovieSceneEntitySystemRunner.h
@@ -63,7 +63,7 @@ ENUM_CLASS_FLAGS(ERunnerUpdateFlags)
 
 DECLARE_MULTICAST_DELEGATE(FMovieSceneEntitySystemEventTriggers);
 
-class MOVIESCENE_API FMovieSceneEntitySystemRunner : public TSharedFromThis<FMovieSceneEntitySystemRunner>
+class FMovieSceneEntitySystemRunner : public TSharedFromThis<FMovieSceneEntitySystemRunner>
 {
 public:
 
@@ -73,38 +73,38 @@ public:
 
 public:
 	/** Creates an unbound runner */
-	FMovieSceneEntitySystemRunner();
+	MOVIESCENE_API FMovieSceneEntitySystemRunner();
 	/** Destructor */
-	~FMovieSceneEntitySystemRunner();
+	MOVIESCENE_API ~FMovieSceneEntitySystemRunner();
 
 	/** Attach this runner to a linker */
-	void AttachToLinker(UMovieSceneEntitySystemLinker* InLinker);
+	MOVIESCENE_API void AttachToLinker(UMovieSceneEntitySystemLinker* InLinker);
 	/** Returns whether this runner is attached to a linker */
-	bool IsAttachedToLinker() const;
+	MOVIESCENE_API bool IsAttachedToLinker() const;
 	/** Detaches this runner from a linker */
-	void DetachFromLinker();
+	MOVIESCENE_API void DetachFromLinker();
 
-	int32 GetQueuedUpdateCount() const;
+	MOVIESCENE_API int32 GetQueuedUpdateCount() const;
 	/** Returns whether this runner has any outstanding updates. */
-	bool HasQueuedUpdates() const;
+	MOVIESCENE_API bool HasQueuedUpdates() const;
 	/** Returns whether the given instance is queued for any updates. */
-	bool HasQueuedUpdates(FInstanceHandle Instance) const;
+	MOVIESCENE_API bool HasQueuedUpdates(FInstanceHandle Instance) const;
 	/** Queue the given instance for an update with the given context. */
-	void QueueUpdate(const FMovieSceneContext& Context, FInstanceHandle Instance, UE::MovieScene::ERunnerUpdateFlags UpdateFlags = UE::MovieScene::ERunnerUpdateFlags::None);
-	void QueueUpdate(const FMovieSceneContext& Context, FInstanceHandle Instance, FSimpleDelegate&& InOnFlushedDelegate, UE::MovieScene::ERunnerUpdateFlags UpdateFlags = UE::MovieScene::ERunnerUpdateFlags::None);
+	MOVIESCENE_API void QueueUpdate(const FMovieSceneContext& Context, FInstanceHandle Instance, UE::MovieScene::ERunnerUpdateFlags UpdateFlags = UE::MovieScene::ERunnerUpdateFlags::None);
+	MOVIESCENE_API void QueueUpdate(const FMovieSceneContext& Context, FInstanceHandle Instance, FSimpleDelegate&& InOnFlushedDelegate, UE::MovieScene::ERunnerUpdateFlags UpdateFlags = UE::MovieScene::ERunnerUpdateFlags::None);
 
 	/**
 	 * Queue a final update for the specified instance, optionally destroying it after finishing it
 	 * @return true if the update requires a flush, or false if it was finished and/or destroyed immediately without requiring a flush
 	 */
-	bool QueueFinalUpdate(FInstanceHandle Instance);
-	bool QueueFinalUpdate(FInstanceHandle Instance, FSimpleDelegate&& InOnLastFlushDelegate);
-	bool QueueFinalUpdateAndDestroy(FInstanceHandle Instance);
+	MOVIESCENE_API bool QueueFinalUpdate(FInstanceHandle Instance);
+	MOVIESCENE_API bool QueueFinalUpdate(FInstanceHandle Instance, FSimpleDelegate&& InOnLastFlushDelegate);
+	MOVIESCENE_API bool QueueFinalUpdateAndDestroy(FInstanceHandle Instance);
 
 	/**
 	 * Abandon the specified instance handle and destroy it immediately. May flush this runner if necessary
 	 */
-	void AbandonAndDestroyInstance(FInstanceHandle Instance);
+	MOVIESCENE_API void AbandonAndDestroyInstance(FInstanceHandle Instance);
 
 	/**
 	 * Flushes the update queue and applies any outstanding evaluation logic
@@ -113,7 +113,7 @@ public:
 	 *                      and will process the outstanding work on the next call to Flush. A value of 0.0 signifies no budget - the queue
 	 *                      will be fully processed without leaving any outstanding work
 	 */
-	void Flush(double BudgetMs = 0.f);
+	MOVIESCENE_API void Flush(double BudgetMs = 0.f);
 
 	/**
 	 * Flushes any outstanding update tasks in the current evaluation scope with a given budget. Only performs work if this runner is part-way through an evaluation
@@ -123,18 +123,18 @@ public:
 	 *                      will be fully processed without leaving any outstanding work
 	 * @param TargetState   The desired state to reach. The runner will stop flushing as soon as all the states specified in TargetState have been flushed.
 	 */
-	void FlushOutstanding(double BudgetMs = 0.f, UE::MovieScene::ERunnerFlushState TargetState = UE::MovieScene::ERunnerFlushState::None);
+	MOVIESCENE_API void FlushOutstanding(double BudgetMs = 0.f, UE::MovieScene::ERunnerFlushState TargetState = UE::MovieScene::ERunnerFlushState::None);
 
 	/**
 	 * Called in the event that the structure of the entity manager has been unexpectedly changed while this runner is active.
 	 * This allows the runner to be reset so it runs from the start of its evaluation loop next time it is evaluated, rather than half-way through
 	 */
-	void ResetFlushState();
+	MOVIESCENE_API void ResetFlushState();
 
 	/**
 	 * Discard any queued updates for the specified sequence instance.
 	 */
-	void DiscardQueuedUpdates(FInstanceHandle Instance);
+	MOVIESCENE_API void DiscardQueuedUpdates(FInstanceHandle Instance);
 
 	/** Access this runner's currently executing phase */
 	UE::MovieScene::ESystemPhase GetCurrentPhase() const
@@ -145,38 +145,38 @@ public:
 	/**
 	 * Check whether this runner is currently inside an active evaluation loop
 	 */
-	bool IsCurrentlyEvaluating() const;
+	MOVIESCENE_API bool IsCurrentlyEvaluating() const;
 
 	/**
 	 * Run a single evaluation phase
 	 *
 	 * @return Whether the phase was run
 	 */
-	bool FlushSingleEvaluationPhase();
+	MOVIESCENE_API bool FlushSingleEvaluationPhase();
 
 public:
 
-	UMovieSceneEntitySystemLinker* GetLinker() const;
-	FEntityManager* GetEntityManager() const;
-	FInstanceRegistry* GetInstanceRegistry() const;
+	MOVIESCENE_API UMovieSceneEntitySystemLinker* GetLinker() const;
+	MOVIESCENE_API FEntityManager* GetEntityManager() const;
+	MOVIESCENE_API FInstanceRegistry* GetInstanceRegistry() const;
 
 public:
 	
 	// Internal API
 
-	void MarkForUpdate(FInstanceHandle InInstanceHandle, UE::MovieScene::ERunnerUpdateFlags UpdateFlags);
-	FMovieSceneEntitySystemEventTriggers& GetQueuedEventTriggers();
+	MOVIESCENE_API void MarkForUpdate(FInstanceHandle InInstanceHandle, UE::MovieScene::ERunnerUpdateFlags UpdateFlags);
+	MOVIESCENE_API FMovieSceneEntitySystemEventTriggers& GetQueuedEventTriggers();
 
 private:
 
-	void OnLinkerAbandon(UMovieSceneEntitySystemLinker* Linker);
+	MOVIESCENE_API void OnLinkerAbandon(UMovieSceneEntitySystemLinker* Linker);
 
 private:
 
 	/**
 	 * Queue the final update of a given instance, optionally destroying it after is finishes
 	 */
-	bool QueueFinalUpdateImpl(FInstanceHandle Instance, FSimpleDelegate&& InOnLastFlushDelegate, bool bDestroyInstance);
+	MOVIESCENE_API bool QueueFinalUpdateImpl(FInstanceHandle Instance, FSimpleDelegate&& InOnLastFlushDelegate, bool bDestroyInstance);
 
 	/**
 	 * Flush the next item in our update loop based off the contents of FlushState
@@ -184,42 +184,42 @@ private:
 	 * @param  Linker   The linker we are arrached to
 	 * @return True if the loop is allowed to continue, or false if we should not flush any more
 	 */
-	bool FlushNext(UMovieSceneEntitySystemLinker* Linker);
+	MOVIESCENE_API bool FlushNext(UMovieSceneEntitySystemLinker* Linker);
 
 	/**
 	 * Set up initial state before any evaluation runs. Only called once regardless of the number of pending updates we have to process
 	 * Primarily used for setting up external 'is evaluating' flags for re-entrancy and async checks.
 	 */
-	bool StartEvaluation(UMovieSceneEntitySystemLinker* Linker);
+	MOVIESCENE_API bool StartEvaluation(UMovieSceneEntitySystemLinker* Linker);
 
 	/** Update sequence instances based on currently queued update requests, or outstanding dissected updates */
-	bool GameThread_UpdateSequenceInstances(UMovieSceneEntitySystemLinker* Linker);
+	MOVIESCENE_API bool GameThread_UpdateSequenceInstances(UMovieSceneEntitySystemLinker* Linker);
 	/** Execute the spawn phase of the entity system graph, if there is anything to do */
-	bool GameThread_SpawnPhase(UMovieSceneEntitySystemLinker* Linker);
+	MOVIESCENE_API bool GameThread_SpawnPhase(UMovieSceneEntitySystemLinker* Linker);
 	/** Execute the instantiation phase of the entity system graph, if there is anything to do */
-	bool GameThread_InstantiationPhase(UMovieSceneEntitySystemLinker* Linker);
+	MOVIESCENE_API bool GameThread_InstantiationPhase(UMovieSceneEntitySystemLinker* Linker);
 	/** Called immediately after instantiation to execute cleanup and bookkeeping tasks. Skipped if instantiation is skipped.*/
-	bool GameThread_PostInstantiation(UMovieSceneEntitySystemLinker* Linker);
+	MOVIESCENE_API bool GameThread_PostInstantiation(UMovieSceneEntitySystemLinker* Linker);
 	/** Main entity-system evaluation phase. Blocks this thread until completion. */
-	bool GameThread_EvaluationPhase(UMovieSceneEntitySystemLinker* Linker);
+	MOVIESCENE_API bool GameThread_EvaluationPhase(UMovieSceneEntitySystemLinker* Linker);
 	/** Finalization phase for triggering external events and other behavior. */
-	void GameThread_EvaluationFinalizationPhase(UMovieSceneEntitySystemLinker* Linker);
+	MOVIESCENE_API void GameThread_EvaluationFinalizationPhase(UMovieSceneEntitySystemLinker* Linker);
 	/** Post-evaluation phase for triggering events. */
-	void GameThread_EventTriggerPhase(UMovieSceneEntitySystemLinker* Linker);
+	MOVIESCENE_API void GameThread_EventTriggerPhase(UMovieSceneEntitySystemLinker* Linker);
 	/** Post-evaluation phase for clean up. */
-	void GameThread_PostEvaluationPhase(UMovieSceneEntitySystemLinker* Linker);
+	MOVIESCENE_API void GameThread_PostEvaluationPhase(UMovieSceneEntitySystemLinker* Linker);
 
 	/**
 	 * Counterpart for StartEvaluation.
 	 * Called only when our UpdateQueue and DissectedUpdates have been fully processed and there is nothing left to do.
 	 */
-	void EndEvaluation(UMovieSceneEntitySystemLinker* Linker);
+	MOVIESCENE_API void EndEvaluation(UMovieSceneEntitySystemLinker* Linker);
 
 	/** Called to set/unset the necessary flags to enter a new flush state and progress to the next */
-	void EnterFlushState(UE::MovieScene::ERunnerFlushState EnteredFlushState);
+	MOVIESCENE_API void EnterFlushState(UE::MovieScene::ERunnerFlushState EnteredFlushState);
 
 	/** Skip the specified flush states if they are currently pending */
-	void SkipFlushState(UE::MovieScene::ERunnerFlushState FlushStateToSkip);
+	MOVIESCENE_API void SkipFlushState(UE::MovieScene::ERunnerFlushState FlushStateToSkip);
 
 private:
 
diff --git a/Engine/Source/Runtime/MovieScene/Public/EntitySystem/MovieSceneEntitySystemTypes.h b/Engine/Source/Runtime/MovieScene/Public/EntitySystem/MovieSceneEntitySystemTypes.h
--- a/Engine/Source/Runtime/MovieScene/Public/EntitySystem/MovieSceneEntitySystemTypes.h
+++ b/Engine/Source/Runtime/MovieScene/Public/EntitySystem/MovieSceneEntitySystemTypes.h
@@ -294,7 +294,7 @@ struct FInterrogationInstance
 };
 
 
-struct MOVIESCENE_API FEntityComponentFilter
+struct FEntityComponentFilter
 {
 
 	void Reset()
@@ -304,9 +304,9 @@ struct MOVIESCENE_API FEntityComponentFilter
 		ComplexMasks.Reset();
 	}
 
-	bool Match(const FComponentMask& Input) const;
+	MOVIESCENE_API bool Match(const FComponentMask& Input) const;
 
-	bool IsValid() const;
+	MOVIESCENE_API bool IsValid() const;
 
 	FEntityComponentFilter& All(const FComponentMask& InComponentMask)
 	{
diff --git a/Engine/Source/Runtime/MovieScene/Public/EntitySystem/MovieSceneEvalTimeSystem.h b/Engine/Source/Runtime/MovieScene/Public/EntitySystem/MovieSceneEvalTimeSystem.h
--- a/Engine/Source/Runtime/MovieScene/Public/EntitySystem/MovieSceneEvalTimeSystem.h
+++ b/Engine/Source/Runtime/MovieScene/Public/EntitySystem/MovieSceneEvalTimeSystem.h
@@ -24,18 +24,18 @@ namespace UE::MovieScene
 	};
 }
 
-UCLASS()
-class MOVIESCENE_API UMovieSceneEvalTimeSystem : public UMovieSceneEntitySystem
+UCLASS(MinimalAPI)
+class UMovieSceneEvalTimeSystem : public UMovieSceneEntitySystem
 {
 public:
 
 	GENERATED_BODY()
 
-	UMovieSceneEvalTimeSystem(const FObjectInitializer& ObjInit);
+	MOVIESCENE_API UMovieSceneEvalTimeSystem(const FObjectInitializer& ObjInit);
 
-	virtual bool IsRelevantImpl(UMovieSceneEntitySystemLinker* InLinker) const override;
-	virtual void OnSchedulePersistentTasks(UE::MovieScene::IEntitySystemScheduler* TaskScheduler) override;
-	virtual void OnRun(FSystemTaskPrerequisites& InPrerequisites, FSystemSubsequentTasks& Subsequents) override;
+	MOVIESCENE_API virtual bool IsRelevantImpl(UMovieSceneEntitySystemLinker* InLinker) const override;
+	MOVIESCENE_API virtual void OnSchedulePersistentTasks(UE::MovieScene::IEntitySystemScheduler* TaskScheduler) override;
+	MOVIESCENE_API virtual void OnRun(FSystemTaskPrerequisites& InPrerequisites, FSystemSubsequentTasks& Subsequents) override;
 
 private:
 
diff --git a/Engine/Source/Runtime/MovieScene/Public/EntitySystem/MovieSceneEvaluationHookSystem.h b/Engine/Source/Runtime/MovieScene/Public/EntitySystem/MovieSceneEvaluationHookSystem.h
--- a/Engine/Source/Runtime/MovieScene/Public/EntitySystem/MovieSceneEvaluationHookSystem.h
+++ b/Engine/Source/Runtime/MovieScene/Public/EntitySystem/MovieSceneEvaluationHookSystem.h
@@ -75,31 +75,31 @@ struct FMovieSceneEvaluationInstanceKey
 };
 
 
-UCLASS()
-class MOVIESCENE_API UMovieSceneEvaluationHookSystem : public UMovieSceneEntitySystem
+UCLASS(MinimalAPI)
+class UMovieSceneEvaluationHookSystem : public UMovieSceneEntitySystem
 {
 public:
 	GENERATED_BODY()
 
-	UMovieSceneEvaluationHookSystem(const FObjectInitializer& ObjInit);
+	MOVIESCENE_API UMovieSceneEvaluationHookSystem(const FObjectInitializer& ObjInit);
 
-	void AddEvent(UE::MovieScene::FInstanceHandle RootInstance, const FMovieSceneEvaluationHookEvent& InEvent);
+	MOVIESCENE_API void AddEvent(UE::MovieScene::FInstanceHandle RootInstance, const FMovieSceneEvaluationHookEvent& InEvent);
 
-	void SortEvents();
+	MOVIESCENE_API void SortEvents();
 
 protected:
 
-	bool HasEvents() const;
-	void TriggerAllEvents();
+	MOVIESCENE_API bool HasEvents() const;
+	MOVIESCENE_API void TriggerAllEvents();
 
 private:
 
-	virtual void OnSchedulePersistentTasks(UE::MovieScene::IEntitySystemScheduler* TaskScheduler) override;
-	virtual void OnRun(FSystemTaskPrerequisites& InPrerequisites, FSystemSubsequentTasks& Subsequents) override;
-	virtual bool IsRelevantImpl(UMovieSceneEntitySystemLinker* InLinker) const override;
+	MOVIESCENE_API virtual void OnSchedulePersistentTasks(UE::MovieScene::IEntitySystemScheduler* TaskScheduler) override;
+	MOVIESCENE_API virtual void OnRun(FSystemTaskPrerequisites& InPrerequisites, FSystemSubsequentTasks& Subsequents) override;
+	MOVIESCENE_API virtual bool IsRelevantImpl(UMovieSceneEntitySystemLinker* InLinker) const override;
 
 	void UpdateHooks();
 
 	UPROPERTY()
 	TMap<FMovieSceneEvaluationInstanceKey, FMovieSceneEvaluationHookEventContainer> PendingEventsByRootInstance;
-};
\ No newline at end of file
+};
diff --git a/Engine/Source/Runtime/MovieScene/Public/EntitySystem/MovieScenePropertyRegistry.h b/Engine/Source/Runtime/MovieScene/Public/EntitySystem/MovieScenePropertyRegistry.h
--- a/Engine/Source/Runtime/MovieScene/Public/EntitySystem/MovieScenePropertyRegistry.h
+++ b/Engine/Source/Runtime/MovieScene/Public/EntitySystem/MovieScenePropertyRegistry.h
@@ -148,7 +148,7 @@ template<typename PropertyTraits> struct TPropertyDefinitionBuilder;
  * Central registry of all property types animatable by sequencer.
  * Once registered, properties cannot be de-registered. This vastly simplifies the lifetime and ID management of the class
  */
-class MOVIESCENE_API FPropertyRegistry
+class FPropertyRegistry
 {
 public:
 
@@ -165,7 +165,7 @@ public:
 	 * @param CustomAccessors A view to an array of custom accessors (as retrieved from ICustomPropertyRegistration::GetAccessors)
 	 * @return An optional variant specifying the resolved property if it resolved successfully
 	 */
-	static TOptional< FResolvedFastProperty > ResolveFastProperty(UObject* Object, const FMovieScenePropertyBinding& PropertyBinding, FCustomAccessorView CustomAccessors);
+	static MOVIESCENE_API TOptional< FResolvedFastProperty > ResolveFastProperty(UObject* Object, const FMovieScenePropertyBinding& PropertyBinding, FCustomAccessorView CustomAccessors);
 
 	/**
 	 * Resolve a property to either a fast ptr offset, or a custom property accessor based on the specified array falling back to a slow instance binding if possible
@@ -175,7 +175,7 @@ public:
 	 * @param CustomAccessors A view to an array of custom accessors (as retrieved from ICustomPropertyRegistration::GetAccessors)
 	 * @return An optional variant specifying the resolved property if it resolved successfully
 	 */
-	static TOptional< FResolvedProperty > ResolveProperty(UObject* Object, const FMovieScenePropertyBinding& PropertyBinding, FCustomAccessorView CustomAccessors);
+	static MOVIESCENE_API TOptional< FResolvedProperty > ResolveProperty(UObject* Object, const FMovieScenePropertyBinding& PropertyBinding, FCustomAccessorView CustomAccessors);
 
 	/**
 	 * Define a new animatable composite property type from its components.
diff --git a/Engine/Source/Runtime/MovieScene/Public/EntitySystem/MovieSceneRootInstantiatorSystem.h b/Engine/Source/Runtime/MovieScene/Public/EntitySystem/MovieSceneRootInstantiatorSystem.h
--- a/Engine/Source/Runtime/MovieScene/Public/EntitySystem/MovieSceneRootInstantiatorSystem.h
+++ b/Engine/Source/Runtime/MovieScene/Public/EntitySystem/MovieSceneRootInstantiatorSystem.h
@@ -15,18 +15,18 @@ class UObject;
 namespace UE { namespace MovieScene { struct FEntityAllocation; } }
 
 
-UCLASS()
-class MOVIESCENE_API UMovieSceneRootInstantiatorSystem : public UMovieSceneEntityInstantiatorSystem
+UCLASS(MinimalAPI)
+class UMovieSceneRootInstantiatorSystem : public UMovieSceneEntityInstantiatorSystem
 {
 	GENERATED_BODY()
 
 public:
 
-	UMovieSceneRootInstantiatorSystem(const FObjectInitializer& ObjInit);
+	MOVIESCENE_API UMovieSceneRootInstantiatorSystem(const FObjectInitializer& ObjInit);
 
 private:
 
-	virtual void OnRun(FSystemTaskPrerequisites& InPrerequisites, FSystemSubsequentTasks& Subsequents) override final;
+	MOVIESCENE_API virtual void OnRun(FSystemTaskPrerequisites& InPrerequisites, FSystemSubsequentTasks& Subsequents) override final;
 
 	void InstantiateAllocation(const UE::MovieScene::FEntityAllocation* ParentAllocation);
-};
\ No newline at end of file
+};
diff --git a/Engine/Source/Runtime/MovieScene/Public/EntitySystem/MovieSceneSequenceInstance.h b/Engine/Source/Runtime/MovieScene/Public/EntitySystem/MovieSceneSequenceInstance.h
--- a/Engine/Source/Runtime/MovieScene/Public/EntitySystem/MovieSceneSequenceInstance.h
+++ b/Engine/Source/Runtime/MovieScene/Public/EntitySystem/MovieSceneSequenceInstance.h
@@ -48,7 +48,7 @@ ENUM_CLASS_FLAGS(ESequenceInstanceUpdateFlags);
  * A sequence instance represents a specific instance of a currently playing sequence, either as a top-level sequence in an IMovieScenePlayer, or as a sub sequence.
  * Any given sequence asset may have any number of instances created for it at any given time depending on how many times it is referenced by playing sequences
  */
-struct MOVIESCENE_API FSequenceInstance
+struct FSequenceInstance
 {
 	/** Ledger that tracks all currently instantiated entities for this instance */
 	FEntityLedger Ledger;
@@ -60,7 +60,7 @@ public:
 	 *
 	 * @param Linker     The linker that owns this sequence instance
 	 */
-	void ConditionalRecompile(UMovieSceneEntitySystemLinker* Linker);
+	MOVIESCENE_API void ConditionalRecompile(UMovieSceneEntitySystemLinker* Linker);
 
 	/**
 	 * Called only for top-level sequence instances before any updates to it with the specified context.
@@ -70,7 +70,7 @@ public:
 	 * @param Context    The overall context that this sequence instance is being evaluated with
 	 * @param OutDissections   An array to populate with dissected time ranges that should be evaluated separately, in order.
 	 */
-	void DissectContext(UMovieSceneEntitySystemLinker* Linker, const FMovieSceneContext& Context, TArray<TRange<FFrameTime>>& OutDissections);
+	MOVIESCENE_API void DissectContext(UMovieSceneEntitySystemLinker* Linker, const FMovieSceneContext& Context, TArray<TRange<FFrameTime>>& OutDissections);
 
 
 	/**
@@ -79,7 +79,7 @@ public:
 	 * @param Linker     The linker that owns this sequence instance
 	 * @param InContext  The context that this sequence instance is to be evaluated with
 	 */
-	void Start(UMovieSceneEntitySystemLinker* Linker, const FMovieSceneContext& InContext);
+	MOVIESCENE_API void Start(UMovieSceneEntitySystemLinker* Linker, const FMovieSceneContext& InContext);
 
 
 	/**
@@ -87,7 +87,7 @@ public:
 	 *
 	 * @param Linker     The linker that owns this sequence instance
 	 */
-	void PreEvaluation(UMovieSceneEntitySystemLinker* Linker);
+	MOVIESCENE_API void PreEvaluation(UMovieSceneEntitySystemLinker* Linker);
 
 
 	/**
@@ -96,7 +96,7 @@ public:
 	 * @param Linker     The linker that owns this sequence instance
 	 * @param InContext  The (potentially dissected) context that this sequence instance is to be evaluated with
 	 */
-	void Update(UMovieSceneEntitySystemLinker* Linker, const FMovieSceneContext& InContext);
+	MOVIESCENE_API void Update(UMovieSceneEntitySystemLinker* Linker, const FMovieSceneContext& InContext);
 
 
 	/**
@@ -105,7 +105,7 @@ public:
 	 * @param Linker     The linker that owns this sequence instance
 	 * @return           Whether the instance can be finished immediately
 	 */
-	bool CanFinishImmediately(UMovieSceneEntitySystemLinker* Linker) const;
+	MOVIESCENE_API bool CanFinishImmediately(UMovieSceneEntitySystemLinker* Linker) const;
 
 
 	/**
@@ -113,19 +113,19 @@ public:
 	 *
 	 * @param Linker     The linker that owns this sequence instance
 	 */
-	void Finish(UMovieSceneEntitySystemLinker* Linker);
+	MOVIESCENE_API void Finish(UMovieSceneEntitySystemLinker* Linker);
 
 	/**
 	 * Called when this sequence instance has been evaluated in order for it to do any clean-up or other post-update work
 	 *
 	 * @param Linker     The linker that owns this sequence instance
 	 */
-	void PostEvaluation(UMovieSceneEntitySystemLinker* Linker);
+	MOVIESCENE_API void PostEvaluation(UMovieSceneEntitySystemLinker* Linker);
 
 	/**
 	 * Called to run legacy track templates
 	 */
-	void RunLegacyTrackTemplates();
+	MOVIESCENE_API void RunLegacyTrackTemplates();
 
 public:
 
@@ -134,7 +134,7 @@ public:
 	 *
 	 * @return A pointer to this instance's player
 	 */
-	IMovieScenePlayer* GetPlayer() const;
+	MOVIESCENE_API IMovieScenePlayer* GetPlayer() const;
 
 	/**
 	 * Retrieve the IMovieScenePlayer's unique index
@@ -247,17 +247,17 @@ public:
 	/**
 	 * Attempt to locate a sub instance based on its sequence ID
 	 */
-	FInstanceHandle FindSubInstance(FMovieSceneSequenceID SequenceID) const;
+	MOVIESCENE_API FInstanceHandle FindSubInstance(FMovieSceneSequenceID SequenceID) const;
 
 	/**
 	 * Attempt to locate an entity given its owner and ID
 	 */
-	FMovieSceneEntityID FindEntity(UObject* Owner, uint32 EntityID) const;
+	MOVIESCENE_API FMovieSceneEntityID FindEntity(UObject* Owner, uint32 EntityID) const;
 
 	/**
 	 * Attempt to locate all entities given their owner
 	 */
-	void FindEntities(UObject* Owner, TArray<FMovieSceneEntityID>& OutEntityIDs) const;
+	MOVIESCENE_API void FindEntities(UObject* Owner, TArray<FMovieSceneEntityID>& OutEntityIDs) const;
 
 	/**
 	 * Retrieve the legacy evaluator for this sequence, if it is available (may return nullptr)
@@ -270,7 +270,7 @@ public:
 	/**
 	 * Retrieve a path for this sequence instance back to the root
 	 */
-	FSubSequencePath GetSubSequencePath() const;
+	MOVIESCENE_API FSubSequencePath GetSubSequencePath() const;
 
 public:
 
@@ -295,40 +295,40 @@ public:
 	/**
 	 * Invalidate any cached data that may be being used for evaluation due to a change in the source asset data
 	 */
-	void InvalidateCachedData(UMovieSceneEntitySystemLinker* Linker);
+	MOVIESCENE_API void InvalidateCachedData(UMovieSceneEntitySystemLinker* Linker);
 
 	/**
 	 * Destroy this sequence instance immediately - Finish must previously have been called
 	 */
-	void DestroyImmediately(UMovieSceneEntitySystemLinker* Linker);
+	MOVIESCENE_API void DestroyImmediately(UMovieSceneEntitySystemLinker* Linker);
 
 	/**
 	 * Called to override the simulated root sequence ID for this instance. Only valid for IsRootSequence() instances.
 	 */
-	void OverrideRootSequence(UMovieSceneEntitySystemLinker* Linker, FMovieSceneSequenceID NewRootSequenceID);
+	MOVIESCENE_API void OverrideRootSequence(UMovieSceneEntitySystemLinker* Linker, FMovieSceneSequenceID NewRootSequenceID);
 
 public:
 
 	/** Constructor for top level sequences */
-	explicit FSequenceInstance(UMovieSceneEntitySystemLinker* Linker, IMovieScenePlayer* Player, FRootInstanceHandle ThisInstanceHandle);
+	MOVIESCENE_API explicit FSequenceInstance(UMovieSceneEntitySystemLinker* Linker, IMovieScenePlayer* Player, FRootInstanceHandle ThisInstanceHandle);
 
 	/** Constructor for sub sequences */
-	explicit FSequenceInstance(UMovieSceneEntitySystemLinker* Linker, IMovieScenePlayer* Player, FInstanceHandle ThisInstanceHandle, FInstanceHandle InParentInstanceHandle, FRootInstanceHandle RootInstanceHandle, FMovieSceneSequenceID InSequenceID);
+	MOVIESCENE_API explicit FSequenceInstance(UMovieSceneEntitySystemLinker* Linker, IMovieScenePlayer* Player, FInstanceHandle ThisInstanceHandle, FInstanceHandle InParentInstanceHandle, FRootInstanceHandle RootInstanceHandle, FMovieSceneSequenceID InSequenceID);
 
 	/** Destructor */
-	~FSequenceInstance();
+	MOVIESCENE_API ~FSequenceInstance();
 
 	/** Sequence instances are not copyable */
 	FSequenceInstance(const FSequenceInstance&) = delete;
 	FSequenceInstance& operator=(const FSequenceInstance&) = delete;
 
 	/** Move constructors implemented in cpp to avoid includes for TUniquePtr */
-	FSequenceInstance(FSequenceInstance&&);
-	FSequenceInstance& operator=(FSequenceInstance&&);
+	MOVIESCENE_API FSequenceInstance(FSequenceInstance&&);
+	MOVIESCENE_API FSequenceInstance& operator=(FSequenceInstance&&);
 
 private:
 
-	void InitializeLegacyEvaluator(UMovieSceneEntitySystemLinker* Linker);
+	MOVIESCENE_API void InitializeLegacyEvaluator(UMovieSceneEntitySystemLinker* Linker);
 
 private:
 
diff --git a/Engine/Source/Runtime/MovieScene/Public/EntitySystem/MovieSceneSystemTaskDependencies.h b/Engine/Source/Runtime/MovieScene/Public/EntitySystem/MovieSceneSystemTaskDependencies.h
--- a/Engine/Source/Runtime/MovieScene/Public/EntitySystem/MovieSceneSystemTaskDependencies.h
+++ b/Engine/Source/Runtime/MovieScene/Public/EntitySystem/MovieSceneSystemTaskDependencies.h
@@ -94,21 +94,21 @@ private:
 
 
 
-struct MOVIESCENE_API FSystemSubsequentTasks
+struct FSystemSubsequentTasks
 {
 	using FComponentTypeID = UE::MovieScene::FComponentTypeID;
 
-	void AddRootTask(FGraphEventRef RootTask);
+	MOVIESCENE_API void AddRootTask(FGraphEventRef RootTask);
 
-	void AddComponentTask(FComponentTypeID ComponentType, FGraphEventRef ComponentTask);
+	MOVIESCENE_API void AddComponentTask(FComponentTypeID ComponentType, FGraphEventRef ComponentTask);
 
 private:
 
 	friend FMovieSceneEntitySystemGraph;
 
-	FSystemSubsequentTasks(FMovieSceneEntitySystemGraph* InGraph, FGraphEventArray* InAllTasks, EEntityThreadingModel InThreadingModel);
+	MOVIESCENE_API FSystemSubsequentTasks(FMovieSceneEntitySystemGraph* InGraph, FGraphEventArray* InAllTasks, EEntityThreadingModel InThreadingModel);
 
-	void ResetNode(uint16 InNodeID);
+	MOVIESCENE_API void ResetNode(uint16 InNodeID);
 
 	TSharedPtr<UE::MovieScene::FSystemTaskPrerequisites> Subsequents;
 	FMovieSceneEntitySystemGraph* Graph;
@@ -119,4 +119,4 @@ private:
 
 
 } // namespace MovieScene
-} // namespace UE
\ No newline at end of file
+} // namespace UE
diff --git a/Engine/Source/Runtime/MovieScene/Public/EntitySystem/TrackInstance/MovieSceneTrackInstance.h b/Engine/Source/Runtime/MovieScene/Public/EntitySystem/TrackInstance/MovieSceneTrackInstance.h
--- a/Engine/Source/Runtime/MovieScene/Public/EntitySystem/TrackInstance/MovieSceneTrackInstance.h
+++ b/Engine/Source/Runtime/MovieScene/Public/EntitySystem/TrackInstance/MovieSceneTrackInstance.h
@@ -101,8 +101,8 @@ struct TStructOpsTypeTraits<FMovieSceneTrackInstanceInput> : public TStructOpsTy
 /**
  * Base class for all track instances. Can also be used for root tracks where AnimatedObject will be nullptr
  */
-UCLASS(transient)
-class MOVIESCENE_API UMovieSceneTrackInstance : public UObject
+UCLASS(transient, MinimalAPI)
+class UMovieSceneTrackInstance : public UObject
 {
 public:
 	GENERATED_BODY()
@@ -113,25 +113,25 @@ public:
 	 * @param InAnimatedObject     (optional) The object that we should animate
 	 * @param InLinker             The entity system linker that owns this animator. Used for retrieving contexts and player information.
 	 */
-	void Initialize(UObject* InAnimatedObject, UMovieSceneEntitySystemLinker* InLinker);
+	MOVIESCENE_API void Initialize(UObject* InAnimatedObject, UMovieSceneEntitySystemLinker* InLinker);
 
 
 	/**
 	 * Run this animator for the current frame
 	 */
-	void Animate();
+	MOVIESCENE_API void Animate();
 
 
 	/**
 	 * Update the inputs that contribute to this animator
 	 */
-	void UpdateInputs(TArray<FMovieSceneTrackInstanceInput>&& InNewInputs);
+	MOVIESCENE_API void UpdateInputs(TArray<FMovieSceneTrackInstanceInput>&& InNewInputs);
 
 
 	/**
 	 * Destroy this animator. The animator may still have inputs when it is destroyed.
 	 */
-	void Destroy();
+	MOVIESCENE_API void Destroy();
 
 
 public:
@@ -163,7 +163,7 @@ public:
 		return Inputs;
 	}
 
-	virtual UWorld* GetWorld() const override;
+	MOVIESCENE_API virtual UWorld* GetWorld() const override;
 #if WITH_EDITOR
 	virtual bool ImplementsGetWorld() const override { return true; }
 #endif
diff --git a/Engine/Source/Runtime/MovieScene/Public/EntitySystem/TrackInstance/MovieSceneTrackInstanceSystem.h b/Engine/Source/Runtime/MovieScene/Public/EntitySystem/TrackInstance/MovieSceneTrackInstanceSystem.h
--- a/Engine/Source/Runtime/MovieScene/Public/EntitySystem/TrackInstance/MovieSceneTrackInstanceSystem.h
+++ b/Engine/Source/Runtime/MovieScene/Public/EntitySystem/TrackInstance/MovieSceneTrackInstanceSystem.h
@@ -36,35 +36,35 @@ struct FMovieSceneTrackInstanceEntry
 };
 
 
-UCLASS()
-class MOVIESCENE_API UMovieSceneTrackInstanceInstantiator : public UMovieSceneEntityInstantiatorSystem
+UCLASS(MinimalAPI)
+class UMovieSceneTrackInstanceInstantiator : public UMovieSceneEntityInstantiatorSystem
 {
 	GENERATED_BODY()
 
 public:
 
-	UMovieSceneTrackInstanceInstantiator(const FObjectInitializer& ObjInit);
+	MOVIESCENE_API UMovieSceneTrackInstanceInstantiator(const FObjectInitializer& ObjInit);
 
-	int32 MakeOutput(UObject* BoundObject, UClass* TrackInstanceClass);
+	MOVIESCENE_API int32 MakeOutput(UObject* BoundObject, UClass* TrackInstanceClass);
 
-	int32 FindOutput(UObject* BoundObject, UClass* TrackInstanceClass) const;
+	MOVIESCENE_API int32 FindOutput(UObject* BoundObject, UClass* TrackInstanceClass) const;
 
 	const TSparseArray<FMovieSceneTrackInstanceEntry>& GetTrackInstances() const
 	{
 		return TrackInstances;
 	}
 
-	static void AddReferencedObjects(UObject* InThis, FReferenceCollector& Collector);
+	static MOVIESCENE_API void AddReferencedObjects(UObject* InThis, FReferenceCollector& Collector);
 
 private:
 
-	virtual void OnLink() override final;
-	virtual void OnUnlink() override final;
-	virtual void OnRun(FSystemTaskPrerequisites& InPrerequisites, FSystemSubsequentTasks& Subsequents) override final;
+	MOVIESCENE_API virtual void OnLink() override final;
+	MOVIESCENE_API virtual void OnUnlink() override final;
+	MOVIESCENE_API virtual void OnRun(FSystemTaskPrerequisites& InPrerequisites, FSystemSubsequentTasks& Subsequents) override final;
 
-	virtual void OnTagGarbage() override;
+	MOVIESCENE_API virtual void OnTagGarbage() override;
 
-	virtual void Serialize(FArchive& Ar) override;
+	MOVIESCENE_API virtual void Serialize(FArchive& Ar) override;
 
 	TSparseArray<FMovieSceneTrackInstanceEntry> TrackInstances;
 	TMultiMap<TObjectPtr<UObject>, int32> BoundObjectToInstances;
@@ -75,8 +75,8 @@ private:
 };
 
 
-UCLASS()
-class MOVIESCENE_API UMovieSceneTrackInstanceSystem : public UMovieSceneEntitySystem
+UCLASS(MinimalAPI)
+class UMovieSceneTrackInstanceSystem : public UMovieSceneEntitySystem
 {
 	GENERATED_BODY()
 
@@ -84,9 +84,9 @@ class MOVIESCENE_API UMovieSceneTrackInstanceSystem : public UMovieSceneEntitySy
 
 private:
 
-	virtual void OnLink() override final;
-	virtual void OnSchedulePersistentTasks(UE::MovieScene::IEntitySystemScheduler* TaskScheduler) override final;
-	virtual void OnRun(FSystemTaskPrerequisites& InPrerequisites, FSystemSubsequentTasks& Subsequents) override final;
+	MOVIESCENE_API virtual void OnLink() override final;
+	MOVIESCENE_API virtual void OnSchedulePersistentTasks(UE::MovieScene::IEntitySystemScheduler* TaskScheduler) override final;
+	MOVIESCENE_API virtual void OnRun(FSystemTaskPrerequisites& InPrerequisites, FSystemSubsequentTasks& Subsequents) override final;
 
 	void EvaluateAllInstances();
 
diff --git a/Engine/Source/Runtime/MovieScene/Public/Evaluation/Blending/MovieSceneInitialValueStore.h b/Engine/Source/Runtime/MovieScene/Public/Evaluation/Blending/MovieSceneInitialValueStore.h
--- a/Engine/Source/Runtime/MovieScene/Public/Evaluation/Blending/MovieSceneInitialValueStore.h
+++ b/Engine/Source/Runtime/MovieScene/Public/Evaluation/Blending/MovieSceneInitialValueStore.h
@@ -16,12 +16,12 @@ template<typename DataType> struct TBlendableTokenStack;
 template<typename DataType> struct TMovieSceneBlendingActuator;
 
 /** Pre animated token producer that reverts the object's initial value from the actuator when its state is restored */
-struct MOVIESCENE_API FMovieSceneRemoveInitialValueTokenProducer : IMovieScenePreAnimatedTokenProducer
+struct FMovieSceneRemoveInitialValueTokenProducer : IMovieScenePreAnimatedTokenProducer
 {
 	/** Construction from the object whose initial value to remove, and the actuator to remove it from */
-	FMovieSceneRemoveInitialValueTokenProducer(TWeakPtr<IMovieSceneBlendingActuator> InWeakActuator);
+	MOVIESCENE_API FMovieSceneRemoveInitialValueTokenProducer(TWeakPtr<IMovieSceneBlendingActuator> InWeakActuator);
 
-	virtual IMovieScenePreAnimatedTokenPtr CacheExistingState(UObject& InObject) const override;
+	MOVIESCENE_API virtual IMovieScenePreAnimatedTokenPtr CacheExistingState(UObject& InObject) const override;
 
 private:
 	/** The actuator to remove the initial value from */
@@ -29,12 +29,12 @@ private:
 };
 
 /** Pre animated token producer that reverts a global initial value from the actuator when its state is restored */
-struct MOVIESCENE_API FMovieSceneRemoveInitialGlobalValueTokenProducer : IMovieScenePreAnimatedGlobalTokenProducer
+struct FMovieSceneRemoveInitialGlobalValueTokenProducer : IMovieScenePreAnimatedGlobalTokenProducer
 {
 	/** Construction from the object whose initial value to remove, and the actuator to remove it from */
-	FMovieSceneRemoveInitialGlobalValueTokenProducer(TWeakPtr<IMovieSceneBlendingActuator> InWeakActuator);
+	MOVIESCENE_API FMovieSceneRemoveInitialGlobalValueTokenProducer(TWeakPtr<IMovieSceneBlendingActuator> InWeakActuator);
 
-	virtual IMovieScenePreAnimatedGlobalTokenPtr CacheExistingState() const override;
+	MOVIESCENE_API virtual IMovieScenePreAnimatedGlobalTokenPtr CacheExistingState() const override;
 
 private:
 	/** The actuator to remove the initial value from */
diff --git a/Engine/Source/Runtime/MovieScene/Public/Evaluation/IMovieSceneCustomClockSource.h b/Engine/Source/Runtime/MovieScene/Public/Evaluation/IMovieSceneCustomClockSource.h
--- a/Engine/Source/Runtime/MovieScene/Public/Evaluation/IMovieSceneCustomClockSource.h
+++ b/Engine/Source/Runtime/MovieScene/Public/Evaluation/IMovieSceneCustomClockSource.h
@@ -16,8 +16,8 @@ class UObject;
 struct FFrame;
 struct FQualifiedFrameTime;
 
-UINTERFACE()
-class MOVIESCENE_API UMovieSceneCustomClockSource : public UInterface
+UINTERFACE(MinimalAPI)
+class UMovieSceneCustomClockSource : public UInterface
 {
 	GENERATED_BODY()
 };
@@ -44,16 +44,16 @@ public:
 	virtual FFrameTime OnRequestCurrentTime(const FQualifiedFrameTime& InCurrentTime, float InPlayRate) { return 0; }
 };
 
-struct MOVIESCENE_API FMovieSceneTimeController_Custom : FMovieSceneTimeController
+struct FMovieSceneTimeController_Custom : FMovieSceneTimeController
 {
-	explicit FMovieSceneTimeController_Custom(const FSoftObjectPath& InObjectPath, TWeakObjectPtr<> PlaybackContext);
+	MOVIESCENE_API explicit FMovieSceneTimeController_Custom(const FSoftObjectPath& InObjectPath, TWeakObjectPtr<> PlaybackContext);
 
 private:
 
-	virtual void OnTick(float DeltaSeconds, float InPlayRate) override final;
-	virtual void OnStartPlaying(const FQualifiedFrameTime& InStartTime) override final;
-	virtual void OnStopPlaying(const FQualifiedFrameTime& InStopTime) override final;
-	virtual FFrameTime OnRequestCurrentTime(const FQualifiedFrameTime& InCurrentTime, float InPlayRate) override final;
+	MOVIESCENE_API virtual void OnTick(float DeltaSeconds, float InPlayRate) override final;
+	MOVIESCENE_API virtual void OnStartPlaying(const FQualifiedFrameTime& InStartTime) override final;
+	MOVIESCENE_API virtual void OnStopPlaying(const FQualifiedFrameTime& InStopTime) override final;
+	MOVIESCENE_API virtual FFrameTime OnRequestCurrentTime(const FQualifiedFrameTime& InCurrentTime, float InPlayRate) override final;
 
 private:
 
@@ -66,4 +66,4 @@ private:
 	IMovieSceneCustomClockSource* InterfacePtr;
 
 	FSoftObjectPath ObjectPath;
-};
\ No newline at end of file
+};
diff --git a/Engine/Source/Runtime/MovieScene/Public/Evaluation/IMovieSceneEvaluationHook.h b/Engine/Source/Runtime/MovieScene/Public/Evaluation/IMovieSceneEvaluationHook.h
--- a/Engine/Source/Runtime/MovieScene/Public/Evaluation/IMovieSceneEvaluationHook.h
+++ b/Engine/Source/Runtime/MovieScene/Public/Evaluation/IMovieSceneEvaluationHook.h
@@ -45,8 +45,8 @@ struct FEvaluationHookParams
 } // namespace UE
 
 
-UINTERFACE()
-class MOVIESCENE_API UMovieSceneEvaluationHook : public UInterface
+UINTERFACE(MinimalAPI)
+class UMovieSceneEvaluationHook : public UInterface
 {
 public:
 	GENERATED_BODY()
@@ -56,7 +56,7 @@ public:
 /**
  * All evaluation hooks are executed at the end of the frame (at a time when re-entrancy is permitted), and cannot have any component dependencies
  */
-class MOVIESCENE_API IMovieSceneEvaluationHook
+class IMovieSceneEvaluationHook
 {
 public:
 
diff --git a/Engine/Source/Runtime/MovieScene/Public/Evaluation/MovieSceneEvaluationField.h b/Engine/Source/Runtime/MovieScene/Public/Evaluation/MovieSceneEvaluationField.h
--- a/Engine/Source/Runtime/MovieScene/Public/Evaluation/MovieSceneEvaluationField.h
+++ b/Engine/Source/Runtime/MovieScene/Public/Evaluation/MovieSceneEvaluationField.h
@@ -209,7 +209,7 @@ template<> struct TStructOpsTypeTraits<FMovieSceneEvaluationFieldEntityTree> : p
  * maintained along with guaranteeing no redundant entries exist.
  */
 USTRUCT()
-struct MOVIESCENE_API FMovieSceneEntityComponentField
+struct FMovieSceneEntityComponentField
 {
 	GENERATED_BODY()
 
@@ -255,7 +255,7 @@ struct MOVIESCENE_API FMovieSceneEntityComponentField
 	 * @param OutRange    Will receive the hull of the range that was intersected for which the resulting OutEntities remains constant
 	 * @param OutEntities A set that will be populated with all the entities that exist at the specified time
 	 */
-	void QueryPersistentEntities(FFrameNumber QueryTime, TRange<FFrameNumber>& OutRange, FMovieSceneEvaluationFieldEntitySet& OutEntities) const;
+	MOVIESCENE_API void QueryPersistentEntities(FFrameNumber QueryTime, TRange<FFrameNumber>& OutRange, FMovieSceneEvaluationFieldEntitySet& OutEntities) const;
 
 	/**
 	 * Query the persistent entities for any given time within a sequence.
@@ -265,12 +265,12 @@ struct MOVIESCENE_API FMovieSceneEntityComponentField
 	 * @param QueryCallback A handler for dealing with the resulting entities
 	 * @param OutRange      Will receive the hull of the range that was intersected for which the resulting OutEntities remains constant
 	 */
-	void QueryPersistentEntities(FFrameNumber QueryTime, TFunctionRef<bool(const FMovieSceneEvaluationFieldEntityQuery&)> QueryCallback, TRange<FFrameNumber>& OutRange) const;
+	MOVIESCENE_API void QueryPersistentEntities(FFrameNumber QueryTime, TFunctionRef<bool(const FMovieSceneEvaluationFieldEntityQuery&)> QueryCallback, TRange<FFrameNumber>& OutRange) const;
 
 	/**
 	 * Check whether this field contains any one-shot entities
 	 */
-	bool HasAnyOneShotEntities() const;
+	MOVIESCENE_API bool HasAnyOneShotEntities() const;
 
 	/**
 	 * Query the one-shot entities that overlap with the specified query range.
@@ -279,7 +279,7 @@ struct MOVIESCENE_API FMovieSceneEntityComponentField
 	 * @param QueryRange  The ranger over which to query the field (in the TickResolution of the sequence this was generated from)
 	 * @param OutEntities A set that will be populated with all the entities that overlapped at all with the specified range
 	 */
-	void QueryOneShotEntities(const TRange<FFrameNumber>& QueryRange, FMovieSceneEvaluationFieldEntitySet& OutEntityIndices) const;
+	MOVIESCENE_API void QueryOneShotEntities(const TRange<FFrameNumber>& QueryRange, FMovieSceneEvaluationFieldEntitySet& OutEntityIndices) const;
 
 private:
 
@@ -311,29 +311,29 @@ private:
  * Builder class used for populating an FMovieSceneEntityComponentField with data.
  * Ensures that null or redundant entities or meta-data are not added to the field, and that all indices are valid and correct.
  */
-struct MOVIESCENE_API FMovieSceneEntityComponentFieldBuilder
+struct FMovieSceneEntityComponentFieldBuilder
 {
 	static constexpr uint32 InvalidEntityID = ~0u;
 
 	/**
 	 * Construction from a field to populate
 	 */
-	FMovieSceneEntityComponentFieldBuilder(FMovieSceneEntityComponentField* InField);
+	MOVIESCENE_API FMovieSceneEntityComponentFieldBuilder(FMovieSceneEntityComponentField* InField);
 
 	/**
 	 * Destructor that cleans up redundant data if necessary
 	 */
-	~FMovieSceneEntityComponentFieldBuilder();
+	MOVIESCENE_API ~FMovieSceneEntityComponentFieldBuilder();
 
 	/**
 	 * Access the shared meta-data for all the entities created by this builder.
 	 */
-	FMovieSceneEvaluationFieldSharedEntityMetaData& GetSharedMetaData();
+	MOVIESCENE_API FMovieSceneEvaluationFieldSharedEntityMetaData& GetSharedMetaData();
 
 	/**
 	 * Access the index of the shared meta-data of this builder.
 	 */
-	int32 GetSharedMetaDataIndex() const;
+	MOVIESCENE_API int32 GetSharedMetaDataIndex() const;
 
 	/**
 	 * Add meta-data to this tree returning its index within this builder
@@ -341,7 +341,7 @@ struct MOVIESCENE_API FMovieSceneEntityComponentFieldBuilder
 	 * @param InMetaData     The meta-data to add.
 	 * @return A unique index for this meta-data within this builder, or INDEX_NONE if the meta-data is redundant
 	 */
-	int32 AddMetaData(const FMovieSceneEvaluationFieldEntityMetaData& InMetaData);
+	MOVIESCENE_API int32 AddMetaData(const FMovieSceneEvaluationFieldEntityMetaData& InMetaData);
 
 	/**
 	 * Retrieve an index for the entity that is identified by the specified owner and ID
@@ -350,7 +350,7 @@ struct MOVIESCENE_API FMovieSceneEntityComponentFieldBuilder
 	 * @param EntityID       (Optional) An identifier used to identify the entity inside IMovieSceneEntityProvider::ImportEntityImpl. Could be an index within an array or a set of flags.
 	 * @return An index into this builder used to uniquely identify this entity.
 	 */
-	int32 FindOrAddEntity(UObject* EntityOwner, uint32 EntityID = 0);
+	MOVIESCENE_API int32 FindOrAddEntity(UObject* EntityOwner, uint32 EntityID = 0);
 
 	/**
 	 * Add a persistent entity to the field for a given range. Equivalent to AddPersistentEntity(Range, FindOrAddEntity(EntityOwner, EntityID)).
@@ -361,7 +361,7 @@ struct MOVIESCENE_API FMovieSceneEntityComponentFieldBuilder
 	 * @param EntityID       (Optional) An identifier used to identify the entity inside IMovieSceneEntityProvider::ImportEntityImpl. Could be an index within an array or a set of flags.
 	 * @param MetaDataIndex  (Optional) Meta-data to use for this entitiy within this range. See AddMetaData.
 	 */
-	void AddPersistentEntity(const TRange<FFrameNumber>& Range, UObject* EntityOwner, uint32 EntityID = 0, int32 InMetaDataIndex = INDEX_NONE);
+	MOVIESCENE_API void AddPersistentEntity(const TRange<FFrameNumber>& Range, UObject* EntityOwner, uint32 EntityID = 0, int32 InMetaDataIndex = INDEX_NONE);
 
 	/**
 	 * Add a persistent entity to the field for a given range
@@ -371,7 +371,7 @@ struct MOVIESCENE_API FMovieSceneEntityComponentFieldBuilder
 	 * @param LocalIndex     The index to the entity retrieved from FindOrAddEntity.
 	 * @param MetaDataIndex  (Optional) Meta-data to use for this entitiy within this range. See AddMetaData.
 	 */
-	void AddPersistentEntity(const TRange<FFrameNumber>& Range, int32 LocalIndex, int32 InMetaDataIndex = INDEX_NONE);
+	MOVIESCENE_API void AddPersistentEntity(const TRange<FFrameNumber>& Range, int32 LocalIndex, int32 InMetaDataIndex = INDEX_NONE);
 
 	/**
 	 * Add a one-shot entity to the field for a given range. Equivalent to AddOneShotEntity(Range, FindOrAddEntity(EntityOwner, EntityID)).
@@ -382,7 +382,7 @@ struct MOVIESCENE_API FMovieSceneEntityComponentFieldBuilder
 	 * @param EntityID       (Optional) An identifier used to identify the entity inside IMovieSceneEntityProvider::ImportEntityImpl. Could be an index within an array or a set of flags.
 	 * @param MetaDataIndex  (Optional) Meta-data to use for this entitiy within this range. See AddMetaData.
 	 */
-	void AddOneShotEntity(const TRange<FFrameNumber>& OneShotRange, UObject* EntityOwner, uint32 EntityID = 0, int32 InMetaDataIndex = INDEX_NONE);
+	MOVIESCENE_API void AddOneShotEntity(const TRange<FFrameNumber>& OneShotRange, UObject* EntityOwner, uint32 EntityID = 0, int32 InMetaDataIndex = INDEX_NONE);
 
 	/**
 	 * Add a one-shot entity to the field for a given range
@@ -392,15 +392,15 @@ struct MOVIESCENE_API FMovieSceneEntityComponentFieldBuilder
 	 * @param LocalIndex     The index to the entity retrieved from FindOrAddEntity.
 	 * @param MetaDataIndex  (Optional) Meta-data to use for this entitiy within this range. See AddMetaData.
 	 */
-	void AddOneShotEntity(const TRange<FFrameNumber>& OneShotRange, int32 LocalIndex, int32 InMetaDataIndex = INDEX_NONE);
+	MOVIESCENE_API void AddOneShotEntity(const TRange<FFrameNumber>& OneShotRange, int32 LocalIndex, int32 InMetaDataIndex = INDEX_NONE);
 
 private:
 
 	/** Convert a user-facing local index into KeyToFieldIndex, into an index within FMovieSceneEntityComponentField::Entities */
-	int32 LocalEntityIndexToFieldIndex(int32 LocalIndex);
+	MOVIESCENE_API int32 LocalEntityIndexToFieldIndex(int32 LocalIndex);
 
 	/** Convert a user-facing local meta-data index into MetaDataToFieldIndex, into an index within FMovieSceneEntityComponentField::EntityMetaData */
-	int32 LocalMetaDataIndexToFieldIndex(int32 LocalIndex);
+	MOVIESCENE_API int32 LocalMetaDataIndexToFieldIndex(int32 LocalIndex);
 
 	/** Array of entity keys and their field index within FMovieSceneEntityComponentField::Entities */
 	struct FKeyToIndex
diff --git a/Engine/Source/Runtime/MovieScene/Public/Evaluation/MovieSceneEvaluationTree.h b/Engine/Source/Runtime/MovieScene/Public/Evaluation/MovieSceneEvaluationTree.h
--- a/Engine/Source/Runtime/MovieScene/Public/Evaluation/MovieSceneEvaluationTree.h
+++ b/Engine/Source/Runtime/MovieScene/Public/Evaluation/MovieSceneEvaluationTree.h
@@ -325,7 +325,7 @@ struct FMovieSceneEvaluationTreeNode
  *	----------------------------------------------------------------------------------------------------------------------------------------
  *		Unique Ranges		[ 		3			| 0,2,3 |  0,2,3,5   | 0,2,3 |		1,2		|		1,2,4		|			4				]
  */
-struct MOVIESCENE_API FMovieSceneEvaluationTree
+struct FMovieSceneEvaluationTree
 {
 	FMovieSceneEvaluationTree()
 		: RootNode(FMovieSceneEvaluationTreeNode(TRange<FFrameNumber>::All(), FMovieSceneEvaluationTreeNodeHandle::Invalid()))
@@ -358,7 +358,7 @@ struct MOVIESCENE_API FMovieSceneEvaluationTree
 	 * @param Time The time at which we should start iterating
 	 * @return A bi-directional iterator that is set to the unique range that overlaps the current time
 	 */
-	FMovieSceneEvaluationTreeRangeIterator IterateFromTime(FFrameNumber Time) const;
+	MOVIESCENE_API FMovieSceneEvaluationTreeRangeIterator IterateFromTime(FFrameNumber Time) const;
 
 	/**
 	 * Start iterating this tree from the specified lower boundary
@@ -366,7 +366,7 @@ struct MOVIESCENE_API FMovieSceneEvaluationTree
 	 * @param InStartingLowerBound The lowerbound at which we should start iterating
 	 * @return A bi-directional iterator that is set to the unique range that overlaps the specified lowerbound
 	 */
-	FMovieSceneEvaluationTreeRangeIterator IterateFromLowerBound(TRangeBound<FFrameNumber> InStartingLowerBound) const;
+	MOVIESCENE_API FMovieSceneEvaluationTreeRangeIterator IterateFromLowerBound(TRangeBound<FFrameNumber> InStartingLowerBound) const;
 
 	/**
 	 * Access this tree's root node (infinite range)
@@ -404,7 +404,7 @@ struct MOVIESCENE_API FMovieSceneEvaluationTree
 	 * @param InNode		The node to get children for
 	 * @return Array view of the node's children
 	 */
-	TArrayView<const FMovieSceneEvaluationTreeNode> GetChildren(const FMovieSceneEvaluationTreeNode& InNode) const;
+	MOVIESCENE_API TArrayView<const FMovieSceneEvaluationTreeNode> GetChildren(const FMovieSceneEvaluationTreeNode& InNode) const;
 
 	/**
 	 * Get the children associated with the specified node
@@ -412,7 +412,7 @@ struct MOVIESCENE_API FMovieSceneEvaluationTree
 	 * @param InNode		The node to get children for
 	 * @return Array view of the node's children
 	 */
-	TArrayView< FMovieSceneEvaluationTreeNode> GetChildren(const FMovieSceneEvaluationTreeNode& InNode);
+	MOVIESCENE_API TArrayView< FMovieSceneEvaluationTreeNode> GetChildren(const FMovieSceneEvaluationTreeNode& InNode);
 
 	/**
 	 * Check whether the specified handle corresponds to a node within this tree
@@ -436,7 +436,7 @@ struct MOVIESCENE_API FMovieSceneEvaluationTree
 	 *
 	 * @param InTimeRange 		The time range to add to the tree
 	 */
-	void AddTimeRange(TRange<FFrameNumber> InTimeRange);
+	MOVIESCENE_API void AddTimeRange(TRange<FFrameNumber> InTimeRange);
 
 	/**
 	 * Serialize this evaluation tree
@@ -474,7 +474,7 @@ protected:
 	 * @param InParent 			The current parent node
 	 * @param InPredicate 		(Optional) Predicate that determines whether the time should be added for a particular point
 	 */
-	void AddTimeRange(TRange<FFrameNumber> InTimeRange, const IMovieSceneEvaluationTreeNodeOperator& InOperator, FMovieSceneEvaluationTreeNodeHandle InParent, const TFunctionRef<bool(FMovieSceneEvaluationTreeNodeHandle)>* Predicate);
+	MOVIESCENE_API void AddTimeRange(TRange<FFrameNumber> InTimeRange, const IMovieSceneEvaluationTreeNodeOperator& InOperator, FMovieSceneEvaluationTreeNodeHandle InParent, const TFunctionRef<bool(FMovieSceneEvaluationTreeNodeHandle)>* Predicate);
 
 	/**
 	 * Helper function that creates a new child for the specified parent node
@@ -484,7 +484,7 @@ protected:
 	 * @param InsertIndex		The index at which to insert the new child (children must be sorted)
 	 * @param InParent 			The node to add the child to
 	 */
-	void InsertNewChild(TRange<FFrameNumber> InEffectiveRange, const IMovieSceneEvaluationTreeNodeOperator& InOperator, int32 InsertIndex, FMovieSceneEvaluationTreeNodeHandle InParent);
+	MOVIESCENE_API void InsertNewChild(TRange<FFrameNumber> InEffectiveRange, const IMovieSceneEvaluationTreeNodeOperator& InOperator, int32 InsertIndex, FMovieSceneEvaluationTreeNodeHandle InParent);
 
 	/** This tree's root node */
 	FMovieSceneEvaluationTreeNode RootNode;
@@ -496,13 +496,13 @@ protected:
 /**
  * Type that iterates contiguous range/data combinations sequentially (including empty space between time ranges)
  */
-struct MOVIESCENE_API FMovieSceneEvaluationTreeRangeIterator
+struct FMovieSceneEvaluationTreeRangeIterator
 {
 	/** Iterate the tree from -infinity */
-	FMovieSceneEvaluationTreeRangeIterator(const FMovieSceneEvaluationTree& InTree);
+	MOVIESCENE_API FMovieSceneEvaluationTreeRangeIterator(const FMovieSceneEvaluationTree& InTree);
 
 	/** Iterate the tree ranges starting at the range that overlaps the specified lower bound */
-	FMovieSceneEvaluationTreeRangeIterator(const FMovieSceneEvaluationTree& InTree, TRangeBound<FFrameNumber> StartingBound);
+	MOVIESCENE_API FMovieSceneEvaluationTreeRangeIterator(const FMovieSceneEvaluationTree& InTree, TRangeBound<FFrameNumber> StartingBound);
 
 	/** Move onto the next time range */
 	FMovieSceneEvaluationTreeRangeIterator& operator++()
@@ -581,10 +581,10 @@ private:
 	}
 
 	/** Iterate onto the next range based on whether we're going forwards or backwards */
-	void Iter(bool bForwards);
+	MOVIESCENE_API void Iter(bool bForwards);
 
 	/** Find the next child within the given parent that is >= the specified bound based on whether we're going forwards or backwards */
-	FMovieSceneEvaluationTreeNodeHandle FindNextChild(FMovieSceneEvaluationTreeNodeHandle ParentNodeHandle, TRangeBound<FFrameNumber> PredicateBound, bool bForwards);
+	MOVIESCENE_API FMovieSceneEvaluationTreeNodeHandle FindNextChild(FMovieSceneEvaluationTreeNodeHandle ParentNodeHandle, TRangeBound<FFrameNumber> PredicateBound, bool bForwards);
 
 private:
 	/** The unique time range that we're currently at */
diff --git a/Engine/Source/Runtime/MovieScene/Public/Evaluation/MovieScenePlayback.h b/Engine/Source/Runtime/MovieScene/Public/Evaluation/MovieScenePlayback.h
--- a/Engine/Source/Runtime/MovieScene/Public/Evaluation/MovieScenePlayback.h
+++ b/Engine/Source/Runtime/MovieScene/Public/Evaluation/MovieScenePlayback.h
@@ -29,32 +29,32 @@ enum class EPlayDirection
 
 
 /** MovieScene evaluation context. Should remain bitwise copyable, and contain no external state since this has the potential to be used on a thread */
-struct MOVIESCENE_API FMovieSceneEvaluationRange
+struct FMovieSceneEvaluationRange
 {
 	/**
 	 * Construct this range from a single fixed time
 	 */
-	FMovieSceneEvaluationRange(FFrameTime InTime, FFrameRate InFrameRate);
+	MOVIESCENE_API FMovieSceneEvaluationRange(FFrameTime InTime, FFrameRate InFrameRate);
 
 	/**
 	 * Construct this range from a raw range and a direction
 	 */
-	FMovieSceneEvaluationRange(TRange<FFrameTime> InRange, FFrameRate InFrameRate, EPlayDirection InDirection);
+	MOVIESCENE_API FMovieSceneEvaluationRange(TRange<FFrameTime> InRange, FFrameRate InFrameRate, EPlayDirection InDirection);
 
 	/**
 	 * Construct this range from 2 times, and whether the range should include the previous time or not
 	 */
-	FMovieSceneEvaluationRange(FFrameTime InCurrentTime, FFrameTime InPreviousTime, FFrameRate InFrameRate, bool bInclusivePreviousTime = false);
+	MOVIESCENE_API FMovieSceneEvaluationRange(FFrameTime InCurrentTime, FFrameTime InPreviousTime, FFrameRate InFrameRate, bool bInclusivePreviousTime = false);
 
 	/**
 	 * Convert a frame time range to a frame number range comprising all the frame numbers traversed in the range
 	 */
-	static TRange<FFrameNumber> TimeRangeToNumberRange(const TRange<FFrameTime>& InFrameTimeRange);
+	static MOVIESCENE_API TRange<FFrameNumber> TimeRangeToNumberRange(const TRange<FFrameTime>& InFrameTimeRange);
 
 	/**
 	 * Convert a frame number range to a frame time range
 	 */
-	static TRange<FFrameTime> NumberRangeToTimeRange(const TRange<FFrameNumber>& InFrameTimeRange);
+	static MOVIESCENE_API TRange<FFrameTime> NumberRangeToTimeRange(const TRange<FFrameNumber>& InFrameTimeRange);
 
 	/**
 	 * Get the range that we should be evaluating
@@ -76,7 +76,7 @@ struct MOVIESCENE_API FMovieSceneEvaluationRange
 	 * Get the range of frame numbers traversed over this evaluation range by flooring the lower bound, and ceiling the upper bound.
 	 * For example: a time range of [1.5, 5.6] will yield the equivalent of [1, 6). A time range of (2.0, 2.9) will yield the equivalent of [2,3).
 	 */
-	TRange<FFrameNumber> GetTraversedFrameNumberRange() const;
+	MOVIESCENE_API TRange<FFrameNumber> GetTraversedFrameNumberRange() const;
 
 	/**
 	 * Get the direction to evaluate our range
@@ -162,7 +162,7 @@ struct MOVIESCENE_API FMovieSceneEvaluationRange
 	/**
 	 * Reset this range to a new range while retaining the current framerate and direction
 	 */
-	void ResetRange(const TRange<FFrameTime>& NewRange);
+	MOVIESCENE_API void ResetRange(const TRange<FFrameTime>& NewRange);
 
 protected:
 
@@ -529,7 +529,7 @@ protected:
 };
 
 /** Helper class designed to abstract the complexity of calculating evaluation ranges for previous times and fixed time intervals */
-struct MOVIESCENE_API FMovieScenePlaybackPosition
+struct FMovieScenePlaybackPosition
 {
 	FMovieScenePlaybackPosition()
 		: InputRate(0,0), OutputRate(0,0), EvaluationType(EMovieSceneEvaluationType::WithSubFrames)
@@ -568,13 +568,13 @@ public:
 	 * @param InOutputRate          The framerate to use when returning any frame range from this class
 	 * @param InputEvaluationType   Whether we're using frame-locked or sub-frame evaluation
 	 */
-	void SetTimeBase(FFrameRate InInputRate, FFrameRate InOutputRate, EMovieSceneEvaluationType InputEvaluationType);
+	MOVIESCENE_API void SetTimeBase(FFrameRate InInputRate, FFrameRate InOutputRate, EMovieSceneEvaluationType InputEvaluationType);
 
 	/**
 	 * Reset this position to the specified time.
 	 * @note Future calls to 'PlayTo' will include this time in its resulting evaluation range
 	 */
-	void Reset(FFrameTime StartPos);
+	MOVIESCENE_API void Reset(FFrameTime StartPos);
 
 	/**
 	 * Get the last position that was set, in InputRate space
@@ -595,7 +595,7 @@ public:
 	 * @param NewPosition         The new frame time to set, in InputRate space
 	 * @return A range encompassing only the specified time, in OutputRate space.
 	 */
-	FMovieSceneEvaluationRange JumpTo(FFrameTime NewPosition);
+	MOVIESCENE_API FMovieSceneEvaluationRange JumpTo(FFrameTime NewPosition);
 
 	/**
 	 * Play from the previously evaluated play time, to the specified time
@@ -603,26 +603,26 @@ public:
 	 * @param NewPosition         The new frame time to set, in InputRate space
 	 * @return An evaluation range from the previously evaluated time to the specified time, in OutputRate space.
 	 */
-	FMovieSceneEvaluationRange PlayTo(FFrameTime NewPosition);
+	MOVIESCENE_API FMovieSceneEvaluationRange PlayTo(FFrameTime NewPosition);
 
 	/**
 	 * Get a range that encompasses the last evaluated range in OutputRate space.
 	 * @return An optional evaluation range in OutputRate space.
 	 */
-	TOptional<FMovieSceneEvaluationRange> GetLastRange() const;
+	MOVIESCENE_API TOptional<FMovieSceneEvaluationRange> GetLastRange() const;
 
 	/**
 	 * Get a range encompassing only the current time, if available (in OutputRate space)
 	 * @return An optional evaluation range in OutputRate space.
 	 */
-	FMovieSceneEvaluationRange GetCurrentPositionAsRange() const;
+	MOVIESCENE_API FMovieSceneEvaluationRange GetCurrentPositionAsRange() const;
 
 private:
 
 	/**
 	 * Check this class's invariants
 	 */
-	void CheckInvariants() const;
+	MOVIESCENE_API void CheckInvariants() const;
 
 private:
 
diff --git a/Engine/Source/Runtime/MovieScene/Public/Evaluation/MovieScenePropertyTemplate.h b/Engine/Source/Runtime/MovieScene/Public/Evaluation/MovieScenePropertyTemplate.h
--- a/Engine/Source/Runtime/MovieScene/Public/Evaluation/MovieScenePropertyTemplate.h
+++ b/Engine/Source/Runtime/MovieScene/Public/Evaluation/MovieScenePropertyTemplate.h
@@ -227,29 +227,29 @@ struct FMovieScenePropertySectionData
 };
 
 USTRUCT()
-struct MOVIESCENE_API FMovieScenePropertySectionTemplate : public FMovieSceneEvalTemplate
+struct FMovieScenePropertySectionTemplate : public FMovieSceneEvalTemplate
 {
 	GENERATED_BODY()
 	
 	FMovieScenePropertySectionTemplate(){}
-	FMovieScenePropertySectionTemplate(FName PropertyName, const FString& InPropertyPath);
+	MOVIESCENE_API FMovieScenePropertySectionTemplate(FName PropertyName, const FString& InPropertyPath);
 public:
 	//use thse keys for setting and iterating the correct types.
-	const static FMovieSceneInterrogationKey GetFloatInterrogationKey();
-	const static FMovieSceneInterrogationKey GetInt32InterrogationKey();
-	const static FMovieSceneInterrogationKey GetVector4InterrogationKey();
-	const static FMovieSceneInterrogationKey GetVectorInterrogationKey();
-	const static FMovieSceneInterrogationKey GetVector2DInterrogationKey();
-	const static FMovieSceneInterrogationKey GetColorInterrogationKey();
+	MOVIESCENE_API const static FMovieSceneInterrogationKey GetFloatInterrogationKey();
+	MOVIESCENE_API const static FMovieSceneInterrogationKey GetInt32InterrogationKey();
+	MOVIESCENE_API const static FMovieSceneInterrogationKey GetVector4InterrogationKey();
+	MOVIESCENE_API const static FMovieSceneInterrogationKey GetVectorInterrogationKey();
+	MOVIESCENE_API const static FMovieSceneInterrogationKey GetVector2DInterrogationKey();
+	MOVIESCENE_API const static FMovieSceneInterrogationKey GetColorInterrogationKey();
 
 protected:
 
 	virtual UScriptStruct& GetScriptStructImpl() const override { return *StaticStruct(); }
 	/** Setup is only called if derived classes enable RequiresSetupFlag */
-	virtual void Setup(FPersistentEvaluationData& PersistentData, IMovieScenePlayer& Player) const override;
+	MOVIESCENE_API virtual void Setup(FPersistentEvaluationData& PersistentData, IMovieScenePlayer& Player) const override;
 
 	/** Access an animation type identifier that uniquely represents the property this section animates */
-	FMovieSceneAnimTypeID GetPropertyTypeID() const;
+	MOVIESCENE_API FMovieSceneAnimTypeID GetPropertyTypeID() const;
 
 	/** Ensure that an actutor for this property type exists */
 	template<typename PropertyType> FMovieSceneBlendingActuatorID EnsureActuator(FMovieSceneBlendingAccumulator& Accumulator) const
diff --git a/Engine/Source/Runtime/MovieScene/Public/Evaluation/MovieSceneRootOverridePath.h b/Engine/Source/Runtime/MovieScene/Public/Evaluation/MovieSceneRootOverridePath.h
--- a/Engine/Source/Runtime/MovieScene/Public/Evaluation/MovieSceneRootOverridePath.h
+++ b/Engine/Source/Runtime/MovieScene/Public/Evaluation/MovieSceneRootOverridePath.h
@@ -20,12 +20,12 @@ namespace MovieScene
  * A path of unaccumulated sequence IDs ordered from child->parent->grandparent that is used to generate unique sequenceIDs for inner sequences
  * Optimized for Remap rather than Push/Pop by keeping sequence IDs child-parent order (the order they are required for remapping)
  */
-struct MOVIESCENE_API FSubSequencePath
+struct FSubSequencePath
 {
 	/**
 	 * Default construction to a root path
 	 */
-	FSubSequencePath();
+	MOVIESCENE_API FSubSequencePath();
 
 	/**
 	 * Set up this path from a specific sequence ID that points to a particular sequence in the specified hierarchy
@@ -33,7 +33,7 @@ struct MOVIESCENE_API FSubSequencePath
 	 * @param LeafID 			ID of the child-most sequence to include in this path
 	 * @param Player 			Player from which to retrieve the hierarchy
 	 */
-	explicit FSubSequencePath(FMovieSceneSequenceID LeafID, IMovieScenePlayer& Player);
+	MOVIESCENE_API explicit FSubSequencePath(FMovieSceneSequenceID LeafID, IMovieScenePlayer& Player);
 
 	/**
 	 * Set up this path from a specific sequence ID that points to a particular sequence in the specified hierarchy
@@ -41,7 +41,7 @@ struct MOVIESCENE_API FSubSequencePath
 	 * @param LeafID 			ID of the child-most sequence to include in this path
 	 * @param RootHierarchy 	Hierarchy to get sequence IDs from
 	 */
-	explicit FSubSequencePath(FMovieSceneSequenceID LeafID, const FMovieSceneSequenceHierarchy* RootHierarchy);
+	MOVIESCENE_API explicit FSubSequencePath(FMovieSceneSequenceID LeafID, const FMovieSceneSequenceHierarchy* RootHierarchy);
 
 
 	/**
@@ -51,7 +51,7 @@ struct MOVIESCENE_API FSubSequencePath
 	 * @param B					The second sequence path
 	 * @return The leaf-most parent common to both paths, or MovieSceneSequenceID::Root
 	 */
-	static FMovieSceneSequenceID FindCommonParent(const FSubSequencePath& A, const FSubSequencePath& B);
+	static MOVIESCENE_API FMovieSceneSequenceID FindCommonParent(const FSubSequencePath& A, const FSubSequencePath& B);
 
 
 	/**
@@ -77,7 +77,7 @@ struct MOVIESCENE_API FSubSequencePath
 	/**
 	 * Reset this path to its default state (pointing to the root sequence)
 	 */
-	void Reset();
+	MOVIESCENE_API void Reset();
 
 
 	/**
@@ -86,7 +86,7 @@ struct MOVIESCENE_API FSubSequencePath
 	 * @param LeafID 			ID of the child-most sequence to include in this path
 	 * @param RootHierarchy 	Hierarchy to get sequence IDs from
 	 */
-	void Reset(FMovieSceneSequenceID LeafID, const FMovieSceneSequenceHierarchy* RootHierarchy);
+	MOVIESCENE_API void Reset(FMovieSceneSequenceID LeafID, const FMovieSceneSequenceHierarchy* RootHierarchy);
 
 
 	/**
@@ -95,7 +95,7 @@ struct MOVIESCENE_API FSubSequencePath
 	 * @param SequenceID		ID of the sequence to check for
 	 * @return true if this path contains the sequence ID (or SequenceID == MovieSceneSequenceID::Root), false otherwise
 	 */
-	bool Contains(FMovieSceneSequenceID SequenceID) const;
+	MOVIESCENE_API bool Contains(FMovieSceneSequenceID SequenceID) const;
 
 
 	/**
@@ -104,7 +104,7 @@ struct MOVIESCENE_API FSubSequencePath
 	 * @param SequenceID		ID of the parent sequence to count generations to
 	 * @return The number of generations between the two nodes. (ie, 0 where SequenceID == Leaf, 1 for Immediate parents, 2 for grandparents etc)
 	 */
-	int32 NumGenerationsFromLeaf(FMovieSceneSequenceID SequenceID) const;
+	MOVIESCENE_API int32 NumGenerationsFromLeaf(FMovieSceneSequenceID SequenceID) const;
 
 
 	/**
@@ -113,7 +113,7 @@ struct MOVIESCENE_API FSubSequencePath
 	 * @param SequenceID		ID of the child sequence to count generations to
 	 * @return The number of generations between the two nodes. (ie, 0 where SequenceID == Root, 1 for Immediate children, 2 for grandchildren)
 	 */
-	int32 NumGenerationsFromRoot(FMovieSceneSequenceID SequenceID) const;
+	MOVIESCENE_API int32 NumGenerationsFromRoot(FMovieSceneSequenceID SequenceID) const;
 
 	/**
 	 * 
@@ -121,15 +121,15 @@ struct MOVIESCENE_API FSubSequencePath
 	 * @param 
 	 * @return
 	 */
-	FMovieSceneSequenceID MakeLocalSequenceID(FMovieSceneSequenceID ParentSequenceID) const;
+	MOVIESCENE_API FMovieSceneSequenceID MakeLocalSequenceID(FMovieSceneSequenceID ParentSequenceID) const;
 
-	FMovieSceneSequenceID MakeLocalSequenceID(FMovieSceneSequenceID ParentSequenceID, FMovieSceneSequenceID TargetSequenceID) const;
+	MOVIESCENE_API FMovieSceneSequenceID MakeLocalSequenceID(FMovieSceneSequenceID ParentSequenceID, FMovieSceneSequenceID TargetSequenceID) const;
 
-	void PushGeneration(FMovieSceneSequenceID AccumulatedSequenceID, FMovieSceneSequenceID UnaccumulatedSequenceID);
+	MOVIESCENE_API void PushGeneration(FMovieSceneSequenceID AccumulatedSequenceID, FMovieSceneSequenceID UnaccumulatedSequenceID);
 
-	void PopTo(FMovieSceneSequenceID ParentSequenceID);
+	MOVIESCENE_API void PopTo(FMovieSceneSequenceID ParentSequenceID);
 
-	void PopGenerations(int32 NumGenerations);
+	MOVIESCENE_API void PopGenerations(int32 NumGenerations);
 
 private:
 
diff --git a/Engine/Source/Runtime/MovieScene/Public/Evaluation/MovieSceneSequenceHierarchy.h b/Engine/Source/Runtime/MovieScene/Public/Evaluation/MovieSceneSequenceHierarchy.h
--- a/Engine/Source/Runtime/MovieScene/Public/Evaluation/MovieSceneSequenceHierarchy.h
+++ b/Engine/Source/Runtime/MovieScene/Public/Evaluation/MovieSceneSequenceHierarchy.h
@@ -70,7 +70,7 @@ struct FMovieSceneSubSequenceData
 	/**
 	* Gets the signature of the sub-section this points to. 
 	*/
-	MOVIESCENE_API FGuid GetSubSectionSignature() const { return SubSectionSignature; }
+	FGuid GetSubSectionSignature() const { return SubSectionSignature; }
 
 	/**
 	 * Re-creates a sub-section parameter struct.
diff --git a/Engine/Source/Runtime/MovieScene/Public/Evaluation/MovieSceneSequenceInstanceData.h b/Engine/Source/Runtime/MovieScene/Public/Evaluation/MovieSceneSequenceInstanceData.h
--- a/Engine/Source/Runtime/MovieScene/Public/Evaluation/MovieSceneSequenceInstanceData.h
+++ b/Engine/Source/Runtime/MovieScene/Public/Evaluation/MovieSceneSequenceInstanceData.h
@@ -54,7 +54,7 @@ private:
 
 /**  */
 USTRUCT()
-struct MOVIESCENE_API FMovieSceneSequenceInstanceDataPtr
+struct FMovieSceneSequenceInstanceDataPtr
 #if CPP
 	: TInlineValue<FMovieSceneSequenceInstanceData, 16>
 #endif
@@ -88,18 +88,18 @@ struct MOVIESCENE_API FMovieSceneSequenceInstanceDataPtr
 	}
 	
 	/** Copy construction/assignment */
-	FMovieSceneSequenceInstanceDataPtr(const FMovieSceneSequenceInstanceDataPtr& RHS);
-	FMovieSceneSequenceInstanceDataPtr& operator=(const FMovieSceneSequenceInstanceDataPtr& RHS);
+	MOVIESCENE_API FMovieSceneSequenceInstanceDataPtr(const FMovieSceneSequenceInstanceDataPtr& RHS);
+	MOVIESCENE_API FMovieSceneSequenceInstanceDataPtr& operator=(const FMovieSceneSequenceInstanceDataPtr& RHS);
 
 	/** Templates are moveable */
 	FMovieSceneSequenceInstanceDataPtr(FMovieSceneSequenceInstanceDataPtr&&) = default;
 	FMovieSceneSequenceInstanceDataPtr& operator=(FMovieSceneSequenceInstanceDataPtr&&) = default;
 
 	/** Serialize the template */
-	bool Serialize(FArchive& Ar);
+	MOVIESCENE_API bool Serialize(FArchive& Ar);
 };
 
 template<> struct TStructOpsTypeTraits<FMovieSceneSequenceInstanceDataPtr> : public TStructOpsTypeTraitsBase2<FMovieSceneSequenceInstanceDataPtr>
 {
 	enum { WithSerializer = true, WithCopy = true };
-};
\ No newline at end of file
+};
diff --git a/Engine/Source/Runtime/MovieScene/Public/Evaluation/MovieSceneSequenceWeights.h b/Engine/Source/Runtime/MovieScene/Public/Evaluation/MovieSceneSequenceWeights.h
--- a/Engine/Source/Runtime/MovieScene/Public/Evaluation/MovieSceneSequenceWeights.h
+++ b/Engine/Source/Runtime/MovieScene/Public/Evaluation/MovieSceneSequenceWeights.h
@@ -16,11 +16,11 @@ namespace UE::MovieScene
 /**
  * Container class for keeping track of externally created entities for providing dynamic sequence weights
  */
-class MOVIESCENE_API FSequenceWeights
+class FSequenceWeights
 {
 public:
-	FSequenceWeights(UMovieSceneEntitySystemLinker* InLinker, FRootInstanceHandle InRootInstanceHandle);
-	~FSequenceWeights();
+	MOVIESCENE_API FSequenceWeights(UMovieSceneEntitySystemLinker* InLinker, FRootInstanceHandle InRootInstanceHandle);
+	MOVIESCENE_API ~FSequenceWeights();
 
 	/**
 	 * Set the weight of the specifieid sequence ID by creating a new HierarchicalEasingProvider entity and assigning its weight
@@ -28,7 +28,7 @@ public:
 	 * @param InSequenceID    The sequence ID to assign a weight to
 	 * @param Weight          The weight to assign. This is multiplied with all other weights contributing to this sequence
 	 */
-	void SetWeight(FMovieSceneSequenceID InSequenceID, double Weight = 1.0);
+	MOVIESCENE_API void SetWeight(FMovieSceneSequenceID InSequenceID, double Weight = 1.0);
 
 
 	/**
@@ -36,13 +36,13 @@ public:
 	 *
 	 * @param Weights         Map containing all desired sequence weights organized by their sequence ID
 	 */
-	void SetWeights(const TSortedMap<FMovieSceneSequenceID, double>& Weights);
+	MOVIESCENE_API void SetWeights(const TSortedMap<FMovieSceneSequenceID, double>& Weights);
 
 
 	/**
 	 * Remove a previously allocated weight for the specified sequence ID
 	 */
-	void RemoveWeight(FMovieSceneSequenceID InSequenceID);
+	MOVIESCENE_API void RemoveWeight(FMovieSceneSequenceID InSequenceID);
 
 private:
 
@@ -57,4 +57,4 @@ private:
 };
 
 
-} // namespace UE::MovieScene
\ No newline at end of file
+} // namespace UE::MovieScene
diff --git a/Engine/Source/Runtime/MovieScene/Public/Evaluation/PreAnimatedState/MovieScenePreAnimatedObjectGroupManager.h b/Engine/Source/Runtime/MovieScene/Public/Evaluation/PreAnimatedState/MovieScenePreAnimatedObjectGroupManager.h
--- a/Engine/Source/Runtime/MovieScene/Public/Evaluation/PreAnimatedState/MovieScenePreAnimatedObjectGroupManager.h
+++ b/Engine/Source/Runtime/MovieScene/Public/Evaluation/PreAnimatedState/MovieScenePreAnimatedObjectGroupManager.h
@@ -19,13 +19,13 @@ namespace MovieScene
 struct FPreAnimatedStorageGroupHandle;
 template <typename StorageType> struct TAutoRegisterPreAnimatedStorageID;
 
-struct MOVIESCENE_API FPreAnimatedObjectGroupManager : TPreAnimatedStateGroupManager<FObjectKey>
+struct FPreAnimatedObjectGroupManager : TPreAnimatedStateGroupManager<FObjectKey>
 {
-	static TAutoRegisterPreAnimatedStorageID<FPreAnimatedObjectGroupManager> GroupManagerID;
+	static MOVIESCENE_API TAutoRegisterPreAnimatedStorageID<FPreAnimatedObjectGroupManager> GroupManagerID;
 
-	void OnObjectsReplaced(const TMap<UObject*, UObject*>& ReplacementMap);
+	MOVIESCENE_API void OnObjectsReplaced(const TMap<UObject*, UObject*>& ReplacementMap);
 
-	void GetGroupsByClass(UClass* GeneratedClass, TArray<FPreAnimatedStorageGroupHandle>& OutGroupHandles);
+	MOVIESCENE_API void GetGroupsByClass(UClass* GeneratedClass, TArray<FPreAnimatedStorageGroupHandle>& OutGroupHandles);
 };
 
 
diff --git a/Engine/Source/Runtime/MovieScene/Public/Evaluation/PreAnimatedState/MovieScenePreAnimatedObjectStorage.h b/Engine/Source/Runtime/MovieScene/Public/Evaluation/PreAnimatedState/MovieScenePreAnimatedObjectStorage.h
--- a/Engine/Source/Runtime/MovieScene/Public/Evaluation/PreAnimatedState/MovieScenePreAnimatedObjectStorage.h
+++ b/Engine/Source/Runtime/MovieScene/Public/Evaluation/PreAnimatedState/MovieScenePreAnimatedObjectStorage.h
@@ -26,11 +26,11 @@ namespace MovieScene
  *	 StorageType CachePreAnimatedValue(const FObjectKey& Object);
  *	 void RestorePreAnimatedValue(const FObjectKey& Object, StorageType& InOutCachedValue, const FRestoreStateParams& Params);
  */
-struct MOVIESCENE_API FBoundObjectPreAnimatedStateTraits : FPreAnimatedStateTraits
+struct FBoundObjectPreAnimatedStateTraits : FPreAnimatedStateTraits
 {
 	enum { NeedsInitialize = true, SupportsGrouping = true, SupportsReplaceObject = true };
 
-	void Initialize(FPreAnimatedStorageID InStorageID, FPreAnimatedStateExtension* InParentExtension);
+	MOVIESCENE_API void Initialize(FPreAnimatedStorageID InStorageID, FPreAnimatedStateExtension* InParentExtension);
 
 	/* Defined as a template rather than a variadic function to prevent error C4840 */
 	template<typename... T>
@@ -44,8 +44,8 @@ struct MOVIESCENE_API FBoundObjectPreAnimatedStateTraits : FPreAnimatedStateTrai
 		return MakeGroupImpl(BoundObject);
 	}
 
-	FPreAnimatedStorageGroupHandle MakeGroupImpl(UObject* BoundObject);
-	FPreAnimatedStorageGroupHandle MakeGroupImpl(const FObjectComponent& BoundObject);
+	MOVIESCENE_API FPreAnimatedStorageGroupHandle MakeGroupImpl(UObject* BoundObject);
+	MOVIESCENE_API FPreAnimatedStorageGroupHandle MakeGroupImpl(const FObjectComponent& BoundObject);
 
 	template<typename ...T>
 	void ReplaceObject(TTuple<FObjectKey, T...>& InOutKey, const FObjectKey& NewObject)
diff --git a/Engine/Source/Runtime/MovieScene/Public/Evaluation/PreAnimatedState/MovieScenePreAnimatedObjectTokenStorage.h b/Engine/Source/Runtime/MovieScene/Public/Evaluation/PreAnimatedState/MovieScenePreAnimatedObjectTokenStorage.h
--- a/Engine/Source/Runtime/MovieScene/Public/Evaluation/PreAnimatedState/MovieScenePreAnimatedObjectTokenStorage.h
+++ b/Engine/Source/Runtime/MovieScene/Public/Evaluation/PreAnimatedState/MovieScenePreAnimatedObjectTokenStorage.h
@@ -42,11 +42,11 @@ struct FPreAnimatedObjectTokenTraits : FBoundObjectPreAnimatedStateTraits
 };
 
 
-struct MOVIESCENE_API FAnimTypePreAnimatedStateObjectStorage : TPreAnimatedStateStorage<FPreAnimatedObjectTokenTraits>
+struct FAnimTypePreAnimatedStateObjectStorage : TPreAnimatedStateStorage<FPreAnimatedObjectTokenTraits>
 {
-	static TAutoRegisterPreAnimatedStorageID<FAnimTypePreAnimatedStateObjectStorage> StorageID;
+	static MOVIESCENE_API TAutoRegisterPreAnimatedStorageID<FAnimTypePreAnimatedStateObjectStorage> StorageID;
 
-	FPreAnimatedStateEntry MakeEntry(UObject* Object, FMovieSceneAnimTypeID AnimTypeID);
+	MOVIESCENE_API FPreAnimatedStateEntry MakeEntry(UObject* Object, FMovieSceneAnimTypeID AnimTypeID);
 };
 
 
diff --git a/Engine/Source/Runtime/MovieScene/Public/Evaluation/PreAnimatedState/MovieScenePreAnimatedRootTokenStorage.h b/Engine/Source/Runtime/MovieScene/Public/Evaluation/PreAnimatedState/MovieScenePreAnimatedRootTokenStorage.h
--- a/Engine/Source/Runtime/MovieScene/Public/Evaluation/PreAnimatedState/MovieScenePreAnimatedRootTokenStorage.h
+++ b/Engine/Source/Runtime/MovieScene/Public/Evaluation/PreAnimatedState/MovieScenePreAnimatedRootTokenStorage.h
@@ -32,16 +32,16 @@ struct FPreAnimatedRootTokenTraits : FPreAnimatedStateTraits
 	}
 };
 
-struct MOVIESCENE_API FAnimTypePreAnimatedStateRootStorage : TPreAnimatedStateStorage<FPreAnimatedRootTokenTraits>
+struct FAnimTypePreAnimatedStateRootStorage : TPreAnimatedStateStorage<FPreAnimatedRootTokenTraits>
 {
-	static TAutoRegisterPreAnimatedStorageID<FAnimTypePreAnimatedStateRootStorage> StorageID;
+	static MOVIESCENE_API TAutoRegisterPreAnimatedStorageID<FAnimTypePreAnimatedStateRootStorage> StorageID;
 
-	FPreAnimatedStateEntry MakeEntry(FMovieSceneAnimTypeID AnimTypeID);
+	MOVIESCENE_API FPreAnimatedStateEntry MakeEntry(FMovieSceneAnimTypeID AnimTypeID);
 
 public:
 
 	FPreAnimatedStorageID GetStorageType() const override { return StorageID; }
-	void Initialize(FPreAnimatedStorageID InStorageID, FPreAnimatedStateExtension* ParentExtension) override;
+	MOVIESCENE_API void Initialize(FPreAnimatedStorageID InStorageID, FPreAnimatedStateExtension* ParentExtension) override;
 };
 
 } // namespace MovieScene
diff --git a/Engine/Source/Runtime/MovieScene/Public/Evaluation/PreAnimatedState/MovieScenePreAnimatedStateExtension.h b/Engine/Source/Runtime/MovieScene/Public/Evaluation/PreAnimatedState/MovieScenePreAnimatedStateExtension.h
--- a/Engine/Source/Runtime/MovieScene/Public/Evaluation/PreAnimatedState/MovieScenePreAnimatedStateExtension.h
+++ b/Engine/Source/Runtime/MovieScene/Public/Evaluation/PreAnimatedState/MovieScenePreAnimatedStateExtension.h
@@ -64,16 +64,16 @@ struct IPreAnimatedStateGroupManager
  *     - Pre-animated state is grouped into buckets in order to preserve significant ordering constraints (for instance, each object must restore state in the reverse order it was cached)
  *     - Groups are managed by registered IPreAnimatedStateGroupManager instances. The main implementation is FPreAnimatedObjectGroupManager, which maintains an object to group mapping
  */
-struct MOVIESCENE_API FPreAnimatedStateExtension
+struct FPreAnimatedStateExtension
 {
 public:
 
-	FPreAnimatedStateExtension();
-	~FPreAnimatedStateExtension();
+	MOVIESCENE_API FPreAnimatedStateExtension();
+	MOVIESCENE_API ~FPreAnimatedStateExtension();
 
 public:
 
-	void Initialize(UMovieSceneEntitySystemLinker* InLinker);
+	MOVIESCENE_API void Initialize(UMovieSceneEntitySystemLinker* InLinker);
 
 	/**
 	 * Called from TAutoRegisterPreAnimatedStorageID to register a new application-wide storage type
@@ -185,12 +185,12 @@ public:
 	/**
 	 * Called by group managers to allocate a new group
 	 */
-	FPreAnimatedStorageGroupHandle AllocateGroup(TSharedPtr<IPreAnimatedStateGroupManager> GroupManager);
+	MOVIESCENE_API FPreAnimatedStorageGroupHandle AllocateGroup(TSharedPtr<IPreAnimatedStateGroupManager> GroupManager);
 
 	/**
 	 * Called by group managers to free an existing group
 	 */
-	void FreeGroup(FPreAnimatedStorageGroupHandle Index);
+	MOVIESCENE_API void FreeGroup(FPreAnimatedStorageGroupHandle Index);
 
 public:
 
@@ -226,7 +226,7 @@ public:
 	 *
 	 * @param Params    Parameters for restoration - if TerminalInstanceHandle is invalid then _all_ state will be restored, regardless of the instance it was cached from
 	 */
-	void RestoreGlobalState(const FRestoreStateParams& Params);
+	MOVIESCENE_API void RestoreGlobalState(const FRestoreStateParams& Params);
 
 	/**
 	 * Restore any state cached for the specified group
@@ -234,72 +234,72 @@ public:
 	 * @params GroupHandle  Handle to the group to restore
 	 * @param  Params       Parameters for restoration
 	 */
-	void RestoreStateForGroup(FPreAnimatedStorageGroupHandle GroupHandle, const FRestoreStateParams& Params);
+	MOVIESCENE_API void RestoreStateForGroup(FPreAnimatedStorageGroupHandle GroupHandle, const FRestoreStateParams& Params);
 
 	/**
 	 * Called during blueprint re-instancing to replace the object bound to a specific group handle with another.
 	 */
-	void ReplaceObjectForGroup(FPreAnimatedStorageGroupHandle GroupHandle, const FObjectKey& OldObject, const FObjectKey& NewObject);
+	MOVIESCENE_API void ReplaceObjectForGroup(FPreAnimatedStorageGroupHandle GroupHandle, const FObjectKey& OldObject, const FObjectKey& NewObject);
 
 	/**
 	 * Discard any transient state and all meta-data for any currently animating objects, whilst preserving the cached values internally.
 	 * Calling this function will cause any currently animating 'RestoreState' sections to re-cache their values if they are re-evaluated
 	 * Any 'RestoreState' sections which are deleted or subsequently not-evaluated will not cause their values to be restored
 	 */
-	void DiscardTransientState();
+	MOVIESCENE_API void DiscardTransientState();
 
 	/**
 	 * Discard any and all cached values for the specified group without restoring them.
 	 * @note This function should only be used to forcibly serialize animated values into a level
 	 */
-	void DiscardStateForGroup(FPreAnimatedStorageGroupHandle GroupHandle);
+	MOVIESCENE_API void DiscardStateForGroup(FPreAnimatedStorageGroupHandle GroupHandle);
 
 	/**
 	 * Search for any captured state that originated from the specified root instance handle
 	 * WARNING: This is a linear search across all state, and so is potentially very slow
 	 */
-	bool ContainsAnyStateForInstanceHandle(FRootInstanceHandle RootInstanceHandle) const;
+	MOVIESCENE_API bool ContainsAnyStateForInstanceHandle(FRootInstanceHandle RootInstanceHandle) const;
 
 
 	// Use FScopedPreAnimatedCaptureSource to capture from a specific source rather than globally
-	void SavePreAnimatedState(FMovieSceneAnimTypeID InTokenType, const IMovieScenePreAnimatedGlobalTokenProducer& Producer);
-	void SavePreAnimatedState(UObject& InObject, FMovieSceneAnimTypeID InTokenType, const IMovieScenePreAnimatedTokenProducer& Producer);
+	MOVIESCENE_API void SavePreAnimatedState(FMovieSceneAnimTypeID InTokenType, const IMovieScenePreAnimatedGlobalTokenProducer& Producer);
+	MOVIESCENE_API void SavePreAnimatedState(UObject& InObject, FMovieSceneAnimTypeID InTokenType, const IMovieScenePreAnimatedTokenProducer& Producer);
 
-	void SavePreAnimatedStateDirectly(FMovieSceneAnimTypeID InTokenType, const IMovieScenePreAnimatedGlobalTokenProducer& Producer);
-	void SavePreAnimatedStateDirectly(UObject& InObject, FMovieSceneAnimTypeID InTokenType, const IMovieScenePreAnimatedTokenProducer& Producer);
+	MOVIESCENE_API void SavePreAnimatedStateDirectly(FMovieSceneAnimTypeID InTokenType, const IMovieScenePreAnimatedGlobalTokenProducer& Producer);
+	MOVIESCENE_API void SavePreAnimatedStateDirectly(UObject& InObject, FMovieSceneAnimTypeID InTokenType, const IMovieScenePreAnimatedTokenProducer& Producer);
 
 public:
 
-	FPreAnimatedEntityCaptureSource* GetEntityMetaData() const;
-	FPreAnimatedEntityCaptureSource* GetOrCreateEntityMetaData();
+	MOVIESCENE_API FPreAnimatedEntityCaptureSource* GetEntityMetaData() const;
+	MOVIESCENE_API FPreAnimatedEntityCaptureSource* GetOrCreateEntityMetaData();
 
-	FPreAnimatedTrackInstanceCaptureSources* GetTrackInstanceMetaData() const;
-	FPreAnimatedTrackInstanceCaptureSources* GetOrCreateTrackInstanceMetaData();
+	MOVIESCENE_API FPreAnimatedTrackInstanceCaptureSources* GetTrackInstanceMetaData() const;
+	MOVIESCENE_API FPreAnimatedTrackInstanceCaptureSources* GetOrCreateTrackInstanceMetaData();
 
-	FPreAnimatedTrackInstanceInputCaptureSources* GetTrackInstanceInputMetaData() const;
-	FPreAnimatedTrackInstanceInputCaptureSources* GetOrCreateTrackInstanceInputMetaData();
+	MOVIESCENE_API FPreAnimatedTrackInstanceInputCaptureSources* GetTrackInstanceInputMetaData() const;
+	MOVIESCENE_API FPreAnimatedTrackInstanceInputCaptureSources* GetOrCreateTrackInstanceInputMetaData();
 
-	bool HasActiveCaptureSource() const;
+	MOVIESCENE_API bool HasActiveCaptureSource() const;
 
-	void AddWeakCaptureSource(TWeakPtr<IPreAnimatedCaptureSource> InWeakMetaData);
-	void RemoveWeakCaptureSource(TWeakPtr<IPreAnimatedCaptureSource> InWeakMetaData);
+	MOVIESCENE_API void AddWeakCaptureSource(TWeakPtr<IPreAnimatedCaptureSource> InWeakMetaData);
+	MOVIESCENE_API void RemoveWeakCaptureSource(TWeakPtr<IPreAnimatedCaptureSource> InWeakMetaData);
 
-	void EnsureMetaData(const FPreAnimatedStateEntry& Entry);
-	void AddSourceMetaData(const FPreAnimatedStateEntry& Entry);
-	bool MetaDataExists(const FPreAnimatedStateEntry& Entry) const;
+	MOVIESCENE_API void EnsureMetaData(const FPreAnimatedStateEntry& Entry);
+	MOVIESCENE_API void AddSourceMetaData(const FPreAnimatedStateEntry& Entry);
+	MOVIESCENE_API bool MetaDataExists(const FPreAnimatedStateEntry& Entry) const;
 
-	void AddMetaData(const FPreAnimatedStateMetaData& MetaData);
-	void RemoveMetaData(const FPreAnimatedStateMetaData& MetaData);
-	void UpdateMetaData(const FPreAnimatedStateMetaData& MetaData);
+	MOVIESCENE_API void AddMetaData(const FPreAnimatedStateMetaData& MetaData);
+	MOVIESCENE_API void RemoveMetaData(const FPreAnimatedStateMetaData& MetaData);
+	MOVIESCENE_API void UpdateMetaData(const FPreAnimatedStateMetaData& MetaData);
 
-	EPreAnimatedStorageRequirement GetStorageRequirement(const FPreAnimatedStateEntry& Entry) const;
+	MOVIESCENE_API EPreAnimatedStorageRequirement GetStorageRequirement(const FPreAnimatedStateEntry& Entry) const;
 
 private:
 
-	void FreeGroupInternal(FPreAnimatedStorageGroupHandle Handle);
-	bool ShouldCaptureAnyState() const;
+	MOVIESCENE_API void FreeGroupInternal(FPreAnimatedStorageGroupHandle Handle);
+	MOVIESCENE_API bool ShouldCaptureAnyState() const;
 
-	void AddReferencedObjects(UMovieSceneEntitySystemLinker*, FReferenceCollector& ReferenceCollector);
+	MOVIESCENE_API void AddReferencedObjects(UMovieSceneEntitySystemLinker*, FReferenceCollector& ReferenceCollector);
 
 public:
 
@@ -344,10 +344,10 @@ private:
 		bool bEligibleForGlobalRestore = false;
 	};
 
-	FAggregatePreAnimatedStateMetaData* FindMetaData(const FPreAnimatedStateEntry& Entry);
-	const FAggregatePreAnimatedStateMetaData* FindMetaData(const FPreAnimatedStateEntry& Entry) const;
+	MOVIESCENE_API FAggregatePreAnimatedStateMetaData* FindMetaData(const FPreAnimatedStateEntry& Entry);
+	MOVIESCENE_API const FAggregatePreAnimatedStateMetaData* FindMetaData(const FPreAnimatedStateEntry& Entry) const;
 
-	FAggregatePreAnimatedStateMetaData* GetOrAddMetaDataInternal(const FPreAnimatedStateEntry& Entry);
+	MOVIESCENE_API FAggregatePreAnimatedStateMetaData* GetOrAddMetaDataInternal(const FPreAnimatedStateEntry& Entry);
 
 	struct FPreAnimatedGroupMetaData
 	{
@@ -369,7 +369,7 @@ private:
 
 private:
 
-	static FPreAnimatedStorageID RegisterStorageInternal();
+	static MOVIESCENE_API FPreAnimatedStorageID RegisterStorageInternal();
 };
 
 
diff --git a/Engine/Source/Runtime/MovieScene/Public/Evaluation/PreAnimatedState/MovieScenePreAnimatedStateStorage.h b/Engine/Source/Runtime/MovieScene/Public/Evaluation/PreAnimatedState/MovieScenePreAnimatedStateStorage.h
--- a/Engine/Source/Runtime/MovieScene/Public/Evaluation/PreAnimatedState/MovieScenePreAnimatedStateStorage.h
+++ b/Engine/Source/Runtime/MovieScene/Public/Evaluation/PreAnimatedState/MovieScenePreAnimatedStateStorage.h
@@ -95,7 +95,7 @@ struct TPreAnimatedStateTaskParams
  * (optional, if SupportsReplaceObject is true)
  *   void ReplaceObject(KeyType&, const FObjectKey&);
  */
-struct MOVIESCENE_API FPreAnimatedStateTraits
+struct FPreAnimatedStateTraits
 {
 	enum
 	{
diff --git a/Engine/Source/Runtime/MovieScene/Public/Generators/MovieSceneEasingCurves.h b/Engine/Source/Runtime/MovieScene/Public/Generators/MovieSceneEasingCurves.h
--- a/Engine/Source/Runtime/MovieScene/Public/Generators/MovieSceneEasingCurves.h
+++ b/Engine/Source/Runtime/MovieScene/Public/Generators/MovieSceneEasingCurves.h
@@ -41,15 +41,15 @@ enum class EMovieSceneBuiltInEasing : uint8
 };
 
 
-UCLASS(BlueprintType, meta=(DisplayName = "Built-In Function"))
-class MOVIESCENE_API UMovieSceneBuiltInEasingFunction : public UObject, public IMovieSceneEasingFunction
+UCLASS(BlueprintType, meta=(DisplayName = "Built-In Function"), MinimalAPI)
+class UMovieSceneBuiltInEasingFunction : public UObject, public IMovieSceneEasingFunction
 {
 public:
 	GENERATED_BODY()
 
-	UMovieSceneBuiltInEasingFunction(const FObjectInitializer& Initiailizer);
+	MOVIESCENE_API UMovieSceneBuiltInEasingFunction(const FObjectInitializer& Initiailizer);
 
-	virtual float Evaluate(float InTime) const override;
+	MOVIESCENE_API virtual float Evaluate(float InTime) const override;
 
 #if WITH_EDITOR
 	virtual FText GetDisplayName() const override { return StaticEnum<EMovieSceneBuiltInEasing>()->GetDisplayNameTextByValue((int64)Type); }
@@ -60,13 +60,13 @@ public:
 };
 
 
-UCLASS(meta=(DisplayName="Curve Asset"))
-class MOVIESCENE_API UMovieSceneEasingExternalCurve : public UObject, public IMovieSceneEasingFunction
+UCLASS(meta=(DisplayName="Curve Asset"), MinimalAPI)
+class UMovieSceneEasingExternalCurve : public UObject, public IMovieSceneEasingFunction
 {
 public:
 	GENERATED_BODY()
 
-	virtual float Evaluate(float InTime) const override;
+	MOVIESCENE_API virtual float Evaluate(float InTime) const override;
 
 #if WITH_EDITOR
 	virtual FText GetDisplayName() const override { return StaticClass()->GetDisplayNameText(); }
