[Engine]

diff --git a/Engine/Source/Runtime/MovieScene/Public/Generators/MovieSceneEasingFunction.h b/Engine/Source/Runtime/MovieScene/Public/Generators/MovieSceneEasingFunction.h
--- a/Engine/Source/Runtime/MovieScene/Public/Generators/MovieSceneEasingFunction.h
+++ b/Engine/Source/Runtime/MovieScene/Public/Generators/MovieSceneEasingFunction.h
@@ -7,14 +7,14 @@
 #include "MovieSceneEasingFunction.generated.h"
 
 
-UINTERFACE(Category="Sequencer", Blueprintable, meta=(DisplayName = "Easing Function"))
-class MOVIESCENE_API UMovieSceneEasingFunction : public UInterface
+UINTERFACE(Category="Sequencer", Blueprintable, meta=(DisplayName = "Easing Function"), MinimalAPI)
+class UMovieSceneEasingFunction : public UInterface
 {
 public:
 	GENERATED_BODY()
 };
 
-class MOVIESCENE_API IMovieSceneEasingFunction
+class IMovieSceneEasingFunction
 {
 public:
 	GENERATED_BODY()
@@ -22,7 +22,7 @@ public:
 	/**
 	 * Evaluate using the specified script interface. Handles both native and k2 implemented interfaces.
 	 */
-	static float EvaluateWith(const TScriptInterface<IMovieSceneEasingFunction>& ScriptInterface, float Time);
+	static MOVIESCENE_API float EvaluateWith(const TScriptInterface<IMovieSceneEasingFunction>& ScriptInterface, float Time);
 
 #if WITH_EDITOR
 	virtual FText GetDisplayName() const { return FText::GetEmpty(); }
@@ -32,8 +32,8 @@ protected:
 
 	/** Evaluate the easing with an interpolation value between 0 and 1 */
 	UFUNCTION(BlueprintCallable, BlueprintImplementableEvent, Category = "Sequencer|Section", meta=(CallInEditor="true"))
-	float OnEvaluate(float Interp) const;
+	MOVIESCENE_API float OnEvaluate(float Interp) const;
 
 	/** Evaluate the easing with an interpolation value between 0 and 1 */
 	virtual float Evaluate(float Interp) const { return 0.f; }
-};
\ No newline at end of file
+};
diff --git a/Engine/Source/Runtime/MovieScene/Public/IMovieSceneBoundObjectProxy.h b/Engine/Source/Runtime/MovieScene/Public/IMovieSceneBoundObjectProxy.h
--- a/Engine/Source/Runtime/MovieScene/Public/IMovieSceneBoundObjectProxy.h
+++ b/Engine/Source/Runtime/MovieScene/Public/IMovieSceneBoundObjectProxy.h
@@ -7,8 +7,8 @@
 #include "IMovieSceneBoundObjectProxy.generated.h"
 
 
-UINTERFACE()
-class MOVIESCENE_API UMovieSceneBoundObjectProxy
+UINTERFACE(MinimalAPI)
+class UMovieSceneBoundObjectProxy
 	: public UInterface
 {
 public:
@@ -16,7 +16,7 @@ public:
 };
 
 
-class MOVIESCENE_API IMovieSceneBoundObjectProxy
+class IMovieSceneBoundObjectProxy
 {
 public:
 	GENERATED_BODY()
@@ -27,6 +27,6 @@ public:
 	 * @return Pointer to the object that should be animated, or nullptr if it's not valid.
 	 */
 	UFUNCTION(BlueprintCallable, BlueprintImplementableEvent, Category = "Sequencer", DisplayName = "GetBoundObjectForSequencer", meta=(CallInEditor="true"))
-	UObject* BP_GetBoundObjectForSequencer(UObject* ResolvedObject);
+	MOVIESCENE_API UObject* BP_GetBoundObjectForSequencer(UObject* ResolvedObject);
 	virtual UObject* NativeGetBoundObjectForSequencer(UObject* ResolvedObject) { return ResolvedObject; }
 };
diff --git a/Engine/Source/Runtime/MovieScene/Public/IMovieSceneMetaData.h b/Engine/Source/Runtime/MovieScene/Public/IMovieSceneMetaData.h
--- a/Engine/Source/Runtime/MovieScene/Public/IMovieSceneMetaData.h
+++ b/Engine/Source/Runtime/MovieScene/Public/IMovieSceneMetaData.h
@@ -5,8 +5,8 @@
 #include "UObject/Interface.h"
 #include "IMovieSceneMetaData.generated.h"
 
-UINTERFACE()
-class MOVIESCENE_API UMovieSceneMetaDataInterface : public UInterface
+UINTERFACE(MinimalAPI)
+class UMovieSceneMetaDataInterface : public UInterface
 {
 	GENERATED_BODY()
 };
@@ -36,4 +36,4 @@ public:
 	virtual void ExtendAssetRegistryTagMetaData(TMap<FName, UObject::FAssetRegistryTagMetadata>& OutMetadata) const {}
 
 #endif
-};
\ No newline at end of file
+};
diff --git a/Engine/Source/Runtime/MovieScene/Public/IMovieScenePlaybackClient.h b/Engine/Source/Runtime/MovieScene/Public/IMovieScenePlaybackClient.h
--- a/Engine/Source/Runtime/MovieScene/Public/IMovieScenePlaybackClient.h
+++ b/Engine/Source/Runtime/MovieScene/Public/IMovieScenePlaybackClient.h
@@ -11,8 +11,8 @@
 struct FMovieSceneEvaluationRange;
 
 /** Movie scene binding overrides interface */
-UINTERFACE()
-class MOVIESCENE_API UMovieScenePlaybackClient
+UINTERFACE(MinimalAPI)
+class UMovieScenePlaybackClient
 	: public UInterface
 {
 public:
@@ -20,7 +20,7 @@ public:
 };
 
 
-class MOVIESCENE_API IMovieScenePlaybackClient
+class IMovieScenePlaybackClient
 {
 public:
 	GENERATED_BODY()
diff --git a/Engine/Source/Runtime/MovieScene/Public/IMovieSceneSequencePlayerObserver.h b/Engine/Source/Runtime/MovieScene/Public/IMovieSceneSequencePlayerObserver.h
--- a/Engine/Source/Runtime/MovieScene/Public/IMovieSceneSequencePlayerObserver.h
+++ b/Engine/Source/Runtime/MovieScene/Public/IMovieSceneSequencePlayerObserver.h
@@ -8,8 +8,8 @@
 #include "IMovieSceneSequencePlayerObserver.generated.h"
 
 
-UINTERFACE()
-class MOVIESCENE_API UMovieSceneSequencePlayerObserver : public UInterface
+UINTERFACE(MinimalAPI)
+class UMovieSceneSequencePlayerObserver : public UInterface
 {
 public:
 	GENERATED_BODY()
@@ -21,4 +21,4 @@ public:
 	GENERATED_BODY()
 
 	virtual bool CanObserveSequence() const = 0;
-};
\ No newline at end of file
+};
diff --git a/Engine/Source/Runtime/MovieScene/Public/INodeAndChannelMappings.h b/Engine/Source/Runtime/MovieScene/Public/INodeAndChannelMappings.h
--- a/Engine/Source/Runtime/MovieScene/Public/INodeAndChannelMappings.h
+++ b/Engine/Source/Runtime/MovieScene/Public/INodeAndChannelMappings.h
@@ -99,13 +99,13 @@ struct FControlRigFbxCurveData
 };
 
 
-UINTERFACE(meta = (CannotImplementInterfaceInBlueprint))
-class MOVIESCENE_API UNodeAndChannelMappings : public UInterface
+UINTERFACE(meta = (CannotImplementInterfaceInBlueprint), MinimalAPI)
+class UNodeAndChannelMappings : public UInterface
 {
 	GENERATED_UINTERFACE_BODY()
 };
 
-class MOVIESCENE_API INodeAndChannelMappings
+class INodeAndChannelMappings
 {
 	GENERATED_IINTERFACE_BODY()
 
diff --git a/Engine/Source/Runtime/MovieScene/Public/MovieScene.h b/Engine/Source/Runtime/MovieScene/Public/MovieScene.h
--- a/Engine/Source/Runtime/MovieScene/Public/MovieScene.h
+++ b/Engine/Source/Runtime/MovieScene/Public/MovieScene.h
@@ -221,8 +221,8 @@ public:
 /**
  * Structure that represents a group of nodes
  */
-UCLASS()
-class MOVIESCENE_API UMovieSceneNodeGroup : public UObject
+UCLASS(MinimalAPI)
+class UMovieSceneNodeGroup : public UObject
 {
 	GENERATED_BODY()
 
@@ -231,17 +231,17 @@ class MOVIESCENE_API UMovieSceneNodeGroup : public UObject
 #if WITH_EDITORONLY_DATA
 public:
 	const FName GetName() const { return Name; }
-	void SetName(const FName& Name);
+	MOVIESCENE_API void SetName(const FName& Name);
 
-	void AddNode(const FString& Path);
-	void RemoveNode(const FString& Path);
+	MOVIESCENE_API void AddNode(const FString& Path);
+	MOVIESCENE_API void RemoveNode(const FString& Path);
 	TArrayView<FString> GetNodes() { return Nodes; }
-	bool ContainsNode(const FString& Path) const;
+	MOVIESCENE_API bool ContainsNode(const FString& Path) const;
 
-	void UpdateNodePath(const FString& OldPath, const FString& NewPath);
+	MOVIESCENE_API void UpdateNodePath(const FString& OldPath, const FString& NewPath);
 
 	bool GetEnableFilter() const { return bEnableFilter; }
-	void SetEnableFilter(bool bInEnableFilter);
+	MOVIESCENE_API void SetEnableFilter(bool bInEnableFilter);
 
 	/** Event that is triggered whenever this node group has changed */
 	DECLARE_EVENT(UMovieSceneNodeGroup, FOnNodeGroupChanged)
@@ -280,8 +280,8 @@ private:
 /**
  * Structure that represents a collection of NodeGroups
  */
-UCLASS()
-class MOVIESCENE_API UMovieSceneNodeGroupCollection : public UObject
+UCLASS(MinimalAPI)
+class UMovieSceneNodeGroupCollection : public UObject
 {
 	GENERATED_BODY()
 
@@ -290,17 +290,17 @@ class MOVIESCENE_API UMovieSceneNodeGroupCollection : public UObject
 #if WITH_EDITORONLY_DATA
 public:
 	/** Called after this object has been deserialized */
-	virtual void PostLoad() override;
-	virtual void PostEditUndo() override;
+	MOVIESCENE_API virtual void PostLoad() override;
+	MOVIESCENE_API virtual void PostEditUndo() override;
 
-	void AddNodeGroup(UMovieSceneNodeGroup* NodeGroup);
-	void RemoveNodeGroup(UMovieSceneNodeGroup* NodeGroup);
+	MOVIESCENE_API void AddNodeGroup(UMovieSceneNodeGroup* NodeGroup);
+	MOVIESCENE_API void RemoveNodeGroup(UMovieSceneNodeGroup* NodeGroup);
 
 	bool Contains(UMovieSceneNodeGroup* NodeGroup) const { return NodeGroups.Contains(NodeGroup); }
 	const int32 Num() const { return NodeGroups.Num(); }
 	bool HasAnyActiveFilter() const { return bAnyActiveFilter; }
 
-	void UpdateNodePath(const FString& OldPath, const FString& NewPath);
+	MOVIESCENE_API void UpdateNodePath(const FString& OldPath, const FString& NewPath);
 
 	/** Event that is triggered whenever this collection of node groups, or an included node group, has changed */
 	DECLARE_EVENT(UMovieSceneNodeGroupCollection, FOnNodeGroupCollectionChanged)
@@ -308,14 +308,14 @@ public:
 
 private:
 
-	void Refresh();
+	MOVIESCENE_API void Refresh();
 
 	UPROPERTY()
 	TArray<TObjectPtr<UMovieSceneNodeGroup>> NodeGroups;
 
 	bool bAnyActiveFilter;
 
-	void OnNodeGroupChanged();
+	MOVIESCENE_API void OnNodeGroupChanged();
 
 	/** Event that is triggered whenever this collection of node groups, or an included node group, has changed */
 	FOnNodeGroupCollectionChanged OnNodeGroupCollectionChangedEvent;
@@ -336,8 +336,8 @@ public:
 /**
  * Implements a movie scene asset.
  */
-UCLASS(DefaultToInstanced)
-class MOVIESCENE_API UMovieScene
+UCLASS(DefaultToInstanced, MinimalAPI)
+class UMovieScene
 	: public UMovieSceneSignedObject
 {
 	GENERATED_UCLASS_BODY()
@@ -345,17 +345,17 @@ class MOVIESCENE_API UMovieScene
 public:
 
 	/**~ UObject implementation */
-	virtual void Serialize( FArchive& Ar ) override;
-	virtual bool IsPostLoadThreadSafe() const override;
-	virtual void PostInitProperties() override;
-	virtual void PostLoad() override;
+	MOVIESCENE_API virtual void Serialize( FArchive& Ar ) override;
+	MOVIESCENE_API virtual bool IsPostLoadThreadSafe() const override;
+	MOVIESCENE_API virtual void PostInitProperties() override;
+	MOVIESCENE_API virtual void PostLoad() override;
 #if WITH_EDITORONLY_DATA
-	static void DeclareConstructClasses(TArray<FTopLevelAssetPath>& OutConstructClasses, const UClass* SpecificSubclass);
+	static MOVIESCENE_API void DeclareConstructClasses(TArray<FTopLevelAssetPath>& OutConstructClasses, const UClass* SpecificSubclass);
 #endif
 
 
 #if WITH_EDITOR
-	virtual void PostEditUndo() override;
+	MOVIESCENE_API virtual void PostEditUndo() override;
 #endif
 
 public:
@@ -371,7 +371,7 @@ public:
 	 * @param ObjectTemplate The object template to use for the spawnable
 	 * @return Guid of the newly-added spawnable.
 	 */
-	FGuid AddSpawnable(const FString& Name, UObject& ObjectTemplate);
+	MOVIESCENE_API FGuid AddSpawnable(const FString& Name, UObject& ObjectTemplate);
 
 	/*
 	 * Adds an existing spawnable to this movie scene.
@@ -379,7 +379,7 @@ public:
 	 * @param InNewSpawnable The posssesable to add.
 	 * @param InNewBinding The object binding to add.
 	 */
-	void AddSpawnable(const FMovieSceneSpawnable& InNewSpawnable, const FMovieSceneBinding& InNewBinding);
+	MOVIESCENE_API void AddSpawnable(const FMovieSceneSpawnable& InNewSpawnable, const FMovieSceneBinding& InNewBinding);
 
 	/**
 	 * Removes a spawnable from this movie scene.
@@ -387,7 +387,7 @@ public:
 	 * @param Guid The guid of a spawnable to find and remove.
 	 * @return true if anything was removed.
 	 */
-	bool RemoveSpawnable(const FGuid& Guid);
+	MOVIESCENE_API bool RemoveSpawnable(const FGuid& Guid);
 
 	/**
 	 * Attempt to find a spawnable using some custom predicate
@@ -395,7 +395,7 @@ public:
 	 * @param InPredicate A predicate to test each spawnable against
 	 * @return Spawnable object that was found (or nullptr if not found).
 	 */
-	FMovieSceneSpawnable* FindSpawnable( const TFunctionRef<bool(FMovieSceneSpawnable&)>& InPredicate );
+	MOVIESCENE_API FMovieSceneSpawnable* FindSpawnable( const TFunctionRef<bool(FMovieSceneSpawnable&)>& InPredicate );
 
 	/**
 	 * Tries to locate a spawnable in this MovieScene for the specified spawnable GUID.
@@ -403,7 +403,7 @@ public:
 	 * @param Guid The spawnable guid to search for.
 	 * @return Spawnable object that was found (or nullptr if not found).
 	 */
-	FMovieSceneSpawnable* FindSpawnable(const FGuid& Guid);
+	MOVIESCENE_API FMovieSceneSpawnable* FindSpawnable(const FGuid& Guid);
 
 	/**
 	 * Grabs a reference to a specific spawnable by index.
@@ -411,14 +411,14 @@ public:
 	 * @param Index of spawnable to return. Must be between 0 and GetSpawnableCount()
 	 * @return Returns the specified spawnable by index.
 	 */
-	FMovieSceneSpawnable& GetSpawnable(int32 Index);
+	MOVIESCENE_API FMovieSceneSpawnable& GetSpawnable(int32 Index);
 
 	/**
 	 * Get the number of spawnable objects in this scene.
 	 *
 	 * @return Spawnable object count.
 	 */
-	int32 GetSpawnableCount() const;
+	MOVIESCENE_API int32 GetSpawnableCount() const;
 	
 public:
 
@@ -429,7 +429,7 @@ public:
 	 * @param Class The class of object that will be possessed.
 	 * @return Guid of the newly-added possessable.
 	 */
-	FGuid AddPossessable(const FString& Name, UClass* Class);
+	MOVIESCENE_API FGuid AddPossessable(const FString& Name, UClass* Class);
 
 	/*
 	 * Adds an existing possessable to this movie scene.
@@ -437,19 +437,19 @@ public:
 	 * @param InNewPossessable The posssesable to add.
 	 * @param InNewBinding The object binding to add.
 	 */
-	void AddPossessable(const FMovieScenePossessable& InNewPossessable, const FMovieSceneBinding& InNewBinding);
+	MOVIESCENE_API void AddPossessable(const FMovieScenePossessable& InNewPossessable, const FMovieSceneBinding& InNewBinding);
 
 	/**
 	 * Removes a possessable from this movie scene.
 	 *
 	 * @param PossessableGuid Guid of possessable to remove.
 	 */
-	bool RemovePossessable(const FGuid& PossessableGuid);
+	MOVIESCENE_API bool RemovePossessable(const FGuid& PossessableGuid);
 	
 	/*
 	* Replace an existing possessable with another 
 	*/
-	bool ReplacePossessable(const FGuid& OldGuid, const FMovieScenePossessable& InNewPosessable);
+	MOVIESCENE_API bool ReplacePossessable(const FGuid& OldGuid, const FMovieScenePossessable& InNewPosessable);
 
 	/**
 	 * Tries to locate a possessable in this MovieScene for the specified possessable GUID.
@@ -457,7 +457,7 @@ public:
 	 * @param Guid The possessable guid to search for.
 	 * @return Possessable object that was found (or nullptr if not found).
 	 */
-	struct FMovieScenePossessable* FindPossessable(const FGuid& Guid);
+	MOVIESCENE_API struct FMovieScenePossessable* FindPossessable(const FGuid& Guid);
 
 	/**
 	 * Attempt to find a possessable using some custom prdeicate
@@ -465,7 +465,7 @@ public:
 	 * @param InPredicate A predicate to test each possessable against
 	 * @return Possessable object that was found (or nullptr if not found).
 	 */
-	FMovieScenePossessable* FindPossessable( const TFunctionRef<bool(FMovieScenePossessable&)>& InPredicate );
+	MOVIESCENE_API FMovieScenePossessable* FindPossessable( const TFunctionRef<bool(FMovieScenePossessable&)>& InPredicate );
 
 	/**
 	 * Grabs a reference to a specific possessable by index.
@@ -473,14 +473,14 @@ public:
 	 * @param Index of possessable to return.
 	 * @return Returns the specified possessable by index.
 	 */
-	FMovieScenePossessable& GetPossessable(const int32 Index);
+	MOVIESCENE_API FMovieScenePossessable& GetPossessable(const int32 Index);
 
 	/**
 	 * Get the number of possessable objects in this scene.
 	 *
 	 * @return Possessable object count.
 	 */
-	int32 GetPossessableCount() const;
+	MOVIESCENE_API int32 GetPossessableCount() const;
 
 public:
 
@@ -494,7 +494,7 @@ public:
 	 * @param Type The newly created type.
 	 * @see  FindTrack, RemoveTrack
 	 */
-	UMovieSceneTrack* AddTrack(TSubclassOf<UMovieSceneTrack> TrackClass, const FGuid& ObjectGuid);
+	MOVIESCENE_API UMovieSceneTrack* AddTrack(TSubclassOf<UMovieSceneTrack> TrackClass, const FGuid& ObjectGuid);
 
 	/**
 	* Adds a given track.
@@ -504,7 +504,7 @@ public:
 	* @see  FindTrack, RemoveTrack
 	* @return true if the track is successfully added, false otherwise.
 	*/
-	bool AddGivenTrack(UMovieSceneTrack* InTrack, const FGuid& ObjectGuid);
+	MOVIESCENE_API bool AddGivenTrack(UMovieSceneTrack* InTrack, const FGuid& ObjectGuid);
 
 	/**
 	 * Adds a track.
@@ -531,7 +531,7 @@ public:
 	 * @return The found track or nullptr if one does not exist.
 	 * @see AddTrack, RemoveTrack, FindTracks
 	 */
-	UMovieSceneTrack* FindTrack(TSubclassOf<UMovieSceneTrack> TrackClass, const FGuid& ObjectGuid, const FName& TrackName = NAME_None) const;
+	MOVIESCENE_API UMovieSceneTrack* FindTrack(TSubclassOf<UMovieSceneTrack> TrackClass, const FGuid& ObjectGuid, const FName& TrackName = NAME_None) const;
 	
 	/**
 	 * Finds a track.
@@ -557,7 +557,7 @@ public:
 	 * @return The found tracks or an empty array if none exist
 	 * @see AddTrack, RemoveTrack, FindTrack
 	 */
-	TArray<UMovieSceneTrack*> FindTracks(TSubclassOf<UMovieSceneTrack> TrackClass, const FGuid& ObjectGuid, const FName& TrackName = NAME_None) const;
+	MOVIESCENE_API TArray<UMovieSceneTrack*> FindTracks(TSubclassOf<UMovieSceneTrack> TrackClass, const FGuid& ObjectGuid, const FName& TrackName = NAME_None) const;
 
 	/**
 	 * Removes a track.
@@ -566,7 +566,7 @@ public:
 	 * @return true if anything was removed.
 	 * @see AddTrack, FindTrack
 	 */
-	bool RemoveTrack(UMovieSceneTrack& Track);
+	MOVIESCENE_API bool RemoveTrack(UMovieSceneTrack& Track);
 
 	/**
 	 * Find a track binding Guid from a UMovieSceneTrack
@@ -575,15 +575,15 @@ public:
 	 * @param	OutGuid		The binding's Guid if one was found.
 	 * @return true if a binding was found for this track.
 	 */
-	bool FindTrackBinding(const UMovieSceneTrack& InTrack, FGuid& OutGuid) const;
+	MOVIESCENE_API bool FindTrackBinding(const UMovieSceneTrack& InTrack, FGuid& OutGuid) const;
 
 #if WITH_EDITOR
 
 	DECLARE_DELEGATE_RetVal_OneParam(bool, FIsTrackClassAllowedEvent, UClass*);
 
-	static FIsTrackClassAllowedEvent IsTrackClassAllowedEvent;
+	static MOVIESCENE_API FIsTrackClassAllowedEvent IsTrackClassAllowedEvent;
 
-	static bool IsTrackClassAllowed(UClass* InClass);
+	static MOVIESCENE_API bool IsTrackClassAllowed(UClass* InClass);
 
 	void OnDynamicBindingUserDefinedPinRenamed(UK2Node* InNode, FName OldPinName, FName NewPinName)
 	{
@@ -592,7 +592,7 @@ public:
 
 	DECLARE_MULTICAST_DELEGATE_FourParams(FFixupDynamicBindingPayloadParameterNameEvent, UMovieScene*, UK2Node*, FName, FName);
 
-	static FFixupDynamicBindingPayloadParameterNameEvent FixupDynamicBindingPayloadParameterNameEvent;
+	static MOVIESCENE_API FFixupDynamicBindingPayloadParameterNameEvent FixupDynamicBindingPayloadParameterNameEvent;
 
 #endif
 
@@ -607,7 +607,7 @@ public:
 	 * @param Type	The newly created type
 	 * @see FindTrack, GetTracks, IsTrack, RemoveTrack
 	 */
-	UMovieSceneTrack* AddTrack(TSubclassOf<UMovieSceneTrack> TrackClass);
+	MOVIESCENE_API UMovieSceneTrack* AddTrack(TSubclassOf<UMovieSceneTrack> TrackClass);
 	
 	UE_DEPRECATED(5.2, "AddMasterTrack is deprecated. Please use AddTrack instead")
 	UMovieSceneTrack* AddMasterTrack(TSubclassOf<UMovieSceneTrack> TrackClass) { return AddTrack(TrackClass); }
@@ -638,7 +638,7 @@ public:
 	* @see  FindTrack, RemoveTrack
 	* @return true if the track is successfully added, false otherwise.
 	*/
-	bool AddGivenTrack(UMovieSceneTrack* InTrack);
+	MOVIESCENE_API bool AddGivenTrack(UMovieSceneTrack* InTrack);
 
 	UE_DEPRECATED(5.2, "AddGivenMasterTrack is deprecated. Please use AddGivenTrack instead")
 	bool AddGivenMasterTrack(UMovieSceneTrack* InTrack) { return AddGivenTrack(InTrack); }
@@ -650,7 +650,7 @@ public:
 	 * @return The found track or nullptr if one does not exist.
 	 * @see AddTrack, GetTracks, IsTrack, RemoveTrack
 	 */
-	UMovieSceneTrack* FindTrack(TSubclassOf<UMovieSceneTrack> TrackClass) const;
+	MOVIESCENE_API UMovieSceneTrack* FindTrack(TSubclassOf<UMovieSceneTrack> TrackClass) const;
 
 	UE_DEPRECATED(5.2, "FindMasterTrack is deprecated. Please use FindTrack instead")
 	UMovieSceneTrack* FindMasterTrack(TSubclassOf<UMovieSceneTrack> TrackClass) const { return FindTrack(TrackClass); }
@@ -692,7 +692,7 @@ public:
 	 * @return true if the track is a track, false otherwise.
 	 * @see AddTrack, FindTrack, GetTracks, RemoveTrack
 	 */
-	bool ContainsTrack(const UMovieSceneTrack& Track) const;
+	MOVIESCENE_API bool ContainsTrack(const UMovieSceneTrack& Track) const;
 
 	UE_DEPRECATED(5.2, "IsAMasterTrack is deprecated. Please use ContainsTrack instead")
 	bool IsAMasterTrack(const UMovieSceneTrack& Track) const { return ContainsTrack(Track); }
@@ -706,7 +706,7 @@ public:
 	 * @param SourceBindingId The identifier of the binding ID to move all tracks and children from
 	 * @param DestinationBindingId The identifier of the binding ID to move the contents to
 	 */
-	void MoveBindingContents(const FGuid& SourceBindingId, const FGuid& DestinationBindingId);
+	MOVIESCENE_API void MoveBindingContents(const FGuid& SourceBindingId, const FGuid& DestinationBindingId);
 
 	/**
 	 * Tries to find an FMovieSceneBinding for the specified Guid.
@@ -731,15 +731,15 @@ public:
 	 * @param TrackClass  The camera cut track class type
 	 * @return The created camera cut track
 	 */
-	UMovieSceneTrack* AddCameraCutTrack( TSubclassOf<UMovieSceneTrack> TrackClass );
+	MOVIESCENE_API UMovieSceneTrack* AddCameraCutTrack( TSubclassOf<UMovieSceneTrack> TrackClass );
 	
 	/** @return The camera cut track if it exists. */
-	UMovieSceneTrack* GetCameraCutTrack() const;
+	MOVIESCENE_API UMovieSceneTrack* GetCameraCutTrack() const;
 
 	/** Removes the camera cut track if it exists. */
-	void RemoveCameraCutTrack();
+	MOVIESCENE_API void RemoveCameraCutTrack();
 
-	void SetCameraCutTrack(UMovieSceneTrack* Track);
+	MOVIESCENE_API void SetCameraCutTrack(UMovieSceneTrack* Track);
 
 public:
 
@@ -748,7 +748,7 @@ public:
 	 *
 	 * @return A list of sections with object bindings and names.
 	 */
-	TArray<UMovieSceneSection*> GetAllSections() const;
+	MOVIESCENE_API TArray<UMovieSceneSection*> GetAllSections() const;
 
 	/**
 	 * @return All object bindings.
@@ -770,7 +770,7 @@ public:
 	 * @param ObjectId The object identifier.
 	 * @return The object's display name.
 	 */
-	FText GetObjectDisplayName(const FGuid& ObjectId);
+	MOVIESCENE_API FText GetObjectDisplayName(const FGuid& ObjectId);
 
 	/** Get the playback time range of this movie scene, relative to its 0-time offset. */
 	TRange<FFrameNumber> GetPlaybackRange() const
@@ -842,7 +842,7 @@ public:
 	/**
 	 * Retrieve a time controller from this sequence instance, if the clock source is set to custom
 	 */
-	TSharedPtr<FMovieSceneTimeController> MakeCustomTimeController(UObject* PlaybackContext);
+	MOVIESCENE_API TSharedPtr<FMovieSceneTimeController> MakeCustomTimeController(UObject* PlaybackContext);
 
 	/**
 	 * Assign the clock source to be used for this moviescene
@@ -868,12 +868,12 @@ public:
 	/**
 	 * Get the earliest timecode source out of all of the movie scene sections contained within this movie scene.
 	 */
-	FMovieSceneTimecodeSource GetEarliestTimecodeSource() const;
+	MOVIESCENE_API FMovieSceneTimecodeSource GetEarliestTimecodeSource() const;
 
 	/*
 	* Replace an existing binding with another 
 	*/
-	void ReplaceBinding(const FGuid& OldGuid, const FGuid& NewGuid, const FString& Name);
+	MOVIESCENE_API void ReplaceBinding(const FGuid& OldGuid, const FGuid& NewGuid, const FString& Name);
 
 	/*
 	* Replace an existing binding with another. Assumes ownership of any
@@ -881,7 +881,7 @@ public:
 	* @param BindingToReplaceGuid	Binding Guid that should be replaced
 	* @param NewBinding				Binding Data that should replace the original one specified by BindingToReplaceGuid.
 	*/
-	void ReplaceBinding(const FGuid& BindingToReplaceGuid, const FMovieSceneBinding& NewBinding);
+	MOVIESCENE_API void ReplaceBinding(const FGuid& BindingToReplaceGuid, const FMovieSceneBinding& NewBinding);
 
 #if WITH_EDITORONLY_DATA
 	/**
@@ -903,47 +903,47 @@ public:
 	 * @param ObjectId The object identifier.
 	 * @return The object's display name.
 	 */
-	void SetObjectDisplayName(const FGuid& ObjectId, const FText& DisplayName);
+	MOVIESCENE_API void SetObjectDisplayName(const FGuid& ObjectId, const FText& DisplayName);
 
 	/**
 	 * Gets the root folders for this movie scene.
 	 */
-	TArrayView<UMovieSceneFolder* const> GetRootFolders();
+	MOVIESCENE_API TArrayView<UMovieSceneFolder* const> GetRootFolders();
 
 	/**
 	 * Gets a copy of the root folders for this movie scene.
 	 */
-	void GetRootFolders(TArray<UMovieSceneFolder*>& InRootFolders);
+	MOVIESCENE_API void GetRootFolders(TArray<UMovieSceneFolder*>& InRootFolders);
 
 	/**
 	 * Gets the number of root folders in this movie scene.
 	 */
-	int32 GetNumRootFolders() const;
+	MOVIESCENE_API int32 GetNumRootFolders() const;
 
 	/**
 	 * Gets the i-th root folder for this movie scene.
 	 */
-	UMovieSceneFolder* GetRootFolder(int32 FolderIndex) const;
+	MOVIESCENE_API UMovieSceneFolder* GetRootFolder(int32 FolderIndex) const;
 
 	/**
 	 * Adds a root folder for this movie scene.
 	 */
-	void AddRootFolder(UMovieSceneFolder* Folder);
+	MOVIESCENE_API void AddRootFolder(UMovieSceneFolder* Folder);
 
 	/**
 	 * Removes a root folder for this movie scene (does not delete tracks or objects contained within)
 	 */
-	int32 RemoveRootFolder(UMovieSceneFolder* Folder);
+	MOVIESCENE_API int32 RemoveRootFolder(UMovieSceneFolder* Folder);
 
 	/**
 	 * Removes a root folder for this movie scene (does not delete tracks or objects contained within)
 	 */
-	bool RemoveRootFolder(int32 FolderIndex);
+	MOVIESCENE_API bool RemoveRootFolder(int32 FolderIndex);
 
 	/**
 	 * Removes all root folders from this movie scene (does not delete tracks or objects contained within)
 	 */
-	void EmptyRootFolders();
+	MOVIESCENE_API void EmptyRootFolders();
 
 	/**
 	 * Gets the nodes marked as solo in the editor, as node tree paths
@@ -970,7 +970,7 @@ public:
 	 * @param bAlwaysMarkDirty Whether to always mark the playback range dirty when changing it. 
 	 *        In the case where the playback range is dynamic and based on section bounds, the playback range doesn't need to be dirtied when set
 	 */
-	void SetPlaybackRange(FFrameNumber Start, int32 Duration, bool bAlwaysMarkDirty = true);
+	MOVIESCENE_API void SetPlaybackRange(FFrameNumber Start, int32 Duration, bool bAlwaysMarkDirty = true);
 
 	/**
 	 * Set the playback range for this movie scene
@@ -979,7 +979,7 @@ public:
 	 * @param bAlwaysMarkDirty Whether to always mark the playback range dirty when changing it. 
 	 *        In the case where the playback range is dynamic and based on section bounds, the playback range doesn't need to be dirtied when set
 	 */
-	void SetPlaybackRange(const TRange<FFrameNumber>& NewRange, bool bAlwaysMarkDirty = true);
+	MOVIESCENE_API void SetPlaybackRange(const TRange<FFrameNumber>& NewRange, bool bAlwaysMarkDirty = true);
 
 	/**
 	 * Set the start and end working range (outer) for this movie scene
@@ -987,7 +987,7 @@ public:
 	 * @param Start The offset from 0-time to view this movie scene.
 	 * @param End The offset from 0-time to view this movie scene
 	 */
-	void SetWorkingRange(float Start, float End);
+	MOVIESCENE_API void SetWorkingRange(float Start, float End);
 
 	/**
 	 * Set the start and end view range (inner) for this movie scene
@@ -995,7 +995,7 @@ public:
 	 * @param Start The offset from 0-time to view this movie scene
 	 * @param End The offset from 0-time to view this movie scene
 	 */
-	void SetViewRange(float Start, float End);
+	MOVIESCENE_API void SetViewRange(float Start, float End);
 
 #if WITH_EDITORONLY_DATA
 
@@ -1014,22 +1014,22 @@ public:
 	/**
 	* Return whether the playback range is locked.
 	*/
-	bool IsPlaybackRangeLocked() const;
+	MOVIESCENE_API bool IsPlaybackRangeLocked() const;
 
 	/**
 	* Set whether the playback range is locked.
 	*/
-	void SetPlaybackRangeLocked(bool bLocked);
+	MOVIESCENE_API void SetPlaybackRangeLocked(bool bLocked);
 
 	/**
 	 * Return whether marked frames are locked.
 	 */
-	bool AreMarkedFramesLocked() const;
+	MOVIESCENE_API bool AreMarkedFramesLocked() const;
 
 	/**
 	 * Set whether marked frames are locked.
 	 */
-	void SetMarkedFramesLocked(bool bLocked);
+	MOVIESCENE_API void SetMarkedFramesLocked(bool bLocked);
 
 	/**
 	 * @return The editor only data for use with this movie scene
@@ -1051,27 +1051,27 @@ public:
 	/*
 	 * @return Whether the section is in a group
 	 */
-	bool IsSectionInGroup(const UMovieSceneSection& InSection) const;
+	MOVIESCENE_API bool IsSectionInGroup(const UMovieSceneSection& InSection) const;
 
 	/*
 	 * Create a group containing InSections, merging any existing groups the sections are in
 	 */
-	void GroupSections(const TArray<UMovieSceneSection*> InSections);
+	MOVIESCENE_API void GroupSections(const TArray<UMovieSceneSection*> InSections);
 
 	/*
 	 * Remove InSection from any group it currently is in
 	 */
-	void UngroupSection(const UMovieSceneSection& InSection);
+	MOVIESCENE_API void UngroupSection(const UMovieSceneSection& InSection);
 
 	/*
 	 * @return The group containing the InSection, or a nullptr if InSection is not grouped.
 	 */
-	const FMovieSceneSectionGroup* GetSectionGroup(const UMovieSceneSection& InSection) const;
+	MOVIESCENE_API const FMovieSceneSectionGroup* GetSectionGroup(const UMovieSceneSection& InSection) const;
 
 	/*
 	 * Cleans stale UMovieSceneSection pointers, and removes any section groups which are no longer valid, e.g. contain less that two valid sections
 	 */
-	void CleanSectionGroups();
+	MOVIESCENE_API void CleanSectionGroups();
 
 	UMovieSceneNodeGroupCollection& GetNodeGroups() { return *NodeGroupCollection; }
 
@@ -1090,7 +1090,7 @@ public:
 	 * @InMarkIndex The given user marked frame index to edit
 	 * @InFrameNumber The frame number to set
 	 */
-	void SetMarkedFrame(int32 InMarkIndex, FFrameNumber InFrameNumber);
+	MOVIESCENE_API void SetMarkedFrame(int32 InMarkIndex, FFrameNumber InFrameNumber);
 
 	/*
 	 * Add a given user marked frame.
@@ -1099,38 +1099,38 @@ public:
 	 * @InMarkedFrame The given user marked frame to add
 	 * @return The index to the newly added marked frame
 	 */
-	int32 AddMarkedFrame(const FMovieSceneMarkedFrame& InMarkedFrame);
+	MOVIESCENE_API int32 AddMarkedFrame(const FMovieSceneMarkedFrame& InMarkedFrame);
 
 	/*
 	 * Delete the user marked frame by index.
 	 *
 	 * @DeleteIndex The index to the user marked frame to delete
 	 */
-	void DeleteMarkedFrame(int32 DeleteIndex);
+	MOVIESCENE_API void DeleteMarkedFrame(int32 DeleteIndex);
 
 	/*
 	 * Delete all user marked frames
 	 */
-	void DeleteMarkedFrames();
+	MOVIESCENE_API void DeleteMarkedFrames();
 
 	/*
 	 * Sort the marked frames in chronological order
 	 */
-	void SortMarkedFrames();
+	MOVIESCENE_API void SortMarkedFrames();
 
 	/*
 	 * Find the user marked frame by label
 	 *
 	 * @InLabel The label to the user marked frame to find
 	 */
-	int32 FindMarkedFrameByLabel(const FString& InLabel) const;
+	MOVIESCENE_API int32 FindMarkedFrameByLabel(const FString& InLabel) const;
 
 	/*
 	 * Find the user marked frame by frame number
 	 *
 	 * @InFrameNumber The frame number of the user marked frame to find
 	 */
-	int32 FindMarkedFrameByFrameNumber(FFrameNumber InFrameNumber) const;
+	MOVIESCENE_API int32 FindMarkedFrameByFrameNumber(FFrameNumber InFrameNumber) const;
 
 	/*
 	 * Find the next/previous user marked frame from the given frame number
@@ -1138,7 +1138,7 @@ public:
 	 * @InFrameNumber The frame number to find the next/previous user marked frame from
 	 * @bForward Find forward from the given frame number.
 	 */
-	int32 FindNextMarkedFrame(FFrameNumber InFrameNumber, bool bForward);
+	MOVIESCENE_API int32 FindNextMarkedFrame(FFrameNumber InFrameNumber, bool bForward);
 
 #if WITH_EDITORONLY_DATA
 	/*
@@ -1168,22 +1168,22 @@ public:
 	/*
 	 * Add a new binding group for the specified name
 	 */
-	void AddNewBindingTag(const FName& NewTag);
+	MOVIESCENE_API void AddNewBindingTag(const FName& NewTag);
 
 	/*
 	 * Tag the specified binding ID with the specified name
 	 */
-	void TagBinding(const FName& NewTag, const UE::MovieScene::FFixedObjectBindingID& BindingToTag);
+	MOVIESCENE_API void TagBinding(const FName& NewTag, const UE::MovieScene::FFixedObjectBindingID& BindingToTag);
 
 	/*
 	 * Remove a tag from the specified object binding
 	 */
-	void UntagBinding(const FName& Tag, const UE::MovieScene::FFixedObjectBindingID& Binding);
+	MOVIESCENE_API void UntagBinding(const FName& Tag, const UE::MovieScene::FFixedObjectBindingID& Binding);
 
 	/*
 	 * Remove the specified tag from any binding and forget about it completely
 	 */
-	void RemoveTag(const FName& TagToRemove);
+	MOVIESCENE_API void RemoveTag(const FName& TagToRemove);
 
 protected:
 
@@ -1192,21 +1192,21 @@ protected:
 	 *
 	 * @param Guid The guid bound to animation data to remove
 	 */
-	void RemoveBinding(const FGuid& Guid);
+	MOVIESCENE_API void RemoveBinding(const FGuid& Guid);
 
 protected:
 
 	/** Called before this object is being deserialized. */
-	virtual void PreSave(FObjectPreSaveContext ObjectSaveContext) override;
+	MOVIESCENE_API virtual void PreSave(FObjectPreSaveContext ObjectSaveContext) override;
 
 	/** Perform legacy upgrade of time ranges */
-	void UpgradeTimeRanges();
+	MOVIESCENE_API void UpgradeTimeRanges();
 
 private:
 
 #if WITH_EDITOR
-	void OptimizeForCook();
-	void RemoveNullTracks();
+	MOVIESCENE_API void OptimizeForCook();
+	MOVIESCENE_API void RemoveNullTracks();
 #endif
 
 private:
diff --git a/Engine/Source/Runtime/MovieScene/Public/MovieSceneBindingOverrides.h b/Engine/Source/Runtime/MovieScene/Public/MovieSceneBindingOverrides.h
--- a/Engine/Source/Runtime/MovieScene/Public/MovieSceneBindingOverrides.h
+++ b/Engine/Source/Runtime/MovieScene/Public/MovieSceneBindingOverrides.h
@@ -49,8 +49,8 @@ struct FMovieSceneBindingOverrideData
 /**
  * A one-to-many definition of movie scene object binding IDs to overridden objects that should be bound to that binding.
  */
-UCLASS(DefaultToInstanced, EditInlineNew)
-class MOVIESCENE_API UMovieSceneBindingOverrides
+UCLASS(DefaultToInstanced, EditInlineNew, MinimalAPI)
+class UMovieSceneBindingOverrides
 	: public UObject
 {
 public:
@@ -58,9 +58,9 @@ public:
 	GENERATED_BODY()
 
 	/** Default constructor */
-	UMovieSceneBindingOverrides(const FObjectInitializer& Init);
+	MOVIESCENE_API UMovieSceneBindingOverrides(const FObjectInitializer& Init);
 
-	bool LocateBoundObjects(const FGuid& InBindingId, FMovieSceneSequenceID InSequenceID, TArray<UObject*, TInlineAllocator<1>>& OutObjects) const;
+	MOVIESCENE_API bool LocateBoundObjects(const FGuid& InBindingId, FMovieSceneSequenceID InSequenceID, TArray<UObject*, TInlineAllocator<1>>& OutObjects) const;
 
 public:
 
@@ -68,35 +68,35 @@ public:
 	 * Overrides the specified binding with the specified objects,
 	 * optionally still allowing the bindings defined in the sequence.
 	 */
-	void SetBinding(FMovieSceneObjectBindingID Binding, const TArray<UObject*>& Objects, bool bAllowBindingsFromAsset = false);
+	MOVIESCENE_API void SetBinding(FMovieSceneObjectBindingID Binding, const TArray<UObject*>& Objects, bool bAllowBindingsFromAsset = false);
 
 	/**
 	 * Adds the specified actor to the overridden bindings for the specified binding ID,
 	 * optionally still allowing the bindings defined in the sequence.
 	 */
-	void AddBinding(FMovieSceneObjectBindingID Binding, UObject* Object, bool bAllowBindingsFromAsset = false);
+	MOVIESCENE_API void AddBinding(FMovieSceneObjectBindingID Binding, UObject* Object, bool bAllowBindingsFromAsset = false);
 
 	/**
 	 * Removes the specified actor from the specified binding's actor array.
 	 */
-	void RemoveBinding(FMovieSceneObjectBindingID Binding, UObject* Object);
+	MOVIESCENE_API void RemoveBinding(FMovieSceneObjectBindingID Binding, UObject* Object);
 
 	/**
 	 * Resets the specified binding back to the defaults defined by the sequence.
 	 */
-	void ResetBinding(FMovieSceneObjectBindingID Binding);
+	MOVIESCENE_API void ResetBinding(FMovieSceneObjectBindingID Binding);
 
 	/** Resets all overridden bindings back to the defaults defined by the sequence. */
-	void ResetBindings();
+	MOVIESCENE_API void ResetBindings();
 
 protected:
 
 	/** Rebuild the lookup map for efficient lookup. */
-	void RebuildLookupMap() const;
+	MOVIESCENE_API void RebuildLookupMap() const;
 
 #if WITH_EDITOR
 
-	virtual void PostEditChangeProperty(FPropertyChangedEvent& PropertyChangedEvent) override;
+	MOVIESCENE_API virtual void PostEditChangeProperty(FPropertyChangedEvent& PropertyChangedEvent) override;
 
 #endif
 
diff --git a/Engine/Source/Runtime/MovieScene/Public/MovieSceneBindingOwnerInterface.h b/Engine/Source/Runtime/MovieScene/Public/MovieSceneBindingOwnerInterface.h
--- a/Engine/Source/Runtime/MovieScene/Public/MovieSceneBindingOwnerInterface.h
+++ b/Engine/Source/Runtime/MovieScene/Public/MovieSceneBindingOwnerInterface.h
@@ -15,15 +15,15 @@ struct FMovieScenePossessable;
 
 
 /** Interface used in the editor to provide contextual information about movie scene bindings */
-UINTERFACE()
-class MOVIESCENE_API UMovieSceneBindingOwnerInterface
+UINTERFACE(MinimalAPI)
+class UMovieSceneBindingOwnerInterface
 	: public UInterface
 {
 public:
 	GENERATED_BODY()
 };
 
-class MOVIESCENE_API IMovieSceneBindingOwnerInterface
+class IMovieSceneBindingOwnerInterface
 {
 public:
 	GENERATED_BODY()
@@ -52,4 +52,4 @@ public:
 		return nullptr;
 	}
 #endif
-};
\ No newline at end of file
+};
diff --git a/Engine/Source/Runtime/MovieScene/Public/MovieSceneBindingProxy.h b/Engine/Source/Runtime/MovieScene/Public/MovieSceneBindingProxy.h
--- a/Engine/Source/Runtime/MovieScene/Public/MovieSceneBindingProxy.h
+++ b/Engine/Source/Runtime/MovieScene/Public/MovieSceneBindingProxy.h
@@ -19,7 +19,7 @@ class UMovieSceneSequence;
 // managing multiple bindings with their corresponding sequences is necessary.
 //
 USTRUCT(BlueprintType)
-struct MOVIESCENE_API FMovieSceneBindingProxy
+struct FMovieSceneBindingProxy
 {
 	GENERATED_BODY()
 
@@ -47,7 +47,7 @@ struct MOVIESCENE_API FMovieSceneBindingProxy
 		return HashCombine(GetTypeHash(In.BindingID), GetTypeHash(In.Sequence));
 	}
 
-	UMovieScene* GetMovieScene() const;
+	MOVIESCENE_API UMovieScene* GetMovieScene() const;
 
 	UPROPERTY(BlueprintReadOnly, Category=Binding)
 	FGuid BindingID;
diff --git a/Engine/Source/Runtime/MovieScene/Public/MovieSceneClipboard.h b/Engine/Source/Runtime/MovieScene/Public/MovieSceneClipboard.h
--- a/Engine/Source/Runtime/MovieScene/Public/MovieSceneClipboard.h
+++ b/Engine/Source/Runtime/MovieScene/Public/MovieSceneClipboard.h
@@ -43,7 +43,7 @@ class UMovieSceneTrack;
  * In general, keys are stored relaative to the minumum key-time in the clipboard, or some arbitrary time.
  * This cardinal time is stored with the clipboard environment.
  */
-class MOVIESCENE_API FMovieSceneClipboardKey
+class FMovieSceneClipboardKey
 {
 public:
 	/** Templated constructor accepting a specific value type */
@@ -55,14 +55,14 @@ public:
 	}
 
 	/** Copy construction/assignment */
-	FMovieSceneClipboardKey(const FMovieSceneClipboardKey& In);
-	FMovieSceneClipboardKey& operator=(const FMovieSceneClipboardKey& In);
+	MOVIESCENE_API FMovieSceneClipboardKey(const FMovieSceneClipboardKey& In);
+	MOVIESCENE_API FMovieSceneClipboardKey& operator=(const FMovieSceneClipboardKey& In);
 
 	/** Get the time at which this key is defined */
-	FFrameNumber GetTime() const;
+	MOVIESCENE_API FFrameNumber GetTime() const;
 
 	/** Set the time at which this key is defined */
-	void SetTime(FFrameNumber InTime);
+	MOVIESCENE_API void SetTime(FFrameNumber InTime);
 
 	/**
 	 * Get the value of this key as the specified type.
@@ -181,11 +181,11 @@ public:
 
 private:
 	/** Static map of conversions between types */
-	static TMap<FName, TMap<FName, FConversionFunction>> ConversionMap;
+	static MOVIESCENE_API TMap<FName, TMap<FName, FConversionFunction>> ConversionMap;
 };
 
 /** Container for a collection of keys arranged in a track. */
-class MOVIESCENE_API FMovieSceneClipboardKeyTrack
+class FMovieSceneClipboardKeyTrack
 {
 public:
 
@@ -202,12 +202,12 @@ public:
 public:
 
 	/** Move construction/assignment */
-	FMovieSceneClipboardKeyTrack(FMovieSceneClipboardKeyTrack&& In);
-	FMovieSceneClipboardKeyTrack& operator=(FMovieSceneClipboardKeyTrack&& In);
+	MOVIESCENE_API FMovieSceneClipboardKeyTrack(FMovieSceneClipboardKeyTrack&& In);
+	MOVIESCENE_API FMovieSceneClipboardKeyTrack& operator=(FMovieSceneClipboardKeyTrack&& In);
 
 	/** Copy construction/assignment */
-	FMovieSceneClipboardKeyTrack(const FMovieSceneClipboardKeyTrack& In);
-	FMovieSceneClipboardKeyTrack& operator=(const FMovieSceneClipboardKeyTrack& In);
+	MOVIESCENE_API FMovieSceneClipboardKeyTrack(const FMovieSceneClipboardKeyTrack& In);
+	MOVIESCENE_API FMovieSceneClipboardKeyTrack& operator=(const FMovieSceneClipboardKeyTrack& In);
 
 	/** Check the type of the keys contained within this track */
 	template<typename KeyType>
@@ -253,7 +253,7 @@ public:
 	}
 
 	/** Get the name of this track */
-	const FName& GetName() const;
+	MOVIESCENE_API const FName& GetName() const;
 
 private:
 
@@ -289,31 +289,31 @@ struct FMovieSceneClipboardEnvironment
 };
 
 /** A clipboard representing serializable copied data for a movie scene */
-class MOVIESCENE_API FMovieSceneClipboard
+class FMovieSceneClipboard
 {
 public:
 	/** Default Constructor */
-	FMovieSceneClipboard();
+	MOVIESCENE_API FMovieSceneClipboard();
 
 	/** Move construction/assignment */
-	FMovieSceneClipboard(FMovieSceneClipboard&& In);
-	FMovieSceneClipboard& operator=(FMovieSceneClipboard&& In);
+	MOVIESCENE_API FMovieSceneClipboard(FMovieSceneClipboard&& In);
+	MOVIESCENE_API FMovieSceneClipboard& operator=(FMovieSceneClipboard&& In);
 
 	/** Copy construction/assignment */
-	FMovieSceneClipboard(const FMovieSceneClipboard& In);
-	FMovieSceneClipboard& operator=(const FMovieSceneClipboard& In);
+	MOVIESCENE_API FMovieSceneClipboard(const FMovieSceneClipboard& In);
+	MOVIESCENE_API FMovieSceneClipboard& operator=(const FMovieSceneClipboard& In);
 
 	/** Get the key track groups that were copied */
-	const TArray<TArray<FMovieSceneClipboardKeyTrack>>& GetKeyTrackGroups() const;
+	MOVIESCENE_API const TArray<TArray<FMovieSceneClipboardKeyTrack>>& GetKeyTrackGroups() const;
 
 	/** Get a text description of this clipboard for display on UI */
-	FText GetDisplayText() const;
+	MOVIESCENE_API FText GetDisplayText() const;
 
 	/** Get the environment to which this clipboard relates */
-	const FMovieSceneClipboardEnvironment& GetEnvironment() const;
+	MOVIESCENE_API const FMovieSceneClipboardEnvironment& GetEnvironment() const;
 
 	/** Get the environment to which this clipboard relates */
-	FMovieSceneClipboardEnvironment& GetEnvironment();
+	MOVIESCENE_API FMovieSceneClipboardEnvironment& GetEnvironment();
 
 private:
 	friend class FMovieSceneClipboardBuilder;
@@ -326,12 +326,12 @@ private:
 };
 
 /** Class responsible for building a clipboard for a movie scene */
-class MOVIESCENE_API FMovieSceneClipboardBuilder
+class FMovieSceneClipboardBuilder
 {
 public:
 
 	/** Generate a clipboard for the current state of this builder, resetting the builder back to its default state */
-	FMovieSceneClipboard Commit(TOptional<FFrameNumber> CopyRelativeTo);
+	MOVIESCENE_API FMovieSceneClipboard Commit(TOptional<FFrameNumber> CopyRelativeTo);
 
 	/**
 	 * Find or add a key track. Key tracks are group primarily by MovieSceneTrack instance, then by name
diff --git a/Engine/Source/Runtime/MovieScene/Public/MovieSceneCommonHelpers.h b/Engine/Source/Runtime/MovieScene/Public/MovieSceneCommonHelpers.h
--- a/Engine/Source/Runtime/MovieScene/Public/MovieSceneCommonHelpers.h
+++ b/Engine/Source/Runtime/MovieScene/Public/MovieSceneCommonHelpers.h
@@ -21,7 +21,7 @@ class USoundBase;
 struct FRichCurve;
 enum class EMovieSceneKeyInterpolation : uint8;
 
-class MOVIESCENE_API MovieSceneHelpers
+class MovieSceneHelpers
 {
 public:
 
@@ -32,10 +32,10 @@ public:
 	* the package to be dirty as it causes Sequencer to unnecessairly dirty
 	* actors.
 	*/
-	struct MOVIESCENE_API FMovieSceneScopedPackageDirtyGuard
+	struct FMovieSceneScopedPackageDirtyGuard
 	{
-		FMovieSceneScopedPackageDirtyGuard(class USceneComponent* InComponent);
-		virtual ~FMovieSceneScopedPackageDirtyGuard();
+		MOVIESCENE_API FMovieSceneScopedPackageDirtyGuard(class USceneComponent* InComponent);
+		MOVIESCENE_API virtual ~FMovieSceneScopedPackageDirtyGuard();
 
 	private:
 		class USceneComponent* Component;
@@ -45,7 +45,7 @@ public:
 	/** 
 	 * @return Whether the section is keyable (active, on a track that is not muted, etc 
 	 */
-	static bool IsSectionKeyable(const UMovieSceneSection*);
+	static MOVIESCENE_API bool IsSectionKeyable(const UMovieSceneSection*);
 
 	/**
 	 * Finds a section that exists at a given time
@@ -54,7 +54,7 @@ public:
 	 * @param RowIndex  Limit the search to a given row index
 	 * @return The found section or null
 	 */
-	static UMovieSceneSection* FindSectionAtTime( TArrayView<UMovieSceneSection* const> Sections, FFrameNumber Time, int32 RowIndex = INDEX_NONE );
+	static MOVIESCENE_API UMovieSceneSection* FindSectionAtTime( TArrayView<UMovieSceneSection* const> Sections, FFrameNumber Time, int32 RowIndex = INDEX_NONE );
 
 	/**
 	 * Finds the nearest section to the given time
@@ -63,13 +63,13 @@ public:
 	 * @param RowIndex  Limit the search to a given row index
 	 * @return The found section or null
 	 */
-	static UMovieSceneSection* FindNearestSectionAtTime( TArrayView<UMovieSceneSection* const> Sections, FFrameNumber Time, int32 RowIndex = INDEX_NONE );
+	static MOVIESCENE_API UMovieSceneSection* FindNearestSectionAtTime( TArrayView<UMovieSceneSection* const> Sections, FFrameNumber Time, int32 RowIndex = INDEX_NONE );
 
 	/** Find the next section that doesn't overlap - the section that has the next closest start time to the requested start time */
-	static UMovieSceneSection* FindNextSection(TArrayView<UMovieSceneSection* const> Sections, FFrameNumber Time);
+	static MOVIESCENE_API UMovieSceneSection* FindNextSection(TArrayView<UMovieSceneSection* const> Sections, FFrameNumber Time);
 
 	/** Find the previous section that doesn't overlap - the section that has the previous closest start time to the requested start time */
-	static UMovieSceneSection* FindPreviousSection(TArrayView<UMovieSceneSection* const> Sections, FFrameNumber Time);
+	static MOVIESCENE_API UMovieSceneSection* FindPreviousSection(TArrayView<UMovieSceneSection* const> Sections, FFrameNumber Time);
 
 	/*
 	 * Fix up consecutive sections so that there are no gaps
@@ -80,7 +80,7 @@ public:
 	 * @param bCleanUp Should we cleanup any invalid sections?
 	 * @return Whether the list of sections was modified as part of the clean-up
 	 */
-	static bool FixupConsecutiveSections(TArray<UMovieSceneSection*>& Sections, UMovieSceneSection& Section, bool bDelete, bool bCleanUp = false);
+	static MOVIESCENE_API bool FixupConsecutiveSections(TArray<UMovieSceneSection*>& Sections, UMovieSceneSection& Section, bool bDelete, bool bCleanUp = false);
 
 	/**
 	 * Fix up consecutive sections so that there are no gaps, but there can be overlaps, in which case the sections
@@ -92,22 +92,22 @@ public:
 	 * @param bCleanUp Should we cleanup any invalid sections?
 	 * @return Whether the list of sections was modified as part of the clean-up
 	 */
-	static bool FixupConsecutiveBlendingSections(TArray<UMovieSceneSection*>& Sections, UMovieSceneSection& Section, bool bDelete, bool bCleanUp = false);
+	static MOVIESCENE_API bool FixupConsecutiveBlendingSections(TArray<UMovieSceneSection*>& Sections, UMovieSceneSection& Section, bool bDelete, bool bCleanUp = false);
 
 	/*
  	 * Sort consecutive sections so that they are in order based on start time
  	 */
-	static void SortConsecutiveSections(TArray<UMovieSceneSection*>& Sections);
+	static MOVIESCENE_API void SortConsecutiveSections(TArray<UMovieSceneSection*>& Sections);
 
 	/*
 	 * Gather up descendant movie scenes from the incoming sequence
 	 */
-	static void GetDescendantMovieScenes(UMovieSceneSequence* InSequence, TArray<UMovieScene*> & InMovieScenes);
+	static MOVIESCENE_API void GetDescendantMovieScenes(UMovieSceneSequence* InSequence, TArray<UMovieScene*> & InMovieScenes);
 
 	/*
 	 * Gather up descendant movie scene sub-sections from the incoming movie scene
 	 */
-	static void GetDescendantSubSections(const UMovieScene* InMovieScene, TArray<UMovieSceneSubSection*>& InSubSections);
+	static MOVIESCENE_API void GetDescendantSubSections(const UMovieScene* InMovieScene, TArray<UMovieSceneSubSection*>& InSubSections);
 	
 	/**
 	 * Get the scene component from the runtime object
@@ -115,7 +115,7 @@ public:
 	 * @param Object The object to get the scene component for
 	 * @return The found scene component
 	 */	
-	static USceneComponent* SceneComponentFromRuntimeObject(UObject* Object);
+	static MOVIESCENE_API USceneComponent* SceneComponentFromRuntimeObject(UObject* Object);
 
 	/**
 	 * Get the active camera component from the actor 
@@ -123,7 +123,7 @@ public:
 	 * @param InActor The actor to look for the camera component on
 	 * @return The active camera component
 	 */
-	static UCameraComponent* CameraComponentFromActor(const AActor* InActor);
+	static MOVIESCENE_API UCameraComponent* CameraComponentFromActor(const AActor* InActor);
 
 	/**
 	 * Find and return camera component from the runtime object
@@ -131,7 +131,7 @@ public:
 	 * @param Object The object to get the camera component for
 	 * @return The found camera component
 	 */	
-	static UCameraComponent* CameraComponentFromRuntimeObject(UObject* RuntimeObject);
+	static MOVIESCENE_API UCameraComponent* CameraComponentFromRuntimeObject(UObject* RuntimeObject);
 
 	/**
 	 * Set the runtime object movable
@@ -139,7 +139,7 @@ public:
 	 * @param Object The object to set the mobility for
 	 * @param Mobility The mobility of the runtime object
 	 */
-	static void SetRuntimeObjectMobility(UObject* Object, EComponentMobility::Type ComponentMobility = EComponentMobility::Movable);
+	static MOVIESCENE_API void SetRuntimeObjectMobility(UObject* Object, EComponentMobility::Type ComponentMobility = EComponentMobility::Movable);
 
 	/*
 	 * Get the duration for the given sound
@@ -147,7 +147,7 @@ public:
 	 * @param Sound The sound to get the duration for
 	 * @return The duration in seconds
 	 */
-	static float GetSoundDuration(USoundBase* Sound);
+	static MOVIESCENE_API float GetSoundDuration(USoundBase* Sound);
 
 	/**
 	 * Sort predicate that sorts lower bounds of a range
@@ -168,7 +168,7 @@ public:
 	/**
 	 * Sort predicate that sorts overlapping sections by row primarily, then by overlap priority
 	 */
-	static bool SortOverlappingSections(const UMovieSceneSection* A, const UMovieSceneSection* B);
+	static MOVIESCENE_API bool SortOverlappingSections(const UMovieSceneSection* A, const UMovieSceneSection* B);
 
 	/*
 	* Get weight needed to modify the global difference in order to correctly key this section due to it possibly being blended by other sections.
@@ -176,7 +176,7 @@ public:
 	* @param  Time we are at.
 	* @return Returns the weight that needs to be applied to the global difference to correctly key this section.
 	*/
-	static float CalculateWeightForBlending(UMovieSceneSection* SectionToKey, FFrameNumber Time);
+	static MOVIESCENE_API float CalculateWeightForBlending(UMovieSceneSection* SectionToKey, FFrameNumber Time);
 
 	/*
 	 * Return a name unique to the spawnable names in the given movie scene
@@ -184,7 +184,7 @@ public:
 	 * @param InName The requested name to make unique.
 	 * @return The unique name
 	 */
-	static FString MakeUniqueSpawnableName(UMovieScene* InMovieScene, const FString& InName);
+	static MOVIESCENE_API FString MakeUniqueSpawnableName(UMovieScene* InMovieScene, const FString& InName);
 
 	/**
 	 * Return a copy of the source object, suitable for use as a spawnable template.
@@ -193,17 +193,17 @@ public:
 	 * @param InName The name to use for the spawnable template
 	 * @return The spawnable template
 	 */
-	static UObject* MakeSpawnableTemplateFromInstance(UObject& InSourceObject, UMovieScene* InMovieScene, FName InName);
+	static MOVIESCENE_API UObject* MakeSpawnableTemplateFromInstance(UObject& InSourceObject, UMovieScene* InMovieScene, FName InName);
 };
 
 /**
  * Manages bindings to keyed properties for a track instance. 
  * Calls UFunctions to set the value on runtime objects
  */
-class MOVIESCENE_API FTrackInstancePropertyBindings
+class FTrackInstancePropertyBindings
 {
 public:
-	FTrackInstancePropertyBindings( FName InPropertyName, const FString& InPropertyPath);
+	MOVIESCENE_API FTrackInstancePropertyBindings( FName InPropertyName, const FString& InPropertyPath);
 
 	/**
 	 * Calls the setter function for a specific runtime object or if the setter function does not exist, the property is set directly
@@ -242,14 +242,14 @@ public:
 	 * @param InRuntimeObject The runtime object whose function to call
 	 * @param PropertyValue The new value to assign to the property
 	 */
-	void CallFunctionForEnum( UObject& InRuntimeObject, int64 PropertyValue );
+	MOVIESCENE_API void CallFunctionForEnum( UObject& InRuntimeObject, int64 PropertyValue );
 
 	/**
 	 * Rebuilds the property and function mappings for a single runtime object, and adds them to the cache
 	 *
 	 * @param InRuntimeObject	The object to cache mappings for
 	 */
-	void CacheBinding( const UObject& InRuntimeObject );
+	MOVIESCENE_API void CacheBinding( const UObject& InRuntimeObject );
 
 	/**
 	 * Gets the FProperty that is bound to the track instance
@@ -257,7 +257,7 @@ public:
 	 * @param Object	The Object that owns the property
 	 * @return			The property on the object if it exists
 	 */
-	FProperty* GetProperty(const UObject& Object) const;
+	MOVIESCENE_API FProperty* GetProperty(const UObject& Object) const;
 
 	/**
 	 * Gets the current value of a property on an object
@@ -325,7 +325,7 @@ public:
 	 * @param Object	The object to get the property from
 	 * @return ValueType	The current value
 	 */
-	int64 GetCurrentValueForEnum(const UObject& Object);
+	MOVIESCENE_API int64 GetCurrentValueForEnum(const UObject& Object);
 
 	/**
 	 * Sets the current value of a property on an object
@@ -361,7 +361,7 @@ public:
 		return PropertyName;
 	}
 
-	static FProperty* FindProperty(const UObject* Object, const FString& InPropertyPath);
+	static MOVIESCENE_API FProperty* FindProperty(const UObject* Object, const FString& InPropertyPath);
 
 private:
 
@@ -431,10 +431,10 @@ private:
 		return false;
 	}
 
-	static FPropertyAddress FindPropertyAddressRecursive(void* BasePointer, UStruct* InStruct, TArray<FString>& InPropertyNames, uint32 Index);
-	static FPropertyAddress FindPropertyAddress(const UObject& Object, const FString& InPropertyPath);
+	static MOVIESCENE_API FPropertyAddress FindPropertyAddressRecursive(void* BasePointer, UStruct* InStruct, TArray<FString>& InPropertyNames, uint32 Index);
+	static MOVIESCENE_API FPropertyAddress FindPropertyAddress(const UObject& Object, const FString& InPropertyPath);
 
-	static FProperty* FindPropertyRecursive(UStruct* InStruct, TArray<FString>& InPropertyNames, uint32 Index);
+	static MOVIESCENE_API FProperty* FindPropertyRecursive(UStruct* InStruct, TArray<FString>& InPropertyNames, uint32 Index);
 
 	/** Find or add the FPropertyAndFunction for the specified object */
 	FPropertyAndFunction FindOrAdd(const UObject& InObject)
diff --git a/Engine/Source/Runtime/MovieScene/Public/MovieSceneDynamicBinding.h b/Engine/Source/Runtime/MovieScene/Public/MovieSceneDynamicBinding.h
--- a/Engine/Source/Runtime/MovieScene/Public/MovieSceneDynamicBinding.h
+++ b/Engine/Source/Runtime/MovieScene/Public/MovieSceneDynamicBinding.h
@@ -17,7 +17,7 @@ class UMovieSceneEntitySystemLinker;
 
 /** Value definition for any type-agnostic variable (exported as text) */
 USTRUCT(BlueprintType)
-struct MOVIESCENE_API FMovieSceneDynamicBindingPayloadVariable
+struct FMovieSceneDynamicBindingPayloadVariable
 {
 	GENERATED_BODY()
 
@@ -29,7 +29,7 @@ struct MOVIESCENE_API FMovieSceneDynamicBindingPayloadVariable
  * Data for a dynamic binding endpoint call.
  */
 USTRUCT()
-struct MOVIESCENE_API FMovieSceneDynamicBinding
+struct FMovieSceneDynamicBinding
 {
 	GENERATED_BODY()
 
@@ -66,7 +66,7 @@ struct MOVIESCENE_API FMovieSceneDynamicBinding
  * Optional parameter struct for dynamic binding resolver functions.
  */
 USTRUCT(BlueprintType)
-struct MOVIESCENE_API FMovieSceneDynamicBindingResolveParams
+struct FMovieSceneDynamicBindingResolveParams
 {
 	GENERATED_BODY()
 
@@ -84,7 +84,7 @@ struct MOVIESCENE_API FMovieSceneDynamicBindingResolveParams
 };
 
 USTRUCT(BlueprintType)
-struct MOVIESCENE_API FMovieSceneDynamicBindingResolveResult
+struct FMovieSceneDynamicBindingResolveResult
 {
 	GENERATED_BODY()
 
@@ -109,7 +109,7 @@ struct MOVIESCENE_API FMovieSceneDynamicBindingResolveResult
  * blueprint function graphs.
  */
 USTRUCT()
-struct MOVIESCENE_API FMovieSceneDynamicBindingContainer
+struct FMovieSceneDynamicBindingContainer
 {
 	GENERATED_BODY()
 
@@ -120,8 +120,8 @@ struct MOVIESCENE_API FMovieSceneDynamicBindingContainer
 /**
  * Default dynamic binding resolver library, with several basic resolver functions.
  */
-UCLASS(meta=(SequencerBindingResolverLibrary))
-class MOVIESCENE_API UBuiltInDynamicBindingResolverLibrary : public UBlueprintFunctionLibrary
+UCLASS(meta=(SequencerBindingResolverLibrary), MinimalAPI)
+class UBuiltInDynamicBindingResolverLibrary : public UBlueprintFunctionLibrary
 {
 	GENERATED_BODY()
 
@@ -129,6 +129,6 @@ public:
 
 	/** Resolve the bound object to the player's pawn */
 	UFUNCTION(BlueprintPure, Category="Sequencer|Dynamic Binding", meta=(WorldContext="WorldContextObject"))
-	static FMovieSceneDynamicBindingResolveResult ResolveToPlayerPawn(UObject* WorldContextObject, int32 PlayerControllerIndex = 0);
+	static MOVIESCENE_API FMovieSceneDynamicBindingResolveResult ResolveToPlayerPawn(UObject* WorldContextObject, int32 PlayerControllerIndex = 0);
 };
 
diff --git a/Engine/Source/Runtime/MovieScene/Public/MovieSceneFolder.h b/Engine/Source/Runtime/MovieScene/Public/MovieSceneFolder.h
--- a/Engine/Source/Runtime/MovieScene/Public/MovieSceneFolder.h
+++ b/Engine/Source/Runtime/MovieScene/Public/MovieSceneFolder.h
@@ -27,82 +27,82 @@ class UMovieSceneTrack;
 namespace UE { namespace MovieScene { class IFolderEventHandler; } }
 
 /** Represents a folder used for organizing objects in tracks in a movie scene. */
-UCLASS(DefaultToInstanced)
-class MOVIESCENE_API UMovieSceneFolder : public UObject
+UCLASS(DefaultToInstanced, MinimalAPI)
+class UMovieSceneFolder : public UObject
 {
 	GENERATED_UCLASS_BODY()
 
 	/** Gets the name of this folder. */
-	FName GetFolderName() const;
+	MOVIESCENE_API FName GetFolderName() const;
 
 	/** Sets the name of this folder. Automatically calls Modify() on the folder object. */
-	void SetFolderName( FName InFolderName );
+	MOVIESCENE_API void SetFolderName( FName InFolderName );
 
 	/** Gets the folders contained by this folder. */
-	TArrayView<UMovieSceneFolder* const> GetChildFolders() const;
+	MOVIESCENE_API TArrayView<UMovieSceneFolder* const> GetChildFolders() const;
 
 	/** Adds a child folder to this folder. Removes the folder from any other folders (including root folders). Automatically calls Modify() on the folder object. */
-	void AddChildFolder( UMovieSceneFolder* InChildFolder );
+	MOVIESCENE_API void AddChildFolder( UMovieSceneFolder* InChildFolder );
 
 	/** Removes a child folder from this folder. Automatically calls Modify() on the folder object. */
-	void RemoveChildFolder( UMovieSceneFolder* InChildFolder );
+	MOVIESCENE_API void RemoveChildFolder( UMovieSceneFolder* InChildFolder );
 
 	/** Gets the tracks contained by this folder. */
-	const TArray<UMovieSceneTrack*>& GetChildTracks() const;
+	MOVIESCENE_API const TArray<UMovieSceneTrack*>& GetChildTracks() const;
 
 	UE_DEPRECATED(5.2, "GetChildMasterTracks is deprecated. Please use GetChildTracks instead")
 	const TArray<UMovieSceneTrack*>& GetChildMasterTracks() const { return GetChildTracks(); }
 
 	/** Adds a track to this folder. Automatically calls Modify() on the folder object. */
-	void AddChildTrack( UMovieSceneTrack* InTrack );
+	MOVIESCENE_API void AddChildTrack( UMovieSceneTrack* InTrack );
 
 	UE_DEPRECATED(5.2, "AddChildMasterTrack is deprecated. Please use AddChildTrack instead")
 	void AddChildMasterTrack(UMovieSceneTrack* InMasterTrack) { AddChildTrack(InMasterTrack); }
 
 	/** Removes a track from this folder. Automatically calls Modify() on the folder object. */
-	void RemoveChildTrack( UMovieSceneTrack* InTrack );
+	MOVIESCENE_API void RemoveChildTrack( UMovieSceneTrack* InTrack );
 
 	UE_DEPRECATED(5.2, "RemoveChildMasterTrack is deprecated. Please use RemoveChildTrack instead")
 	void RemoveChildMasterTrack(UMovieSceneTrack* InMasterTrack) { return RemoveChildTrack(InMasterTrack); }
 
 	/** Clear all child tracks from this folder. */
-	void ClearChildTracks();
+	MOVIESCENE_API void ClearChildTracks();
 
 	UE_DEPRECATED(5.2, "ClearChildMasterTracks is deprecated. Please use ClearChildTracks instead")
 	void ClearChildMasterTracks() { ClearChildTracks(); }
 
 	/** Gets the guids for the object bindings contained by this folder. */
-	const TArray<FGuid>& GetChildObjectBindings() const;
+	MOVIESCENE_API const TArray<FGuid>& GetChildObjectBindings() const;
 
 	/** Adds a guid for an object binding to this folder. Automatically calls Modify() on the folder object. */
-	void AddChildObjectBinding(const FGuid& InObjectBinding );
+	MOVIESCENE_API void AddChildObjectBinding(const FGuid& InObjectBinding );
 
 	/** Removes a guid for an object binding from this folder. Automatically calls Modify() on the folder object. */
-	void RemoveChildObjectBinding( const FGuid& InObjectBinding );
+	MOVIESCENE_API void RemoveChildObjectBinding( const FGuid& InObjectBinding );
 
 	/** Clear all child object bindings from this folder. */
-	void ClearChildObjectBindings();
+	MOVIESCENE_API void ClearChildObjectBindings();
 
 	/** Called after this object has been deserialized */
-	virtual void PostLoad() override;
+	MOVIESCENE_API virtual void PostLoad() override;
 
 	/** Searches for a guid in this folder and its child folders, if found returns the folder containing the guid. */
-	UMovieSceneFolder* FindFolderContaining(const FGuid& InObjectBinding);
+	MOVIESCENE_API UMovieSceneFolder* FindFolderContaining(const FGuid& InObjectBinding);
 
 	/** Searches for a track in this folder and its child folders, if found returns the folder containing the track. */
-	UMovieSceneFolder* FindFolderContaining(const UMovieSceneTrack* InTrack);
+	MOVIESCENE_API UMovieSceneFolder* FindFolderContaining(const UMovieSceneTrack* InTrack);
 
 	/** Get the folder path for this folder, stopping at the given root folders */
-	static void CalculateFolderPath(UMovieSceneFolder* InFolder, TArrayView<UMovieSceneFolder* const> RootFolders, TArray<FName>& FolderPath);
+	static MOVIESCENE_API void CalculateFolderPath(UMovieSceneFolder* InFolder, TArrayView<UMovieSceneFolder* const> RootFolders, TArray<FName>& FolderPath);
 	
 	/** For the given set of folders, return the folder that has the matching folder path */
-	static UMovieSceneFolder* GetFolderWithPath(const TArray<FName>& InFolderPath, const TArray<UMovieSceneFolder*>& InFolders, TArrayView<UMovieSceneFolder* const> RootFolders);
+	static MOVIESCENE_API UMovieSceneFolder* GetFolderWithPath(const TArray<FName>& InFolderPath, const TArray<UMovieSceneFolder*>& InFolders, TArrayView<UMovieSceneFolder* const> RootFolders);
 
-	virtual void Serialize( FArchive& Archive );
+	MOVIESCENE_API virtual void Serialize( FArchive& Archive );
 
-	FName MakeUniqueChildFolderName(FName InName) const;
+	MOVIESCENE_API FName MakeUniqueChildFolderName(FName InName) const;
 
-	static FName MakeUniqueChildFolderName(FName InName, TArrayView<UMovieSceneFolder* const> InFolders);
+	static MOVIESCENE_API FName MakeUniqueChildFolderName(FName InName, TArrayView<UMovieSceneFolder* const> InFolders);
 
 #if WITH_EDITORONLY_DATA
 	/**
@@ -145,8 +145,8 @@ class MOVIESCENE_API UMovieSceneFolder : public UObject
 #endif
 
 #if WITH_EDITOR
-	virtual void PostEditUndo() override;
-	virtual void PostEditUndo(TSharedPtr<ITransactionObjectAnnotation> TransactionAnnotation) override;
+	MOVIESCENE_API virtual void PostEditUndo() override;
+	MOVIESCENE_API virtual void PostEditUndo(TSharedPtr<ITransactionObjectAnnotation> TransactionAnnotation) override;
 #endif
 
 	/** Event handlers for handling changes to this object */
diff --git a/Engine/Source/Runtime/MovieScene/Public/MovieSceneKeyProxy.h b/Engine/Source/Runtime/MovieScene/Public/MovieSceneKeyProxy.h
--- a/Engine/Source/Runtime/MovieScene/Public/MovieSceneKeyProxy.h
+++ b/Engine/Source/Runtime/MovieScene/Public/MovieSceneKeyProxy.h
@@ -11,8 +11,8 @@
 #include "GenericPlatform/GenericPlatformMath.h"
 #include "MovieSceneKeyProxy.generated.h"
 
-UINTERFACE()
-class MOVIESCENE_API UMovieSceneKeyProxy : public UInterface
+UINTERFACE(MinimalAPI)
+class UMovieSceneKeyProxy : public UInterface
 {
 	GENERATED_BODY()
 };
@@ -23,7 +23,7 @@ class MOVIESCENE_API UMovieSceneKeyProxy : public UInterface
  * When used, UpdateValuesFromRawData should be called every frame to optionally retrieve the current values
  * of the key for this proxy.
  */
-class MOVIESCENE_API IMovieSceneKeyProxy
+class IMovieSceneKeyProxy
 {
 public:
 
@@ -110,4 +110,4 @@ void IMovieSceneKeyProxy::RefreshCurrentValue(TMovieSceneChannelHandle<ChannelTy
 			OutTime  = ChannelData.GetTimes()[KeyIndex];
 		}
 	}
-}
\ No newline at end of file
+}
diff --git a/Engine/Source/Runtime/MovieScene/Public/MovieSceneKeyStructHelper.h b/Engine/Source/Runtime/MovieScene/Public/MovieSceneKeyStructHelper.h
--- a/Engine/Source/Runtime/MovieScene/Public/MovieSceneKeyStructHelper.h
+++ b/Engine/Source/Runtime/MovieScene/Public/MovieSceneKeyStructHelper.h
@@ -20,7 +20,7 @@ struct FMovieSceneChannelValueHelper;
 /**
  * Helper class that marshals user-facing data from an edit UI into particluar keys on various underlying channels
  */
-struct MOVIESCENE_API FMovieSceneKeyStructHelper
+struct FMovieSceneKeyStructHelper
 {
 	/**
 	 * Default construction
@@ -38,24 +38,24 @@ struct MOVIESCENE_API FMovieSceneKeyStructHelper
 	 *
 	 * @param InHelper    The utility class to add, bound to the key handle it wants to edit
 	 */
-	void Add(FMovieSceneChannelValueHelper&& InHelper);
+	MOVIESCENE_API void Add(FMovieSceneChannelValueHelper&& InHelper);
 
 	/**
 	 * Set the user facing values on the UI based on the unified starting time
 	 */
-	void SetStartingValues();
+	MOVIESCENE_API void SetStartingValues();
 
 	/**
 	 * Get unified starting time that should be shown on the UI
 	 */
-	TOptional<FFrameNumber> GetUnifiedKeyTime() const;
+	MOVIESCENE_API TOptional<FFrameNumber> GetUnifiedKeyTime() const;
 
 	/**
 	 * Propagate the user-facing UI values to the keys that are being represented in this class
 	 *
 	 * @param InUnifiedTime  A time to set all keys to
 	 */
-	void Apply(FFrameNumber InUnifiedTime);
+	MOVIESCENE_API void Apply(FFrameNumber InUnifiedTime);
 
 private:
 	/** Unified key time that represents all the keys */
@@ -211,4 +211,4 @@ TOptional<TTuple<FKeyHandle, FFrameNumber>> FMovieSceneChannelValueHelper::FindF
 	}
 
 	return TOptional<TTuple<FKeyHandle, FFrameNumber>>();
-}
\ No newline at end of file
+}
diff --git a/Engine/Source/Runtime/MovieScene/Public/MovieSceneLatentActionManager.h b/Engine/Source/Runtime/MovieScene/Public/MovieSceneLatentActionManager.h
--- a/Engine/Source/Runtime/MovieScene/Public/MovieSceneLatentActionManager.h
+++ b/Engine/Source/Runtime/MovieScene/Public/MovieSceneLatentActionManager.h
@@ -16,14 +16,14 @@ using FMovieSceneSequenceLatentActionDelegate = TDelegate<void(), FNotThreadSafe
 /**
  * Utility class for running latent actions created from sequence players.
  */
-class MOVIESCENE_API FMovieSceneLatentActionManager
+class FMovieSceneLatentActionManager
 {
 public:
-	void AddLatentAction(FMovieSceneSequenceLatentActionDelegate Delegate);
-	void ClearLatentActions(UObject* Object);
-	void ClearLatentActions();
+	MOVIESCENE_API void AddLatentAction(FMovieSceneSequenceLatentActionDelegate Delegate);
+	MOVIESCENE_API void ClearLatentActions(UObject* Object);
+	MOVIESCENE_API void ClearLatentActions();
 
-	void RunLatentActions(TFunctionRef<void()> FlushCallback);
+	MOVIESCENE_API void RunLatentActions(TFunctionRef<void()> FlushCallback);
 
 	bool IsEmpty() const { return LatentActions.Num() == 0; }
 
diff --git a/Engine/Source/Runtime/MovieScene/Public/MovieSceneMetaData.h b/Engine/Source/Runtime/MovieScene/Public/MovieSceneMetaData.h
--- a/Engine/Source/Runtime/MovieScene/Public/MovieSceneMetaData.h
+++ b/Engine/Source/Runtime/MovieScene/Public/MovieSceneMetaData.h
@@ -12,29 +12,29 @@
  * Movie scene meta-data that is stored on UMovieScene assets
  * Meta-data is retrieved through ULevelSequence::FindMetaData<ULevelSequenceMetaData>()
  */
-UCLASS(config = EditorSettings, PerObjectConfig, BlueprintType)
-class MOVIESCENE_API UMovieSceneMetaData : public UObject, public IMovieSceneMetaDataInterface
+UCLASS(config = EditorSettings, PerObjectConfig, BlueprintType, MinimalAPI)
+class UMovieSceneMetaData : public UObject, public IMovieSceneMetaDataInterface
 {
 public:
 	GENERATED_BODY()
 
-	UMovieSceneMetaData (const FObjectInitializer& ObjInit);
+	MOVIESCENE_API UMovieSceneMetaData (const FObjectInitializer& ObjInit);
 
 public:
 
 	/** The asset registry tag that contains the author for this meta-data */
-	static const FName AssetRegistryTag_Author;
+	static MOVIESCENE_API const FName AssetRegistryTag_Author;
 
 	/** The asset registry tag that contains the notes for this meta-data */
-	static const FName AssetRegistryTag_Notes;
+	static MOVIESCENE_API const FName AssetRegistryTag_Notes;
 
 	/** The asset registry tag that contains the created date for this meta-data */
-	static const FName AssetRegistryTag_Created;
+	static MOVIESCENE_API const FName AssetRegistryTag_Created;
 
 	/**
 	 * Access the global config instance that houses default settings for movie scene meta data for a given project
 	 */
-	static UMovieSceneMetaData* GetConfigInstance();
+	static MOVIESCENE_API UMovieSceneMetaData* GetConfigInstance();
 
 	/**
 	 * Create a new meta-data object from the project defaults
@@ -42,26 +42,26 @@ public:
 	 * @param Outer    The object to allocate the new meta-data within
 	 * @param Name     The name for the new object. Must not already exist
 	 */
-	static UMovieSceneMetaData* CreateFromDefaults(UObject* Outer, FName Name);
+	static MOVIESCENE_API UMovieSceneMetaData* CreateFromDefaults(UObject* Outer, FName Name);
 
 public:
 
 	/**
 	 * Extend the default asset registry tags
 	 */
-	virtual void ExtendAssetRegistryTags(TArray<UObject::FAssetRegistryTag>& OutTags) const override;
+	MOVIESCENE_API virtual void ExtendAssetRegistryTags(TArray<UObject::FAssetRegistryTag>& OutTags) const override;
 
 #if WITH_EDITOR
 
 	/**
 	 * Extend the default asset registry tag meta-data
 	 */
-	virtual void ExtendAssetRegistryTagMetaData(TMap<FName, UObject::FAssetRegistryTagMetadata>& OutMetadata) const override;
+	MOVIESCENE_API virtual void ExtendAssetRegistryTagMetaData(TMap<FName, UObject::FAssetRegistryTagMetadata>& OutMetadata) const override;
 
 #endif
 
 	/** Return whether this metadata has any valid data */
-	bool IsEmpty() const;
+	MOVIESCENE_API bool IsEmpty() const;
 
 public:
 
@@ -69,19 +69,19 @@ public:
 	 * @return The author for this metadata
 	 */
 	UFUNCTION(BlueprintCallable, Category = "Meta Data")
-	FString GetAuthor() const;
+	MOVIESCENE_API FString GetAuthor() const;
 
 	/**
 	 * @return The created date for this metadata
 	 */
 	UFUNCTION(BlueprintCallable, Category = "Meta Data")
-	FDateTime GetCreated() const;
+	MOVIESCENE_API FDateTime GetCreated() const;
 
 	/**
 	 * @return The notes for this metadata
 	 */
 	UFUNCTION(BlueprintCallable, Category = "Meta Data")
-	FString GetNotes() const;
+	MOVIESCENE_API FString GetNotes() const;
 
 public:
 
@@ -89,19 +89,19 @@ public:
 	 * Set this metadata's author
 	 */
 	UFUNCTION(BlueprintCallable, Category = "Meta Data")
-	void SetAuthor(FString InAuthor);
+	MOVIESCENE_API void SetAuthor(FString InAuthor);
 
 	/**
 	 * Set this metadata's created date
 	 */
 	UFUNCTION(BlueprintCallable, Category = "Meta Data")
-	void SetCreated(FDateTime InCreated);
+	MOVIESCENE_API void SetCreated(FDateTime InCreated);
 
 	/**
 	 * Set this metadata's notes
 	 */
 	UFUNCTION(BlueprintCallable, Category = "Meta Data")
-	void SetNotes(FString InNotes);
+	MOVIESCENE_API void SetNotes(FString InNotes);
 
 private:
 
diff --git a/Engine/Source/Runtime/MovieScene/Public/MovieSceneNameableTrack.h b/Engine/Source/Runtime/MovieScene/Public/MovieSceneNameableTrack.h
--- a/Engine/Source/Runtime/MovieScene/Public/MovieSceneNameableTrack.h
+++ b/Engine/Source/Runtime/MovieScene/Public/MovieSceneNameableTrack.h
@@ -47,7 +47,7 @@ public:
 	 *
 	 * @return Whether this track can be renamed.
 	 */
-	MOVIESCENE_API virtual bool CanRename() const { return true; }
+	virtual bool CanRename() const { return true; }
 
 	/** 
 	 * Validate the new display name. 
diff --git a/Engine/Source/Runtime/MovieScene/Public/MovieScenePossessable.h b/Engine/Source/Runtime/MovieScene/Public/MovieScenePossessable.h
--- a/Engine/Source/Runtime/MovieScene/Public/MovieScenePossessable.h
+++ b/Engine/Source/Runtime/MovieScene/Public/MovieScenePossessable.h
@@ -159,7 +159,7 @@ public:
 	}
 
 	/* Set the optional binding id for binding to a spawnable */
-	MOVIESCENE_API void SetSpawnableObjectBindingID(const FMovieSceneObjectBindingID& InSpawnableObjectBindingID)
+	void SetSpawnableObjectBindingID(const FMovieSceneObjectBindingID& InSpawnableObjectBindingID)
 	{
 		SpawnableObjectBindingID = InSpawnableObjectBindingID;
 	}
diff --git a/Engine/Source/Runtime/MovieScene/Public/MovieSceneSection.h b/Engine/Source/Runtime/MovieScene/Public/MovieSceneSection.h
--- a/Engine/Source/Runtime/MovieScene/Public/MovieSceneSection.h
+++ b/Engine/Source/Runtime/MovieScene/Public/MovieSceneSection.h
@@ -251,7 +251,7 @@ class UMovieSceneSection
 {
 	GENERATED_UCLASS_BODY()
 
-	MOVIESCENE_API ~UMovieSceneSection() {};
+	~UMovieSceneSection() {};
 
 public:
 
@@ -319,7 +319,7 @@ public:
 	 * 
 	 * @param NewRange	The new range of times
 	 */
-	MOVIESCENE_API virtual void SetRange(const TRange<FFrameNumber>& NewRange)
+	virtual void SetRange(const TRange<FFrameNumber>& NewRange)
 	{
 		
 		// Skip TryModify for objects that still need initialization (i.e. we're in the object's constructor), because modifying objects in their constructor can lead to non-deterministic cook issues.
@@ -440,7 +440,7 @@ public:
 	 * Sets this section's blend type
 	 */
 	UFUNCTION(BlueprintCallable, Category = "Sequencer|Section")
-	MOVIESCENE_API virtual void SetBlendType(EMovieSceneBlendType InBlendType)
+	virtual void SetBlendType(EMovieSceneBlendType InBlendType)
 	{
 		if (GetSupportedBlendTypes().Contains(InBlendType))
 		{
@@ -612,7 +612,7 @@ public:
 	virtual void OnBindingIDsUpdated(const TMap<UE::MovieScene::FFixedObjectBindingID, UE::MovieScene::FFixedObjectBindingID>& OldFixedToNewFixedMap, FMovieSceneSequenceID LocalSequenceID, const FMovieSceneSequenceHierarchy* Hierarchy, IMovieScenePlayer& Player) {}
 
 	/** Get the referenced bindings for this section */
-	MOVIESCENE_API virtual void GetReferencedBindings(TArray<FGuid>& OutBindings) {}
+	virtual void GetReferencedBindings(TArray<FGuid>& OutBindings) {}
 
 	/**
 	 * Gets a list of all overlapping sections
@@ -642,19 +642,19 @@ public:
 	MOVIESCENE_API FMovieSceneChannelProxy& GetChannelProxy() const;
 
 	/** Does this movie section support infinite ranges for evaluation */
-	MOVIESCENE_API bool GetSupportsInfiniteRange() const { return bSupportsInfiniteRange; }
+	bool GetSupportsInfiniteRange() const { return bSupportsInfiniteRange; }
 
 	/**
 	*  Whether or not we draw a curve for a particular channel owned by this section.
 	*  Defaults to true.
 	*/
-	MOVIESCENE_API virtual bool ShowCurveForChannel(const void *Channel) const  { return true; }
+	virtual bool ShowCurveForChannel(const void *Channel) const  { return true; }
 
 	/** 
 	*  Get The Total Weight Value for this Section
 	*  For Most Sections it's just the Ease Value, but for some Sections also have an extra Weight Curve
 	*/
-	MOVIESCENE_API virtual float GetTotalWeightValue(FFrameTime InTime) const { return EvaluateEasing(InTime); }
+	virtual float GetTotalWeightValue(FFrameTime InTime) const { return EvaluateEasing(InTime); }
 
 
 	/**
diff --git a/Engine/Source/Runtime/MovieScene/Public/MovieSceneSequencePlayer.h b/Engine/Source/Runtime/MovieScene/Public/MovieSceneSequencePlayer.h
--- a/Engine/Source/Runtime/MovieScene/Public/MovieSceneSequencePlayer.h
+++ b/Engine/Source/Runtime/MovieScene/Public/MovieSceneSequencePlayer.h
@@ -153,8 +153,8 @@ struct FMovieSceneSequencePlayToParams
 /**
  * Abstract class that provides consistent player behaviour for various animation players
  */
-UCLASS(Abstract, BlueprintType)
-class MOVIESCENE_API UMovieSceneSequencePlayer
+UCLASS(Abstract, BlueprintType, MinimalAPI)
+class UMovieSceneSequencePlayer
 	: public UObject
 	, public IMovieScenePlayer
 	, public IMovieSceneSequenceTickManagerClient
@@ -166,47 +166,47 @@ public:
 	UPROPERTY(replicated)
 	TScriptInterface<IMovieSceneSequencePlayerObserver> Observer;
 
-	UMovieSceneSequencePlayer(const FObjectInitializer&);
-	virtual ~UMovieSceneSequencePlayer();
+	MOVIESCENE_API UMovieSceneSequencePlayer(const FObjectInitializer&);
+	MOVIESCENE_API virtual ~UMovieSceneSequencePlayer();
 
 	/** Start playback forwards from the current time cursor position, using the current play rate. */
 	UFUNCTION(BlueprintCallable, Category = "Sequencer|Player")
-	void Play();
+	MOVIESCENE_API void Play();
 
 	/** Reverse playback. */
 	UFUNCTION(BlueprintCallable, Category = "Sequencer|Player")
-	void PlayReverse();
+	MOVIESCENE_API void PlayReverse();
 
 	/** Changes the direction of playback (go in reverse if it was going forward, or vice versa) */
 	UFUNCTION(BlueprintCallable, Category = "Sequencer|Player")
-	void ChangePlaybackDirection();
+	MOVIESCENE_API void ChangePlaybackDirection();
 
 	/**
 	 * Start playback from the current time cursor position, looping the specified number of times.
 	 * @param NumLoops - The number of loops to play. -1 indicates infinite looping.
 	 */
 	UFUNCTION(BlueprintCallable, Category = "Sequencer|Player")
-	void PlayLooping(int32 NumLoops = -1);
+	MOVIESCENE_API void PlayLooping(int32 NumLoops = -1);
 	
 	/** Pause playback. */
 	UFUNCTION(BlueprintCallable, Category = "Sequencer|Player")
-	void Pause();
+	MOVIESCENE_API void Pause();
 	
 	/** Scrub playback. */
 	UFUNCTION(BlueprintCallable, Category = "Sequencer|Player")
-	void Scrub();
+	MOVIESCENE_API void Scrub();
 
 	/** Stop playback and move the cursor to the end (or start, for reversed playback) of the sequence. */
 	UFUNCTION(BlueprintCallable, Category = "Sequencer|Player")
-	void Stop();
+	MOVIESCENE_API void Stop();
 
 	/** Stop playback without moving the cursor. */
 	UFUNCTION(BlueprintCallable, Category = "Sequencer|Player")
-	void StopAtCurrentTime();
+	MOVIESCENE_API void StopAtCurrentTime();
 
 	/** Go to end and stop. */
 	UFUNCTION(BlueprintCallable, Category = "Sequencer|Player", meta = (ToolTip = "Go to end of the sequence and stop. Adheres to 'When Finished' section rules."))
-	void GoToEndAndStop();
+	MOVIESCENE_API void GoToEndAndStop();
 
 public:
 
@@ -215,19 +215,19 @@ public:
 	 * @return The current playback position
 	 */
 	UFUNCTION(BlueprintCallable, Category = "Sequencer|Player")
-	FQualifiedFrameTime GetCurrentTime() const;
+	MOVIESCENE_API FQualifiedFrameTime GetCurrentTime() const;
 
 	/**
 	 * Get the total duration of the sequence
 	 */
 	UFUNCTION(BlueprintCallable, Category = "Sequencer|Player")
-	FQualifiedFrameTime GetDuration() const;
+	MOVIESCENE_API FQualifiedFrameTime GetDuration() const;
 
 	/**
 	 * Get this sequence's duration in frames
 	 */
 	UFUNCTION(BlueprintCallable, Category = "Sequencer|Player")
-	int32 GetFrameDuration() const;
+	MOVIESCENE_API int32 GetFrameDuration() const;
 
 	/**
 	 * Get this sequence's display rate.
@@ -239,7 +239,7 @@ public:
 	 * Set the frame-rate that this player should play with, making all frame numbers in the specified time-space
 	 */
 	UFUNCTION(BlueprintCallable, Category = "Sequencer|Player")
-	void SetFrameRate(FFrameRate FrameRate);
+	MOVIESCENE_API void SetFrameRate(FFrameRate FrameRate);
 
 	/**
 	 * Get the offset within the level sequence to start playing
@@ -261,13 +261,13 @@ public:
 	 * @param InWeight    The weight to suuply to all elements in this sequence
 	 */
 	UFUNCTION(BlueprintCallable, Category = "Sequencer|Player")
-	void SetWeight(double InWeight);
+	MOVIESCENE_API void SetWeight(double InWeight);
 
 	/**
 	 * Removes a previously assigned weight
 	 */
 	UFUNCTION(BlueprintCallable, Category = "Sequencer|Player")
-	void RemoveWeight();
+	MOVIESCENE_API void RemoveWeight();
 
 	/**
 	 * Set a manual weight to be multiplied with all blendable elements within the specified sequence
@@ -276,12 +276,12 @@ public:
 	 *
 	 * @param InWeight    The weight to suuply to all elements in this sequence
 	 */
-	void SetWeight(double InWeight, FMovieSceneSequenceID SequenceID);
+	MOVIESCENE_API void SetWeight(double InWeight, FMovieSceneSequenceID SequenceID);
 
 	/**
 	 * Removes a previously assigned weight
 	 */
-	void RemoveWeight(FMovieSceneSequenceID SequenceID);
+	MOVIESCENE_API void RemoveWeight(FMovieSceneSequenceID SequenceID);
 
 public:
 
@@ -292,7 +292,7 @@ public:
 	 * @param Duration        The number of frames to play
 	 */
 	UFUNCTION(BlueprintCallable, Category = "Sequencer|Player", DisplayName="Set Play Range (Frames)")
-	void SetFrameRange( int32 StartFrame, int32 Duration, float SubFrames = 0.f );
+	MOVIESCENE_API void SetFrameRange( int32 StartFrame, int32 Duration, float SubFrames = 0.f );
 
 	/**
 	 * Set the valid play range for this sequence, determined by a starting time  and a duration (in seconds)
@@ -301,7 +301,7 @@ public:
 	 * @param Duration        The length to play for
 	 */
 	UFUNCTION(BlueprintCallable, Category = "Sequencer|Player", DisplayName="Set Play Range (Seconds)")
-	void SetTimeRange( float StartTime, float Duration );
+	MOVIESCENE_API void SetTimeRange( float StartTime, float Duration );
 
 public:
 
@@ -313,7 +313,7 @@ public:
 	 * @param PlaybackParams The position settings (ie. the position to play to)
 	 */
 	UFUNCTION(BlueprintCallable, Category = "Sequencer|Player")
-	void PlayTo(FMovieSceneSequencePlaybackParams PlaybackParams, FMovieSceneSequencePlayToParams PlayToParams);
+	MOVIESCENE_API void PlayTo(FMovieSceneSequencePlaybackParams PlaybackParams, FMovieSceneSequencePlayToParams PlayToParams);
 
 	/**
 	 * Set the current time of the player by evaluating from the current time to the specified time, as if the sequence is playing. 
@@ -322,38 +322,38 @@ public:
 	 * @param PlaybackParams The position settings (ie. the position to set playback to)
 	 */
 	UFUNCTION(BlueprintCallable, Category = "Sequencer|Player")
-	void SetPlaybackPosition(FMovieSceneSequencePlaybackParams PlaybackParams);
+	MOVIESCENE_API void SetPlaybackPosition(FMovieSceneSequencePlaybackParams PlaybackParams);
 
 	/**
 	 * Restore any changes made by this player to their original state
 	 */
 	UFUNCTION(BlueprintCallable, Category="Game|Cinematic")
-	void RestoreState();
+	MOVIESCENE_API void RestoreState();
 
 public:
 
 	/** Check whether the sequence is actively playing. */
 	UFUNCTION(BlueprintCallable, Category = "Sequencer|Player")
-	bool IsPlaying() const;
+	MOVIESCENE_API bool IsPlaying() const;
 
 	/** Check whether the sequence is paused. */
 	UFUNCTION(BlueprintCallable, Category = "Sequencer|Player")
-	bool IsPaused() const;
+	MOVIESCENE_API bool IsPaused() const;
 
 	/** Check whether playback is reversed. */
 	UFUNCTION(BlueprintCallable, Category = "Sequencer|Player")
-	bool IsReversed() const;
+	MOVIESCENE_API bool IsReversed() const;
 
 	/** Get the playback rate of this player. */
 	UFUNCTION(BlueprintCallable, Category = "Sequencer|Player")
-	float GetPlayRate() const;
+	MOVIESCENE_API float GetPlayRate() const;
 
 	/**
 	 * Set the playback rate of this player. Negative values will play the animation in reverse.
 	 * @param PlayRate - The new rate of playback for the animation.
 	 */
 	UFUNCTION(BlueprintCallable, Category = "Sequencer|Player")
-	void SetPlayRate(float PlayRate);
+	MOVIESCENE_API void SetPlayRate(float PlayRate);
 
 	/** Set whether to disable camera cuts */
 	UFUNCTION(BlueprintCallable, Category = "Sequencer|Player")
@@ -392,34 +392,34 @@ public:
 
 	/** Retrieve all objects currently bound to the specified binding identifier */
 	UFUNCTION(BlueprintCallable, Category = "Sequencer|Player")
-	TArray<UObject*> GetBoundObjects(FMovieSceneObjectBindingID ObjectBinding);
+	MOVIESCENE_API TArray<UObject*> GetBoundObjects(FMovieSceneObjectBindingID ObjectBinding);
 
 	/** Get the object bindings for the requested object */
 	UFUNCTION(BlueprintCallable, Category = "Sequencer|Player")
-	TArray<FMovieSceneObjectBindingID> GetObjectBindings(UObject* InObject);
+	MOVIESCENE_API TArray<FMovieSceneObjectBindingID> GetObjectBindings(UObject* InObject);
 
 public:
 
 	/** Ensure that this player's tick manager is set up correctly for the specified context */
-	void InitializeForTick(UObject* Context);
+	MOVIESCENE_API void InitializeForTick(UObject* Context);
 
 	/** Assign this player's playback settings */
-	void SetPlaybackSettings(const FMovieSceneSequencePlaybackSettings& InSettings);
+	MOVIESCENE_API void SetPlaybackSettings(const FMovieSceneSequencePlaybackSettings& InSettings);
 
 	/** Initialize this player using its existing playback settings */
-	void Initialize(UMovieSceneSequence* InSequence);
+	MOVIESCENE_API void Initialize(UMovieSceneSequence* InSequence);
 
 	/** Initialize this player with a sequence and some settings */
-	void Initialize(UMovieSceneSequence* InSequence, const FMovieSceneSequencePlaybackSettings& InSettings);
+	MOVIESCENE_API void Initialize(UMovieSceneSequence* InSequence, const FMovieSceneSequencePlaybackSettings& InSettings);
 
 	/** Update the sequence for the current time, if playing */
-	void Update(const float DeltaSeconds);
+	MOVIESCENE_API void Update(const float DeltaSeconds);
 
 	/** Update the sequence for the current time, if playing, asynchronously */
-	void UpdateAsync(const float DeltaSeconds);
+	MOVIESCENE_API void UpdateAsync(const float DeltaSeconds);
 
 	/** Perform any tear-down work when this player is no longer (and will never) be needed */
-	void TearDown();
+	MOVIESCENE_API void TearDown();
 
 public:
 
@@ -436,7 +436,7 @@ public:
 	 * @return the name of the sequence, or None if no sequence is set
 	 */
 	UFUNCTION(BlueprintCallable, Category = "Sequencer|Player")
-	FString GetSequenceName(bool bAddClientInfo = false) const;
+	MOVIESCENE_API FString GetSequenceName(bool bAddClientInfo = false) const;
 
 	/**
 	 * Access this player's tick manager
@@ -446,37 +446,37 @@ public:
 	/**
 	 * Assign a playback client interface for this sequence player, defining instance data and binding overrides
 	 */
-	void SetPlaybackClient(TScriptInterface<IMovieScenePlaybackClient> InPlaybackClient);
+	MOVIESCENE_API void SetPlaybackClient(TScriptInterface<IMovieScenePlaybackClient> InPlaybackClient);
 
 	/**
 	 * Retrieve the currently assigned time controller
 	 */
-	TSharedPtr<FMovieSceneTimeController> GetTimeController() const;
+	MOVIESCENE_API TSharedPtr<FMovieSceneTimeController> GetTimeController() const;
 
 	/**
 	 * Assign a time controller for this sequence player allowing custom time management implementations.
 	 * Will reset the supplied time controller to the current time.
 	 */
-	void SetTimeController(TSharedPtr<FMovieSceneTimeController> InTimeController);
+	MOVIESCENE_API void SetTimeController(TSharedPtr<FMovieSceneTimeController> InTimeController);
 	
 
 	/**
 	 * Assign a time controller for this sequence player allowing custom time management implementations.
 	 * Will not reset the supplied time controller in any way, so the sequence will receive its time directly from the controller.
 	 */
-	void SetTimeControllerDirectly(TSharedPtr<FMovieSceneTimeController> InTimeController);
+	MOVIESCENE_API void SetTimeControllerDirectly(TSharedPtr<FMovieSceneTimeController> InTimeController);
 
 	/**
 	 * Sets whether to listen or ignore playback replication events.
 	 * @param bState If true, ignores playback replication.
 	 */
-	void SetIgnorePlaybackReplication(bool bState);
+	MOVIESCENE_API void SetIgnorePlaybackReplication(bool bState);
 
 protected:
 
-	void PlayInternal();
-	void StopInternal(FFrameTime TimeToResetTo);
-	void FinishPlaybackInternal(FFrameTime TimeToFinishAt);
+	MOVIESCENE_API void PlayInternal();
+	MOVIESCENE_API void StopInternal(FFrameTime TimeToResetTo);
+	MOVIESCENE_API void FinishPlaybackInternal(FFrameTime TimeToFinishAt);
 
 	struct FMovieSceneUpdateArgs
 	{
@@ -484,27 +484,27 @@ protected:
 		bool bIsAsync = false;
 	};
 
-	void UpdateMovieSceneInstance(FMovieSceneEvaluationRange InRange, EMovieScenePlayerStatus::Type PlayerStatus, bool bHasJumped = false);
-	virtual void UpdateMovieSceneInstance(FMovieSceneEvaluationRange InRange, EMovieScenePlayerStatus::Type PlayerStatus, const FMovieSceneUpdateArgs& Args);
+	MOVIESCENE_API void UpdateMovieSceneInstance(FMovieSceneEvaluationRange InRange, EMovieScenePlayerStatus::Type PlayerStatus, bool bHasJumped = false);
+	MOVIESCENE_API virtual void UpdateMovieSceneInstance(FMovieSceneEvaluationRange InRange, EMovieScenePlayerStatus::Type PlayerStatus, const FMovieSceneUpdateArgs& Args);
 
-	void UpdateTimeCursorPosition(FFrameTime NewPosition, EUpdatePositionMethod Method, bool bHasJumpedOverride = false);
-	bool ShouldStopOrLoop(FFrameTime NewPosition) const;
+	MOVIESCENE_API void UpdateTimeCursorPosition(FFrameTime NewPosition, EUpdatePositionMethod Method, bool bHasJumpedOverride = false);
+	MOVIESCENE_API bool ShouldStopOrLoop(FFrameTime NewPosition) const;
 	/** 
 	* If the current sequence should pause (due to NewPosition overshooting a previously set ShouldPause) 
 	* then a range of time that should be evaluated to reach there will be returned. If we should not pause
 	* then the TOptional will be unset.
 	* */
-	TOptional<TRange<FFrameTime>> GetPauseRange(const FFrameTime& NewPosition) const;
+	MOVIESCENE_API TOptional<TRange<FFrameTime>> GetPauseRange(const FFrameTime& NewPosition) const;
 
-	UWorld* GetPlaybackWorld() const;
+	MOVIESCENE_API UWorld* GetPlaybackWorld() const;
 
-	FFrameTime GetLastValidTime() const;
+	MOVIESCENE_API FFrameTime GetLastValidTime() const;
 
-	FFrameRate GetDisplayRate() const;
+	MOVIESCENE_API FFrameRate GetDisplayRate() const;
 
-	bool NeedsQueueLatentAction() const;
-	void QueueLatentAction(FMovieSceneSequenceLatentActionDelegate Delegate);
-	void RunLatentActions();
+	MOVIESCENE_API bool NeedsQueueLatentAction() const;
+	MOVIESCENE_API void QueueLatentAction(FMovieSceneSequenceLatentActionDelegate Delegate);
+	MOVIESCENE_API void RunLatentActions();
 
 public:
 	//~ IMovieScenePlayer interface
@@ -512,9 +512,9 @@ public:
 
 protected:
 	//~ IMovieScenePlayer interface
-	virtual UMovieSceneEntitySystemLinker* ConstructEntitySystemLinker() override;
-	virtual EMovieScenePlayerStatus::Type GetPlaybackStatus() const override;
-	virtual FMovieSceneSpawnRegister& GetSpawnRegister() override;
+	MOVIESCENE_API virtual UMovieSceneEntitySystemLinker* ConstructEntitySystemLinker() override;
+	MOVIESCENE_API virtual EMovieScenePlayerStatus::Type GetPlaybackStatus() const override;
+	MOVIESCENE_API virtual FMovieSceneSpawnRegister& GetSpawnRegister() override;
 	virtual UObject* AsUObject() override { return this; }
 
 	virtual void SetPlaybackStatus(EMovieScenePlayerStatus::Type InPlaybackStatus) override {}
@@ -522,23 +522,23 @@ protected:
 	virtual void GetViewportSettings(TMap<FViewportClient*, EMovieSceneViewportParams>& ViewportParamsMap) const override {}
 	virtual bool CanUpdateCameraCut() const override { return !PlaybackSettings.bDisableCameraCuts; }
 	virtual void UpdateCameraCut(UObject* CameraObject, const EMovieSceneCameraCutParams& CameraCutParams) override {}
-	virtual void ResolveBoundObjects(const FGuid& InBindingId, FMovieSceneSequenceID SequenceID, UMovieSceneSequence& Sequence, UObject* ResolutionContext, TArray<UObject*, TInlineAllocator<1>>& OutObjects) const override;
+	MOVIESCENE_API virtual void ResolveBoundObjects(const FGuid& InBindingId, FMovieSceneSequenceID SequenceID, UMovieSceneSequence& Sequence, UObject* ResolutionContext, TArray<UObject*, TInlineAllocator<1>>& OutObjects) const override;
 	virtual IMovieScenePlaybackClient* GetPlaybackClient() override { return PlaybackClient ? &*PlaybackClient : nullptr; }
-	virtual bool IsDisablingEventTriggers(FFrameTime& DisabledUntilTime) const override;
-	virtual bool HasDynamicWeighting() const override;
-	virtual void PreEvaluation(const FMovieSceneContext& Context) override;
-	virtual void PostEvaluation(const FMovieSceneContext& Context) override;
+	MOVIESCENE_API virtual bool IsDisablingEventTriggers(FFrameTime& DisabledUntilTime) const override;
+	MOVIESCENE_API virtual bool HasDynamicWeighting() const override;
+	MOVIESCENE_API virtual void PreEvaluation(const FMovieSceneContext& Context) override;
+	MOVIESCENE_API virtual void PostEvaluation(const FMovieSceneContext& Context) override;
 
 	/*~ Begin UObject interface */
 	virtual bool IsSupportedForNetworking() const { return true; }
-	virtual int32 GetFunctionCallspace(UFunction* Function, FFrame* Stack) override;
-	virtual bool CallRemoteFunction(UFunction* Function, void* Parameters, FOutParmRec* OutParms, FFrame* Stack) override;
-	virtual void PostNetReceive() override;
-	virtual void BeginDestroy() override;
+	MOVIESCENE_API virtual int32 GetFunctionCallspace(UFunction* Function, FFrame* Stack) override;
+	MOVIESCENE_API virtual bool CallRemoteFunction(UFunction* Function, void* Parameters, FOutParmRec* OutParms, FFrame* Stack) override;
+	MOVIESCENE_API virtual void PostNetReceive() override;
+	MOVIESCENE_API virtual void BeginDestroy() override;
 	/*~ End UObject interface */
 
 	//~ Begin IMovieSceneSequenceTickManagerClient interface
-	virtual void TickFromSequenceTickManager(float DeltaSeconds, FMovieSceneEntitySystemRunner* Runner) override;
+	MOVIESCENE_API virtual void TickFromSequenceTickManager(float DeltaSeconds, FMovieSceneEntitySystemRunner* Runner) override;
 	//~ End IMovieSceneSequenceTickManagerClient interface
 
 protected:
@@ -551,10 +551,10 @@ protected:
 	
 private:
 
-	void UpdateTimeCursorPosition_Internal(FFrameTime NewPosition, EUpdatePositionMethod Method, bool bHasJumpedOverride);
+	MOVIESCENE_API void UpdateTimeCursorPosition_Internal(FFrameTime NewPosition, EUpdatePositionMethod Method, bool bHasJumpedOverride);
 
-	void RunPreEvaluationCallbacks();
-	void RunPostEvaluationCallbacks();
+	MOVIESCENE_API void RunPreEvaluationCallbacks();
+	MOVIESCENE_API void RunPostEvaluationCallbacks();
 	
 private:
 
@@ -562,46 +562,46 @@ private:
 	 * Called on the server whenever an explicit change in time has occurred through one of the (Play|Jump|Scrub)To methods
 	 */
 	UFUNCTION(netmulticast, reliable)
-	void RPC_ExplicitServerUpdateEvent(EUpdatePositionMethod Method, FFrameTime RelevantTime);
+	MOVIESCENE_API void RPC_ExplicitServerUpdateEvent(EUpdatePositionMethod Method, FFrameTime RelevantTime);
 
 	/**
 	 * Called on the server when Stop() is called in order to differentiate Stops from Pauses.
 	 */
 	UFUNCTION(netmulticast, reliable)
-	void RPC_OnStopEvent(FFrameTime StoppedTime);
+	MOVIESCENE_API void RPC_OnStopEvent(FFrameTime StoppedTime);
 
 	/**
 	 * Called on the server when playback has reached the end. Could lead to stopping or pausing.
 	 */
 	UFUNCTION(netmulticast, reliable)
-	void RPC_OnFinishPlaybackEvent(FFrameTime StoppedTime);
+	MOVIESCENE_API void RPC_OnFinishPlaybackEvent(FFrameTime StoppedTime);
 
 	/**
 	 * Check whether this sequence player is an authority, as determined by its outer Actor
 	 */
-	bool HasAuthority() const;
+	MOVIESCENE_API bool HasAuthority() const;
 
 	/**
 	 * Update the replicated properties required for synchronizing to clients of this sequence player
 	 */
-	void UpdateNetworkSyncProperties();
+	MOVIESCENE_API void UpdateNetworkSyncProperties();
 
 	/**
 	 * Analyse the set of samples we have estimating the server time if we have confidence over the data.
 	 * Should only be called once per frame.
 	 * @return An estimation of the server time, or the current local time if we cannot make a strong estimate
 	 */
-	FFrameTime UpdateServerTimeSamples();
+	MOVIESCENE_API FFrameTime UpdateServerTimeSamples();
 
 	/**
 	 * Check and correct network synchronization for the clients of this sequence player.
 	 */
-	void UpdateNetworkSync();
+	MOVIESCENE_API void UpdateNetworkSync();
 
 	/**
 	 * Compute the latency for the client connection.
 	 */
-	float GetPing() const;
+	MOVIESCENE_API float GetPing() const;
 
 protected:
 
diff --git a/Engine/Source/Runtime/MovieScene/Public/MovieSceneSequenceTickInterval.h b/Engine/Source/Runtime/MovieScene/Public/MovieSceneSequenceTickInterval.h
--- a/Engine/Source/Runtime/MovieScene/Public/MovieSceneSequenceTickInterval.h
+++ b/Engine/Source/Runtime/MovieScene/Public/MovieSceneSequenceTickInterval.h
@@ -15,7 +15,7 @@ class UObject;
  * Structure defining a concrete tick interval for a Sequencer based evaluation
  */
 USTRUCT(BlueprintType)
-struct MOVIESCENE_API FMovieSceneSequenceTickInterval
+struct FMovieSceneSequenceTickInterval
 {
 	GENERATED_BODY()
 
@@ -24,12 +24,12 @@ struct MOVIESCENE_API FMovieSceneSequenceTickInterval
 	/**
 	 * Generate a tick interval from an actor's primary tick function
 	 */
-	FMovieSceneSequenceTickInterval(const AActor* InActor);
+	MOVIESCENE_API FMovieSceneSequenceTickInterval(const AActor* InActor);
 
 	/**
 	 * Generate a tick interval from an component's primary tick function
 	 */
-	FMovieSceneSequenceTickInterval(const UActorComponent* InActorComponent);
+	MOVIESCENE_API FMovieSceneSequenceTickInterval(const UActorComponent* InActorComponent);
 
 public:
 
@@ -54,13 +54,13 @@ public:
 	/**
 	 * Round this interval to the nearest Sequencer.TickIntervalGroupingResolutionMs milliseconds
 	 */
-	int32 RoundTickIntervalMs() const;
+	MOVIESCENE_API int32 RoundTickIntervalMs() const;
 
 	/**
 	 * Resolve this tick interval within the specified context object (usually a movie scene player)
 	 * inheriting properties from the first valid parent if possible
 	 */
-	static FMovieSceneSequenceTickInterval GetInheritedInterval(UObject* ContextObject);
+	static MOVIESCENE_API FMovieSceneSequenceTickInterval GetInheritedInterval(UObject* ContextObject);
 
 	/**
 	 * Equality comparison operator
@@ -80,4 +80,4 @@ public:
 	{
 		return !(A == B);
 	}
-};
\ No newline at end of file
+};
