Sequencer: Fix leaking pre animated state for tracks set to keep state. Looping sequences with spawnables in subscenes could cause memory to continually get allocated in pre animated storage for each time spawnables were created. This was occurring for sections that were set to 'Keep State'. Even though states were not being cached, memory was being allocated in the cache tied to the object keys because 'MakeEntry' was being called. This memory was not even required as we weren't ever caching into it or restoring state. To fix this, logic has been reorganized in the CachePreAnimatedState functions to check whether we are meant to be tracking state before creating entries.

In addition on garbage collection we go through the metadata and remove any metadata for objects that may have been garbage collected, e.g. Spawnables. This could also be a source of leaks for the looping sequence states.

 Also, removed the function 'ShouldSavePreAnimatedState' as it wasn't being used and simplified the logic.

Also fixed a case where Data Layer preanimated state was firing an ensure by filtering on the RestoreState tag, which should be unnecessary as the BeginTrackingEntity function already handles that case. So that's been removed.

diff --git a/Engine/Source/Runtime/MovieScene/Private/EntitySystem/MovieSceneEntitySystemLinker.cpp b/Engine/Source/Runtime/MovieScene/Private/EntitySystem/MovieSceneEntitySystemLinker.cpp
--- a/Engine/Source/Runtime/MovieScene/Private/EntitySystem/MovieSceneEntitySystemLinker.cpp
+++ b/Engine/Source/Runtime/MovieScene/Private/EntitySystem/MovieSceneEntitySystemLinker.cpp
@@ -390,6 +390,9 @@ void UMovieSceneEntitySystemLinker::CleanGarbage()
 {
 	using namespace UE::MovieScene;
 
+	// Cleanup any stale preanimated state. This can exist even without Entities being marked as NeedsUnlink.
+	PreAnimatedState.DiscardStaleObjectState();
+
 	FBuiltInComponentTypes* BuiltInComponents = FBuiltInComponentTypes::Get();
 	FComponentTypeID NeedsUnlink = BuiltInComponents->Tags.NeedsUnlink;
 	if (!EntityManager.ContainsComponent(NeedsUnlink))
@@ -411,6 +414,7 @@ void UMovieSceneEntitySystemLinker::CleanGarbage()
 	//       to finding entities tagged NeedsUnlink.
 	Events.CleanTaggedGarbage.Broadcast(this);
 
+
 	TArray<FMovieSceneEntityID> UnresolvedEntities;
 
 	FEntityTaskBuilder()
diff --git a/Engine/Source/Runtime/MovieScene/Private/Evaluation/PreAnimatedState/MovieScenePreAnimatedObjectGroupManager.cpp b/Engine/Source/Runtime/MovieScene/Private/Evaluation/PreAnimatedState/MovieScenePreAnimatedObjectGroupManager.cpp
--- a/Engine/Source/Runtime/MovieScene/Private/Evaluation/PreAnimatedState/MovieScenePreAnimatedObjectGroupManager.cpp
+++ b/Engine/Source/Runtime/MovieScene/Private/Evaluation/PreAnimatedState/MovieScenePreAnimatedObjectGroupManager.cpp
@@ -2,6 +2,7 @@
 
 #include "Evaluation/PreAnimatedState/MovieScenePreAnimatedObjectGroupManager.h"
 #include "Evaluation/PreAnimatedState/MovieScenePreAnimatedStorageID.inl"
+#include "EntitySystem/BuiltInComponentTypes.h"
 
 namespace UE
 {
@@ -50,5 +51,16 @@ void FPreAnimatedObjectGroupManager::GetGroupsByClass(UClass* GeneratedClass, TA
 	}
 }
 
+void FPreAnimatedObjectGroupManager::GatherStaleStorageGroups(TArray<FPreAnimatedStorageGroupHandle>& StaleGroupStorage) const
+{
+	for (auto It = StorageGroupsByKey.CreateConstIterator(); It; ++It)
+	{
+		if (FBuiltInComponentTypes::IsBoundObjectGarbage(It.Key().ResolveObjectPtr()))
+		{
+			StaleGroupStorage.Add(It.Value());
+		}
+	}
+}
+
 } // namespace MovieScene
 } // namespace UE
diff --git a/Engine/Source/Runtime/MovieScene/Private/Evaluation/PreAnimatedState/MovieScenePreAnimatedObjectStorage.cpp b/Engine/Source/Runtime/MovieScene/Private/Evaluation/PreAnimatedState/MovieScenePreAnimatedObjectStorage.cpp
--- a/Engine/Source/Runtime/MovieScene/Private/Evaluation/PreAnimatedState/MovieScenePreAnimatedObjectStorage.cpp
+++ b/Engine/Source/Runtime/MovieScene/Private/Evaluation/PreAnimatedState/MovieScenePreAnimatedObjectStorage.cpp
@@ -14,6 +14,16 @@ void FBoundObjectPreAnimatedStateTraits::Initialize(FPreAnimatedStorageID InStor
 	ObjectGroupManager = InParentExtension->GetOrCreateGroupManager<FPreAnimatedObjectGroupManager>();
 }
 
+FPreAnimatedStorageGroupHandle FBoundObjectPreAnimatedStateTraits::FindGroupImpl(UObject* BoundObject)
+{
+	return ObjectGroupManager->FindGroupForKey(BoundObject);
+}
+
+FPreAnimatedStorageGroupHandle FBoundObjectPreAnimatedStateTraits::FindGroupImpl(const FObjectComponent& BoundObject)
+{
+	return ObjectGroupManager->FindGroupForKey(BoundObject.GetObject());
+}
+
 FPreAnimatedStorageGroupHandle FBoundObjectPreAnimatedStateTraits::MakeGroupImpl(UObject* BoundObject)
 {
 	return ObjectGroupManager->MakeGroupForKey(BoundObject);
diff --git a/Engine/Source/Runtime/MovieScene/Private/Evaluation/PreAnimatedState/MovieScenePreAnimatedObjectTokenStorage.cpp b/Engine/Source/Runtime/MovieScene/Private/Evaluation/PreAnimatedState/MovieScenePreAnimatedObjectTokenStorage.cpp
--- a/Engine/Source/Runtime/MovieScene/Private/Evaluation/PreAnimatedState/MovieScenePreAnimatedObjectTokenStorage.cpp
+++ b/Engine/Source/Runtime/MovieScene/Private/Evaluation/PreAnimatedState/MovieScenePreAnimatedObjectTokenStorage.cpp
@@ -11,6 +11,19 @@ namespace MovieScene
 
 TAutoRegisterPreAnimatedStorageID<FAnimTypePreAnimatedStateObjectStorage> FAnimTypePreAnimatedStateObjectStorage::StorageID;
 
+FPreAnimatedStateEntry FAnimTypePreAnimatedStateObjectStorage::FindEntry(UObject* Object, FMovieSceneAnimTypeID AnimTypeID)
+{
+	FPreAnimatedObjectTokenTraits::KeyType Key{ Object, AnimTypeID };
+
+	// Begin by finding or creating a pre-animated state group for this bound object
+	FPreAnimatedStorageGroupHandle Group = this->Traits.FindGroup(Object);
+
+	// Find the storage index for the specific anim-type and object we're animating
+	FPreAnimatedStorageIndex StorageIndex = FindStorageIndex(Key);
+
+	return FPreAnimatedStateEntry{ Group, FPreAnimatedStateCachedValueHandle{ StorageID, StorageIndex } };
+}
+
 FPreAnimatedStateEntry FAnimTypePreAnimatedStateObjectStorage::MakeEntry(UObject* Object, FMovieSceneAnimTypeID AnimTypeID)
 {
 	FPreAnimatedObjectTokenTraits::KeyType Key{ Object, AnimTypeID };
diff --git a/Engine/Source/Runtime/MovieScene/Private/Evaluation/PreAnimatedState/MovieScenePreAnimatedRootTokenStorage.cpp b/Engine/Source/Runtime/MovieScene/Private/Evaluation/PreAnimatedState/MovieScenePreAnimatedRootTokenStorage.cpp
--- a/Engine/Source/Runtime/MovieScene/Private/Evaluation/PreAnimatedState/MovieScenePreAnimatedRootTokenStorage.cpp
+++ b/Engine/Source/Runtime/MovieScene/Private/Evaluation/PreAnimatedState/MovieScenePreAnimatedRootTokenStorage.cpp
@@ -15,6 +15,13 @@ void FAnimTypePreAnimatedStateRootStorage::Initialize(FPreAnimatedStorageID InSt
 	TPreAnimatedStateStorage<FPreAnimatedRootTokenTraits>::Initialize(InStorageID, InParentExtension);
 }
 
+FPreAnimatedStateEntry FAnimTypePreAnimatedStateRootStorage::FindEntry(FMovieSceneAnimTypeID AnimTypeID)
+{
+	// Find the storage index for the specific anim-type and object we're animating
+	FPreAnimatedStorageIndex StorageIndex = FindStorageIndex(AnimTypeID);
+	return FPreAnimatedStateEntry{ FPreAnimatedStorageGroupHandle(), FPreAnimatedStateCachedValueHandle{ StorageID, StorageIndex } };
+}
+
 FPreAnimatedStateEntry FAnimTypePreAnimatedStateRootStorage::MakeEntry(FMovieSceneAnimTypeID AnimTypeID)
 {
 	// Find the storage index for the specific anim-type and object we're animating
diff --git a/Engine/Source/Runtime/MovieScene/Private/Evaluation/PreAnimatedState/MovieScenePreAnimatedStateExtension.cpp b/Engine/Source/Runtime/MovieScene/Private/Evaluation/PreAnimatedState/MovieScenePreAnimatedStateExtension.cpp
--- a/Engine/Source/Runtime/MovieScene/Private/Evaluation/PreAnimatedState/MovieScenePreAnimatedStateExtension.cpp
+++ b/Engine/Source/Runtime/MovieScene/Private/Evaluation/PreAnimatedState/MovieScenePreAnimatedStateExtension.cpp
@@ -396,6 +396,27 @@ void FPreAnimatedStateExtension::RestoreGlobalState(const FRestoreStateParams& P
 	bEntriesInvalidated = true;
 }
 
+void FPreAnimatedStateExtension::DiscardStaleObjectState()
+{
+	TArray<FPreAnimatedStorageGroupHandle> StaleStorageGroups;
+	// Gather any groups from the group managers whose keys (e.g. FObjectKeys have become invalid)
+	for (auto It = GroupManagers.CreateIterator(); It; ++It)
+	{
+		TWeakPtr<IPreAnimatedStateGroupManager> GroupManager = It.Value();
+		if (auto GroupManagerPtr = GroupManager.Pin())
+		{
+			GroupManagerPtr->GatherStaleStorageGroups(StaleStorageGroups);
+		}
+	}
+	// Discard the state for such groups
+	for (FPreAnimatedStorageGroupHandle GroupHandle : StaleStorageGroups)
+	{
+		DiscardStateForGroup(GroupHandle);
+	}
+
+	bEntriesInvalidated = true;
+}
+
 void FPreAnimatedStateExtension::DiscardTransientState()
 {
 	if (FPreAnimatedEntityCaptureSource* EntityMetaData = GetEntityMetaData())
diff --git a/Engine/Source/Runtime/MovieScene/Public/Evaluation/PreAnimatedState/MovieScenePreAnimatedObjectGroupManager.h b/Engine/Source/Runtime/MovieScene/Public/Evaluation/PreAnimatedState/MovieScenePreAnimatedObjectGroupManager.h
--- a/Engine/Source/Runtime/MovieScene/Public/Evaluation/PreAnimatedState/MovieScenePreAnimatedObjectGroupManager.h
+++ b/Engine/Source/Runtime/MovieScene/Public/Evaluation/PreAnimatedState/MovieScenePreAnimatedObjectGroupManager.h
@@ -26,6 +26,8 @@ struct FPreAnimatedObjectGroupManager : TPreAnimatedStateGroupManager<FObjectKey
 	MOVIESCENE_API void OnObjectsReplaced(const TMap<UObject*, UObject*>& ReplacementMap);
 
 	MOVIESCENE_API void GetGroupsByClass(UClass* GeneratedClass, TArray<FPreAnimatedStorageGroupHandle>& OutGroupHandles);
+
+	MOVIESCENE_API void GatherStaleStorageGroups(TArray<FPreAnimatedStorageGroupHandle>& StaleGroupStorage) const override;
 };
 
 
diff --git a/Engine/Source/Runtime/MovieScene/Public/Evaluation/PreAnimatedState/MovieScenePreAnimatedObjectStorage.h b/Engine/Source/Runtime/MovieScene/Public/Evaluation/PreAnimatedState/MovieScenePreAnimatedObjectStorage.h
--- a/Engine/Source/Runtime/MovieScene/Public/Evaluation/PreAnimatedState/MovieScenePreAnimatedObjectStorage.h
+++ b/Engine/Source/Runtime/MovieScene/Public/Evaluation/PreAnimatedState/MovieScenePreAnimatedObjectStorage.h
@@ -32,6 +32,18 @@ struct FBoundObjectPreAnimatedStateTraits : FPreAnimatedStateTraits
 
 	MOVIESCENE_API void Initialize(FPreAnimatedStorageID InStorageID, FPreAnimatedStateExtension* InParentExtension);
 
+	/* Defined as a template rather than a variadic function to prevent error C4840 */
+	template<typename... T>
+	FPreAnimatedStorageGroupHandle FindGroup(UObject* BoundObject, T&&... Unused)
+	{
+		return FindGroupImpl(BoundObject);
+	}
+	template<typename... T>
+	FPreAnimatedStorageGroupHandle FindGroup(const FObjectComponent& BoundObject, T&&... Unused)
+	{
+		return FindGroupImpl(BoundObject);
+	}
+
 	/* Defined as a template rather than a variadic function to prevent error C4840 */
 	template<typename... T>
 	FPreAnimatedStorageGroupHandle MakeGroup(UObject* BoundObject, T&&... Unused)
@@ -44,6 +56,9 @@ struct FBoundObjectPreAnimatedStateTraits : FPreAnimatedStateTraits
 		return MakeGroupImpl(BoundObject);
 	}
 
+	MOVIESCENE_API FPreAnimatedStorageGroupHandle FindGroupImpl(UObject* BoundObject);
+	MOVIESCENE_API FPreAnimatedStorageGroupHandle FindGroupImpl(const FObjectComponent& BoundObject);
+
 	MOVIESCENE_API FPreAnimatedStorageGroupHandle MakeGroupImpl(UObject* BoundObject);
 	MOVIESCENE_API FPreAnimatedStorageGroupHandle MakeGroupImpl(const FObjectComponent& BoundObject);
 
@@ -151,10 +166,12 @@ public:
 
 	void CachePreAnimatedValue(const FCachePreAnimatedValueParams& Params, UObject* BoundObject, EPreAnimatedCaptureSourceTracking TrackingMode = EPreAnimatedCaptureSourceTracking::CacheIfTracked)
 	{
-		FPreAnimatedStateEntry Entry = MakeEntry(BoundObject);
-
-		if (this->TrackCaptureSource(Entry, TrackingMode))
+		if (this->ShouldTrackCaptureSource(TrackingMode, BoundObject))
 		{
+			FPreAnimatedStateEntry Entry = MakeEntry(BoundObject);
+
+			this->TrackCaptureSource(Entry, TrackingMode);
+
 			CachePreAnimatedValue(Params, Entry, BoundObject);
 		}
 	}
@@ -179,16 +196,17 @@ public:
 	template<typename OnCacheValue /* StorageType(const KeyType&) */>
 	void CachePreAnimatedValue(const FCachePreAnimatedValueParams& Params, UObject* BoundObject, OnCacheValue&& CacheCallback, EPreAnimatedCaptureSourceTracking TrackingMode = EPreAnimatedCaptureSourceTracking::CacheIfTracked)
 	{
-		FPreAnimatedStateEntry Entry = MakeEntry(BoundObject);
-
-		if (this->TrackCaptureSource(Entry, TrackingMode))
+		if (this->ShouldTrackCaptureSource(TrackingMode, BoundObject))
 		{
+			FPreAnimatedStateEntry Entry = MakeEntry(BoundObject);
+
+			this->TrackCaptureSource(Entry, TrackingMode);
 			const FPreAnimatedStorageIndex StorageIndex = Entry.ValueHandle.StorageIndex;
 
 			EPreAnimatedStorageRequirement StorageRequirement = this->ParentExtension->GetStorageRequirement(Entry);
 			if (!this->IsStorageRequirementSatisfied(StorageIndex, StorageRequirement))
 			{
-				KeyType Key { BoundObject };
+				KeyType Key{ BoundObject };
 				StorageType NewValue = CacheCallback(Key);
 				this->AssignPreAnimatedValue(StorageIndex, StorageRequirement, MoveTemp(NewValue));
 			}
diff --git a/Engine/Source/Runtime/MovieScene/Public/Evaluation/PreAnimatedState/MovieScenePreAnimatedObjectTokenStorage.h b/Engine/Source/Runtime/MovieScene/Public/Evaluation/PreAnimatedState/MovieScenePreAnimatedObjectTokenStorage.h
--- a/Engine/Source/Runtime/MovieScene/Public/Evaluation/PreAnimatedState/MovieScenePreAnimatedObjectTokenStorage.h
+++ b/Engine/Source/Runtime/MovieScene/Public/Evaluation/PreAnimatedState/MovieScenePreAnimatedObjectTokenStorage.h
@@ -46,6 +46,7 @@ struct FAnimTypePreAnimatedStateObjectStorage : TPreAnimatedStateStorage<FPreAni
 {
 	static MOVIESCENE_API TAutoRegisterPreAnimatedStorageID<FAnimTypePreAnimatedStateObjectStorage> StorageID;
 
+	MOVIESCENE_API FPreAnimatedStateEntry FindEntry(UObject* Object, FMovieSceneAnimTypeID AnimTypeID);
 	MOVIESCENE_API FPreAnimatedStateEntry MakeEntry(UObject* Object, FMovieSceneAnimTypeID AnimTypeID);
 };
 
diff --git a/Engine/Source/Runtime/MovieScene/Public/Evaluation/PreAnimatedState/MovieScenePreAnimatedPropertyStorage.h b/Engine/Source/Runtime/MovieScene/Public/Evaluation/PreAnimatedState/MovieScenePreAnimatedPropertyStorage.h
--- a/Engine/Source/Runtime/MovieScene/Public/Evaluation/PreAnimatedState/MovieScenePreAnimatedPropertyStorage.h
+++ b/Engine/Source/Runtime/MovieScene/Public/Evaluation/PreAnimatedState/MovieScenePreAnimatedPropertyStorage.h
@@ -198,13 +198,15 @@ struct TPreAnimatedPropertyStorageImpl<PropertyTraits, TPropertyMetaData<MetaDat
 				continue;
 			}
 
-			FPreAnimatedStateEntry Entry = this->MakeEntry(BoundObject, PropertyBindings[Index].PropertyPath);
-
-			if (!this->TrackCaptureSource(Entry, EPreAnimatedCaptureSourceTracking::CacheIfTracked))
+			if (!this->ShouldTrackCaptureSource(EPreAnimatedCaptureSourceTracking::CacheIfTracked, BoundObject, PropertyBindings[Index].PropertyPath))
 			{
 				continue;
 			}
 
+			FPreAnimatedStateEntry Entry = this->MakeEntry(BoundObject, PropertyBindings[Index].PropertyPath);
+
+			this->TrackCaptureSource(Entry, EPreAnimatedCaptureSourceTracking::CacheIfTracked);
+
 			EPreAnimatedStorageRequirement StorageRequirement = this->ParentExtension->GetStorageRequirement(Entry);
 			if (!this->IsStorageRequirementSatisfied(Entry.ValueHandle.StorageIndex, StorageRequirement))
 			{
diff --git a/Engine/Source/Runtime/MovieScene/Public/Evaluation/PreAnimatedState/MovieScenePreAnimatedRootTokenStorage.h b/Engine/Source/Runtime/MovieScene/Public/Evaluation/PreAnimatedState/MovieScenePreAnimatedRootTokenStorage.h
--- a/Engine/Source/Runtime/MovieScene/Public/Evaluation/PreAnimatedState/MovieScenePreAnimatedRootTokenStorage.h
+++ b/Engine/Source/Runtime/MovieScene/Public/Evaluation/PreAnimatedState/MovieScenePreAnimatedRootTokenStorage.h
@@ -36,6 +36,7 @@ struct FAnimTypePreAnimatedStateRootStorage : TPreAnimatedStateStorage<FPreAnima
 {
 	static MOVIESCENE_API TAutoRegisterPreAnimatedStorageID<FAnimTypePreAnimatedStateRootStorage> StorageID;
 
+	MOVIESCENE_API FPreAnimatedStateEntry FindEntry(FMovieSceneAnimTypeID AnimTypeID);
 	MOVIESCENE_API FPreAnimatedStateEntry MakeEntry(FMovieSceneAnimTypeID AnimTypeID);
 
 public:
diff --git a/Engine/Source/Runtime/MovieScene/Public/Evaluation/PreAnimatedState/MovieScenePreAnimatedStateExtension.h b/Engine/Source/Runtime/MovieScene/Public/Evaluation/PreAnimatedState/MovieScenePreAnimatedStateExtension.h
--- a/Engine/Source/Runtime/MovieScene/Public/Evaluation/PreAnimatedState/MovieScenePreAnimatedStateExtension.h
+++ b/Engine/Source/Runtime/MovieScene/Public/Evaluation/PreAnimatedState/MovieScenePreAnimatedStateExtension.h
@@ -51,6 +51,7 @@ struct IPreAnimatedStateGroupManager
 	virtual ~IPreAnimatedStateGroupManager(){}
 	virtual void InitializeGroupManager(FPreAnimatedStateExtension* Extension) = 0;
 	virtual void OnGroupDestroyed(FPreAnimatedStorageGroupHandle Group) = 0;
+	virtual void GatherStaleStorageGroups(TArray<FPreAnimatedStorageGroupHandle>& StaleGroupStorage) const = 0;
 };
 
 
@@ -236,11 +237,17 @@ public:
 	 */
 	MOVIESCENE_API void RestoreStateForGroup(FPreAnimatedStorageGroupHandle GroupHandle, const FRestoreStateParams& Params);
 
+	/**
+	* Called during Garbage Collection to clean up preanimated state on invalid bound objects. Does not restore state.
+	*/
+	void DiscardStaleObjectState();
+
 	/**
 	 * Called during blueprint re-instancing to replace the object bound to a specific group handle with another.
 	 */
 	MOVIESCENE_API void ReplaceObjectForGroup(FPreAnimatedStorageGroupHandle GroupHandle, const FObjectKey& OldObject, const FObjectKey& NewObject);
 
+
 	/**
 	 * Discard any transient state and all meta-data for any currently animating objects, whilst preserving the cached values internally.
 	 * Calling this function will cause any currently animating 'RestoreState' sections to re-cache their values if they are re-evaluated
@@ -389,6 +396,12 @@ struct TPreAnimatedStateGroupManager : IPreAnimatedStateGroupManager, TSharedFro
 		StorageGroupsToKey.Remove(Group);
 	}
 
+	
+	virtual void GatherStaleStorageGroups(TArray<FPreAnimatedStorageGroupHandle>& StaleGroupStorage) const override
+	{
+		
+	}
+
 	FPreAnimatedStorageGroupHandle FindGroupForKey(const KeyType& InKey) const
 	{
 		return StorageGroupsByKey.FindRef(InKey);
diff --git a/Engine/Source/Runtime/MovieScene/Public/Evaluation/PreAnimatedState/MovieScenePreAnimatedStateStorage.h b/Engine/Source/Runtime/MovieScene/Public/Evaluation/PreAnimatedState/MovieScenePreAnimatedStateStorage.h
--- a/Engine/Source/Runtime/MovieScene/Public/Evaluation/PreAnimatedState/MovieScenePreAnimatedStateStorage.h
+++ b/Engine/Source/Runtime/MovieScene/Public/Evaluation/PreAnimatedState/MovieScenePreAnimatedStateStorage.h
@@ -65,11 +65,6 @@ struct TPreAnimatedStateTaskParams
 		TrackingMode = EPreAnimatedCaptureSourceTracking::CacheIfTracked;
 		AdditionalFilter.All({ FBuiltInComponentTypes::Get()->Tags.NeedsLink });
 	}
-
-	static constexpr bool ShouldSavePreAnimatedState(typename TCallTraits<InputTypes>::ParamType... Inputs)
-	{
-		return true;
-	}
 };
 
 /**
@@ -457,6 +452,29 @@ public:
 		return true;
 	}
 
+	/**
+	* Given a list of arguments suitable for building a storage key, builds that key, and then 
+	* tries to find an existing pre-animated state entry for that key if it already exists.
+	* If it doesn't exist, it returns an empty FPreAnimatedStateEntry.
+	*
+	* Note that the arguments are also used for find the group handle, if the storage traits required grouping.
+	*/
+	template<typename... KeyArgs>
+	FPreAnimatedStateEntry FindEntry(KeyArgs&&... InKeyArgs)
+	{
+		KeyType Key{ Forward<KeyArgs>(InKeyArgs)... };
+
+		FPreAnimatedStorageGroupHandle GroupHandle;
+		if constexpr (StorageTraits::SupportsGrouping)
+		{
+			GroupHandle = this->Traits.FindGroup(InKeyArgs...);
+		}
+
+		FPreAnimatedStorageIndex StorageIndex = this->FindStorageIndex(Key);
+
+		return FPreAnimatedStateEntry{ GroupHandle, FPreAnimatedStateCachedValueHandle{ this->StorageID, StorageIndex } };
+	}
+
 	/**
 	 * Creates a new pre-animated state entry from a list of arguments suitable for building the storage key.
 	 *
@@ -562,11 +580,6 @@ public:
 			FPreAnimatedEntityCaptureSource* EntityMetaData = this->ParentExtension->GetOrCreateEntityMetaData();
 			for (int32 Index = 0; Index < Num; ++Index)
 			{
-				if (!InParams.ShouldSavePreAnimatedState(Inputs[Index]...))
-				{
-					continue;
-				}
-
 				FPreAnimatedStateEntry Entry = MakeEntry(Inputs[Index]...);
 				EntityMetaData->BeginTrackingEntity(Entry, EntityIDs[Index], RootInstanceHandles[Index], bWantsRestore);
 			}
@@ -600,11 +613,6 @@ public:
 
 		TPreAnimatedStateTaskParams<ContributorTypes...> Params;
 
-		if (!Params.ShouldSavePreAnimatedState(InComponents...))
-		{
-			return;
-		}
-
 		FPreAnimatedStateEntry Entry = MakeEntry(InComponents...);
 		FPreAnimatedEntityCaptureSource* EntityMetaData = this->ParentExtension->GetOrCreateEntityMetaData();
 		EntityMetaData->BeginTrackingEntity(Entry, EntityID, RootInstanceHandle, bWantsRestoreState);
@@ -630,20 +638,19 @@ public:
 			const int32 Num = Item.GetAllocation()->Num();
 			for (int32 Index = 0; Index < Num; ++Index)
 			{
-				if (!InParams.ShouldSavePreAnimatedState(Values[Index]...))
+				if (!ShouldTrackCaptureSource(InParams.TrackingMode, Values[Index]...))
 				{
 					continue;
 				}
 
 				FPreAnimatedStateEntry Entry = MakeEntry(Values[Index]...);
-				if (TrackCaptureSource(Entry, InParams.TrackingMode))
+				TrackCaptureSource(Entry, InParams.TrackingMode);
+				
+				EPreAnimatedStorageRequirement StorageRequirement = this->ParentExtension->GetStorageRequirement(Entry);
+				if (!this->IsStorageRequirementSatisfied(Entry.ValueHandle.StorageIndex, StorageRequirement))
 				{
-					EPreAnimatedStorageRequirement StorageRequirement = this->ParentExtension->GetStorageRequirement(Entry);
-					if (!this->IsStorageRequirementSatisfied(Entry.ValueHandle.StorageIndex, StorageRequirement))
-					{
-						StorageType NewValue = this->Traits.CachePreAnimatedValue(Values[Index]...);
-						this->AssignPreAnimatedValue(Entry.ValueHandle.StorageIndex, StorageRequirement, MoveTemp(NewValue));
-					}
+					StorageType NewValue = this->Traits.CachePreAnimatedValue(Values[Index]...);
+					this->AssignPreAnimatedValue(Entry.ValueHandle.StorageIndex, StorageRequirement, MoveTemp(NewValue));
 				}
 			}
 		};
@@ -671,9 +678,12 @@ public:
 	template<typename... ContributorTypes>
 	void CacheTrackedPreAnimatedValue(EPreAnimatedCaptureSourceTracking TrackingMode, ContributorTypes... Values)
 	{
-		FPreAnimatedStateEntry Entry = MakeEntry(Values...);
-		if (TrackCaptureSource(Entry, TrackingMode))
+		if (ShouldTrackCaptureSource(TrackingMode, Values...))
 		{
+			FPreAnimatedStateEntry Entry = MakeEntry(Values...);
+
+			TrackCaptureSource(Entry, TrackingMode);
+
 			EPreAnimatedStorageRequirement Requirement = ParentExtension->HasActiveCaptureSource()
 				? EPreAnimatedStorageRequirement::Transient
 				: EPreAnimatedStorageRequirement::Persistent;
@@ -706,13 +716,14 @@ public:
 	{
 		ensureMsgf(GroupHandle.IsValid() || !StorageTraits::SupportsGrouping, TEXT("The group handle must be valid for pre-animated state that supports grouping, and invalid if not"));
 
-		// Find the storage index for the specific key we're animating
-		FPreAnimatedStorageIndex StorageIndex = GetOrCreateStorageIndex(InKey);
-		FPreAnimatedStateEntry Entry{ GroupHandle, FPreAnimatedStateCachedValueHandle{ this->StorageID, StorageIndex } };
-		if (!TrackCaptureSource(Entry, TrackingMode))
+		if (!ShouldTrackCaptureSource(TrackingMode, InKey))
 		{
 			return;
 		}
+		// Find the storage index for the specific key we're animating
+		FPreAnimatedStorageIndex StorageIndex = GetOrCreateStorageIndex(InKey);
+		FPreAnimatedStateEntry Entry{ GroupHandle, FPreAnimatedStateCachedValueHandle{ this->StorageID, StorageIndex } };
+		TrackCaptureSource(Entry, TrackingMode);
 
 		EPreAnimatedStorageRequirement Requirement = ParentExtension->HasActiveCaptureSource()
 			? EPreAnimatedStorageRequirement::Transient
@@ -755,11 +766,6 @@ public:
 			FPreAnimatedEntityCaptureSource* EntityMetaData = this->ParentExtension->GetOrCreateEntityMetaData();
 			for (int32 Index = 0; Index < Num; ++Index)
 			{
-				if (!InParams.ShouldSavePreAnimatedState(Inputs[Index]...))
-				{
-					continue;
-				}
-
 				FPreAnimatedStateEntry Entry = MakeEntry(Inputs[Index]...);
 
 				EntityMetaData->BeginTrackingEntity(Entry, EntityIDs[Index], RootInstanceHandles[Index], bWantsRestore);
@@ -770,7 +776,6 @@ public:
 					StorageType NewValue = this->Traits.CachePreAnimatedValue(Inputs[Index]...);
 					this->AssignPreAnimatedValue(Entry.ValueHandle.StorageIndex, StorageRequirement, MoveTemp(NewValue));
 				}
-
 			}
 		};
 
@@ -796,12 +801,47 @@ public:
 
 protected:
 
+	template<typename... KeyArgs>
+	bool ShouldTrackCaptureSource(EPreAnimatedCaptureSourceTracking TrackingMode, KeyArgs&&... InKeyArgs)
+	{
+		switch (TrackingMode)
+		{
+		case EPreAnimatedCaptureSourceTracking::CacheIfTracked:
+		{
+			// If we're capturing global state we always track changes
+			if (this->ParentExtension->IsCapturingGlobalState())
+			{
+				return true;
+			}
+
+			// Only cache the value if an entry already exists, and if metadata already exists.
+			// (ie, something is actively animating this and wants restore state)
+			if (FPreAnimatedStateEntry Entry = FindEntry(InKeyArgs...))
+			{
+				ensureMsgf(this->ParentExtension->MetaDataExists(Entry), TEXT("PreAnimatedStateEntry has allocated storage but no metadata exists to track it."));
+			}
+			return false;
+		}
+		break;
+
+		case EPreAnimatedCaptureSourceTracking::AlwaysCache:
+		{
+			// Always cache the source meta data. If there is no capture source this calls EnsureMetaData as above
+			return true;
+		}
+		break;
+
+		default:
+			ensureMsgf(false, TEXT("Unsupported tracking mode, no pre-animated state caching will occur!"));
+			return false;
+		}
+	}
+
 	/**
 	 * Begins tracking the current entry with the currently set tracking source, if necessary/desirable.
-	 *
-	 * Returns whether caching pre-animated state value is wanted.
+	 * If the entry is valid, we assume the metadata must already exist.
 	 */
-	bool TrackCaptureSource(const FPreAnimatedStateEntry& Entry, EPreAnimatedCaptureSourceTracking TrackingMode)
+	void TrackCaptureSource(const FPreAnimatedStateEntry& Entry, EPreAnimatedCaptureSourceTracking TrackingMode)
 	{
 		switch (TrackingMode)
 		{
@@ -811,12 +851,12 @@ protected:
 				if (this->ParentExtension->IsCapturingGlobalState())
 				{
 					this->ParentExtension->EnsureMetaData(Entry);
-					return true;
+					return;
 				}
 
 				// Only cache the value if tracking meta-data exists for this entry
 				// (ie, something is actively animating this and wants restore state)
-				return this->ParentExtension->MetaDataExists(Entry);
+				ensureMsgf(Entry.IsValid() && this->ParentExtension->MetaDataExists(Entry), TEXT("PreAnimatedStateEntry has allocated storage but no metadata exists to track it."));
 			}
 			break;
 
@@ -824,13 +864,12 @@ protected:
 			{
 				// Always cache the source meta data. If there is no capture source this calls EnsureMetaData as above
 				this->ParentExtension->AddSourceMetaData(Entry);
-				return true;
+				return;
 			}
 			break;
 
 		default:
 			ensureMsgf(false, TEXT("Unsupported tracking mode, no pre-animated state caching will occur!"));
-			return false;
 		}
 	}
 
diff --git a/Engine/Source/Runtime/MovieScene/Public/Evaluation/PreAnimatedState/MovieScenePreAnimatedStateTypes.h b/Engine/Source/Runtime/MovieScene/Public/Evaluation/PreAnimatedState/MovieScenePreAnimatedStateTypes.h
--- a/Engine/Source/Runtime/MovieScene/Public/Evaluation/PreAnimatedState/MovieScenePreAnimatedStateTypes.h
+++ b/Engine/Source/Runtime/MovieScene/Public/Evaluation/PreAnimatedState/MovieScenePreAnimatedStateTypes.h
@@ -161,6 +161,16 @@ struct FPreAnimatedStateEntry
 	/** A handle for the actual storage location of the value */
 	FPreAnimatedStateCachedValueHandle ValueHandle;
 
+	bool IsValid() const
+	{
+		return ValueHandle.StorageIndex.IsValid();
+	}
+
+	explicit operator bool() const
+	{
+		return ValueHandle.StorageIndex.IsValid();
+	}
+
 	friend uint32 GetTypeHash(const FPreAnimatedStateEntry& InEntry)
 	{
 		return GetTypeHash(InEntry.GroupHandle) ^ GetTypeHash(InEntry.ValueHandle);
diff --git a/Engine/Source/Runtime/MovieSceneTracks/Private/Systems/MovieSceneComponentMobilitySystem.cpp b/Engine/Source/Runtime/MovieSceneTracks/Private/Systems/MovieSceneComponentMobilitySystem.cpp
--- a/Engine/Source/Runtime/MovieSceneTracks/Private/Systems/MovieSceneComponentMobilitySystem.cpp
+++ b/Engine/Source/Runtime/MovieSceneTracks/Private/Systems/MovieSceneComponentMobilitySystem.cpp
@@ -84,6 +84,14 @@ static void GetFlattenedHierarchy(USceneComponent* SceneComponent, TArray<UScene
 	}
 }
 
+FPreAnimatedStateEntry FPreAnimatedComponentMobilityStorage::FindEntry(USceneComponent* InSceneComponent)
+{
+	FPreAnimatedStorageGroupHandle GroupHandle = this->Traits.ObjectGroupManager->FindGroupForKey(InSceneComponent);
+	FPreAnimatedStorageIndex       StorageIndex = FindStorageIndex(InSceneComponent);
+
+	return FPreAnimatedStateEntry{ GroupHandle, FPreAnimatedStateCachedValueHandle{ StorageID, StorageIndex } };
+}
+
 FPreAnimatedStateEntry FPreAnimatedComponentMobilityStorage::MakeEntry(USceneComponent* InSceneComponent)
 {
 	FPreAnimatedStorageGroupHandle GroupHandle  = this->Traits.ObjectGroupManager->MakeGroupForKey(InSceneComponent);
@@ -183,6 +191,11 @@ void UMovieSceneComponentMobilitySystem::SavePreAnimatedState(const FPreAnimatio
 		const FEntityAllocation* Allocation = Item.GetAllocation();
 		const bool bWantsRestore = Item.GetAllocationType().Contains(FBuiltInComponentTypes::Get()->Tags.RestoreState);
 
+		if (!Extension->IsCapturingGlobalState() && !bWantsRestore)
+		{
+			return;
+		}
+
 		FCachePreAnimatedValueParams CacheValueParams;
 		CacheValueParams.bForcePersist = Item.GetAllocationType().Contains(FMovieSceneTracksComponentTypes::Get()->AttachParent);
 
@@ -204,7 +217,7 @@ void UMovieSceneComponentMobilitySystem::SavePreAnimatedState(const FPreAnimatio
 			for (USceneComponent* CurrentSceneComponent : FlatHierarchy)
 			{
 				FPreAnimatedStateEntry Entry = ComponentMobilityStorage->MakeEntry(CurrentSceneComponent);
-
+				
 				EntityMetaData->BeginTrackingEntity(Entry, EntityID, InstanceHandle, bWantsRestore);
 				ComponentMobilityStorage->CachePreAnimatedValue(CacheValueParams, Entry, CurrentSceneComponent);
 			}
diff --git a/Engine/Source/Runtime/MovieSceneTracks/Private/Systems/MovieSceneDataLayerSystem.cpp b/Engine/Source/Runtime/MovieSceneTracks/Private/Systems/MovieSceneDataLayerSystem.cpp
--- a/Engine/Source/Runtime/MovieSceneTracks/Private/Systems/MovieSceneDataLayerSystem.cpp
+++ b/Engine/Source/Runtime/MovieSceneTracks/Private/Systems/MovieSceneDataLayerSystem.cpp
@@ -52,6 +52,14 @@ struct FPreAnimatedDataLayerStorageTraits : FBoundObjectPreAnimatedStateTraits
 
 	/** Called when a previously animated data layer needs to be restored */
 	static void RestorePreAnimatedValue(const TObjectKey<UDataLayerInstance>& InKey, EDataLayerRuntimeState PreviousState, const FRestoreStateParams& Params);
+
+	template<typename... T>
+	FPreAnimatedStorageGroupHandle FindGroup(const UDataLayerInstance* BoundObject, T&&... Unused)
+	{
+		return FindGroupImpl(BoundObject);
+	}
+
+	FPreAnimatedStorageGroupHandle FindGroupImpl(const UDataLayerInstance* BoundObject);
 };
 
 /** Container class for all pre-animated data layer state */
@@ -65,8 +73,10 @@ struct FPreAnimatedDataLayerStorage
 	/*~ IPreAnimatedStateGroupManager */
 	void InitializeGroupManager(FPreAnimatedStateExtension* Extension) override;
 	void OnGroupDestroyed(FPreAnimatedStorageGroupHandle Group) override;
+	void GatherStaleStorageGroups(TArray<FPreAnimatedStorageGroupHandle>& StaleGroupStorage) const override {}
 
 	/** Make an entry for the specified data layer */
+	FPreAnimatedStateEntry FindEntry(const UDataLayerInstance* InDataLayer);
 	FPreAnimatedStateEntry MakeEntry(const UDataLayerInstance* InDataLayer);
 
 	/** Save the value of a data layer. Should only be used for runtime / PIE worlds */
@@ -143,6 +153,11 @@ void FPreAnimatedDataLayerStorageTraits::RestorePreAnimatedValue(const TObjectKe
 	}
 }
 
+FPreAnimatedStorageGroupHandle FPreAnimatedDataLayerStorageTraits::FindGroupImpl(const UDataLayerInstance* BoundObject)
+{
+	return ObjectGroupManager->FindGroupForKey(BoundObject);
+}
+
 // ---------------------------------------------------------------------
 // FPreAnimatedDataLayerStorage definitions
 TAutoRegisterPreAnimatedStorageID<FPreAnimatedDataLayerStorage> FPreAnimatedDataLayerStorage::StorageID;
@@ -157,6 +172,16 @@ void FPreAnimatedDataLayerStorage::OnGroupDestroyed(FPreAnimatedStorageGroupHand
 	GroupHandle = FPreAnimatedStorageGroupHandle();
 }
 
+FPreAnimatedStateEntry FPreAnimatedDataLayerStorage::FindEntry(const UDataLayerInstance* InDataLayer)
+{
+	if (!GroupHandle)
+	{
+		return FPreAnimatedStateEntry();
+	}
+	FPreAnimatedStorageIndex StorageIndex = FindStorageIndex(InDataLayer);
+	return FPreAnimatedStateEntry{ GroupHandle, FPreAnimatedStateCachedValueHandle{ StorageID, StorageIndex } };
+}
+
 FPreAnimatedStateEntry FPreAnimatedDataLayerStorage::MakeEntry(const UDataLayerInstance* InDataLayer)
 {
 	if (!GroupHandle)
@@ -169,7 +194,15 @@ FPreAnimatedStateEntry FPreAnimatedDataLayerStorage::MakeEntry(const UDataLayerI
 
 void FPreAnimatedDataLayerStorage::SavePreAnimatedState(const UDataLayerInstance* DataLayer)
 {
-	FPreAnimatedStateEntry         Entry              = MakeEntry(DataLayer);
+	if (!this->ShouldTrackCaptureSource(EPreAnimatedCaptureSourceTracking::CacheIfTracked, DataLayer))
+	{
+		return;
+	}
+
+	FPreAnimatedStateEntry Entry = MakeEntry(DataLayer);
+
+	this->TrackCaptureSource(Entry, EPreAnimatedCaptureSourceTracking::CacheIfTracked);
+
 	EPreAnimatedStorageRequirement StorageRequirement = this->ParentExtension->GetStorageRequirement(Entry);
 
 	if (!IsStorageRequirementSatisfied(Entry.ValueHandle.StorageIndex, StorageRequirement))
@@ -185,7 +218,14 @@ void FPreAnimatedDataLayerStorage::SavePreAnimatedState(const UDataLayerInstance
 #if WITH_EDITOR
 void FPreAnimatedDataLayerStorage::SavePreAnimatedStateInEditor(const UDataLayerInstance* DataLayer)
 {
-	FPreAnimatedStateEntry         Entry              = MakeEntry(DataLayer);
+	if (!this->ShouldTrackCaptureSource(EPreAnimatedCaptureSourceTracking::CacheIfTracked, DataLayer))
+	{
+		return;
+	}
+
+	FPreAnimatedStateEntry Entry = MakeEntry(DataLayer);
+	this->TrackCaptureSource(Entry, EPreAnimatedCaptureSourceTracking::CacheIfTracked);
+
 	EPreAnimatedStorageRequirement StorageRequirement = this->ParentExtension->GetStorageRequirement(Entry);
 
 	if (!IsStorageRequirementSatisfied(Entry.ValueHandle.StorageIndex, StorageRequirement))
@@ -664,6 +704,12 @@ void UMovieSceneDataLayerSystem::BeginTrackingEntities()
 		TRead<FMovieSceneDataLayerComponentData> ComponentData)
 	{
 		const bool bRestoreState = Item.GetAllocationType().Contains(FBuiltInComponentTypes::Get()->Tags.RestoreState);
+
+		if (!this->GetLinker()->PreAnimatedState.IsCapturingGlobalState() && !bRestoreState)
+		{
+			return;
+		}
+
 		for (int32 Index = 0; Index < Item.GetAllocation()->Num(); ++Index)
 		{
 			const UMovieSceneDataLayerSection* Section = ComponentData[Index].Section.Get();
@@ -692,11 +738,6 @@ void UMovieSceneDataLayerSystem::BeginTrackingEntities()
 	};
 
 	FComponentMask Filter{ BuiltInComponents->Tags.NeedsLink };
-	if (!Linker->PreAnimatedState.IsCapturingGlobalState())
-	{
-		// If we're not capturing global state, only visit entities with the RestoreState tag
-		Filter.Set(BuiltInComponents->Tags.RestoreState);
-	}
 
 	// Iterate any data layer components that need link
 	FEntityTaskBuilder()
diff --git a/Engine/Source/Runtime/MovieSceneTracks/Public/Systems/MovieSceneComponentMobilitySystem.h b/Engine/Source/Runtime/MovieSceneTracks/Public/Systems/MovieSceneComponentMobilitySystem.h
--- a/Engine/Source/Runtime/MovieSceneTracks/Public/Systems/MovieSceneComponentMobilitySystem.h
+++ b/Engine/Source/Runtime/MovieSceneTracks/Public/Systems/MovieSceneComponentMobilitySystem.h
@@ -62,6 +62,7 @@ struct FPreAnimatedComponentMobilityStorage : TPreAnimatedStateStorage_ObjectTra
 {
 	static TAutoRegisterPreAnimatedStorageID<FPreAnimatedComponentMobilityStorage> StorageID;
 
+	FPreAnimatedStateEntry FindEntry(USceneComponent* InSceneComponent);
 	FPreAnimatedStateEntry MakeEntry(USceneComponent* InSceneComponent);
 };
 
