Added Csv profiling scope for MovieSceneSequence.

diff --git a/Engine/Source/Runtime/MovieScene/Private/MovieSceneSequenceTickManager.cpp b/Engine/Source/Runtime/MovieScene/Private/MovieSceneSequenceTickManager.cpp
--- a/Engine/Source/Runtime/MovieScene/Private/MovieSceneSequenceTickManager.cpp
+++ b/Engine/Source/Runtime/MovieScene/Private/MovieSceneSequenceTickManager.cpp
@@ -300,6 +300,7 @@ void UMovieSceneSequenceTickManager::UnregisterSequenceActor(AActor* InActor, TS
 
 void UMovieSceneSequenceTickManager::TickSequenceActors(float DeltaSeconds)
 {
+	CSV_SCOPED_TIMING_STAT_EXCLUSIVE(MovieSceneEval);
 	SCOPE_CYCLE_COUNTER(MovieSceneEval_SequenceTickManager);
 
 	// Let all tickable clients update. Some of them won't do anything, others will do synchronous
