Sequencer- Fix broken sequencer autotests and some instances of cacheing preanimated values. Previous change of mine at 26143512 broke this due to unfortunate mistake in failing to return 'true' in ShouldTrackCaptureSource in the case we already have a storage entry.

diff --git a/Engine/Source/Runtime/MovieScene/Public/Evaluation/PreAnimatedState/MovieScenePreAnimatedStateStorage.h b/Engine/Source/Runtime/MovieScene/Public/Evaluation/PreAnimatedState/MovieScenePreAnimatedStateStorage.h
--- a/Engine/Source/Runtime/MovieScene/Public/Evaluation/PreAnimatedState/MovieScenePreAnimatedStateStorage.h
+++ b/Engine/Source/Runtime/MovieScene/Public/Evaluation/PreAnimatedState/MovieScenePreAnimatedStateStorage.h
@@ -818,7 +818,7 @@ protected:
 			// (ie, something is actively animating this and wants restore state)
 			if (FPreAnimatedStateEntry Entry = FindEntry(InKeyArgs...))
 			{
-				ensureMsgf(this->ParentExtension->MetaDataExists(Entry), TEXT("PreAnimatedStateEntry has allocated storage but no metadata exists to track it."));
+				return ensureMsgf(this->ParentExtension->MetaDataExists(Entry), TEXT("PreAnimatedStateEntry has allocated storage but no metadata exists to track it."));
 			}
 			return false;
 		}
