Added missing include and resubmit previous change

diff --git a/Engine/Source/Runtime/MovieScene/Private/MovieSceneSequenceTickManager.cpp b/Engine/Source/Runtime/MovieScene/Private/MovieSceneSequenceTickManager.cpp
--- a/Engine/Source/Runtime/MovieScene/Private/MovieSceneSequenceTickManager.cpp
+++ b/Engine/Source/Runtime/MovieScene/Private/MovieSceneSequenceTickManager.cpp
@@ -12,6 +12,7 @@
 #include "Logging/TokenizedMessage.h"
 #include "Misc/UObjectToken.h"
 #include "Algo/IndexOf.h"
+#include "ProfilingDebugging/CsvProfiler.h"
 
 #include UE_INLINE_GENERATED_CPP_BY_NAME(MovieSceneSequenceTickManager)
 
@@ -300,6 +301,7 @@ void UMovieSceneSequenceTickManager::UnregisterSequenceActor(AActor* InActor, TS
 
 void UMovieSceneSequenceTickManager::TickSequenceActors(float DeltaSeconds)
 {
+	CSV_SCOPED_TIMING_STAT_EXCLUSIVE(MovieSceneEval);
 	SCOPE_CYCLE_COUNTER(MovieSceneEval_SequenceTickManager);
 
 	// Let all tickable clients update. Some of them won't do anything, others will do synchronous
