Sequencer- Change bRestoreState boolean option on FMovieSceneSequencePlaybackSettings to an enumeration named FinishCompletionStateOverride with 3 values (Default, ForceKeepState, ForceRestoreState) with ForceRestoreState being the equivalent option to setting bRestoreState to true.

bRestoreState was used to force PreAnimatedState to be cached for all items globally and then restored when the Sequence is stopped or is finished. When set to false, no such override occurred, and completion mode was determined per track section by the settings on each track section.

We now have a 3rd option- ForceKeepState. If FinishCompletionStateOverride is set to ForceKeepState, then when the sequence is stopped or completed, we do not restore any state, regardless of track section settings- all state active when the sequence ends is persisted.This is done by discarding any preanimated state metadata and storage, preventing it from being used to restore state.

diff --git a/Engine/Plugins/Experimental/RenderGrid/Source/RenderGridEditor/Private/UI/SRenderGridViewerLive.cpp b/Engine/Plugins/Experimental/RenderGrid/Source/RenderGridEditor/Private/UI/SRenderGridViewerLive.cpp
--- a/Engine/Plugins/Experimental/RenderGrid/Source/RenderGridEditor/Private/UI/SRenderGridViewerLive.cpp
+++ b/Engine/Plugins/Experimental/RenderGrid/Source/RenderGridEditor/Private/UI/SRenderGridViewerLive.cpp
@@ -179,7 +179,7 @@ ULevelSequencePlayer* UE::RenderGrid::Private::SRenderGridEditorViewport::GetSeq
 		PlaybackSettings.LoopCount.Value = 0;
 		PlaybackSettings.bAutoPlay = false;
 		PlaybackSettings.bPauseAtEnd = true;
-		PlaybackSettings.bRestoreState = true;
+		PlaybackSettings.FinishCompletionStateOverride = EMovieSceneCompletionModeOverride::ForceRestoreState;
 
 		ALevelSequenceActor* PlayerActor = nullptr;
 		if (ULevelSequencePlayer* Player = ULevelSequencePlayer::CreateLevelSequencePlayer(World, LevelSequence, PlaybackSettings, PlayerActor); IsValid(Player))
diff --git a/Engine/Plugins/MovieScene/ActorSequence/Source/ActorSequence/Private/ActorSequenceComponent.cpp b/Engine/Plugins/MovieScene/ActorSequence/Source/ActorSequence/Private/ActorSequenceComponent.cpp
--- a/Engine/Plugins/MovieScene/ActorSequence/Source/ActorSequence/Private/ActorSequenceComponent.cpp
+++ b/Engine/Plugins/MovieScene/ActorSequence/Source/ActorSequence/Private/ActorSequenceComponent.cpp
@@ -86,3 +86,14 @@ void UActorSequenceComponent::TickComponent(float DeltaSeconds, enum ELevelTick
 	Super::TickComponent(DeltaSeconds, TickType, ThisTickFunction);
 }
 
+void UActorSequenceComponent::PostLoad()
+{
+	Super::PostLoad();
+	// If we previously were using bRestoreState on our PlaybackSettings, upgrade to the enum version.
+	if (PlaybackSettings.bRestoreState_DEPRECATED)
+	{
+		PlaybackSettings.FinishCompletionStateOverride = EMovieSceneCompletionModeOverride::ForceRestoreState;
+		PlaybackSettings.bRestoreState_DEPRECATED = false;
+	}
+}
+
diff --git a/Engine/Plugins/MovieScene/ActorSequence/Source/ActorSequence/Public/ActorSequenceComponent.h b/Engine/Plugins/MovieScene/ActorSequence/Source/ActorSequence/Public/ActorSequenceComponent.h
--- a/Engine/Plugins/MovieScene/ActorSequence/Source/ActorSequence/Public/ActorSequenceComponent.h
+++ b/Engine/Plugins/MovieScene/ActorSequence/Source/ActorSequence/Public/ActorSequenceComponent.h
@@ -51,6 +51,9 @@ public:
 	virtual void EndPlay(const EEndPlayReason::Type EndPlayReason) override;
 	virtual void TickComponent(float DeltaTime, enum ELevelTick TickType, FActorComponentTickFunction *ThisTickFunction) override;
 
+	// UObject interface
+	virtual void PostLoad() override;
+
 protected:
 
 	UPROPERTY(EditAnywhere, Category="Playback", meta=(ShowOnlyInnerProperties))
diff --git a/Engine/Plugins/MovieScene/MovieRenderPipeline/Source/MovieRenderPipelineCore/Private/Graph/MovieGraphSequenceDataSource.cpp b/Engine/Plugins/MovieScene/MovieRenderPipeline/Source/MovieRenderPipelineCore/Private/Graph/MovieGraphSequenceDataSource.cpp
--- a/Engine/Plugins/MovieScene/MovieRenderPipeline/Source/MovieRenderPipelineCore/Private/Graph/MovieGraphSequenceDataSource.cpp
+++ b/Engine/Plugins/MovieScene/MovieRenderPipeline/Source/MovieRenderPipelineCore/Private/Graph/MovieGraphSequenceDataSource.cpp
@@ -88,7 +88,7 @@ void UMovieGraphSequenceDataSource::CacheLevelSequenceData(ULevelSequence* InSeq
 	LevelSequenceActor->PlaybackSettings.LoopCount.Value = 0;
 	LevelSequenceActor->PlaybackSettings.bAutoPlay = false;
 	LevelSequenceActor->PlaybackSettings.bPauseAtEnd = true;
-	LevelSequenceActor->PlaybackSettings.bRestoreState = true;
+	LevelSequenceActor->PlaybackSettings.FinishCompletionStateOverride = EMovieSceneCompletionModeOverride::ForceRestoreState;
 
 	// Ensure the (possibly new) Level Sequence Actor uses our sequence
 	LevelSequenceActor->SetSequence(InSequence);
diff --git a/Engine/Plugins/MovieScene/MovieRenderPipeline/Source/MovieRenderPipelineCore/Private/MoviePipeline.cpp b/Engine/Plugins/MovieScene/MovieRenderPipeline/Source/MovieRenderPipelineCore/Private/MoviePipeline.cpp
--- a/Engine/Plugins/MovieScene/MovieRenderPipeline/Source/MovieRenderPipelineCore/Private/MoviePipeline.cpp
+++ b/Engine/Plugins/MovieScene/MovieRenderPipeline/Source/MovieRenderPipelineCore/Private/MoviePipeline.cpp
@@ -901,7 +901,7 @@ void UMoviePipeline::InitializeLevelSequenceActor()
 	LevelSequenceActor->PlaybackSettings.LoopCount.Value = 0;
 	LevelSequenceActor->PlaybackSettings.bAutoPlay = false;
 	LevelSequenceActor->PlaybackSettings.bPauseAtEnd = true;
-	LevelSequenceActor->PlaybackSettings.bRestoreState = true;
+	LevelSequenceActor->PlaybackSettings.FinishCompletionStateOverride = EMovieSceneCompletionModeOverride::ForceRestoreState;
 
 	// Use our duplicated sequence
 	LevelSequenceActor->SetSequence(TargetSequence);
diff --git a/Engine/Plugins/MovieScene/TemplateSequence/Source/TemplateSequence/Private/TemplateSequenceActor.cpp b/Engine/Plugins/MovieScene/TemplateSequence/Source/TemplateSequence/Private/TemplateSequenceActor.cpp
--- a/Engine/Plugins/MovieScene/TemplateSequence/Source/TemplateSequence/Private/TemplateSequenceActor.cpp
+++ b/Engine/Plugins/MovieScene/TemplateSequence/Source/TemplateSequence/Private/TemplateSequenceActor.cpp
@@ -29,6 +29,18 @@ void ATemplateSequenceActor::PostInitProperties()
 	SequencePlayer->SetPlaybackClient(this);
 }
 
+void ATemplateSequenceActor::PostLoad()
+{
+	Super::PostLoad();
+
+	// If we previously were using bRestoreState on our PlaybackSettings, upgrade to the enum version.
+	if (PlaybackSettings.bRestoreState_DEPRECATED)
+	{
+		PlaybackSettings.FinishCompletionStateOverride = EMovieSceneCompletionModeOverride::ForceRestoreState;
+		PlaybackSettings.bRestoreState_DEPRECATED = false;
+	}
+}
+
 void ATemplateSequenceActor::GetLifetimeReplicatedProps(TArray<FLifetimeProperty>& OutLifetimeProps) const
 {
 	Super::GetLifetimeReplicatedProps(OutLifetimeProps);
diff --git a/Engine/Plugins/MovieScene/TemplateSequence/Source/TemplateSequence/Public/TemplateSequenceActor.h b/Engine/Plugins/MovieScene/TemplateSequence/Source/TemplateSequence/Public/TemplateSequenceActor.h
--- a/Engine/Plugins/MovieScene/TemplateSequence/Source/TemplateSequence/Public/TemplateSequenceActor.h
+++ b/Engine/Plugins/MovieScene/TemplateSequence/Source/TemplateSequence/Public/TemplateSequenceActor.h
@@ -118,6 +118,7 @@ protected:
 
 	//~ Begin UObject interface
 	virtual void PostInitProperties() override;
+	virtual void PostLoad() override;
 	//~ End UObject interface
 
 public:
diff --git a/Engine/Source/Runtime/LevelSequence/Private/LevelSequenceActor.cpp b/Engine/Source/Runtime/LevelSequence/Private/LevelSequenceActor.cpp
--- a/Engine/Source/Runtime/LevelSequence/Private/LevelSequenceActor.cpp
+++ b/Engine/Source/Runtime/LevelSequence/Private/LevelSequenceActor.cpp
@@ -212,6 +212,13 @@ void ALevelSequenceActor::PostLoad()
 		bAutoPlay_DEPRECATED = false;
 	}
 
+	// If we previously were using bRestoreState on our PlaybackSettings, upgrade to the enum version.
+	if (PlaybackSettings.bRestoreState_DEPRECATED)
+	{
+		PlaybackSettings.FinishCompletionStateOverride = EMovieSceneCompletionModeOverride::ForceRestoreState;
+		PlaybackSettings.bRestoreState_DEPRECATED = false;
+	}
+
 	SequencePlayer->SetPlaybackSettings(PlaybackSettings);
 
 #if WITH_EDITORONLY_DATA
diff --git a/Engine/Source/Runtime/MovieScene/Private/Evaluation/MovieScenePreAnimatedState.cpp b/Engine/Source/Runtime/MovieScene/Private/Evaluation/MovieScenePreAnimatedState.cpp
--- a/Engine/Source/Runtime/MovieScene/Private/Evaluation/MovieScenePreAnimatedState.cpp
+++ b/Engine/Source/Runtime/MovieScene/Private/Evaluation/MovieScenePreAnimatedState.cpp
@@ -135,6 +135,17 @@ void FMovieScenePreAnimatedState::RestorePreAnimatedState()
 	}
 }
 
+void FMovieScenePreAnimatedState::DiscardPreAnimatedState()
+{
+	using namespace UE::MovieScene;
+
+	UMovieSceneEntitySystemLinker* Linker = WeakLinker.Get();
+	if (Linker)
+	{
+		Linker->PreAnimatedState.DiscardGlobalState(FRestoreStateParams{ Linker, InstanceHandle });
+	}
+}
+
 void FMovieScenePreAnimatedState::OnFinishedEvaluating(const FMovieSceneEvaluationKey& Key)
 {
 	if (TemplateMetaData)
diff --git a/Engine/Source/Runtime/MovieScene/Private/Evaluation/PreAnimatedState/MovieScenePreAnimatedStateExtension.cpp b/Engine/Source/Runtime/MovieScene/Private/Evaluation/PreAnimatedState/MovieScenePreAnimatedStateExtension.cpp
--- a/Engine/Source/Runtime/MovieScene/Private/Evaluation/PreAnimatedState/MovieScenePreAnimatedStateExtension.cpp
+++ b/Engine/Source/Runtime/MovieScene/Private/Evaluation/PreAnimatedState/MovieScenePreAnimatedStateExtension.cpp
@@ -417,6 +417,102 @@ void FPreAnimatedStateExtension::DiscardStaleObjectState()
 	bEntriesInvalidated = true;
 }
 
+void FPreAnimatedStateExtension::DiscardGlobalState(const FRestoreStateParams& Params)
+{
+	TArray<FPreAnimatedStateMetaData> ExpiredMetaData;
+
+	if (FPreAnimatedEntityCaptureSource* EntityMetaData = GetEntityMetaData())
+	{
+		EntityMetaData->GatherAndRemoveExpiredMetaData(Params, ExpiredMetaData);
+	}
+
+	for (int32 Index = WeakExternalCaptureSources.Num() - 1; Index >= 0; --Index)
+	{
+		TSharedPtr<IPreAnimatedCaptureSource> MetaData = WeakExternalCaptureSources[Index].Pin();
+		if (MetaData)
+		{
+			MetaData->GatherAndRemoveExpiredMetaData(Params, ExpiredMetaData);
+		}
+		else
+		{
+			// Order is not important in this array, so we can use the more efficient RemoveAtSwap algorithm
+			WeakExternalCaptureSources.RemoveAtSwap(Index, 1);
+		}
+	}
+
+	// Remove all contributions
+	for (const FPreAnimatedStateMetaData& MetaData : ExpiredMetaData)
+	{
+		FAggregatePreAnimatedStateMetaData* Aggregate = FindMetaData(MetaData.Entry);
+		if (ensure(Aggregate))
+		{
+			const int32 TotalNum = --Aggregate->NumContributors;
+			if (MetaData.bWantsRestoreState)
+			{
+				--Aggregate->NumRestoreContributors;
+			}
+
+			if (TotalNum == 0)
+			{
+				Aggregate->bWantedRestore = false;
+				Aggregate->TerminalInstanceHandle = MetaData.RootInstanceHandle;
+			}
+		}
+	}
+
+	// Ensure that the entries are discarded in strictly the reverse order they were cached in
+	for (int32 Index = 0; Index < GroupMetaData.GetMaxIndex(); ++Index)
+	{
+		if (!GroupMetaData.IsAllocated(Index))
+		{
+			continue;
+		}
+
+		FPreAnimatedGroupMetaData& Group = GroupMetaData[Index];
+
+		for (int32 AggregateIndex = Group.AggregateMetaData.Num() - 1; AggregateIndex >= 0; --AggregateIndex)
+		{
+			FAggregatePreAnimatedStateMetaData& Aggregate = Group.AggregateMetaData[AggregateIndex];
+			if (Aggregate.NumContributors == 0 && (!Aggregate.TerminalInstanceHandle.IsValid() || Aggregate.TerminalInstanceHandle == Params.TerminalInstanceHandle))
+			{
+				TSharedPtr<IPreAnimatedStorage> Storage = GetStorageChecked(Aggregate.ValueHandle.TypeID);
+				Storage->DiscardPreAnimatedStateStorage(Aggregate.ValueHandle.StorageIndex, EPreAnimatedStorageRequirement::Persistent);
+
+				Group.AggregateMetaData.RemoveAt(AggregateIndex, 1, false);
+			}
+
+			if (Group.AggregateMetaData.Num() == 0)
+			{
+				// Remove at will not re-allocate the array or shuffle items within the sparse array, so this is safe
+				Group.GroupManagerPtr->OnGroupDestroyed(Index);
+				GroupMetaData.RemoveAt(Index);
+			}
+		}
+	}
+
+	for (auto UngroupedIt = UngroupedMetaData.CreateIterator(); UngroupedIt; ++UngroupedIt)
+	{
+		FAggregatePreAnimatedStateMetaData& Aggregate = UngroupedIt.Value();
+		if (Aggregate.NumContributors == 0 && (!Aggregate.TerminalInstanceHandle.IsValid() || Aggregate.TerminalInstanceHandle == Params.TerminalInstanceHandle))
+		{
+			TSharedPtr<IPreAnimatedStorage> Storage = GetStorageChecked(Aggregate.ValueHandle.TypeID);
+			Storage->DiscardPreAnimatedStateStorage(Aggregate.ValueHandle.StorageIndex, EPreAnimatedStorageRequirement::Persistent);
+
+			UngroupedIt.RemoveCurrent();
+		}
+	}
+
+	GroupMetaData.Shrink();
+
+	// Invalidate cached data for any sequence instance that belongs to the terminal instance
+	if (Params.TerminalInstanceHandle.IsValid())
+	{
+		Linker->GetInstanceRegistry()->MutateInstance(Params.TerminalInstanceHandle).InvalidateCachedData(Linker);
+	}
+
+	bEntriesInvalidated = true;
+}
+
 void FPreAnimatedStateExtension::DiscardTransientState()
 {
 	if (FPreAnimatedEntityCaptureSource* EntityMetaData = GetEntityMetaData())
diff --git a/Engine/Source/Runtime/MovieScene/Private/MovieSceneSequencePlayer.cpp b/Engine/Source/Runtime/MovieScene/Private/MovieSceneSequencePlayer.cpp
--- a/Engine/Source/Runtime/MovieScene/Private/MovieSceneSequencePlayer.cpp
+++ b/Engine/Source/Runtime/MovieScene/Private/MovieSceneSequencePlayer.cpp
@@ -263,7 +263,7 @@ void UMovieSceneSequencePlayer::PlayInternal()
 		}
 
 		// Update now
-		if (PlaybackSettings.bRestoreState)
+		if (PlaybackSettings.FinishCompletionStateOverride == EMovieSceneCompletionModeOverride::ForceRestoreState)
 		{
 			RootTemplateInstance.EnableGlobalPreAnimatedStateCapture();
 		}
@@ -412,10 +412,14 @@ void UMovieSceneSequencePlayer::StopInternal(FFrameTime TimeToResetTo)
 		// Reset loop count on stop so that it doesn't persist to the next call to play
 		PlaybackSettings.LoopCount.Value = 0;
 
-		if (PlaybackSettings.bRestoreState)
+		if (PlaybackSettings.FinishCompletionStateOverride == EMovieSceneCompletionModeOverride::ForceRestoreState)
 		{
 			RestorePreAnimatedState();
 		}
+		else if (PlaybackSettings.FinishCompletionStateOverride == EMovieSceneCompletionModeOverride::ForceKeepState)
+		{
+			DiscardPreAnimatedState();
+		}
 
 		// Lambda that is invoked when the request to finish this sequence has been fulfilled
 		auto OnFlushed = [this, TimeToResetTo]
@@ -469,10 +473,14 @@ void UMovieSceneSequencePlayer::StopInternal(FFrameTime TimeToResetTo)
 	}
 	else if (RootTemplateInstance.IsValid() && RootTemplateInstance.HasEverUpdated())
 	{
-		if (PlaybackSettings.bRestoreState)
+		if (PlaybackSettings.FinishCompletionStateOverride == EMovieSceneCompletionModeOverride::ForceRestoreState)
 		{
 			RestorePreAnimatedState();
 		}
+		else if (PlaybackSettings.FinishCompletionStateOverride == EMovieSceneCompletionModeOverride::ForceKeepState)
+		{
+			DiscardPreAnimatedState();
+		}
 
 		TSharedPtr<FMovieSceneEntitySystemRunner> Runner = RootTemplateInstance.GetRunner();
 		if (Runner)
@@ -634,9 +642,9 @@ void UMovieSceneSequencePlayer::SetPlaybackPosition(FMovieSceneSequencePlaybackP
 
 void UMovieSceneSequencePlayer::RestoreState()
 {
-	if (!PlaybackSettings.bRestoreState)
+	if (PlaybackSettings.FinishCompletionStateOverride != EMovieSceneCompletionModeOverride::ForceRestoreState)
 	{
-		UE_LOG(LogMovieScene, Warning, TEXT("Attempting to restore pre-animated state for a player that was not set to capture pre-animated state. Please enable PlaybackSettings.bRestoreState"));
+		UE_LOG(LogMovieScene, Warning, TEXT("Attempting to restore pre-animated state for a player that was not set to capture pre-animated state. Please set PlaybackSettings.FinishCompletionStateOverride to ForceRestoreState"));
 	}
 
 	RestorePreAnimatedState();
diff --git a/Engine/Source/Runtime/MovieScene/Public/Evaluation/MovieScenePreAnimatedState.h b/Engine/Source/Runtime/MovieScene/Public/Evaluation/MovieScenePreAnimatedState.h
--- a/Engine/Source/Runtime/MovieScene/Public/Evaluation/MovieScenePreAnimatedState.h
+++ b/Engine/Source/Runtime/MovieScene/Public/Evaluation/MovieScenePreAnimatedState.h
@@ -104,6 +104,9 @@ public:
 
 	MOVIESCENE_API void RestorePreAnimatedState(UObject& Object, TFunctionRef<bool(FMovieSceneAnimTypeID)> InFilter);
 
+	// Discards all pre-animated state without restoring it.
+	MOVIESCENE_API void DiscardPreAnimatedState();
+
 	/**
 	 * Discard any tokens that relate to entity animation (ie sections or tracks) without restoring the values.
 	 * Any global pre-animated state tokens (that reset the animation when saving a map, for instance) will remain.
diff --git a/Engine/Source/Runtime/MovieScene/Public/Evaluation/PreAnimatedState/MovieScenePreAnimatedStateExtension.h b/Engine/Source/Runtime/MovieScene/Public/Evaluation/PreAnimatedState/MovieScenePreAnimatedStateExtension.h
--- a/Engine/Source/Runtime/MovieScene/Public/Evaluation/PreAnimatedState/MovieScenePreAnimatedStateExtension.h
+++ b/Engine/Source/Runtime/MovieScene/Public/Evaluation/PreAnimatedState/MovieScenePreAnimatedStateExtension.h
@@ -229,6 +229,13 @@ public:
 	 */
 	MOVIESCENE_API void RestoreGlobalState(const FRestoreStateParams& Params);
 
+	/**
+	 * Discards any state for any persistent tokens captured during the course of evaluation without restoring it.
+	 *
+	 * @param Params    Parameters for restoration - if TerminalInstanceHandle is invalid then _all_ state will be discarded, regardless of the instance it was cached from
+	 */
+	MOVIESCENE_API void DiscardGlobalState(const FRestoreStateParams& Params);
+
 	/**
 	 * Restore any state cached for the specified group
 	 *
diff --git a/Engine/Source/Runtime/MovieScene/Public/IMovieScenePlayer.h b/Engine/Source/Runtime/MovieScene/Public/IMovieScenePlayer.h
--- a/Engine/Source/Runtime/MovieScene/Public/IMovieScenePlayer.h
+++ b/Engine/Source/Runtime/MovieScene/Public/IMovieScenePlayer.h
@@ -338,6 +338,15 @@ public:
 		State.ClearObjectCaches(*this);
 	}
 
+	/**
+	 * Discard all pre-animated state without restoring it
+	 */
+	void DiscardPreAnimatedState()
+	{
+		PreAnimatedState.DiscardPreAnimatedState();
+		State.ClearObjectCaches(*this);
+	}
+
 
 	/**
 	 * Invalidate any cached state contained within this player causing all entities to be forcibly re-linked and evaluated
diff --git a/Engine/Source/Runtime/MovieScene/Public/MovieSceneSequencePlaybackSettings.h b/Engine/Source/Runtime/MovieScene/Public/MovieSceneSequencePlaybackSettings.h
--- a/Engine/Source/Runtime/MovieScene/Public/MovieSceneSequencePlaybackSettings.h
+++ b/Engine/Source/Runtime/MovieScene/Public/MovieSceneSequencePlaybackSettings.h
@@ -32,6 +32,15 @@ template<> struct TStructOpsTypeTraits<FMovieSceneSequenceLoopCount> : public TS
 	enum { WithStructuredSerializeFromMismatchedTag = true };
 };
 
+/* Whether to override track sections' default completion mode when a sequence finishes.*/
+UENUM(BlueprintType)
+enum class EMovieSceneCompletionModeOverride
+{
+	None,
+	ForceKeepState,
+	ForceRestoreState
+};
+
 
 /**
  * Settings for the level sequence player actor.
@@ -44,7 +53,7 @@ struct FMovieSceneSequencePlaybackSettings
 		, PlayRate(1.f)
 		, StartTime(0.f)
 		, bRandomStartTime(false)
-		, bRestoreState(false)
+		, bRestoreState_DEPRECATED(false)
 		, bDisableMovementInput(false)
 		, bDisableLookAtInput(false)
 		, bHidePlayer(false)
@@ -81,9 +90,11 @@ struct FMovieSceneSequencePlaybackSettings
 	UPROPERTY(EditAnywhere, BlueprintReadWrite, Category="Playback")
 	uint32 bRandomStartTime : 1;
 
-	/** Flag used to specify whether actor states should be restored on stop */
-	UPROPERTY(EditAnywhere, BlueprintReadWrite, Category="Playback")
-	uint32 bRestoreState : 1;
+	/** Flag used to specify whether actor states should be restored on stop. 
+	* This has been deprecated in favor of FinishCompletionStateOverride.
+	*/
+	UPROPERTY()
+	uint32 bRestoreState_DEPRECATED : 1;
 
 	/** Disable Input from player during play */
 	UPROPERTY(EditAnywhere, BlueprintReadWrite, Category="Cinematic")
@@ -104,6 +115,13 @@ struct FMovieSceneSequencePlaybackSettings
 	/** Disable camera cuts */
 	UPROPERTY(EditAnywhere, BlueprintReadWrite, Category="Cinematic")
 	uint32 bDisableCameraCuts : 1;
+	
+	/* If set to something other than none, when a sequence ends, the completion mode of any track sections still active will be overridden
+	* by the chosen value, either keep state or restore state. Otherwise, completion mode will be determined by each track section.
+	* Note that any track sections that finish before the end of a sequence will have their completion mode determined by the section settings rather than this override.
+	*/
+	UPROPERTY(EditAnywhere, BlueprintReadWrite, Category="Cinematic")
+	EMovieSceneCompletionModeOverride FinishCompletionStateOverride = EMovieSceneCompletionModeOverride::None;
 
 	/** Pause the sequence when playback reaches the end rather than stopping it */
 	UPROPERTY(EditAnywhere, BlueprintReadWrite, Category="Playback")
@@ -117,6 +135,7 @@ struct FMovieSceneSequencePlaybackSettings
 	UPROPERTY(EditAnywhere, Category="Playback")
 	uint32 bDynamicWeighting : 1;
 
+
 	MOVIESCENE_API bool SerializeFromMismatchedTag(const FPropertyTag& Tag, FStructuredArchive::FSlot Slot);
 };
 
