Sequencer: Fix deprecation for bRestoreState. Added meta tags and WITH_EDITORONLY_DATA

BPs that refernce Restore State will need to be fixed up manually.

Fixed up EngineTest BPs that reference Restore State

Note: This should be added to the Upgrade Notes

diff --git a/Engine/Plugins/MovieScene/ActorSequence/Source/ActorSequence/Private/ActorSequenceComponent.cpp b/Engine/Plugins/MovieScene/ActorSequence/Source/ActorSequence/Private/ActorSequenceComponent.cpp
--- a/Engine/Plugins/MovieScene/ActorSequence/Source/ActorSequence/Private/ActorSequenceComponent.cpp
+++ b/Engine/Plugins/MovieScene/ActorSequence/Source/ActorSequence/Private/ActorSequenceComponent.cpp
@@ -89,11 +89,14 @@ void UActorSequenceComponent::TickComponent(float DeltaSeconds, enum ELevelTick
 void UActorSequenceComponent::PostLoad()
 {
 	Super::PostLoad();
+
 	// If we previously were using bRestoreState on our PlaybackSettings, upgrade to the enum version.
+#if WITH_EDITORONLY_DATA
 	if (PlaybackSettings.bRestoreState_DEPRECATED)
 	{
 		PlaybackSettings.FinishCompletionStateOverride = EMovieSceneCompletionModeOverride::ForceRestoreState;
 		PlaybackSettings.bRestoreState_DEPRECATED = false;
 	}
+#endif
 }
 
diff --git a/Engine/Plugins/MovieScene/TemplateSequence/Source/TemplateSequence/Private/TemplateSequenceActor.cpp b/Engine/Plugins/MovieScene/TemplateSequence/Source/TemplateSequence/Private/TemplateSequenceActor.cpp
--- a/Engine/Plugins/MovieScene/TemplateSequence/Source/TemplateSequence/Private/TemplateSequenceActor.cpp
+++ b/Engine/Plugins/MovieScene/TemplateSequence/Source/TemplateSequence/Private/TemplateSequenceActor.cpp
@@ -34,11 +34,13 @@ void ATemplateSequenceActor::PostLoad()
 	Super::PostLoad();
 
 	// If we previously were using bRestoreState on our PlaybackSettings, upgrade to the enum version.
+#if WITH_EDITORONLY_DATA
 	if (PlaybackSettings.bRestoreState_DEPRECATED)
 	{
 		PlaybackSettings.FinishCompletionStateOverride = EMovieSceneCompletionModeOverride::ForceRestoreState;
 		PlaybackSettings.bRestoreState_DEPRECATED = false;
 	}
+#endif
 }
 
 void ATemplateSequenceActor::GetLifetimeReplicatedProps(TArray<FLifetimeProperty>& OutLifetimeProps) const
diff --git a/Engine/Source/Runtime/LevelSequence/Private/LevelSequenceActor.cpp b/Engine/Source/Runtime/LevelSequence/Private/LevelSequenceActor.cpp
--- a/Engine/Source/Runtime/LevelSequence/Private/LevelSequenceActor.cpp
+++ b/Engine/Source/Runtime/LevelSequence/Private/LevelSequenceActor.cpp
@@ -213,11 +213,13 @@ void ALevelSequenceActor::PostLoad()
 	}
 
 	// If we previously were using bRestoreState on our PlaybackSettings, upgrade to the enum version.
+#if WITH_EDITORONLY_DATA
 	if (PlaybackSettings.bRestoreState_DEPRECATED)
 	{
 		PlaybackSettings.FinishCompletionStateOverride = EMovieSceneCompletionModeOverride::ForceRestoreState;
 		PlaybackSettings.bRestoreState_DEPRECATED = false;
 	}
+#endif
 
 	SequencePlayer->SetPlaybackSettings(PlaybackSettings);
 
diff --git a/Engine/Source/Runtime/MovieScene/Public/MovieSceneSequencePlaybackSettings.h b/Engine/Source/Runtime/MovieScene/Public/MovieSceneSequencePlaybackSettings.h
--- a/Engine/Source/Runtime/MovieScene/Public/MovieSceneSequencePlaybackSettings.h
+++ b/Engine/Source/Runtime/MovieScene/Public/MovieSceneSequencePlaybackSettings.h
@@ -34,11 +34,11 @@ template<> struct TStructOpsTypeTraits<FMovieSceneSequenceLoopCount> : public TS
 
 /* Whether to override track sections' default completion mode when a sequence finishes.*/
 UENUM(BlueprintType)
-enum class EMovieSceneCompletionModeOverride
+enum class EMovieSceneCompletionModeOverride : uint8
 {
-	None,
-	ForceKeepState,
-	ForceRestoreState
+	None UMETA(DisplayName = "None"),
+	ForceKeepState UMETA(DisplayName="Force Keep State"),
+	ForceRestoreState UMETA(DisplayName="Force Restore State")
 };
 
 
@@ -53,7 +53,9 @@ struct FMovieSceneSequencePlaybackSettings
 		, PlayRate(1.f)
 		, StartTime(0.f)
 		, bRandomStartTime(false)
+#if WITH_EDITORONLY_DATA
 		, bRestoreState_DEPRECATED(false)
+#endif
 		, bDisableMovementInput(false)
 		, bDisableLookAtInput(false)
 		, bHidePlayer(false)
@@ -90,11 +92,13 @@ struct FMovieSceneSequencePlaybackSettings
 	UPROPERTY(EditAnywhere, BlueprintReadWrite, Category="Playback")
 	uint32 bRandomStartTime : 1;
 
+#if WITH_EDITORONLY_DATA
 	/** Flag used to specify whether actor states should be restored on stop. 
 	* This has been deprecated in favor of FinishCompletionStateOverride.
 	*/
-	UPROPERTY()
+	UPROPERTY(meta = (DeprecatedProperty, DeprecationMessage = "Use Settings.FinishCompletionStateOverride instead"))
 	uint32 bRestoreState_DEPRECATED : 1;
+#endif
 
 	/** Disable Input from player during play */
 	UPROPERTY(EditAnywhere, BlueprintReadWrite, Category="Cinematic")
