Sequencer: fix some issues with replicated sequence playback

There were two bugs that happened rarely:

1) We could have some of the replicated player properties sent to the client, while others weren't. In particular, we would get the LastKnownPosition on the client, but not the LastKnownStatus, which meant we would update and stay stopped, instead of playing. This change therefore adds a custom NetSerialize method on the replicated properties struct so that it is sent as a consistent whole.

2) The other bug had to do with stray network packages making it after "important" RPCs such as Stop and Finish. Sometimes, after the client has stopped from these RPCs, we would get PostNetReceive called with statuses that predate these RPCs. This often meant that, after stopping/finishing playback, the client would re-update to one of the last frames of the sequence. This could easily result in stuck cameras or black (fade) screens. Now we add a serial number to the synchronized properties so that we can ignore these outdated messages.

Other changes include adding some new logging, and simplifying/standardizing existing logging.

diff --git a/Engine/Source/Runtime/Core/Public/Misc/FrameTime.h b/Engine/Source/Runtime/Core/Public/Misc/FrameTime.h
--- a/Engine/Source/Runtime/Core/Public/Misc/FrameTime.h
+++ b/Engine/Source/Runtime/Core/Public/Misc/FrameTime.h
@@ -43,6 +43,11 @@ struct FFrameTime
 	 */
 	FFrameTime& operator=(FFrameNumber InFrameNumber);
 
+	/**
+	 * Serializes the given FrameTime from or into the specified archive
+	 */
+	bool Serialize(FArchive& Ar);
+
 public:
 
 	/**
@@ -95,6 +100,20 @@ private:
 	/** Must be 0.f <= SubFrame < 1.f */
 	float SubFrame;
 
+public:
+
+	/**
+	 * Serializes the given FrameTime from or into the specified archive.
+	 *
+	 * @param Ar            The archive to serialize from or into.
+	 * @param FrameTime     The frame time to serialize.
+	 * @return The archive used for serialization.
+	 */
+	friend FArchive& operator<<(FArchive& Ar, FFrameTime& FrameTime)
+	{
+		FrameTime.Serialize(Ar);
+		return Ar;
+	}
 
 	friend FORCEINLINE_DEBUGGABLE bool operator==(FFrameTime A, FFrameTime B)
 	{
@@ -248,6 +267,14 @@ inline FFrameTime::FFrameTime(FFrameNumber InFrameNumber, float InSubFrame)
 }
 
 
+inline bool FFrameTime::Serialize(FArchive& Ar)
+{
+	Ar << FrameNumber;
+	Ar << SubFrame;
+	return true;
+}
+
+
 inline FFrameTime& FFrameTime::operator=(FFrameNumber InFrameNumber)
 {
 	FrameNumber = InFrameNumber;
diff --git a/Engine/Source/Runtime/MovieScene/Private/MovieSceneSequencePlayer.cpp b/Engine/Source/Runtime/MovieScene/Private/MovieSceneSequencePlayer.cpp
--- a/Engine/Source/Runtime/MovieScene/Private/MovieSceneSequencePlayer.cpp
+++ b/Engine/Source/Runtime/MovieScene/Private/MovieSceneSequencePlayer.cpp
@@ -111,6 +111,17 @@ FFrameTime FMovieSceneSequencePlaybackParams::GetPlaybackPosition(UMovieSceneSeq
 	return PlaybackPosition;
 }
 
+bool FMovieSceneSequenceReplProperties::NetSerialize(FArchive& Ar, UPackageMap* PackageMap, bool& bOutSuccess)
+{
+	Ar << LastKnownPosition;
+	Ar << LastKnownStatus;
+	Ar << LastKnownNumLoops;
+	Ar << LastKnownSerialNumber;
+
+	bOutSuccess = !Ar.IsError();
+	return true;
+}
+
 UMovieSceneSequencePlayer::UMovieSceneSequencePlayer(const FObjectInitializer& Init)
 	: Super(Init)
 	, Status(EMovieScenePlayerStatus::Stopped)
@@ -124,6 +135,7 @@ UMovieSceneSequencePlayer::UMovieSceneSequencePlayer(const FObjectInitializer& I
 	, DurationFrames(0)
 	, DurationSubFrames(0.f)
 	, CurrentNumLoops(0)
+	, SerialNumber(0)
 	, CurrentRunner(nullptr)
 {
 	PlayPosition.Reset(FFrameTime(0));
@@ -147,6 +159,7 @@ void UMovieSceneSequencePlayer::UpdateNetworkSyncProperties()
 		NetSyncProps.LastKnownPosition = PlayPosition.GetCurrentPosition();
 		NetSyncProps.LastKnownStatus   = Status;
 		NetSyncProps.LastKnownNumLoops = CurrentNumLoops;
+		NetSyncProps.LastKnownSerialNumber = SerialNumber;
 	}
 }
 
@@ -241,6 +254,9 @@ void UMovieSceneSequencePlayer::PlayInternal()
 
 	if (!IsPlaying() && Sequence && CanPlay())
 	{
+		const FString SequenceName = GetSequenceName(true);
+		UE_LOG(LogMovieScene, Verbose, TEXT("PlayInternal - %s (current status: %s)"), *SequenceName, *UEnum::GetValueAsString(Status));
+
 		// Set playback status to playing before any calls to update the position
 		Status = EMovieScenePlayerStatus::Playing;
 
@@ -272,9 +288,6 @@ void UMovieSceneSequencePlayer::PlayInternal()
 		Status = EMovieScenePlayerStatus::Playing;
 		TimeController->StartPlaying(GetCurrentTime());
 		
-		UMovieSceneSequence* MovieSceneSequence = RootTemplateInstance.GetSequence(MovieSceneSequenceID::Root);
-		UMovieScene*         MovieScene         = MovieSceneSequence ? MovieSceneSequence->GetMovieScene() : nullptr;
-
 		if (PlayPosition.GetEvaluationType() == EMovieSceneEvaluationType::FrameLocked)
 		{
 			if (!OldMaxTickRate.IsSet())
@@ -291,14 +304,8 @@ void UMovieSceneSequencePlayer::PlayInternal()
 		}
 
 		RunLatentActions();
-
 		UpdateNetworkSyncProperties();
-
-		if (MovieSceneSequence)
-		{
-			UE_LOG(LogMovieScene, Verbose, TEXT("PlayInternal - MovieSceneSequence: %s"), *MovieSceneSequence->GetName());
-		}
-
+		
 		if (bReversePlayback)
 		{
 			if (OnPlayReverse.IsBound())
@@ -324,6 +331,9 @@ void UMovieSceneSequencePlayer::Pause()
 		return;
 	}
 
+	const FString SequenceName = GetSequenceName(true);
+	UE_LOG(LogMovieScene, Verbose, TEXT("Pause - %s (current status: %s)"), *SequenceName, *UEnum::GetValueAsString(Status));
+
 	if (IsPlaying())
 	{
 		Status = EMovieScenePlayerStatus::Paused;
@@ -337,11 +347,8 @@ void UMovieSceneSequencePlayer::Pause()
 			this->RunLatentActions();
 			this->UpdateNetworkSyncProperties();
 
-			UMovieSceneSequence* MovieSceneSequence = this->RootTemplateInstance.GetSequence(MovieSceneSequenceID::Root);
-			if (MovieSceneSequence)
-			{
-				UE_LOG(LogMovieScene, Verbose, TEXT("Pause - MovieSceneSequence: %s"), *MovieSceneSequence->GetName());
-			}
+			const FString SequenceName = this->GetSequenceName(true);
+			UE_LOG(LogMovieScene, Verbose, TEXT("Paused - %s"), *SequenceName);
 
 			if (this->OnPause.IsBound())
 			{
@@ -394,6 +401,9 @@ void UMovieSceneSequencePlayer::StopInternal(FFrameTime TimeToResetTo)
 		return;
 	}
 
+	const FString SequenceName = GetSequenceName(true);
+	UE_LOG(LogMovieScene, Verbose, TEXT("StopInternal - %s (at: %s, current status: %s)"), *SequenceName, *LexToString(TimeToResetTo), *UEnum::GetValueAsString(Status));
+
 	if (IsPlaying() || IsPaused())
 	{
 		Status = EMovieScenePlayerStatus::Stopped;
@@ -430,24 +440,19 @@ void UMovieSceneSequencePlayer::StopInternal(FFrameTime TimeToResetTo)
 				this->OldMaxTickRate.Reset();
 			}
 
+			this->UpdateNetworkSyncProperties();
+
+			const FString SequenceName = this->GetSequenceName(true);
+			UE_LOG(LogMovieScene, Verbose, TEXT("Stopped - %s"), *SequenceName);
+
 			if (this->HasAuthority())
 			{
 				// Explicitly handle Stop() events through an RPC call
-				this->RPC_OnStopEvent(TimeToResetTo);
+				this->RPC_OnStopEvent(TimeToResetTo, SerialNumber + 1);
 			}
-			this->UpdateNetworkSyncProperties();
 
 			this->OnStopped();
 
-			if (this->RootTemplateInstance.IsValid())
-			{
-				UMovieSceneSequence* MovieSceneSequence = RootTemplateInstance.GetSequence(MovieSceneSequenceID::Root);
-				if (MovieSceneSequence)
-				{
-					UE_LOG(LogMovieScene, Verbose, TEXT("Stop - MovieSceneSequence: %s"), *MovieSceneSequence->GetName());
-				}
-			}
-
 			if (this->OnStop.IsBound())
 			{
 				this->OnStop.Broadcast();
@@ -636,7 +641,7 @@ void UMovieSceneSequencePlayer::SetPlaybackPosition(FMovieSceneSequencePlaybackP
 
 	if (HasAuthority())
 	{
-		RPC_ExplicitServerUpdateEvent(InPlaybackParams.UpdateMethod, NewPosition);
+		RPC_ExplicitServerUpdateEvent(InPlaybackParams.UpdateMethod, NewPosition, SerialNumber + 1);
 	}
 }
 
@@ -1108,6 +1113,8 @@ void UMovieSceneSequencePlayer::UpdateTimeCursorPosition_Internal(FFrameTime New
 			FQualifiedFrameTime ExactCurrentTime(StartTimeWithReversed + NewFrameOffset, PlayPosition.GetInputRate());
 			TimeController->Reset(ExactCurrentTime);
 
+			UpdateNetworkSyncProperties();
+
 			OnLooped();
 		}
 
@@ -1129,25 +1136,28 @@ void UMovieSceneSequencePlayer::UpdateTimeCursorPosition_Internal(FFrameTime New
 				(!PlaybackClient || !PlaybackClient->GetIsReplicatedPlayback()) ||
 				HasAuthority() ||
 				NetSyncProps.LastKnownStatus == EMovieScenePlayerStatus::Stopped);
+			const FString SequenceName = GetSequenceName(true);
 			if (bHasAuthorityToFinish)
 			{
+				UE_LOG(LogMovieScene, Verbose,
+						TEXT("Finishing sequence %s at frame %s since we have authority."),
+						*SequenceName, *LexToString(NewPosition));
 				FinishPlaybackInternal(NewPosition);
 
 				// Explicitly tell the clients to finish their playback. They won't have called FinishPlaybackInternal
 				// because it's in the line right above, only for sequence players with some authority
 				// (client only or server).
-				RPC_OnFinishPlaybackEvent(NewPosition);
+				RPC_OnFinishPlaybackEvent(NewPosition, SerialNumber + 1);
 			}
 			else
 			{
-				const FString SequenceName = GetSequenceName(true);
 				UE_LOG(LogMovieScene, Verbose, 
-						TEXT("Keeping sequence %s at frame %d, subframe %f while waiting for playback finish from server."),
-						*SequenceName, NewPosition.FrameNumber.Value, NewPosition.GetSubFrame());
+						TEXT("Keeping sequence %s at frame %s while waiting for playback finish from server."),
+						*SequenceName, *LexToString(NewPosition));
 			}
-		}
 
-		UpdateNetworkSyncProperties();
+			UpdateNetworkSyncProperties();
+		}
 	}
 	else
 	{
@@ -1213,9 +1223,12 @@ void UMovieSceneSequencePlayer::UpdateMovieSceneInstance(FMovieSceneEvaluationRa
 	}
 
 #if !NO_LOGGING
-	const FQualifiedFrameTime CurrentTime = GetCurrentTime();
-	const FString SequenceName = GetSequenceName(true);
-	UE_LOG(LogMovieScene, VeryVerbose, TEXT("Evaluating sequence %s at frame %d, subframe %f (%f fps)."), *SequenceName, CurrentTime.Time.FrameNumber.Value, CurrentTime.Time.GetSubFrame(), CurrentTime.Rate.AsDecimal());
+	if (UE_LOG_ACTIVE(LogMovieScene, VeryVerbose))
+	{
+		const FQualifiedFrameTime CurrentTime = GetCurrentTime();
+		const FString SequenceName = GetSequenceName(true);
+		UE_LOG(LogMovieScene, VeryVerbose, TEXT("Evaluating sequence %s at frame %d, subframe %f (%f fps)."), *SequenceName, CurrentTime.Time.FrameNumber.Value, CurrentTime.Time.GetSubFrame(), CurrentTime.Rate.AsDecimal());
+	}
 #endif
 
 	if (PlaybackClient)
@@ -1485,14 +1498,33 @@ FFrameTime UMovieSceneSequencePlayer::UpdateServerTimeSamples()
 	}
 }
 
-void UMovieSceneSequencePlayer::RPC_ExplicitServerUpdateEvent_Implementation(EUpdatePositionMethod EventMethod, FFrameTime MarkerTime)
+void UMovieSceneSequencePlayer::AdvanceClientSerialNumberTo(int32 NewSerialNumber)
+{
+	if (ensureAlwaysMsgf(!HasAuthority(), TEXT("Trying to advance the serial number on a server player!")))
+	{
+		if (ensureAlwaysMsgf(NewSerialNumber >= SerialNumber, TEXT("Advancing to an older serial number!")))
+		{
+			SerialNumber = NewSerialNumber;
+		}
+	}
+}
+
+void UMovieSceneSequencePlayer::RPC_ExplicitServerUpdateEvent_Implementation(EUpdatePositionMethod EventMethod, FFrameTime MarkerTime, int32 NewSerialNumber)
 {
 	// ~~~~~~~~~~~~~~~~~~~~~~~~~~~~~~~~~~~~~~~~~~~~~~~~~~~~~~~~~~~
 	// Handle an explicit jump/play/scrub command from the server.
 
-	if (HasAuthority() || !Sequence || bIgnorePlaybackReplication)
+	if (HasAuthority())
 	{
 		// Never run network sync operations on authoritative players
+		ensure(NewSerialNumber > SerialNumber);
+		SerialNumber = NewSerialNumber;
+		return;
+	}
+
+	if (!Sequence || bIgnorePlaybackReplication)
+	{
+		// Never run network sync operations on players that have not been initialized yet
 		return;
 	}
 
@@ -1505,12 +1537,16 @@ void UMovieSceneSequencePlayer::RPC_ExplicitServerUpdateEvent_Implementation(EUp
 	{
 		const FFrameTime CurrentTime = PlayPosition.GetCurrentPosition();
 		const FString SequenceName = GetSequenceName(true);
-		UE_LOG(LogMovieScene, Verbose, TEXT("Explicit update event for sequence %s %s @ frame %d, subframe %f. Server has moved to frame %d, subframe %f with EUpdatePositionMethod::%s."),
-			*SequenceName, *UEnum::GetValueAsString(TEXT("MovieScene.EMovieScenePlayerStatus"), Status.GetValue()), CurrentTime.FrameNumber.Value, CurrentTime.GetSubFrame(),
-			NetSyncProps.LastKnownPosition.FrameNumber.Value, NetSyncProps.LastKnownPosition.GetSubFrame(), *UEnum::GetValueAsString(TEXT("MovieScene.EUpdatePositionMethod"), NetSyncProps.LastKnownStatus.GetValue()));
+		UE_LOG(LogMovieScene, Verbose, TEXT("Explicit update event for sequence %s %s @ %s. Server has moved to %s @ %s."),
+			*SequenceName,
+			*UEnum::GetValueAsString(Status.GetValue()), *LexToString(CurrentTime),
+			*UEnum::GetValueAsString(NetSyncProps.LastKnownStatus.GetValue()), *LexToString(MarkerTime));
 	}
 #endif
 
+	// Update our serial number
+	AdvanceClientSerialNumberTo(NewSerialNumber);
+
 	// Explicitly repeat the authoritative update event on this client.
 
 	// Note: in the case of PlayToFrame this will not necessarily sweep the exact same range as the server did
@@ -1519,14 +1555,22 @@ void UMovieSceneSequencePlayer::RPC_ExplicitServerUpdateEvent_Implementation(EUp
 	SetPlaybackPosition(FMovieSceneSequencePlaybackParams(MarkerTime, EventMethod));
 }
 
-void UMovieSceneSequencePlayer::RPC_OnStopEvent_Implementation(FFrameTime StoppedTime)
+void UMovieSceneSequencePlayer::RPC_OnStopEvent_Implementation(FFrameTime StoppedTime, int32 NewSerialNumber)
 {
 	// ~~~~~~~~~~~~~~~~~~~~~~~~~~~~~~~~~~~~~~~~~~~~~~~~
 	// Handle an explicit Stop command from the server.
 
-	if (HasAuthority() || !Sequence || bIgnorePlaybackReplication)
+	if (HasAuthority())
 	{
-		// Never run network sync operations on authoritative players or players that have not been initialized yet
+		// Never run network sync operations on authoritative players
+		ensure(NewSerialNumber > SerialNumber);
+		SerialNumber = NewSerialNumber;
+		return;
+	}
+
+	if (!Sequence || bIgnorePlaybackReplication)
+	{
+		// Never run network sync operations on players that have not been initialized yet
 		return;
 	}
 
@@ -1539,11 +1583,15 @@ void UMovieSceneSequencePlayer::RPC_OnStopEvent_Implementation(FFrameTime Stoppe
 		const FFrameTime CurrentTime = PlayPosition.GetCurrentPosition();
 		const FString SequenceName = GetSequenceName(true);
 		UE_LOG(LogMovieSceneRepl, Verbose, TEXT("Explicit Stop() event for sequence %s %s @ frame %d, subframe %f. Server has stopped at frame %d, subframe %f."),
-			*SequenceName, *UEnum::GetValueAsString(TEXT("MovieScene.EMovieScenePlayerStatus"), Status.GetValue()), CurrentTime.FrameNumber.Value, CurrentTime.GetSubFrame(),
-			NetSyncProps.LastKnownPosition.FrameNumber.Value, NetSyncProps.LastKnownPosition.GetSubFrame());
+			*SequenceName, *UEnum::GetValueAsString(Status.GetValue()),
+			CurrentTime.FrameNumber.Value, CurrentTime.GetSubFrame(),
+			StoppedTime.FrameNumber.Value, StoppedTime.GetSubFrame());
 	}
 #endif
 
+	// Update our serial number
+	AdvanceClientSerialNumberTo(NewSerialNumber);
+
 	EUpdatePositionMethod UpdatePositionMethod;
 	switch (Status.GetValue())
 	{
@@ -1571,20 +1619,35 @@ void UMovieSceneSequencePlayer::RPC_OnStopEvent_Implementation(FFrameTime Stoppe
 	// Now do the last bit of catch-up for the current loop.
 	if (PlayPosition.GetCurrentPosition() < StoppedTime)
 	{
+		UE_LOG(LogMovieSceneRepl, Verbose, TEXT("Catching up to explicit stop time %s"), *LexToString(StoppedTime));
 		SetPlaybackPosition(FMovieSceneSequencePlaybackParams(StoppedTime, UpdatePositionMethod));
 	}
 
 	StopInternal(StoppedTime);
 }
 
-void UMovieSceneSequencePlayer::RPC_OnFinishPlaybackEvent_Implementation(FFrameTime StoppedTime)
+void UMovieSceneSequencePlayer::RPC_OnFinishPlaybackEvent_Implementation(FFrameTime StoppedTime, int32 NewSerialNumber)
 {
-	if (HasAuthority() || !Sequence || bIgnorePlaybackReplication)
+	if (HasAuthority())
 	{
-		// Never run network sync operations on authoritative players or players that have not been initialized yet
+		// Never run network sync operations on authoritative players
+		ensure(NewSerialNumber > SerialNumber);
+		SerialNumber = NewSerialNumber;
+		return;
+	}
+
+	if (!Sequence || bIgnorePlaybackReplication)
+	{
+		// Never run network sync operations on players that have not been initialized yet
 		return;
 	}
 
+	const FString SequenceName = GetSequenceName(true);
+	UE_LOG(LogMovieScene, Verbose, TEXT("Received RPC event to finish sequence %s at frame %s."), *SequenceName, *LexToString(StoppedTime));
+
+	// Update our serial number
+	AdvanceClientSerialNumberTo(NewSerialNumber);
+
 	FinishPlaybackInternal(StoppedTime);
 }
 
@@ -1601,6 +1664,24 @@ void UMovieSceneSequencePlayer::PostNetReceive()
 		return;
 	}
 
+	// Very occasionally a stray network update can come late, and we need to discard it. One such situation
+	// is when the server invokes an RPC to stop and finish the sequence, but late network updates arrive
+	// after that for the last few frames, and the client player ends up restarting the sequence to evaluate
+	// these last few frames even though it has already stopped from the RPCs.
+	if (NetSyncProps.LastKnownSerialNumber < SerialNumber)
+	{
+#if !NO_LOGGING
+		const FFrameTime CurrentTime = PlayPosition.GetCurrentPosition();
+		const FString SequenceName = GetSequenceName(true);
+		UE_LOG(LogMovieSceneRepl, Verbose, 
+				TEXT("Ignoring network update with old serial (%d < %d) for sequence %s %s @ %s. Server was %s @ %s."),
+				NetSyncProps.LastKnownSerialNumber, SerialNumber, *SequenceName,
+				*UEnum::GetValueAsString(Status.GetValue()), *LexToString(CurrentTime),
+				*UEnum::GetValueAsString(NetSyncProps.LastKnownStatus.GetValue()), *LexToString(NetSyncProps.LastKnownPosition));
+#endif
+		return;
+	}
+
 	const bool bHasStartedPlaying = NetSyncProps.LastKnownStatus == EMovieScenePlayerStatus::Playing && Status != EMovieScenePlayerStatus::Playing;
 	const bool bHasChangedStatus  = NetSyncProps.LastKnownStatus   != Status;
 	const bool bHasChangedTime    = NetSyncProps.LastKnownPosition != PlayPosition.GetCurrentPosition();
@@ -1622,9 +1703,10 @@ void UMovieSceneSequencePlayer::PostNetReceive()
 	{
 		const FFrameTime CurrentTime = PlayPosition.GetCurrentPosition();
 		const FString SequenceName = GetSequenceName(true);
-		UE_LOG(LogMovieSceneRepl, VeryVerbose, TEXT("Network sync for sequence %s %s @ frame %d, subframe %f. Server is %s @ frame %d, subframe %f."),
-			*SequenceName, *UEnum::GetValueAsString(TEXT("MovieScene.EMovieScenePlayerStatus"), Status.GetValue()), CurrentTime.FrameNumber.Value, CurrentTime.GetSubFrame(),
-			*UEnum::GetValueAsString(TEXT("MovieScene.EMovieScenePlayerStatus"), NetSyncProps.LastKnownStatus.GetValue()), NetSyncProps.LastKnownPosition.FrameNumber.Value, NetSyncProps.LastKnownPosition.GetSubFrame());
+		UE_LOG(LogMovieSceneRepl, VeryVerbose, TEXT("Network sync for sequence %s %s @ %s. Server is %s @ %s."),
+			*SequenceName,
+			*UEnum::GetValueAsString(Status.GetValue()), *LexToString(CurrentTime),
+			*UEnum::GetValueAsString(NetSyncProps.LastKnownStatus.GetValue()), *LexToString(NetSyncProps.LastKnownPosition));
 	}
 #endif
 
@@ -1746,12 +1828,11 @@ void UMovieSceneSequencePlayer::UpdateNetworkSync()
 			{
 				const FFrameTime CurrentTime = PlayPosition.GetCurrentPosition();
 				const FString SequenceName = GetSequenceName(true);
-				UE_LOG(LogMovieSceneRepl, Log, TEXT("Correcting de-synced play position for sequence %s %s @ frame %d, subframe %f. Server is %s @ frame %d, subframe %f, (smoothed: frame %d, sf %f). Client ping is %.2fms."),
-					*SequenceName, *UEnum::GetValueAsString(TEXT("MovieScene.EMovieScenePlayerStatus"), Status.GetValue()), CurrentTime.FrameNumber.Value, CurrentTime.GetSubFrame(),
-					*UEnum::GetValueAsString(TEXT("MovieScene.EMovieScenePlayerStatus"), NetSyncProps.LastKnownStatus.GetValue()),
-					NetSyncProps.LastKnownPosition.FrameNumber.Value, NetSyncProps.LastKnownPosition.GetSubFrame(),
-					SmoothedServerTime.FrameNumber.Value, SmoothedServerTime.GetSubFrame(),
-					PingMs);
+				UE_LOG(LogMovieSceneRepl, Log, TEXT("Correcting de-synced play position for sequence %s %s @ %s. Server is %s @ %s, (smoothed: %s). Client ping is %.2fms."),
+					*SequenceName, 
+					*UEnum::GetValueAsString(Status.GetValue()), *LexToString(CurrentTime),
+					*UEnum::GetValueAsString(NetSyncProps.LastKnownStatus.GetValue()), *LexToString(NetSyncProps.LastKnownPosition),
+					*LexToString(SmoothedServerTime), PingMs);
 			}
 #endif
 			// We're drastically out of sync with the server so we need to forcibly set the time.
diff --git a/Engine/Source/Runtime/MovieScene/Public/MovieSceneSequencePlayer.h b/Engine/Source/Runtime/MovieScene/Public/MovieSceneSequencePlayer.h
--- a/Engine/Source/Runtime/MovieScene/Public/MovieSceneSequencePlayer.h
+++ b/Engine/Source/Runtime/MovieScene/Public/MovieSceneSequencePlayer.h
@@ -56,6 +56,7 @@ struct FMovieSceneSequenceReplProperties
 	FMovieSceneSequenceReplProperties()
 		: LastKnownStatus(EMovieScenePlayerStatus::Stopped)
 		, LastKnownNumLoops(0)
+		, LastKnownSerialNumber(0)
 	{}
 
 	/** The last known position of the sequence on the server */
@@ -69,6 +70,29 @@ struct FMovieSceneSequenceReplProperties
 	/** The last known number of loops of the sequence on the server */
 	UPROPERTY()
 	int32 LastKnownNumLoops;
+
+	/** The last known serial number on the server */
+	UPROPERTY()
+	int32 LastKnownSerialNumber;
+
+	/**
+	 * Custom serialization method so that all the properties in this struct are
+	 * in sync with the server. Without this, we could sometimes get half of the
+	 * properties up-to-date, and half of the properties left to old values.
+	 * This is especially problematic when LastKnownPosition updates, but LastKnownStatus
+	 * is left as "Stopped".
+	 */
+	bool NetSerialize(FArchive& Ar, UPackageMap* PackageMap, bool& bOutSuccess);
+};
+
+
+template<>
+struct TStructOpsTypeTraits<FMovieSceneSequenceReplProperties> : public TStructOpsTypeTraitsBase2<FMovieSceneSequenceReplProperties>
+{
+	enum
+	{
+		WithNetSerializer = true
+	};
 };
 
 
@@ -551,6 +575,9 @@ private:
 
 	MOVIESCENE_API void RunPreEvaluationCallbacks();
 	MOVIESCENE_API void RunPostEvaluationCallbacks();
+
+	void IncrementServerSerialNumber();
+	void AdvanceClientSerialNumberTo(int32 NewSerialNumber);
 	
 private:
 
@@ -558,19 +585,19 @@ private:
 	 * Called on the server whenever an explicit change in time has occurred through one of the (Play|Jump|Scrub)To methods
 	 */
 	UFUNCTION(netmulticast, reliable)
-	MOVIESCENE_API void RPC_ExplicitServerUpdateEvent(EUpdatePositionMethod Method, FFrameTime RelevantTime);
+	MOVIESCENE_API void RPC_ExplicitServerUpdateEvent(EUpdatePositionMethod Method, FFrameTime RelevantTime, int32 NewSerialNumber);
 
 	/**
 	 * Called on the server when Stop() is called in order to differentiate Stops from Pauses.
 	 */
 	UFUNCTION(netmulticast, reliable)
-	MOVIESCENE_API void RPC_OnStopEvent(FFrameTime StoppedTime);
+	MOVIESCENE_API void RPC_OnStopEvent(FFrameTime StoppedTime, int32 NewSerialNumber);
 
 	/**
 	 * Called on the server when playback has reached the end. Could lead to stopping or pausing.
 	 */
 	UFUNCTION(netmulticast, reliable)
-	MOVIESCENE_API void RPC_OnFinishPlaybackEvent(FFrameTime StoppedTime);
+	MOVIESCENE_API void RPC_OnFinishPlaybackEvent(FFrameTime StoppedTime, int32 NewSerialNumber);
 
 	/**
 	 * Check whether this sequence player is an authority, as determined by its outer Actor
@@ -640,6 +667,13 @@ protected:
 	UPROPERTY(transient)
 	int32 CurrentNumLoops;
 
+	/**
+	 * The serial number for the current update lifespan
+	 * It is incremented every time we pass a "gate" such as an RPC call that stops/finishes the sequence.
+	 */
+	UPROPERTY(transient)
+	int32 SerialNumber;
+
 	/** Specific playback settings for the animation. */
 	UPROPERTY(replicated)
 	FMovieSceneSequencePlaybackSettings PlaybackSettings;
