Sequencer: revert NetSyncProps to the old property-based serialization

This is to not break replay data. Will need to investigate how to redo this change in a backwards-compatible way if the rare issue of inconsistent network packets crops up again.

diff --git a/Engine/Source/Runtime/MovieScene/Private/MovieSceneSequencePlayer.cpp b/Engine/Source/Runtime/MovieScene/Private/MovieSceneSequencePlayer.cpp
--- a/Engine/Source/Runtime/MovieScene/Private/MovieSceneSequencePlayer.cpp
+++ b/Engine/Source/Runtime/MovieScene/Private/MovieSceneSequencePlayer.cpp
@@ -111,17 +111,6 @@ FFrameTime FMovieSceneSequencePlaybackParams::GetPlaybackPosition(UMovieSceneSeq
 	return PlaybackPosition;
 }
 
-bool FMovieSceneSequenceReplProperties::NetSerialize(FArchive& Ar, UPackageMap* PackageMap, bool& bOutSuccess)
-{
-	Ar << LastKnownPosition;
-	Ar << LastKnownStatus;
-	Ar << LastKnownNumLoops;
-	Ar << LastKnownSerialNumber;
-
-	bOutSuccess = !Ar.IsError();
-	return true;
-}
-
 UMovieSceneSequencePlayer::UMovieSceneSequencePlayer(const FObjectInitializer& Init)
 	: Super(Init)
 	, Status(EMovieScenePlayerStatus::Stopped)
diff --git a/Engine/Source/Runtime/MovieScene/Public/MovieSceneSequencePlayer.h b/Engine/Source/Runtime/MovieScene/Public/MovieSceneSequencePlayer.h
--- a/Engine/Source/Runtime/MovieScene/Public/MovieSceneSequencePlayer.h
+++ b/Engine/Source/Runtime/MovieScene/Public/MovieSceneSequencePlayer.h
@@ -74,25 +74,6 @@ struct FMovieSceneSequenceReplProperties
 	/** The last known serial number on the server */
 	UPROPERTY()
 	int32 LastKnownSerialNumber;
-
-	/**
-	 * Custom serialization method so that all the properties in this struct are
-	 * in sync with the server. Without this, we could sometimes get half of the
-	 * properties up-to-date, and half of the properties left to old values.
-	 * This is especially problematic when LastKnownPosition updates, but LastKnownStatus
-	 * is left as "Stopped".
-	 */
-	bool NetSerialize(FArchive& Ar, UPackageMap* PackageMap, bool& bOutSuccess);
-};
-
-
-template<>
-struct TStructOpsTypeTraits<FMovieSceneSequenceReplProperties> : public TStructOpsTypeTraitsBase2<FMovieSceneSequenceReplProperties>
-{
-	enum
-	{
-		WithNetSerializer = true
-	};
 };
 
 
