Sequencer: Remove EditAnywhere from Flags which wasn't intended on being exposed

#rnx

diff --git a/Engine/Source/Runtime/MovieScene/Public/Evaluation/MovieSceneSectionParameters.h b/Engine/Source/Runtime/MovieScene/Public/Evaluation/MovieSceneSectionParameters.h
--- a/Engine/Source/Runtime/MovieScene/Public/Evaluation/MovieSceneSectionParameters.h
+++ b/Engine/Source/Runtime/MovieScene/Public/Evaluation/MovieSceneSectionParameters.h
@@ -95,7 +95,7 @@ public:
 	int32 HierarchicalBias;
 
 	/** Sub-section flags defining how to deal with this sub-sequence */
-	UPROPERTY(config, BlueprintReadWrite, EditAnywhere, Category="Sequence")
+	UPROPERTY(config, BlueprintReadWrite, Category="Sequence")
 	EMovieSceneSubSectionFlags Flags;
 
 	UPROPERTY()
