Sequencer: Proper support for zero-timescale sub-sequence sections. The intention behind setting zero-timescale on a section is to hold at a single frame (the StartFrameOffset) while playing this sub section in its parent sequence.
Previously, this 'kind of' worked, but in practice there were divide by zeros happening and this was causing bugs under certain circumstances, as well as some odd behavior.

To support this change required a bit of a refactor of FMovieSceneSequenceTransform and its 'Warping'. Warping previously just indicated looping, but now it can also indicate timescales of zero. In cases of zero timescale transforms, we keep transforms split into FMovieSceneNestedSequenceTransforms rather than combining them.

Some logic that directly accessed the internal FMovieSceneTimeTransform 'LinearTransform' also had to be refactored to properly use the FMovieSceneSequenceTransform's accessor functions for doing things like inverse. When dealing with zero timescale in FMovieSceneTimeTransforms, there is the ability for loss of information when multiplying and inverting, and so to compensate for that we need every hierarchical step of the transforms stored. This also necessitated some changes in logic for TMovieSceneTimeArray and how it stores its TransformStack, as it was depending on the inverse of an inverse being a lossless operation when acting on a FMovieSceneTimeTransform.

See https://docs.google.com/document/d/1ftdIyi9J8At6fw31K5awKCUK5-75uz7F-MWAq5b1vWc/edit#heading=h.fclfl8j8p52c.

diff --git a/Engine/Plugins/Animation/ControlRig/Source/ControlRigEditor/Private/ControlRigSequencerEditorLibrary.cpp b/Engine/Plugins/Animation/ControlRig/Source/ControlRigEditor/Private/ControlRigSequencerEditorLibrary.cpp
--- a/Engine/Plugins/Animation/ControlRig/Source/ControlRigEditor/Private/ControlRigSequencerEditorLibrary.cpp
+++ b/Engine/Plugins/Animation/ControlRig/Source/ControlRigEditor/Private/ControlRigSequencerEditorLibrary.cpp
@@ -1318,7 +1318,7 @@ static bool LocalGetControlRigControlValues(IMovieScenePlayer* Player, UMovieSce
 			{
 				GlobalTime = FFrameRate::TransformTime(GlobalTime, MovieScene->GetDisplayRate(), MovieScene->GetTickResolution());
 			}
-			GlobalTime = GlobalTime * RootToLocalTransform.InverseLinearOnly();
+			GlobalTime = GlobalTime * RootToLocalTransform.InverseNoLooping();
 			FMovieSceneContext Context = FMovieSceneContext(FMovieSceneEvaluationRange(GlobalTime, TickResolution), Player->GetPlaybackStatus()).SetHasJumped(true);
 
 			Player->GetEvaluationTemplate().EvaluateSynchronousBlocking(Context, *Player);
diff --git a/Engine/Plugins/Animation/ControlRig/Source/ControlRigEditor/Private/Sequencer/ControlRigParameterTrackEditor.cpp b/Engine/Plugins/Animation/ControlRig/Source/ControlRigEditor/Private/Sequencer/ControlRigParameterTrackEditor.cpp
--- a/Engine/Plugins/Animation/ControlRig/Source/ControlRigEditor/Private/Sequencer/ControlRigParameterTrackEditor.cpp
+++ b/Engine/Plugins/Animation/ControlRig/Source/ControlRigEditor/Private/Sequencer/ControlRigParameterTrackEditor.cpp
@@ -4942,7 +4942,7 @@ bool FControlRigParameterTrackEditor::CollapseAllLayers(TSharedPtr<ISequencer>&S
 					{
 						const FFrameNumber& FrameNumber = Frames[Index];
 						FFrameTime GlobalTime(FrameNumber);
-						GlobalTime = GlobalTime * RootToLocalTransform.InverseLinearOnly();
+						GlobalTime = GlobalTime * RootToLocalTransform.InverseNoLooping();
 
 						FMovieSceneContext Context = FMovieSceneContext(FMovieSceneEvaluationRange(GlobalTime, TickResolution), SequencerPtr->GetPlaybackStatus()).SetHasJumped(true);
 
diff --git a/Engine/Plugins/Animation/ControlRig/Source/ControlRigEditor/Private/Sequencer/ControlRigSpaceChannelEditors.cpp b/Engine/Plugins/Animation/ControlRig/Source/ControlRigEditor/Private/Sequencer/ControlRigSpaceChannelEditors.cpp
--- a/Engine/Plugins/Animation/ControlRig/Source/ControlRigEditor/Private/Sequencer/ControlRigSpaceChannelEditors.cpp
+++ b/Engine/Plugins/Animation/ControlRig/Source/ControlRigEditor/Private/Sequencer/ControlRigSpaceChannelEditors.cpp
@@ -447,7 +447,7 @@ FKeyHandle FControlRigSpaceChannelHelpers::SequencerKeyControlRigSpaceChannel(UC
 		if (bSetPreviousKey)
 		{
 			FFrameTime GlobalTime(Time - 1);
-			GlobalTime = GlobalTime * RootToLocalTransform.InverseLinearOnly();
+			GlobalTime = GlobalTime * RootToLocalTransform.InverseNoLooping();
 
 			FMovieSceneContext SceneContext = FMovieSceneContext(FMovieSceneEvaluationRange(GlobalTime, TickResolution), Sequencer->GetPlaybackStatus()).SetHasJumped(true);
 			Sequencer->GetEvaluationTemplate().EvaluateSynchronousBlocking(SceneContext, *Sequencer);
@@ -477,7 +477,7 @@ FKeyHandle FControlRigSpaceChannelHelpers::SequencerKeyControlRigSpaceChannel(UC
 		for (const FFrameNumber& Frame : Frames)
 		{
 			FFrameTime GlobalTime(Frame);
-			GlobalTime = GlobalTime * RootToLocalTransform.InverseLinearOnly();
+			GlobalTime = GlobalTime * RootToLocalTransform.InverseNoLooping();
 
 			FMovieSceneContext SceneContext = FMovieSceneContext(FMovieSceneEvaluationRange(GlobalTime, TickResolution), Sequencer->GetPlaybackStatus()).SetHasJumped(true);
 			Sequencer->GetEvaluationTemplate().EvaluateSynchronousBlocking(SceneContext, *Sequencer);
@@ -687,7 +687,7 @@ void  FControlRigSpaceChannelHelpers::SequencerSpaceChannelKeyDeleted(UControlRi
 		{
 			//evaluate sequencer
 			FFrameTime GlobalTime(Frame);
-			GlobalTime = GlobalTime * RootToLocalTransform.InverseLinearOnly();
+			GlobalTime = GlobalTime * RootToLocalTransform.InverseNoLooping();
 
 			FMovieSceneContext SceneContext = FMovieSceneContext(FMovieSceneEvaluationRange(GlobalTime, TickResolution), Sequencer->GetPlaybackStatus()).SetHasJumped(true);
 			Sequencer->GetEvaluationTemplate().EvaluateSynchronousBlocking(SceneContext, *Sequencer);
@@ -1087,7 +1087,7 @@ void FControlRigSpaceChannelHelpers::SequencerBakeControlInSpace(UControlRig* Co
 
 				//evaluate sequencer
 				FFrameTime GlobalTime(Frame);
-				GlobalTime = GlobalTime * RootToLocalTransform.InverseLinearOnly();
+				GlobalTime = GlobalTime * RootToLocalTransform.InverseNoLooping();
 
 				FMovieSceneContext SceneContext = FMovieSceneContext(FMovieSceneEvaluationRange(GlobalTime, TickResolution), Sequencer->GetPlaybackStatus()).SetHasJumped(true);
 				Sequencer->GetEvaluationTemplate().EvaluateSynchronousBlocking(SceneContext, *Sequencer);
@@ -1127,7 +1127,7 @@ void FControlRigSpaceChannelHelpers::SequencerBakeControlInSpace(UControlRig* Co
 
 				//evaluate sequencer
 				FFrameTime GlobalTime(EndFrame);
-				GlobalTime = GlobalTime * RootToLocalTransform.InverseLinearOnly();
+				GlobalTime = GlobalTime * RootToLocalTransform.InverseNoLooping();
 
 				FMovieSceneContext SceneContext = FMovieSceneContext(FMovieSceneEvaluationRange(GlobalTime, TickResolution), Sequencer->GetPlaybackStatus()).SetHasJumped(true);
 				Sequencer->GetEvaluationTemplate().EvaluateSynchronousBlocking(SceneContext, *Sequencer);
diff --git a/Engine/Plugins/Animation/ControlRig/Source/ControlRigEditor/Private/Tools/ControlRigSnapper.cpp b/Engine/Plugins/Animation/ControlRig/Source/ControlRigEditor/Private/Tools/ControlRigSnapper.cpp
--- a/Engine/Plugins/Animation/ControlRig/Source/ControlRigEditor/Private/Tools/ControlRigSnapper.cpp
+++ b/Engine/Plugins/Animation/ControlRig/Source/ControlRigEditor/Private/Tools/ControlRigSnapper.cpp
@@ -138,7 +138,7 @@ static bool LocalGetControlRigControlTransforms(IMovieScenePlayer* Player, const
 			if (CurrentFrame.IsSet() == false || CurrentFrame.GetValue() != FrameNumber)
 			{
 				FFrameTime GlobalTime(FrameNumber);
-				GlobalTime = GlobalTime * RootToLocalTransform.InverseLinearOnly(); //player evals in root time so need to go back to it.
+				GlobalTime = GlobalTime * RootToLocalTransform.InverseNoLooping(); //player evals in root time so need to go back to it.
 
 				FMovieSceneContext Context = FMovieSceneContext(FMovieSceneEvaluationRange(GlobalTime, TickResolution), Player->GetPlaybackStatus()).SetHasJumped(true);
 				Player->GetEvaluationTemplate().EvaluateSynchronousBlocking(Context, *Player);
@@ -157,7 +157,7 @@ bool FControlRigSnapper::GetControlRigControlTransforms(ISequencer* Sequencer,
 	{
 		FMovieSceneSequenceIDRef Template = Sequencer->GetFocusedTemplateID();
 		FMovieSceneSequenceTransform RootToLocalTransform = Sequencer->GetFocusedMovieSceneSequenceTransform();
-		TOptional<FFrameNumber> FrameNumber = (Sequencer->GetLocalTime().Time * RootToLocalTransform.InverseLinearOnly()).RoundToFrame();
+		TOptional<FFrameNumber> FrameNumber = (Sequencer->GetLocalTime().Time * RootToLocalTransform.InverseNoLooping()).RoundToFrame();
 		return LocalGetControlRigControlTransforms(Sequencer, FrameNumber, Sequencer->GetFocusedMovieSceneSequence(), Template, RootToLocalTransform,
 			ControlRig, ControlName, Frames, ParentTransforms, OutTransforms);
 	
diff --git a/Engine/Plugins/Editor/SequencerAnimTools/Source/SequencerAnimTools/Private/Sequencer/MovieSceneTransformTrail.cpp b/Engine/Plugins/Editor/SequencerAnimTools/Source/SequencerAnimTools/Private/Sequencer/MovieSceneTransformTrail.cpp
--- a/Engine/Plugins/Editor/SequencerAnimTools/Source/SequencerAnimTools/Private/Sequencer/MovieSceneTransformTrail.cpp
+++ b/Engine/Plugins/Editor/SequencerAnimTools/Source/SequencerAnimTools/Private/Sequencer/MovieSceneTransformTrail.cpp
@@ -747,7 +747,7 @@ bool FMovieSceneControlRigTransformTrail::ApplyDelta(const FVector& Pos, const F
 				Context.LocalTime = TickResolution.AsSeconds(FFrameTime(KeyInfo->FrameNumber));
 				Context.KeyMask = (uint32)EControlRigContextChannelToKey::Translation;
 				FFrameTime GlobalTime(KeyInfo->FrameNumber);
-				GlobalTime = GlobalTime * RootToLocalTransform.InverseLinearOnly(); //player evals in root time so need to go back to it.
+				GlobalTime = GlobalTime * RootToLocalTransform.InverseNoLooping(); //player evals in root time so need to go back to it.
 
 				FMovieSceneContext MovieSceneContext = FMovieSceneContext(FMovieSceneEvaluationRange(GlobalTime, TickResolution), Player->GetPlaybackStatus()).SetHasJumped(true);
 
@@ -806,7 +806,7 @@ bool FMovieSceneControlRigTransformTrail::EndTracking()
 		FFrameTime StartTime = GetSequencer()->GetLocalTime().Time;
 		FFrameRate TickResolution = GetSequencer()->GetFocusedTickResolution();
 		FMovieSceneSequenceTransform RootToLocalTransform = GetSequencer()->GetFocusedMovieSceneSequenceTransform();
-		StartTime = StartTime * RootToLocalTransform.InverseLinearOnly(); //player evals in root time so need to go back to it.
+		StartTime = StartTime * RootToLocalTransform.InverseNoLooping(); //player evals in root time so need to go back to it.
 
 		FMovieSceneContext MovieSceneContext = FMovieSceneContext(FMovieSceneEvaluationRange(StartTime, TickResolution), Player->GetPlaybackStatus()).SetHasJumped(true);
 		
@@ -845,7 +845,7 @@ bool FMovieSceneControlRigTransformTrail::HandleAltClick(FEditorViewportClient*
 	DisplayTime = DisplayTime.RoundToFrame();
 	GlobalTime = FFrameRate::TransformTime(DisplayTime, TickResolution, DisplayResolution);
 	FMovieSceneSequenceTransform RootToLocalTransform = GetSequencer()->GetFocusedMovieSceneSequenceTransform();
-	GlobalTime = GlobalTime * RootToLocalTransform.InverseLinearOnly(); //player evals in root time so need to go back to it.
+	GlobalTime = GlobalTime * RootToLocalTransform.InverseNoLooping(); //player evals in root time so need to go back to it.
 
 	Context.LocalTime = TickResolution.AsSeconds(GlobalTime);
 	Context.KeyMask = (uint32)EControlRigContextChannelToKey::Translation;
@@ -857,7 +857,7 @@ bool FMovieSceneControlRigTransformTrail::HandleAltClick(FEditorViewportClient*
 
 	//eval back at current time
 	FFrameTime StartTime = GetSequencer()->GetLocalTime().Time;
-	StartTime = StartTime * RootToLocalTransform.InverseLinearOnly(); //player evals in root time so need to go back to it.
+	StartTime = StartTime * RootToLocalTransform.InverseNoLooping(); //player evals in root time so need to go back to it.
 
 	MovieSceneContext = FMovieSceneContext(FMovieSceneEvaluationRange(StartTime, TickResolution), Player->GetPlaybackStatus()).SetHasJumped(true);
 	Player->GetEvaluationTemplate().EvaluateSynchronousBlocking(MovieSceneContext, *Player);
diff --git a/Engine/Plugins/Editor/SequencerAnimTools/Source/SequencerAnimTools/Private/Sequencer/SequencerTrailHierarchy.cpp b/Engine/Plugins/Editor/SequencerAnimTools/Source/SequencerAnimTools/Private/Sequencer/SequencerTrailHierarchy.cpp
--- a/Engine/Plugins/Editor/SequencerAnimTools/Source/SequencerAnimTools/Private/Sequencer/SequencerTrailHierarchy.cpp
+++ b/Engine/Plugins/Editor/SequencerAnimTools/Source/SequencerAnimTools/Private/Sequencer/SequencerTrailHierarchy.cpp
@@ -258,7 +258,7 @@ void FSequencerTrailHierarchy::UpdateControlRig(const TArray<FFrameNumber> &Fram
 				{
 					const FFrameNumber& FrameNumber = Frames[Index];
 					FFrameTime GlobalTime(FrameNumber);
-					GlobalTime = GlobalTime * RootToLocalTransform.InverseLinearOnly(); //player evals in root time so need to go back to it.
+					GlobalTime = GlobalTime * RootToLocalTransform.InverseNoLooping(); //player evals in root time so need to go back to it.
 
 					FMovieSceneContext Context = FMovieSceneContext(FMovieSceneEvaluationRange(GlobalTime, TickResolution), Player->GetPlaybackStatus()).SetHasJumped(true);
 
@@ -292,7 +292,7 @@ void FSequencerTrailHierarchy::UpdateControlRig(const TArray<FFrameNumber> &Fram
 		if (Player)
 		{
 			FFrameTime StartTime = Sequencer->GetLocalTime().Time;
-			StartTime = StartTime * RootToLocalTransform.InverseLinearOnly(); //player evals in root time so need to go back to it.
+			StartTime = StartTime * RootToLocalTransform.InverseNoLooping(); //player evals in root time so need to go back to it.
 			FMovieSceneContext Context = FMovieSceneContext(FMovieSceneEvaluationRange(StartTime, TickResolution), Player->GetPlaybackStatus()).SetHasJumped(true);
 			Player->GetEvaluationTemplate().EvaluateSynchronousBlocking(Context, *Player);
 			ControlRig->Evaluate_AnyThread();
diff --git a/Engine/Plugins/Experimental/CinematicPrestreaming/Source/CinematicPrestreamingEditor/Private/CinePrestreamingRecorderSetting.cpp b/Engine/Plugins/Experimental/CinematicPrestreaming/Source/CinematicPrestreamingEditor/Private/CinePrestreamingRecorderSetting.cpp
--- a/Engine/Plugins/Experimental/CinematicPrestreaming/Source/CinematicPrestreamingEditor/Private/CinePrestreamingRecorderSetting.cpp
+++ b/Engine/Plugins/Experimental/CinematicPrestreaming/Source/CinematicPrestreamingEditor/Private/CinePrestreamingRecorderSetting.cpp
@@ -247,7 +247,7 @@ void UCinePrestreamingRecorderSetting::CreateAssetsFromData()
 
 			// We only use the leafmost moviescene as that's the one that actually gets modified.
 			OutputSegment->MovieScene = Node->MovieScene.Get();
-			OutputSegment->Range = ShotInfo.TotalOutputRangeRoot * ShotInfo.OuterToInnerTransform;
+			OutputSegment->Range =  ShotInfo.OuterToInnerTransform.TransformRangeUnwarped(ShotInfo.TotalOutputRangeRoot);
 		}
 
 		// Convert our samples and append them to the existing array
diff --git a/Engine/Plugins/Importers/USDImporter/Source/USDStage/Private/USDLevelSequenceHelper.cpp b/Engine/Plugins/Importers/USDImporter/Source/USDStage/Private/USDLevelSequenceHelper.cpp
--- a/Engine/Plugins/Importers/USDImporter/Source/USDStage/Private/USDLevelSequenceHelper.cpp
+++ b/Engine/Plugins/Importers/USDImporter/Source/USDStage/Private/USDLevelSequenceHelper.cpp
@@ -3236,7 +3236,7 @@ void FUsdLevelSequenceHelperImpl::HandleControlRigSectionChange(UMovieSceneContr
 	const UsdUtils::FBlendShapeMap& BlendShapeMap = StageActorValue->GetBlendShapeMap();
 	bool bBaked = UnrealToUsd::ConvertControlRigSection(
 		&Section,
-		SequenceTransform.InverseLinearOnly(),
+		SequenceTransform.InverseNoLooping(),
 		MovieScene,
 		Player,
 		Skeleton->GetReferenceSkeleton(),
diff --git a/Engine/Plugins/Importers/USDImporter/Source/USDUtilities/Private/USDPrimConversion.cpp b/Engine/Plugins/Importers/USDImporter/Source/USDUtilities/Private/USDPrimConversion.cpp
--- a/Engine/Plugins/Importers/USDImporter/Source/USDUtilities/Private/USDPrimConversion.cpp
+++ b/Engine/Plugins/Importers/USDImporter/Source/USDUtilities/Private/USDPrimConversion.cpp
@@ -1533,7 +1533,7 @@ bool UnrealToUsd::ConvertBoolTrack( const UMovieSceneBoolTrack& MovieSceneTrack,
 			for ( const TPair< FFrameNumber, bool >& Pair : EvaluateChannel( BoolSection->GetChannel(), false ) )
 			{
 				FFrameTime TransformedBakedKeyTime{ Pair.Key };
-				TransformedBakedKeyTime *= SequenceTransform.InverseLinearOnly();
+				TransformedBakedKeyTime *= SequenceTransform.InverseNoLooping();
 				FFrameTime UsdFrameTime = FFrameRate::TransformTime( TransformedBakedKeyTime, Resolution, StageFrameRate );
 
 				bool UEValue = Pair.Value;
@@ -1630,7 +1630,7 @@ bool UnrealToUsd::ConvertFloatTrack( const UMovieSceneFloatTrack& MovieSceneTrac
 			for ( const TPair< FFrameNumber, float >& Pair : EvaluateChannel( FloatSection->GetChannel(), 0.0f ) )
 			{
 				FFrameTime TransformedBakedKeyTime{ Pair.Key };
-				TransformedBakedKeyTime *= SequenceTransform.InverseLinearOnly();
+				TransformedBakedKeyTime *= SequenceTransform.InverseNoLooping();
 				FFrameTime UsdFrameTime = FFrameRate::TransformTime( TransformedBakedKeyTime, Resolution, StageFrameRate );
 
 				float UEValue = Pair.Value;
@@ -1746,7 +1746,7 @@ bool UnrealToUsd::ConvertColorTrack( const UMovieSceneColorTrack& MovieSceneTrac
 				FLinearColor Color{ RedValue, GreenValue, BlueValue, AlphaValue };
 
 				FFrameTime TransformedBakedKeyTime{ UntransformedBakeTime };
-				TransformedBakedKeyTime *= SequenceTransform.InverseLinearOnly();
+				TransformedBakedKeyTime *= SequenceTransform.InverseNoLooping();
 				FFrameTime UsdFrameTime = FFrameRate::TransformTime( TransformedBakedKeyTime, Resolution, StageFrameRate );
 
 				WriterFunc( Color, UsdFrameTime.AsDecimal() );
@@ -1945,7 +1945,7 @@ bool UnrealToUsd::Convert3DTransformTrack( const UMovieScene3DTransformTrack& Mo
 				FTransform Transform{ Rotation, Location, Scale };
 
 				FFrameTime TransformedBakedKeyTime{ UntransformedBakeTime };
-				TransformedBakedKeyTime *= SequenceTransform.InverseLinearOnly();
+				TransformedBakedKeyTime *= SequenceTransform.InverseNoLooping();
 				FFrameTime UsdFrameTime = FFrameRate::TransformTime( TransformedBakedKeyTime, Resolution, StageFrameRate );
 
 				WriterFunc( Transform, UsdFrameTime.AsDecimal() );
@@ -3643,7 +3643,7 @@ bool UnrealToUsd::ConvertXformable( const UMovieScene3DTransformTrack& MovieScen
 				}
 
 				FFrameTime GlobalEvalTime( KeyTime );
-				GlobalEvalTime *= SequenceTransform.InverseLinearOnly();
+				GlobalEvalTime *= SequenceTransform.InverseNoLooping();
 
 				Values.Emplace( GlobalEvalTime.GetFrame(), Result );
 			}
diff --git a/Engine/Plugins/MovieScene/LevelSequenceEditor/Source/LevelSequenceEditor/Private/LevelSequenceEditorSubsystem.cpp b/Engine/Plugins/MovieScene/LevelSequenceEditor/Source/LevelSequenceEditor/Private/LevelSequenceEditorSubsystem.cpp
--- a/Engine/Plugins/MovieScene/LevelSequenceEditor/Source/LevelSequenceEditor/Private/LevelSequenceEditorSubsystem.cpp
+++ b/Engine/Plugins/MovieScene/LevelSequenceEditor/Source/LevelSequenceEditor/Private/LevelSequenceEditorSubsystem.cpp
@@ -1046,7 +1046,7 @@ bool ULevelSequenceEditorSubsystem::BakeTransformWithSettings(const TArray<FMovi
 
 	for (FFrameNumber KeyTime: AllFrames)
 	{
-		FMovieSceneEvaluationRange Range(KeyTime * RootToLocalTransform.InverseLinearOnly(), TickResolution);
+		FMovieSceneEvaluationRange Range(KeyTime * RootToLocalTransform.InverseNoLooping(), TickResolution);
 
 		Sequencer->SetGlobalTime(Range.GetTime());
 
diff --git a/Engine/Plugins/MovieScene/MovieRenderPipeline/Source/MovieRenderPipelineCore/Private/MoviePipelineBlueprintLibrary.cpp b/Engine/Plugins/MovieScene/MovieRenderPipeline/Source/MovieRenderPipelineCore/Private/MoviePipelineBlueprintLibrary.cpp
--- a/Engine/Plugins/MovieScene/MovieRenderPipeline/Source/MovieRenderPipelineCore/Private/MoviePipelineBlueprintLibrary.cpp
+++ b/Engine/Plugins/MovieScene/MovieRenderPipeline/Source/MovieRenderPipelineCore/Private/MoviePipelineBlueprintLibrary.cpp
@@ -380,7 +380,7 @@ void UMoviePipelineBlueprintLibrary::UpdateJobShotListFromSequence(ULevelSequenc
 				TRange<FFrameNumber> LocalCameraRange = Section->GetRange();
 
 				// Intersect it with the root range so that if the parent has trimmed down the sub-section we don't render outside that.
-				TRange<FFrameNumber> RootCameraRange = TRange<FFrameNumber>::Intersection(LocalSpace.RootClampRange, LocalCameraRange * LocalSpace.RootToSequenceTransform.InverseLinearOnly());
+				TRange<FFrameNumber> RootCameraRange = TRange<FFrameNumber>::Intersection(LocalSpace.RootClampRange, LocalSpace.RootToSequenceTransform.InverseNoLooping().TransformRangeUnwarped(LocalCameraRange));
 				if (!RootCameraRange.IsEmpty())
 				{
 					if (UMovieSceneCinematicShotSection* ShotSection = Cast<UMovieSceneCinematicShotSection>(Section))
@@ -432,7 +432,7 @@ void UMoviePipelineBlueprintLibrary::UpdateJobShotListFromSequence(ULevelSequenc
 			TTuple<FString, FString> Name;
 			TSharedPtr<MoviePipeline::FCameraCutSubSectionHierarchyNode> LeafNode;
 			TRange<FFrameNumber> CameraCutWarmUpRange;
-			FMovieSceneTimeTransform InnerToOuterTransform;
+			FMovieSceneSequenceTransform InnerToOuterTransform;
 		};
 
 		TArray<FLinearizedEntity> Entities;
@@ -632,11 +632,11 @@ void UMoviePipelineBlueprintLibrary::UpdateJobShotListFromSequence(ULevelSequenc
 				}
 			}
 
-			FMovieSceneTimeTransform InnerToOuterTransform = FMovieSceneTimeTransform();
+			FMovieSceneSequenceTransform InnerToOuterTransform;
 			FMovieSceneSubSequenceData* SubSequenceData = SequenceHierarchyCache.FindSubData(Entity.SequenceID);
 			if (SubSequenceData)
 			{
-				InnerToOuterTransform = SubSequenceData->RootToSequenceTransform.InverseFromAllFirstWarps();
+				InnerToOuterTransform = SubSequenceData->RootToSequenceTransform.InverseFromAllFirstLoops();
 			}
 
 			// To make the camera cut range detection more consistent, we'll convert the start of this Entity into root sequence space, and then we'll convert the start of the
@@ -645,7 +645,7 @@ void UMoviePipelineBlueprintLibrary::UpdateJobShotListFromSequence(ULevelSequenc
 			TRange<FFrameNumber> CameraCutRangeInRoot = EntityRangeInRoot;
 			if (LeafNode->CameraCutSection.IsValid())
 			{
-				CameraCutRangeInRoot = LeafNode->CameraCutSection->GetRange() * InnerToOuterTransform;
+				CameraCutRangeInRoot = InnerToOuterTransform.TransformRangeConstrained(LeafNode->CameraCutSection->GetRange());
 			}
 
 			TRange<FFrameNumber> CameraCutWarmUpRange = TRange<FFrameNumber>::Empty();
@@ -680,7 +680,7 @@ void UMoviePipelineBlueprintLibrary::UpdateJobShotListFromSequence(ULevelSequenc
 			Entity.CameraCutWarmUpRange = TRange<FFrameNumber>::Empty();
 			Entity.LeafNode = LeafNode;
 			Entity.Name = MoviePipeline::GetNameForShot(SequenceHierarchyCache, InSequence, LeafNode);
-			Entity.InnerToOuterTransform = FMovieSceneTimeTransform();
+			Entity.InnerToOuterTransform = FMovieSceneSequenceTransform();
 		}
 
 		// We need to generate all of the linearized segments first so that we have all of the names available.
@@ -768,7 +768,7 @@ void UMoviePipelineBlueprintLibrary::UpdateJobShotListFromSequence(ULevelSequenc
 			NewShot->ShotInfo.SubSectionHierarchy = Entity.LeafNode;
 			NewShot->ShotInfo.TotalOutputRangeRoot = Entity.Range;
 			NewShot->ShotInfo.WarmupRangeRoot = Entity.CameraCutWarmUpRange;
-			NewShot->ShotInfo.OuterToInnerTransform = Entity.InnerToOuterTransform.Inverse();
+			NewShot->ShotInfo.OuterToInnerTransform = Entity.InnerToOuterTransform.InverseNoLooping();
 			NewShot->SidecarCameras = Entity.SidecarCameras;
 			UE_LOG(LogMovieRenderPipeline, Log, TEXT("Registering range: %s (InnerName: %s OuterName: %s)"), *LexToString(NewShot->ShotInfo.TotalOutputRangeRoot), *NewShot->InnerName, *NewShot->OuterName);
 		}
diff --git a/Engine/Plugins/MovieScene/MovieRenderPipeline/Source/MovieRenderPipelineCore/Public/MovieRenderPipelineDataTypes.h b/Engine/Plugins/MovieScene/MovieRenderPipeline/Source/MovieRenderPipelineCore/Public/MovieRenderPipelineDataTypes.h
--- a/Engine/Plugins/MovieScene/MovieRenderPipeline/Source/MovieRenderPipelineCore/Public/MovieRenderPipelineDataTypes.h
+++ b/Engine/Plugins/MovieScene/MovieRenderPipeline/Source/MovieRenderPipelineCore/Public/MovieRenderPipelineDataTypes.h
@@ -536,7 +536,7 @@ public:
 	FFrameTime CurrentTimeInRoot;
 
 	/** Converts from the outermost space into the innermost space. Only works with linear transforms. */
-	FMovieSceneTimeTransform OuterToInnerTransform;
+	FMovieSceneSequenceTransform OuterToInnerTransform;
 	
 	/** The total range of output frames in root space */
 	TRange<FFrameNumber> TotalOutputRangeRoot;
diff --git a/Engine/Plugins/MovieScene/SequencerScripting/Source/SequencerScripting/Private/ExtensionLibraries/MovieSceneSectionExtensions.cpp b/Engine/Plugins/MovieScene/SequencerScripting/Source/SequencerScripting/Private/ExtensionLibraries/MovieSceneSectionExtensions.cpp
--- a/Engine/Plugins/MovieScene/SequencerScripting/Source/SequencerScripting/Private/ExtensionLibraries/MovieSceneSectionExtensions.cpp
+++ b/Engine/Plugins/MovieScene/SequencerScripting/Source/SequencerScripting/Private/ExtensionLibraries/MovieSceneSectionExtensions.cpp
@@ -588,7 +588,7 @@ int32 UMovieSceneSectionExtensions::GetParentSequenceFrame(UMovieSceneSubSection
 		
 	for (int32 SectionIndex = 0; SectionIndex < SubSectionChain.Num(); ++SectionIndex)
 	{
-		LocalFrameTime = LocalFrameTime * SubSectionChain[SectionIndex]->OuterToInnerTransform().InverseLinearOnly();
+		LocalFrameTime = LocalFrameTime * SubSectionChain[SectionIndex]->OuterToInnerTransform().InverseNoLooping();
 	}
 
 	FFrameRate ParentDisplayRate = ParentSequence->GetMovieScene()->GetDisplayRate();
diff --git a/Engine/Plugins/MovieScene/SequencerScripting/Source/SequencerScriptingEditor/Private/SequencerTools.cpp b/Engine/Plugins/MovieScene/SequencerScripting/Source/SequencerScriptingEditor/Private/SequencerTools.cpp
--- a/Engine/Plugins/MovieScene/SequencerScripting/Source/SequencerScriptingEditor/Private/SequencerTools.cpp
+++ b/Engine/Plugins/MovieScene/SequencerScripting/Source/SequencerScriptingEditor/Private/SequencerTools.cpp
@@ -255,7 +255,7 @@ bool ExportFBXInternal(const FSequencerExportFBXParams& InParams, UMovieSceneSeq
 	if (SubSequenceData)
 	{
 		RootToLocalTransform = SubSequenceData->RootToSequenceTransform;
-		StartTime = UE::MovieScene::DiscreteInclusiveLower(SubSequenceData->PlayRange.Value) * RootToLocalTransform.InverseLinearOnly();
+		StartTime = UE::MovieScene::DiscreteInclusiveLower(SubSequenceData->PlayRange.Value) * RootToLocalTransform.InverseNoLooping();
 	}
 
 	bool bDidExport = false;
diff --git a/Engine/Source/Editor/MovieSceneTools/Private/Constraints/ComponentConstraintChannelInterface.cpp b/Engine/Source/Editor/MovieSceneTools/Private/Constraints/ComponentConstraintChannelInterface.cpp
--- a/Engine/Source/Editor/MovieSceneTools/Private/Constraints/ComponentConstraintChannelInterface.cpp
+++ b/Engine/Source/Editor/MovieSceneTools/Private/Constraints/ComponentConstraintChannelInterface.cpp
@@ -380,7 +380,7 @@ void FComponentConstraintChannelInterface::RecomposeTransforms(
 		// add keys
 		for (int32 Index = 0; Index < InFrames.Num(); ++Index)
 		{
-			const FFrameNumber& FrameNumber = (FFrameTime(InFrames[Index]) * RootToLocalTransform.InverseLinearOnly()).GetFrame();
+			const FFrameNumber& FrameNumber = (FFrameTime(InFrames[Index]) * RootToLocalTransform.InverseNoLooping()).GetFrame();
 			const FMovieSceneEvaluationRange EvaluationRange = FMovieSceneEvaluationRange(FFrameTime(FrameNumber), TickResolution);
 			const FMovieSceneContext Context = FMovieSceneContext(EvaluationRange, PlaybackStatus).SetHasJumped(true);
 
diff --git a/Engine/Source/Editor/MovieSceneTools/Private/Constraints/MovieSceneConstraintChannelHelper.cpp b/Engine/Source/Editor/MovieSceneTools/Private/Constraints/MovieSceneConstraintChannelHelper.cpp
--- a/Engine/Source/Editor/MovieSceneTools/Private/Constraints/MovieSceneConstraintChannelHelper.cpp
+++ b/Engine/Source/Editor/MovieSceneTools/Private/Constraints/MovieSceneConstraintChannelHelper.cpp
@@ -93,7 +93,7 @@ void FCompensationEvaluator::ComputeLocalTransforms(
 	for (int32 Index = 0; Index < NumFrames + 1; ++Index)
 	{
 		FFrameNumber FrameNumber = (Index == 0) ? InFrames[0] - 1 : InFrames[Index - 1];
-		FrameNumber = (FFrameTime(FrameNumber) * RootToLocalTransform.InverseLinearOnly()).GetFrame();
+		FrameNumber = (FFrameTime(FrameNumber) * RootToLocalTransform.InverseNoLooping()).GetFrame();
 		// evaluate animation
 		const FMovieSceneEvaluationRange EvaluationRange = FMovieSceneEvaluationRange(FFrameTime(FrameNumber), TickResolution);
 		const FMovieSceneContext Context = FMovieSceneContext(EvaluationRange, PlaybackStatus).SetHasJumped(true);
@@ -247,7 +247,7 @@ void FCompensationEvaluator::ComputeLocalTransformsForBaking(UWorld* InWorld, co
 
 	for (int32 Index = 0; Index < NumFrames; ++Index)
 	{
-		const FFrameNumber& FrameNumber = (FFrameTime(InFrames[Index]) * RootToLocalTransform.InverseLinearOnly()).GetFrame();
+		const FFrameNumber& FrameNumber = (FFrameTime(InFrames[Index]) * RootToLocalTransform.InverseNoLooping()).GetFrame();
 
 		// evaluate animation
 		const FMovieSceneEvaluationRange EvaluationRange = FMovieSceneEvaluationRange(FFrameTime(FrameNumber), TickResolution);
@@ -362,7 +362,7 @@ void FCompensationEvaluator::ComputeLocalTransformsBeforeDeletion(
 
 	for (int32 Index = 0; Index < NumFrames; ++Index)
 	{
-		const FFrameNumber& FrameNumber = (FFrameTime(InFrames[Index]) * RootToLocalTransform.InverseLinearOnly()).GetFrame();
+		const FFrameNumber& FrameNumber = (FFrameTime(InFrames[Index]) * RootToLocalTransform.InverseNoLooping()).GetFrame();
 
 		// evaluate animation
 		const FMovieSceneEvaluationRange EvaluationRange = FMovieSceneEvaluationRange(FFrameTime(FrameNumber), TickResolution);
@@ -459,7 +459,7 @@ void FCompensationEvaluator::ComputeCompensation(UWorld* InWorld, const TSharedP
 	auto EvaluateAt = [InSequencer, &AllConstraints, &BakeHelpers](FFrameNumber InFrame)
 	{
 		FMovieSceneSequenceTransform RootToLocalTransform = InSequencer->GetFocusedMovieSceneSequenceTransform();
-		InFrame = (FFrameTime(InFrame) * RootToLocalTransform.InverseLinearOnly()).GetFrame();
+		InFrame = (FFrameTime(InFrame) * RootToLocalTransform.InverseNoLooping()).GetFrame();
 
 		UMovieScene* MovieScene = InSequencer->GetFocusedMovieSceneSequence()->GetMovieScene();
 		const FFrameRate TickResolution = MovieScene->GetTickResolution();
@@ -567,7 +567,7 @@ void FCompensationEvaluator::CacheTransforms(UWorld* InWorld, const TSharedPtr<I
 		const FMovieSceneEvaluationRange EvaluationRange = FMovieSceneEvaluationRange(FFrameTime(InFrame), TickResolution);
 		const FMovieSceneContext Context = FMovieSceneContext(EvaluationRange, PlaybackStatus).SetHasJumped(true);
 		FMovieSceneSequenceTransform RootToLocalTransform = InSequencer->GetFocusedMovieSceneSequenceTransform();
-		InFrame = (FFrameTime(InFrame) * RootToLocalTransform.InverseLinearOnly()).GetFrame();
+		InFrame = (FFrameTime(InFrame) * RootToLocalTransform.InverseNoLooping()).GetFrame();
 
 		for (IMovieSceneToolsAnimationBakeHelper* BakeHelper : BakeHelpers)
 		{
diff --git a/Engine/Source/Editor/MovieSceneTools/Private/TrackEditors/SubTrackEditorBase.cpp b/Engine/Source/Editor/MovieSceneTools/Private/TrackEditors/SubTrackEditorBase.cpp
--- a/Engine/Source/Editor/MovieSceneTools/Private/TrackEditors/SubTrackEditorBase.cpp
+++ b/Engine/Source/Editor/MovieSceneTools/Private/TrackEditors/SubTrackEditorBase.cpp
@@ -34,7 +34,7 @@ FSubSectionPainterResult FSubSectionPainterUtil::PaintSection(TSharedPtr<const I
     }
 
     const ESlateDrawEffect DrawEffects = InPainter.bParentEnabled ? ESlateDrawEffect::None : ESlateDrawEffect::DisabledEffect;
-    if (SectionObject.Parameters.bCanLoop)
+    if (SectionObject.Parameters.bCanLoop && !FMath::IsNearlyZero(SectionObject.OuterToInnerTransform().GetTimeScale()))
     {
         DoPaintLoopingSection(SectionObject, *InnerSequence, InPainter, DrawEffects);
     }
@@ -158,70 +158,120 @@ void FSubSectionPainterUtil::DoPaintNonLoopingSection(const UMovieSceneSubSectio
     UMovieScene* MovieScene = InnerSequence.GetMovieScene();
     TRange<FFrameNumber> PlaybackRange = MovieScene->GetPlaybackRange();
 
-    // We're in the non-looping case so we know we have a purely linear transform.
-    const FMovieSceneSequenceTransform InnerToOuterTransform = SectionObject.OuterToInnerTransform().InverseLinearOnly();
+    // We're in the non-looping case so we know we have a purely linear transform (unless zero timescale)
+	const FMovieSceneSequenceTransform OuterToInnerTransform = SectionObject.OuterToInnerTransform();
+    const FMovieSceneSequenceTransform InnerToOuterTransform = OuterToInnerTransform.InverseNoLooping();
     const FFrameNumber PlaybackStart = (UE::MovieScene::DiscreteInclusiveLower(PlaybackRange) * InnerToOuterTransform).FloorToFrame();
-    if (SectionRange.Contains(PlaybackStart))
-    {
-        const int32 StartOffset = (PlaybackStart - SectionStartFrame).Value;
-        // add dark tint for left out-of-bounds
-        FSlateDrawElement::MakeBox(
-                InPainter.DrawElements,
-                InPainter.LayerId++,
-                InPainter.SectionGeometry.ToPaintGeometry(
-                    FVector2f(StartOffset * PixelsPerFrame, InPainter.SectionGeometry.Size.Y),
-                    FSlateLayoutTransform()
-                    ),
-                FAppStyle::GetBrush("WhiteBrush"),
-                DrawEffects,
-                FLinearColor::Black.CopyWithNewOpacity(0.5f)
-                );
-
-        // add green line for playback start
-        FSlateDrawElement::MakeBox(
-                InPainter.DrawElements,
-                InPainter.LayerId++,
-                InPainter.SectionGeometry.ToPaintGeometry(
-                    FVector2f(1.0f, InPainter.SectionGeometry.Size.Y),
-                    FSlateLayoutTransform(FVector2f(StartOffset * PixelsPerFrame, 0.f))
-                    ),
-                FAppStyle::GetBrush("WhiteBrush"),
-                DrawEffects,
-                FColor(32, 128, 32, GhostAlpha)	// 120, 75, 50 (HSV)
-                );
-    }
-
-    const FFrameNumber PlaybackEnd = (UE::MovieScene::DiscreteExclusiveUpper(PlaybackRange) * InnerToOuterTransform) .FloorToFrame();
-    if (SectionRange.Contains(PlaybackEnd))
-    {
-        // add dark tint for right out-of-bounds
-        const int32 EndOffset = (PlaybackEnd - SectionStartFrame).Value;
-        FSlateDrawElement::MakeBox(
-                InPainter.DrawElements,
-                InPainter.LayerId++,
-                InPainter.SectionGeometry.ToPaintGeometry(
-                    FVector2f((SectionSize - EndOffset) * PixelsPerFrame, InPainter.SectionGeometry.Size.Y),
-                    FSlateLayoutTransform(FVector2f(EndOffset * PixelsPerFrame, 0.f))
-                    ),
-                FAppStyle::GetBrush("WhiteBrush"),
-                DrawEffects,
-                FLinearColor::Black.CopyWithNewOpacity(0.5f)
-                );
-
-
-        // add red line for playback end
-        FSlateDrawElement::MakeBox(
-                InPainter.DrawElements,
-                InPainter.LayerId++,
-                InPainter.SectionGeometry.ToPaintGeometry(
-                    FVector2f(1.0f, InPainter.SectionGeometry.Size.Y),
-                    FSlateLayoutTransform(FVector2f(EndOffset * PixelsPerFrame, 0.f))
-                    ),
-                FAppStyle::GetBrush("WhiteBrush"),
-                DrawEffects,
-                FColor(128, 32, 32, GhostAlpha)	// 0, 75, 50 (HSV)
-                );
-    }
+	const FFrameNumber PlaybackEnd = (UE::MovieScene::DiscreteExclusiveUpper(PlaybackRange) * InnerToOuterTransform).FloorToFrame();
+	const TRange<FFrameNumber> PlaybackRangeTransformed(PlaybackStart, PlaybackEnd);
+	bool bZeroTimescale = !FMath::IsFinite(InnerToOuterTransform.GetTimeScale());
+
+	// Section contains the beginning but not the end- draw beginning, draw dark tint before it
+	// Section contains the end but not the beginning- draw dark tint, draw dark tint after it
+	// Section contains both end and beginning- draw both, dark tint before beginning, dark tint after end
+	// Section contains neither or zero-timescale and the frame isn't in the section range- draw dark tint over entire section
+	// Zero-timescale and the frame is in the sectionrange- don't draw anything extra
+	const int32 StartOffset = (PlaybackStart - SectionStartFrame).Value;
+	const int32 EndOffset = (PlaybackEnd - SectionStartFrame).Value;
+
+	if (bZeroTimescale)
+	{
+		// If we're zero-timescale we can't use the inverse transform to determine bounds, as the transformation is non-deterministic.
+		// Instead, we just ask whether the playback start is in the interior range, and if it doesn't, draw dark tint over the entire section
+		
+		const FFrameTime InnerFrame = SectionStartFrame * OuterToInnerTransform;
+		if (!PlaybackRange.Contains(InnerFrame.FloorToFrame()))
+		{
+			// add dark tint over the entire section
+			FSlateDrawElement::MakeBox(
+				InPainter.DrawElements,
+				InPainter.LayerId++,
+				InPainter.SectionGeometry.ToPaintGeometry(
+					FVector2f(InPainter.SectionGeometry.Size.X, InPainter.SectionGeometry.Size.Y),
+					FSlateLayoutTransform()
+				),
+				FAppStyle::GetBrush("WhiteBrush"),
+				DrawEffects,
+				FLinearColor::Black.CopyWithNewOpacity(0.5f)
+			);
+		}
+	}
+	else if (TRange<FFrameNumber>::Intersection(PlaybackRangeTransformed, SectionRange).IsEmpty())
+	{
+		// add dark tint over the entire section
+		FSlateDrawElement::MakeBox(
+			InPainter.DrawElements,
+			InPainter.LayerId++,
+			InPainter.SectionGeometry.ToPaintGeometry(
+				FVector2f(InPainter.SectionGeometry.Size.X, InPainter.SectionGeometry.Size.Y),
+				FSlateLayoutTransform()
+			),
+			FAppStyle::GetBrush("WhiteBrush"),
+			DrawEffects,
+			FLinearColor::Black.CopyWithNewOpacity(0.5f)
+		);
+	}
+	else
+	{
+		if (SectionRange.Contains(PlaybackStart))
+		{
+			// add dark tint for left out-of-bounds
+			FSlateDrawElement::MakeBox(
+				InPainter.DrawElements,
+				InPainter.LayerId++,
+				InPainter.SectionGeometry.ToPaintGeometry(
+					FVector2f(StartOffset * PixelsPerFrame, InPainter.SectionGeometry.Size.Y),
+					FSlateLayoutTransform()
+				),
+				FAppStyle::GetBrush("WhiteBrush"),
+				DrawEffects,
+				FLinearColor::Black.CopyWithNewOpacity(0.5f)
+			);
+
+			// add green line for playback start
+			FSlateDrawElement::MakeBox(
+				InPainter.DrawElements,
+				InPainter.LayerId++,
+				InPainter.SectionGeometry.ToPaintGeometry(
+					FVector2f(1.0f, InPainter.SectionGeometry.Size.Y),
+					FSlateLayoutTransform(FVector2f(StartOffset * PixelsPerFrame, 0.f))
+				),
+				FAppStyle::GetBrush("WhiteBrush"),
+				DrawEffects,
+				FColor(32, 128, 32, GhostAlpha)	// 120, 75, 50 (HSV)
+			);
+		}
+
+		if (SectionRange.Contains(PlaybackEnd))
+		{
+			// add dark tint for right out-of-bounds
+			FSlateDrawElement::MakeBox(
+				InPainter.DrawElements,
+				InPainter.LayerId++,
+				InPainter.SectionGeometry.ToPaintGeometry(
+					FVector2f((SectionSize - EndOffset) * PixelsPerFrame, InPainter.SectionGeometry.Size.Y),
+					FSlateLayoutTransform(FVector2f(EndOffset * PixelsPerFrame, 0.f))
+				),
+				FAppStyle::GetBrush("WhiteBrush"),
+				DrawEffects,
+				FLinearColor::Black.CopyWithNewOpacity(0.5f)
+			);
+
+
+			// add red line for playback end
+			FSlateDrawElement::MakeBox(
+				InPainter.DrawElements,
+				InPainter.LayerId++,
+				InPainter.SectionGeometry.ToPaintGeometry(
+					FVector2f(1.0f, InPainter.SectionGeometry.Size.Y),
+					FSlateLayoutTransform(FVector2f(EndOffset * PixelsPerFrame, 0.f))
+				),
+				FAppStyle::GetBrush("WhiteBrush"),
+				DrawEffects,
+				FColor(128, 32, 32, GhostAlpha)	// 0, 75, 50 (HSV)
+			);
+		}
+	}
 }
 
 void FSubSectionPainterUtil::DoPaintLoopingSection(const UMovieSceneSubSection& SectionObject, const UMovieSceneSequence& InnerSequence, FSequencerSectionPainter& InPainter, ESlateDrawEffect DrawEffects)
diff --git a/Engine/Source/Editor/Sequencer/Private/Sequencer.cpp b/Engine/Source/Editor/Sequencer/Private/Sequencer.cpp
--- a/Engine/Source/Editor/Sequencer/Private/Sequencer.cpp
+++ b/Engine/Source/Editor/Sequencer/Private/Sequencer.cpp
@@ -1360,12 +1360,12 @@ void FSequencer::PopToSequenceInstance(FMovieSceneSequenceIDRef SequenceID)
 void FSequencer::UpdateSubSequenceData()
 {
 	const bool bIsScrubbing = GetPlaybackStatus() == EMovieScenePlayerStatus::Scrubbing;
-	const bool bIsSubSequenceWarping = RootToLocalTransform.NestedTransforms.Num() > 0 && RootToLocalTransform.NestedTransforms.Last().IsWarping();
-	const bool bIsScrubbingWarpingSubSequence = bIsScrubbing && bIsSubSequenceWarping;
+	const bool bIsSubSequenceLooping = RootToLocalTransform.NestedTransforms.Num() > 0 && RootToLocalTransform.NestedTransforms.Last().IsLooping();
+	const bool bIsScrubbingLoopingSubSequence = bIsScrubbing && bIsSubSequenceLooping;
 
 	SubSequenceRange = TRange<FFrameNumber>::Empty();
 	RootToLocalTransform = FMovieSceneSequenceTransform();
-	if (!bIsScrubbingWarpingSubSequence)
+	if (!bIsScrubbingLoopingSubSequence)
 	{
 		RootToLocalLoopCounter = FMovieSceneWarpCounter();
 	}
@@ -1389,7 +1389,7 @@ void FSequencer::UpdateSubSequenceData()
 		RootToLocalTransform = SubSequenceData->RootToSequenceTransform;
 
 		const FQualifiedFrameTime RootTime = GetGlobalTime();
-		if (!bIsScrubbingWarpingSubSequence)
+		if (!bIsScrubbingLoopingSubSequence)
 		{
 			FFrameTime LocalTime;
 			RootToLocalTransform.TransformTime(RootTime.Time, LocalTime, RootToLocalLoopCounter);
@@ -1409,7 +1409,7 @@ void FSequencer::UpdateSubSequenceData()
 			FMovieSceneNestedSequenceTransform LeafLooping = RootToLocalTransformWithoutLeafLooping.NestedTransforms.Pop();
 			FFrameTime LocalTimeWithLastLoopUnwarped = RootTime.Time * RootToLocalTransformWithoutLeafLooping;
 			LocalTimeWithLastLoopUnwarped = LocalTimeWithLastLoopUnwarped * LeafLooping.LinearTransform;
-			if (LeafLooping.IsWarping())
+			if (LeafLooping.IsLooping())
 			{
 				LeafLooping.Warping.TransformTimeSpecific(
 						LocalTimeWithLastLoopUnwarped, LocalLoopIndexOnBeginScrubbing, LocalTimeWithLastLoopUnwarped);
@@ -1420,15 +1420,18 @@ void FSequencer::UpdateSubSequenceData()
 			// point (see comment at the beginning of this method). We will commit to the new loop indices
 			// once we're done scrubbing.
 			uint32 CurLoopIndex = 0;
-			while (LocalTimeWithLastLoopUnwarped >= PlayRangeUpperBound)
+			if (PlayRangeSize > 0)
 			{
-				LocalTimeWithLastLoopUnwarped = LocalTimeWithLastLoopUnwarped - PlayRangeSize;
-				++CurLoopIndex;
-			}
-			while (LocalTimeWithLastLoopUnwarped <= PlayRangeLowerBound)
-			{
-				LocalTimeWithLastLoopUnwarped = LocalTimeWithLastLoopUnwarped + PlayRangeSize;
-				--CurLoopIndex;
+				while (LocalTimeWithLastLoopUnwarped >= PlayRangeUpperBound)
+				{
+					LocalTimeWithLastLoopUnwarped = LocalTimeWithLastLoopUnwarped - PlayRangeSize;
+					++CurLoopIndex;
+				}
+				while (LocalTimeWithLastLoopUnwarped <= PlayRangeLowerBound)
+				{
+					LocalTimeWithLastLoopUnwarped = LocalTimeWithLastLoopUnwarped + PlayRangeSize;
+					--CurLoopIndex;
+				}
 			}
 			if (CurLoopIndex != LocalLoopIndexOffsetDuringScrubbing)
 			{
@@ -3035,8 +3038,14 @@ void FSequencer::SetLocalTimeDirectly(FFrameTime NewTime, bool bEvaluate)
 		FSlateApplication::Get().ClearKeyboardFocus(EFocusCause::Cleared);
 	}
 
-	// Transform the time to the root time-space
-	SetGlobalTime(NewTime * RootToLocalTransform.InverseFromWarp(RootToLocalLoopCounter), bEvaluate);
+	// Special-case. If the RootToLocalTransform contains a zero timescale, then scrubbing in the subscene should not change the current time
+	// Inverting the zero time-scale time is technically non-deterministic, and so doing this scrub would just result in losing the current global time
+	// while maintaining the same sub-sequence time (as we're holding a single frame there). So the best case here is just preventing the scrub.
+	if (!FMath::IsNearlyZero(RootToLocalTransform.GetTimeScale()))
+	{
+		// Transform the time to the root time-space
+		SetGlobalTime(NewTime * RootToLocalTransform.InverseFromLoop(RootToLocalLoopCounter), bEvaluate);
+	}
 
 	if (PreviousFocusedWidget.IsValid())
 	{
@@ -3545,7 +3554,7 @@ void FSequencer::RenderMovieInternal(TRange<FFrameNumber> Range, bool bSetFrameO
 
 		if (const FMovieSceneSubSequenceData* SubSequenceData = RootTemplateInstance.FindSubData(GetFocusedTemplateID()))
 		{
-			Range = Range * SubSequenceData->RootToSequenceTransform.InverseLinearOnly();
+			Range = SubSequenceData->RootToSequenceTransform.InverseNoLooping().TransformRangeUnwarped(Range);
 		}
 	}
 
@@ -4664,7 +4673,7 @@ FReply FSequencer::OnPlay(bool bTogglePlay)
 
 		if (GetLocalTime().Time <= MinInclusiveTime || GetLocalTime().Time >= MaxInclusiveTime)
 		{
-			FFrameTime NewGlobalTime = (PlaybackSpeed > 0 ? MinInclusiveTime : MaxInclusiveTime) * RootToLocalTransform.InverseFromWarp(RootToLocalLoopCounter);
+			FFrameTime NewGlobalTime = (PlaybackSpeed > 0 ? MinInclusiveTime : MaxInclusiveTime) * RootToLocalTransform.InverseFromLoop(RootToLocalLoopCounter);
 			SetGlobalTime(NewGlobalTime);
 		}
 
@@ -4894,7 +4903,7 @@ void FSequencer::SetLocalTimeLooped(FFrameTime NewLocalTime)
 {
 	TOptional<EMovieScenePlayerStatus::Type> NewPlaybackStatus;
 
-	const FMovieSceneSequenceTransform LocalToRootTransform = RootToLocalTransform.InverseFromWarp(RootToLocalLoopCounter);
+	const FMovieSceneSequenceTransform LocalToRootTransform = RootToLocalTransform.InverseFromLoop(RootToLocalLoopCounter);
 
 	FFrameTime NewGlobalTime = NewLocalTime * LocalToRootTransform;
 
@@ -9155,7 +9164,7 @@ void FSequencer::Pause()
 		FFrameRate          FocusedDisplayRate = GetFocusedDisplayRate();
 
 		// Snap to the focused play rate
-		FFrameTime RootPosition  = FFrameRate::Snap(LocalTime.Time, LocalTime.Rate, FocusedDisplayRate) * RootToLocalTransform.InverseFromWarp(RootToLocalLoopCounter);
+		FFrameTime RootPosition  = FFrameRate::Snap(LocalTime.Time, LocalTime.Rate, FocusedDisplayRate) * RootToLocalTransform.InverseFromLoop(RootToLocalLoopCounter);
 
 		// Convert the root position from tick resolution time base (the output rate), to the play position input rate
 		FFrameTime InputPosition = ConvertFrameTime(RootPosition, PlayPosition.GetOutputRate(), PlayPosition.GetInputRate());
@@ -9261,7 +9270,7 @@ void FSequencer::StepToNextShot()
 		return;
 	}
 
-	FFrameTime CurrentTime = SubSequenceRange.GetLowerBoundValue() * SubData->OuterToInnerTransform.InverseFromWarp(RootToLocalLoopCounter);
+	FFrameTime CurrentTime = SubSequenceRange.GetLowerBoundValue() * SubData->OuterToInnerTransform.InverseFromLoop(RootToLocalLoopCounter);
 	UMovieSceneSubSection* NextShot = Cast<UMovieSceneSubSection>(MovieSceneHelpers::FindNextSection(CinematicShotTrack->GetAllSections(), CurrentTime.FloorToFrame()));
 	if (!NextShot)
 	{
@@ -9319,7 +9328,7 @@ void FSequencer::StepToPreviousShot()
 		return;
 	}
 
-	FFrameTime CurrentTime = SubSequenceRange.GetLowerBoundValue() * SubData->OuterToInnerTransform.InverseFromWarp(RootToLocalLoopCounter);
+	FFrameTime CurrentTime = SubSequenceRange.GetLowerBoundValue() * SubData->OuterToInnerTransform.InverseFromLoop(RootToLocalLoopCounter);
 	UMovieSceneSubSection* PreviousShot = Cast<UMovieSceneSubSection>(MovieSceneHelpers::FindPreviousSection(CinematicShotTrack->GetAllSections(), CurrentTime.FloorToFrame()));
 	if (!PreviousShot)
 	{
diff --git a/Engine/Source/Editor/Sequencer/Private/SequencerMarkedFrameHelper.cpp b/Engine/Source/Editor/Sequencer/Private/SequencerMarkedFrameHelper.cpp
--- a/Engine/Source/Editor/Sequencer/Private/SequencerMarkedFrameHelper.cpp
+++ b/Engine/Source/Editor/Sequencer/Private/SequencerMarkedFrameHelper.cpp
@@ -79,14 +79,23 @@ static void FindGlobalMarkedFrames(
 				continue;
 			}
 
-			if (!ChildSubData->bCanLoop)
+			if (!ChildSubData->bCanLoop || FMath::IsNearlyZero(ChildSubData->OuterToInnerTransform.GetTimeScale()))
 			{
-				// This child doesn't loop. We just need to take into account its offset and time scale, and how much of it we "see" through the
+				// This child doesn't loop (or it does, but we have zero timescale, so it's irrelevant). We just need to take into account its offset and time scale, and how much of it we "see" through the
 				// parent section.
-				ensure(!ChildSubData->OuterToInnerTransform.IsWarping());
+				ensure(!ChildSubData->OuterToInnerTransform.IsLooping() || FMath::IsNearlyZero(ChildSubData->OuterToInnerTransform.GetTimeScale()));
+
 				const FMovieSceneTimeTransform& OuterToInnerTransform(ChildSubData->OuterToInnerTransform.LinearTransform);
 				OutTimestampedGlobalMarkedFrames.PushTransform(OuterToInnerTransform);
 
+				// We may not be able to loop, but we may have multiple nested transforms if we have a zero timescale
+				const int NumNestedTransforms = ChildSubData->OuterToInnerTransform.NestedTransforms.Num();
+				for (int i = 0; i < NumNestedTransforms; ++i)
+				{
+					const FMovieSceneTimeTransform& SubOuterToInnerTransform(ChildSubData->OuterToInnerTransform.NestedTransforms[i].LinearTransform);
+					OutTimestampedGlobalMarkedFrames.PushTransform(SubOuterToInnerTransform);
+				}
+
 				// Compute the "window" of the child that we can see through the parent section. Note that the parent section could extend past
 				// the end of the child's playback... we don't trim that, so that any marked frame added after the playback end in the child sequence
 				// is still visible in the parent sequence as long as the parent section is long enough. This lets artists see what they're "missing".
@@ -95,8 +104,12 @@ static void FindGlobalMarkedFrames(
 
 				// Gather marked frames in this "window".
 				FindGlobalMarkedFrames(Sequencer, SequenceHierarchy, FocusedSequenceID, ChildID, ChildGatherRange, OutTimestampedGlobalMarkedFrames);
-
-				OutTimestampedGlobalMarkedFrames.PopTransform();
+				
+				// Pop the nested transforms off, including the base linear transform
+				for (int i = 0; i < NumNestedTransforms + 1; ++i)
+				{
+					OutTimestampedGlobalMarkedFrames.PopTransform();
+				}
 			}
 			else if (ensure(ChildSubData->OuterToInnerTransform.NestedTransforms.Num() > 0))
 			{
@@ -183,20 +196,40 @@ void FSequencerMarkedFrameHelper::FindGlobalMarkedFrames(ISequencer& Sequencer,
 	// All the marked frames will be added using their root time, but we want to actually display them in the time space of whatever
 	// is the currently focused sequence. We therefore add the inverse time transform of the focused sequence at the top of the
 	// transform stack if the focused sequence isn't the root sequence (which has no time transform).
-	FMovieSceneTimeTransform FocusedMovieSequenceInverseTransform;
+
+	TMovieSceneTimeArray<FMovieSceneMarkedFrame> TimestampedGlobalMarkedFrames;
 	const FMovieSceneSubSequenceData* FocusedMovieSequenceSubData = SequenceHierarchy->FindSubData(FocusedMovieSequenceID);
 	if (FocusedMovieSequenceSubData)
 	{
-		const int32 FocusedMovieSequenceDepth = FMath::Min(
-				FocusedMovieSequenceSubData->RootToSequenceTransform.NestedTransforms.Num(), 
+		// If we have a warping, but not looping transform, we need to do the inverse in pieces rather than just a single transform, 
+		// because if we have zero-timescale anywhere in the mix, then we lose information if we just take a simple FMovieSceneTimeTransform as an inverse.
+		if (FocusedMovieSequenceSubData->RootToSequenceTransform.IsWarping() && !FocusedMovieSequenceSubData->RootToSequenceTransform.IsLooping())
+		{
+			const int NumNestedTransforms = FocusedMovieSequenceSubData->RootToSequenceTransform.NestedTransforms.Num();
+			for (int i = NumNestedTransforms - 1; i >= 0; --i)
+			{
+				const FMovieSceneTimeTransform& SubOuterToInnerTransform(FocusedMovieSequenceSubData->RootToSequenceTransform.NestedTransforms[i].LinearTransform.Inverse());
+				TimestampedGlobalMarkedFrames.PushTransform(SubOuterToInnerTransform);
+			}
+			TimestampedGlobalMarkedFrames.PushTransform(FocusedMovieSequenceSubData->RootToSequenceTransform.LinearTransform.Inverse());
+		}
+		else
+		{
+			const int32 FocusedMovieSequenceDepth = FMath::Min(
+				FocusedMovieSequenceSubData->RootToSequenceTransform.NestedTransforms.Num(),
 				LoopCounter.Num());
-		TArrayView<uint32> LoopCounterForFocusedMovieSequence = MakeArrayView(LoopCounter.GetData(), FocusedMovieSequenceDepth);
-		FocusedMovieSequenceInverseTransform = FocusedMovieSequenceSubData->RootToSequenceTransform.InverseFromWarp(LoopCounterForFocusedMovieSequence);
+			TArrayView<uint32> LoopCounterForFocusedMovieSequence = MakeArrayView(LoopCounter.GetData(), FocusedMovieSequenceDepth);
+			FMovieSceneSequenceTransform InverseLoopTransform = FocusedMovieSequenceSubData->RootToSequenceTransform.InverseFromLoop(LoopCounterForFocusedMovieSequence);
+
+			TimestampedGlobalMarkedFrames.PushTransform(InverseLoopTransform.LinearTransform);
+			for (int i = 0; i < InverseLoopTransform.NestedTransforms.Num(); ++i)
+			{
+				TimestampedGlobalMarkedFrames.PushTransform(InverseLoopTransform.NestedTransforms[i].LinearTransform);
+			}
+		}
 	}
 
 	// Grab the marked frames from the root sequence, and recursively across the whole hierarchy.
-	TMovieSceneTimeArray<FMovieSceneMarkedFrame> TimestampedGlobalMarkedFrames;
-	TimestampedGlobalMarkedFrames.PushTransform(FocusedMovieSequenceInverseTransform);
 	UE::MovieScene::FindGlobalMarkedFrames(Sequencer, SequenceHierarchy, FocusedMovieSequenceID, MovieSceneSequenceID::Root, TRange<FFrameNumber>::All(), TimestampedGlobalMarkedFrames);
 
 	// Export the modified timestamped entries.
diff --git a/Engine/Source/Editor/UnrealEd/Private/Fbx/FbxMainExport.cpp b/Engine/Source/Editor/UnrealEd/Private/Fbx/FbxMainExport.cpp
--- a/Engine/Source/Editor/UnrealEd/Private/Fbx/FbxMainExport.cpp
+++ b/Engine/Source/Editor/UnrealEd/Private/Fbx/FbxMainExport.cpp
@@ -1712,7 +1712,7 @@ int32 FLevelSequenceAnimTrackAdapter::GetStartFrame() const
 {
 	FFrameRate TickResolution = MovieScene->GetTickResolution();
 	FFrameRate DisplayRate = MovieScene->GetDisplayRate();
-	return FFrameRate::TransformTime(FFrameTime(UE::MovieScene::DiscreteInclusiveLower(MovieScene->GetPlaybackRange()) * RootToLocalTransform.InverseLinearOnly()), TickResolution, DisplayRate).RoundToFrame().Value;
+	return FFrameRate::TransformTime(FFrameTime(UE::MovieScene::DiscreteInclusiveLower(MovieScene->GetPlaybackRange()) * RootToLocalTransform.InverseNoLooping()), TickResolution, DisplayRate).RoundToFrame().Value;
 }
 
 int32 FLevelSequenceAnimTrackAdapter::GetLength() const
@@ -1729,7 +1729,7 @@ void FLevelSequenceAnimTrackAdapter::UpdateAnimation( int32 LocalFrame )
 	FFrameRate DisplayRate = MovieScene->GetDisplayRate();
 
 	FFrameTime LocalTime = FFrameRate::TransformTime(FFrameTime(LocalFrame), DisplayRate, TickResolution);
-	FFrameTime GlobalTime = LocalTime * RootToLocalTransform.InverseLinearOnly();
+	FFrameTime GlobalTime = LocalTime * RootToLocalTransform.InverseNoLooping();
 
 	FMovieSceneContext Context = FMovieSceneContext(FMovieSceneEvaluationRange(GlobalTime, TickResolution), MovieScenePlayer->GetPlaybackStatus()).SetHasJumped(true);
 
@@ -2058,7 +2058,7 @@ void FFbxExporter::AddTimecodeAttributesAndSetKey(const UMovieSceneSection* InSe
 				AnimCurve->KeyModifyBegin();
 					
 				FbxTime FbxTime;
-				const double KeyTimeSeconds = GetExportOptions()->bExportLocalTime ? KeyTime / TickResolution : (KeyTime * RootToLocalTransform.InverseLinearOnly()) / TickResolution;
+				const double KeyTimeSeconds = GetExportOptions()->bExportLocalTime ? KeyTime / TickResolution : (KeyTime * RootToLocalTransform.InverseNoLooping()) / TickResolution;
 
 				FbxTime.SetSecondDouble(KeyTimeSeconds);
 
@@ -2351,7 +2351,7 @@ void FFbxExporter::ExportTransformChannelsToFbxCurve(FbxNode* InFbxNode, FMovieS
 	TRange<FFrameNumber> PlaybackRange = Track->GetTypedOuter<UMovieScene>()->GetPlaybackRange();
 	
 	int32 LocalStartFrame = FFrameRate::TransformTime(FFrameTime(UE::MovieScene::DiscreteInclusiveLower(PlaybackRange)), TickResolution, DisplayRate).RoundToFrame().Value;
-	int32 StartFrame = FFrameRate::TransformTime(FFrameTime(UE::MovieScene::DiscreteInclusiveLower(PlaybackRange) * RootToLocalTransform.InverseLinearOnly()), TickResolution, DisplayRate).RoundToFrame().Value;
+	int32 StartFrame = FFrameRate::TransformTime(FFrameTime(UE::MovieScene::DiscreteInclusiveLower(PlaybackRange) * RootToLocalTransform.InverseNoLooping()), TickResolution, DisplayRate).RoundToFrame().Value;
 	int32 AnimationLength = FFrameRate::TransformTime(FFrameTime(FFrameNumber(UE::MovieScene::DiscreteSize(PlaybackRange))), TickResolution, DisplayRate).RoundToFrame().Value;
 
 	for (int32 FrameCount = 0; FrameCount <= AnimationLength; ++FrameCount)
@@ -2751,7 +2751,7 @@ void FFbxExporter::ExportBezierChannelToFbxCurve(FbxAnimCurve& InFbxCurve, const
 
 		FbxTime FbxTime;
 		FbxAnimCurveKey FbxKey;
-		const double KeyTimeSeconds = GetExportOptions()->bExportLocalTime ? KeyTime / TickResolution : (KeyTime * RootToLocalTransform.InverseLinearOnly()) / TickResolution;
+		const double KeyTimeSeconds = GetExportOptions()->bExportLocalTime ? KeyTime / TickResolution : (KeyTime * RootToLocalTransform.InverseNoLooping()) / TickResolution;
 
 		FbxTime.SetSecondDouble(KeyTimeSeconds);
 
@@ -2847,7 +2847,7 @@ void FFbxExporter::ExportConstantChannelToFbxCurve(FbxAnimCurve& InFbxCurve, con
 		const T KeyValue = Values[Index];
 
 		FbxTime FbxTime;
-		const double KeyTimeSeconds = GetExportOptions()->bExportLocalTime ? KeyTime / TickResolution : (KeyTime * RootToLocalTransform.InverseLinearOnly()) / TickResolution;
+		const double KeyTimeSeconds = GetExportOptions()->bExportLocalTime ? KeyTime / TickResolution : (KeyTime * RootToLocalTransform.InverseNoLooping()) / TickResolution;
 
 		FbxTime.SetSecondDouble(KeyTimeSeconds);
 
@@ -2983,7 +2983,7 @@ void FFbxExporter::ExportLevelSequence3DTransformTrack(FbxNode* FbxNode, IMovieS
 		FbxCurveRotZ->KeyModifyBegin();
 
 		int32 LocalStartFrame = FFrameRate::TransformTime(FFrameTime(UE::MovieScene::DiscreteInclusiveLower(InPlaybackRange)), TickResolution, DisplayRate).RoundToFrame().Value;
-		int32 StartFrame = FFrameRate::TransformTime(FFrameTime(UE::MovieScene::DiscreteInclusiveLower(InPlaybackRange) * RootToLocalTransform.InverseLinearOnly()), TickResolution, DisplayRate).RoundToFrame().Value;
+		int32 StartFrame = FFrameRate::TransformTime(FFrameTime(UE::MovieScene::DiscreteInclusiveLower(InPlaybackRange) * RootToLocalTransform.InverseNoLooping()), TickResolution, DisplayRate).RoundToFrame().Value;
 		int32 AnimationLength = FFrameRate::TransformTime(FFrameTime(FFrameNumber(UE::MovieScene::DiscreteSize(InPlaybackRange))), TickResolution, DisplayRate).RoundToFrame().Value;
 
 		for (int32 FrameCount = 0; FrameCount <= AnimationLength; ++FrameCount)
@@ -3168,7 +3168,7 @@ void FFbxExporter::ExportLevelSequenceBaked3DTransformTrack(IAnimTrackAdapter& A
 		FbxCurveScaleZ->KeyModifyBegin();
 	}
 
-	FMovieSceneTimeTransform LocalToRootTransform = RootToLocalTransform.InverseLinearOnly();
+	FMovieSceneSequenceTransform LocalToRootTransform = RootToLocalTransform.InverseNoLooping();
 
 	TArray<FTransform> RelativeTransforms;
 	int32 LocalStartFrame = FFrameRate::TransformTime(FFrameTime(DiscreteInclusiveLower(InPlaybackRange)), TickResolution, DisplayRate).RoundToFrame().Value;
diff --git a/Engine/Source/Runtime/MovieScene/Private/Compilation/MovieSceneCompiledDataManager.cpp b/Engine/Source/Runtime/MovieScene/Private/Compilation/MovieSceneCompiledDataManager.cpp
--- a/Engine/Source/Runtime/MovieScene/Private/Compilation/MovieSceneCompiledDataManager.cpp
+++ b/Engine/Source/Runtime/MovieScene/Private/Compilation/MovieSceneCompiledDataManager.cpp
@@ -1303,7 +1303,7 @@ void UMovieSceneCompiledDataManager::GatherTrack(const FMovieSceneBinding* Objec
 		// Iterate everything in the field
 		for (const FMovieSceneTrackEvaluationFieldEntry& Entry : EvaluationField.Entries)
 		{
-			FMovieSceneSequenceTransform SequenceToRootTransform  = Params.RootToSequenceTransform.InverseFromWarp(Params.RootToSequenceWarpCounter);
+			FMovieSceneSequenceTransform SequenceToRootTransform  = Params.RootToSequenceTransform.InverseFromLoop(Params.RootToSequenceWarpCounter);
 			TRange<FFrameNumber>         ClampedRangeRoot         = Params.ClampRoot(SequenceToRootTransform.TransformRangeUnwarped(Entry.Range));
 			UMovieSceneSection*          Section                  = Entry.Section;
 
@@ -1569,8 +1569,20 @@ void UMovieSceneCompiledDataManager::PopulateSubSequenceTree(UMovieSceneSubTrack
 			continue;
 		}
 
-		const FMovieSceneTimeTransform SequenceToRootTransform = Params.RootToSequenceTransform.InverseFromWarp(Params.RootToSequenceWarpCounter);
-		TRange<FFrameNumber> EffectiveRange = Params.ClampRoot(Entry.Range * SequenceToRootTransform);
+		const FMovieSceneSequenceTransform SequenceToRootTransform = Params.RootToSequenceTransform.InverseFromLoop(Params.RootToSequenceWarpCounter);
+
+		// In the case the Sequence to Root Transform contains an infinite timescale, then one of the timescales above us is zero. In this case, we cannot
+		// rely on a simple multiply by an inverse transform to figure out an effective range of this entry in root space, as this is non-deterministic.
+		// It all comes down to whether the single frame is inside or outside of this entry. 
+		// If inside, the effective range in root space is the entire clamp range.
+		// If outside, the effective range in root space is empty, and we should not include this entry.
+
+		TRange<FFrameNumber> EffectiveRange = TRange<FFrameNumber>::Empty();
+		if (FMath::IsFinite(SequenceToRootTransform.GetTimeScale()) || !TRange<FFrameNumber>::Intersection(Params.LocalClampRange, Entry.Range).IsEmpty())
+		{
+			EffectiveRange = Params.ClampRoot(SequenceToRootTransform.TransformRangeConstrained(Entry.Range));
+		}
+
 		if (EffectiveRange.IsEmpty())
 		{
 			continue;
@@ -1586,14 +1598,26 @@ void UMovieSceneCompiledDataManager::PopulateSubSequenceTree(UMovieSceneSubTrack
 
 		const ESectionEvaluationFlags SubEntryFlags = Entry.Flags | Params.Flags;
 
-		if (!SubSectionParams.bCanLoop)
+		// If we can't loop, or our timescale is zero, and therefore looping is irrelevant
+		if (!SubSectionParams.bCanLoop || FMath::IsNearlyZero(SubData->RootToSequenceTransform.GetTimeScale()))
 		{
 			FGatherParameters SubParams = Params.CreateForSubData(*SubData, SubSequenceID, Params.RootToSequenceWarpCounter);
 			SubParams.SetClampRange(EffectiveRange);
 			SubParams.Flags |= Entry.Flags;
 			SubParams.NetworkMask = NewMask;
-			SubParams.RootToSequenceWarpCounter.AddNonWarpingLevel();
 
+			for (int i = 0; i < SubParams.RootToSequenceTransform.NestedTransforms.Num(); ++i)
+			{
+				if (SubParams.RootToSequenceTransform.NestedTransforms[i].IsLooping())
+				{
+					SubParams.RootToSequenceWarpCounter.AddWarpingLevel(1);
+				}
+				else
+				{
+					SubParams.RootToSequenceWarpCounter.AddNonWarpingLevel();
+				}
+			}
+			
 			// The section isn't looping, so we can just add it to the tree.
 			InOutHierarchy->AddRange(EffectiveRange, SubSequenceID, SubEntryFlags, SubParams.RootToSequenceWarpCounter);
 
diff --git a/Engine/Source/Runtime/MovieScene/Private/EntitySystem/MovieSceneEvaluationHookSystem.cpp b/Engine/Source/Runtime/MovieScene/Private/EntitySystem/MovieSceneEvaluationHookSystem.cpp
--- a/Engine/Source/Runtime/MovieScene/Private/EntitySystem/MovieSceneEvaluationHookSystem.cpp
+++ b/Engine/Source/Runtime/MovieScene/Private/EntitySystem/MovieSceneEvaluationHookSystem.cpp
@@ -46,7 +46,7 @@ struct FEvaluationHookUpdater
 			FMovieSceneEvaluationHookEvent NewEvent;
 			NewEvent.Hook          = Hooks[Index];
 			NewEvent.Type          = EEvaluationHookEvent::Update;
-			NewEvent.RootTime      = EvalTimes[Index] * SequenceInstance.GetContext().GetSequenceToRootTransform();
+			NewEvent.RootTime      = EvalTimes[Index] * SequenceInstance.GetContext().GetSequenceToRootSequenceTransform();
 			NewEvent.SequenceID    = SequenceInstance.GetSequenceID();
 			NewEvent.bRestoreState = bRestoreState;
 
@@ -156,7 +156,7 @@ void UMovieSceneEvaluationHookSystem::UpdateHooks()
 			FMovieSceneEvaluationHookEvent NewEvent;
 			NewEvent.Hook          = Hooks[Index];
 			NewEvent.Type          = EEvaluationHookEvent::Begin;
-			NewEvent.RootTime      = EvalTimes[Index] * SequenceInstance.GetContext().GetSequenceToRootTransform();
+			NewEvent.RootTime      = EvalTimes[Index] * SequenceInstance.GetContext().GetSequenceToRootSequenceTransform();
 			NewEvent.SequenceID    = SequenceInstance.GetSequenceID();
 			NewEvent.bRestoreState = bRestoreState;
 
@@ -176,7 +176,7 @@ void UMovieSceneEvaluationHookSystem::UpdateHooks()
 			FMovieSceneEvaluationHookEvent NewEvent;
 			NewEvent.Hook          = Hooks[Index];
 			NewEvent.Type          = EEvaluationHookEvent::End;
-			NewEvent.RootTime      = EvalTimes[Index] * SequenceInstance.GetContext().GetSequenceToRootTransform();
+			NewEvent.RootTime      = EvalTimes[Index] * SequenceInstance.GetContext().GetSequenceToRootSequenceTransform();
 			NewEvent.SequenceID    = SequenceInstance.GetSequenceID();
 			NewEvent.bRestoreState = bRestoreState;
 
diff --git a/Engine/Source/Runtime/MovieScene/Private/EntitySystem/MovieSceneSequenceUpdaters.cpp b/Engine/Source/Runtime/MovieScene/Private/EntitySystem/MovieSceneSequenceUpdaters.cpp
--- a/Engine/Source/Runtime/MovieScene/Private/EntitySystem/MovieSceneSequenceUpdaters.cpp
+++ b/Engine/Source/Runtime/MovieScene/Private/EntitySystem/MovieSceneSequenceUpdaters.cpp
@@ -442,7 +442,7 @@ void FSequenceUpdater_Hierarchical::DissectContext(UMovieSceneEntitySystemLinker
 						FFrameTime Unused;
 						SubData->RootToSequenceTransform.TransformTime(RootClampRange.GetLowerBoundValue(), Unused, WarpCounter);
 
-						FMovieSceneTimeTransform InverseTransform = SubData->RootToSequenceTransform.InverseFromWarp(WarpCounter);
+						FMovieSceneSequenceTransform InverseTransform = SubData->RootToSequenceTransform.InverseFromLoop(WarpCounter);
 						Algo::Transform(TraversedFences, RootDissectionTimes, [InverseTransform](FFrameTime In){ return In * InverseTransform; });
 					}
 				}
@@ -702,11 +702,11 @@ void FSequenceUpdater_Hierarchical::Update(UMovieSceneEntitySystemLinker* Linker
 
 					if (WarpCounter.WarpCounts.Num() > 0)
 					{
-						FMovieSceneTimeTransform InverseTransform = SubData->RootToSequenceTransform.InverseFromWarp(WarpCounter);
-						CachedEntityRange = TRange<FFrameNumber>::Intersection(CachedEntityRange, SubData->PlayRange.Value * InverseTransform);
+						FMovieSceneSequenceTransform InverseTransform = SubData->RootToSequenceTransform.InverseFromLoop(WarpCounter);
+						CachedEntityRange = TRange<FFrameNumber>::Intersection(CachedEntityRange, InverseTransform.TransformRangeConstrained(SubData->PlayRange.Value));
 					}
 
-					const FMovieSceneSequenceTransform SequenceToRootOverrideTransform = FMovieSceneSequenceTransform(SubContext.GetSequenceToRootTransform()) * RootContext.GetRootToSequenceTransform();
+					const FMovieSceneSequenceTransform SequenceToRootOverrideTransform = SubContext.GetSequenceToRootSequenceTransform() * RootContext.GetRootToSequenceTransform();
 					SubEntityRange = SequenceToRootOverrideTransform.TransformRangeConstrained(SubEntityRange);
 
 					CachedEntityRange = TRange<FFrameNumber>::Intersection(CachedEntityRange, SubEntityRange);
diff --git a/Engine/Source/Runtime/MovieScene/Private/Evaluation/MovieSceneSequenceHierarchy.cpp b/Engine/Source/Runtime/MovieScene/Private/Evaluation/MovieSceneSequenceHierarchy.cpp
--- a/Engine/Source/Runtime/MovieScene/Private/Evaluation/MovieSceneSequenceHierarchy.cpp
+++ b/Engine/Source/Runtime/MovieScene/Private/Evaluation/MovieSceneSequenceHierarchy.cpp
@@ -52,9 +52,9 @@ FMovieSceneSubSequenceData::FMovieSceneSubSequenceData(const UMovieSceneSubSecti
 	// being truly the full transform.
 	OuterToInnerTransform = RootToSequenceTransform = InSubSection.OuterToInnerTransform();
 
-	if (!InSubSection.Parameters.bCanLoop)
+	if (!InSubSection.Parameters.bCanLoop || FMath::IsNearlyZero(RootToSequenceTransform.GetTimeScale()))
 	{
-		PlayRange.Value = SubSectionRange * RootToSequenceTransform.LinearTransform;
+		PlayRange.Value = RootToSequenceTransform.TransformRangeUnwarped(SubSectionRange);
 		UnwarpedPlayRange.Value = PlayRange.Value;
 	}
 	else
@@ -69,13 +69,14 @@ FMovieSceneSubSequenceData::FMovieSceneSubSequenceData(const UMovieSceneSubSecti
 	}
 
 	// Make sure pre/postroll *ranges* are in the inner sequence's time space. Pre/PostRollFrames are in the outer sequence space.
+
 	if (InSubSection.GetPreRollFrames() > 0)
 	{
-		PreRollRange = UE::MovieScene::MakeDiscreteRangeFromUpper( TRangeBound<FFrameNumber>::FlipInclusion(SubSectionRange.GetLowerBound()), InSubSection.GetPreRollFrames() ) * RootToSequenceTransform.LinearTransform;
+		PreRollRange = RootToSequenceTransform.TransformRangeUnwarped(UE::MovieScene::MakeDiscreteRangeFromUpper(TRangeBound<FFrameNumber>::FlipInclusion(SubSectionRange.GetLowerBound()), InSubSection.GetPreRollFrames()));
 	}
 	if (InSubSection.GetPostRollFrames() > 0)
 	{
-		PostRollRange = UE::MovieScene::MakeDiscreteRangeFromLower( TRangeBound<FFrameNumber>::FlipInclusion(SubSectionRange.GetUpperBound()), InSubSection.GetPostRollFrames() ) * RootToSequenceTransform.LinearTransform;
+		PostRollRange = RootToSequenceTransform.TransformRangeUnwarped(UE::MovieScene::MakeDiscreteRangeFromLower(TRangeBound<FFrameNumber>::FlipInclusion(SubSectionRange.GetUpperBound()), InSubSection.GetPostRollFrames()));
 	}
 }
 
diff --git a/Engine/Source/Runtime/MovieScene/Private/MovieSceneSequenceVisitor.cpp b/Engine/Source/Runtime/MovieScene/Private/MovieSceneSequenceVisitor.cpp
--- a/Engine/Source/Runtime/MovieScene/Private/MovieSceneSequenceVisitor.cpp
+++ b/Engine/Source/Runtime/MovieScene/Private/MovieSceneSequenceVisitor.cpp
@@ -52,7 +52,7 @@ void VisitSubTrackImpl(const FSequenceVisitParams& InParams, UMovieSceneSubTrack
 			continue;
 		}
 
-		TRange<FFrameNumber> EffectiveRange = TRange<FFrameNumber>::Intersection(Entry.Range * SubSequenceSpace->RootToSequenceTransform.InverseLinearOnly(), SubSequenceSpace->RootClampRange);
+		TRange<FFrameNumber> EffectiveRange = TRange<FFrameNumber>::Intersection(SubSequenceSpace->RootToSequenceTransform.InverseNoLooping().TransformRangeConstrained(Entry.Range), SubSequenceSpace->RootClampRange);
 		if (EffectiveRange.IsEmpty())
 		{
 			continue;
diff --git a/Engine/Source/Runtime/MovieScene/Private/Sections/MovieSceneHookSection.cpp b/Engine/Source/Runtime/MovieScene/Private/Sections/MovieSceneHookSection.cpp
--- a/Engine/Source/Runtime/MovieScene/Private/Sections/MovieSceneHookSection.cpp
+++ b/Engine/Source/Runtime/MovieScene/Private/Sections/MovieSceneHookSection.cpp
@@ -48,7 +48,7 @@ void UMovieSceneHookSection::ImportTriggerEntity(UMovieSceneEntitySystemLinker*
 	NewEvent.Hook.Interface = this;
 	NewEvent.Hook.ObjectBindingID = Params.GetObjectBindingID();
 
-	NewEvent.RootTime = Times[EventIndex] * Context.GetSequenceToRootTransform();
+	NewEvent.RootTime = Times[EventIndex] * Context.GetSequenceToRootSequenceTransform();
 	NewEvent.Type = EEvaluationHookEvent::Trigger;
 	NewEvent.SequenceID = ThisInstance.GetSequenceID();
 	NewEvent.TriggerIndex = EventIndex;
diff --git a/Engine/Source/Runtime/MovieScene/Private/Sections/MovieSceneSubSection.cpp b/Engine/Source/Runtime/MovieScene/Private/Sections/MovieSceneSubSection.cpp
--- a/Engine/Source/Runtime/MovieScene/Private/Sections/MovieSceneSubSection.cpp
+++ b/Engine/Source/Runtime/MovieScene/Private/Sections/MovieSceneSubSection.cpp
@@ -57,29 +57,38 @@ FMovieSceneSequenceTransform UMovieSceneSubSection::OuterToInnerTransform() cons
 	const FFrameNumber InnerStartTime = UE::MovieScene::DiscreteInclusiveLower(MovieScenePlaybackRange);
 	const FFrameNumber OuterStartTime = UE::MovieScene::DiscreteInclusiveLower(SubRange);
 
-	// This is the transform for the "placement" (position and scaling) of the sub-sequence.
-	FMovieSceneTimeTransform LinearTransform =
-		// Inner play offset
-		FMovieSceneTimeTransform(InnerStartTime)
-		// Inner play rate
-		* FMovieSceneTimeTransform(0, Parameters.TimeScale * FrameRateScale)
-		// Outer section start time
-		* FMovieSceneTimeTransform(-OuterStartTime);
-	
-	if (!Parameters.bCanLoop)
+
+	FMovieSceneSequenceTransform Result;
+	// We have to special case 0 and infinite timescale so we can keep hold of each transform separately for property inverting.
+	// The linear transform for this special case remains identity.
+	if (FMath::IsNearlyZero(Parameters.TimeScale) || !FMath::IsFinite(Parameters.TimeScale))
 	{
-		return FMovieSceneSequenceTransform(LinearTransform);
+		Result.NestedTransforms.Add(FMovieSceneTimeTransform(-OuterStartTime));
+		Result.NestedTransforms.Add(FMovieSceneTimeTransform(0, 0));
+		Result.NestedTransforms.Add(FMovieSceneTimeTransform(InnerStartTime));
 	}
 	else
+	{
+		// This is the transform for the "placement" (position and scaling) of the sub-sequence.
+		Result.LinearTransform = 
+			// Inner play offset
+			FMovieSceneTimeTransform(InnerStartTime)
+			// Inner play rate
+			* FMovieSceneTimeTransform(0, Parameters.TimeScale * FrameRateScale)
+			// Outer section start time
+			* FMovieSceneTimeTransform(-OuterStartTime);
+	}
+
+	if (Parameters.bCanLoop)
 	{
 		const FFrameNumber InnerEndTime = UE::MovieScene::DiscreteExclusiveUpper(MovieScenePlaybackRange);
 		const FMovieSceneTimeWarping LoopingTransform(InnerStartTime, InnerEndTime);
-		LinearTransform = FMovieSceneTimeTransform(Parameters.FirstLoopStartFrameOffset) * LinearTransform;
 
-		FMovieSceneSequenceTransform Result;
-		Result.NestedTransforms.Add(FMovieSceneNestedSequenceTransform(LinearTransform, LoopingTransform));
+		Result.NestedTransforms.Add(FMovieSceneNestedSequenceTransform(FMovieSceneTimeTransform(Parameters.FirstLoopStartFrameOffset) * Result.LinearTransform, LoopingTransform));
+		Result.LinearTransform = FMovieSceneTimeTransform();
 		return Result;
 	}
+	return Result;
 }
 
 bool UMovieSceneSubSection::GetValidatedInnerPlaybackRange(TRange<FFrameNumber>& OutInnerPlaybackRange) const
@@ -355,7 +364,7 @@ TOptional<TRange<FFrameNumber> > UMovieSceneSubSection::GetAutoSizeRange() const
 	{
 		// We probably want to just auto-size the section to the sub-sequence's scaled playback range... if this section
 		// is looping, however, it's hard to know what we want to do. Let's just size it to one loop.
-		const FMovieSceneTimeTransform InnerToOuter = OuterToInnerTransform().InverseLinearOnly();
+		const FMovieSceneSequenceTransform InnerToOuter = OuterToInnerTransform().InverseNoLooping();
 		const TRange<FFrameNumber> InnerPlaybackRange = UMovieSceneSubSection::GetValidatedInnerPlaybackRange(Parameters, *MovieScene);
 
 		const FFrameTime IncAutoStartTime = FFrameTime(UE::MovieScene::DiscreteInclusiveLower(InnerPlaybackRange)) * InnerToOuter;
@@ -440,7 +449,7 @@ void UMovieSceneSubSection::GetSnapTimes(TArray<FFrameNumber>& OutSnapTimes, boo
 	}
 	else
 	{
-		const FMovieSceneSequenceTransform InnerToOuterTransform = OuterToInnerTransform().InverseLinearOnly();
+		const FMovieSceneSequenceTransform InnerToOuterTransform = OuterToInnerTransform().InverseNoLooping();
 		const FFrameNumber PlaybackStart = (UE::MovieScene::DiscreteInclusiveLower(PlaybackRange) * InnerToOuterTransform).FloorToFrame();
 		if (GetRange().Contains(PlaybackStart))
 		{
diff --git a/Engine/Source/Runtime/MovieScene/Private/Tests/MovieSceneTransformTests.cpp b/Engine/Source/Runtime/MovieScene/Private/Tests/MovieSceneTransformTests.cpp
--- a/Engine/Source/Runtime/MovieScene/Private/Tests/MovieSceneTransformTests.cpp
+++ b/Engine/Source/Runtime/MovieScene/Private/Tests/MovieSceneTransformTests.cpp
@@ -318,7 +318,7 @@ bool FMovieSceneSubSectionCoreSequenceTransformsTest::RunTest(const FString& Par
 			AddError(FString::Printf(TEXT("Accumulated transform does not have the same effect as separate transformations (%i+%.5f != %i+%.5f)"), AccumValue.FrameNumber.Value, AccumValue.GetSubFrame(), SeedValue.FrameNumber.Value, SeedValue.GetSubFrame()));
 		}
 
-		FMovieSceneSequenceTransform InverseTransform = AccumulatedTransform.InverseLinearOnly();
+		FMovieSceneSequenceTransform InverseTransform = AccumulatedTransform.InverseNoLooping();
 
 		FFrameTime InverseValue = AccumValue * InverseTransform;
 		if (InverseValue != 10)
@@ -343,9 +343,9 @@ bool FMovieSceneSubSectionCoreWarpingAndScalingTransformsTest::RunTest(const FSt
 		TestEqual("Transform time 1", FFrameNumber(10) * Transform, FFrameTime(20));
 		TestEqual("Transform time 2", FFrameNumber(18) * Transform, FFrameTime(6));
 
-		FMovieSceneTimeTransform Inv = Transform.InverseFromWarp(TArray<uint32> { 0 });
+		FMovieSceneSequenceTransform Inv = Transform.InverseFromLoop(TArray<uint32> { 0 });
 		TestEqual("Inverse time 3", FFrameNumber(20) * Inv, FFrameTime(10));
-		Inv = Transform.InverseFromWarp(TArray<uint32> { 1 });
+		Inv = Transform.InverseFromLoop(TArray<uint32> { 1 });
 		TestEqual("Inverse time 4", FFrameNumber(20) * Inv, FFrameTime(25));
 	}
 
@@ -356,9 +356,9 @@ bool FMovieSceneSubSectionCoreWarpingAndScalingTransformsTest::RunTest(const FSt
 		TestEqual("Transform time 5", FFrameNumber(10) * Transform, FFrameTime(40));
 		TestEqual("Transform time 6", FFrameNumber(18) * Transform, FFrameTime(26));
 
-		FMovieSceneTimeTransform Inv = Transform.InverseFromWarp(TArray<uint32> { 0 });
+		FMovieSceneSequenceTransform Inv = Transform.InverseFromLoop(TArray<uint32> { 0 });
 		TestEqual("Inverse time 7", FFrameNumber(40) * Inv, FFrameTime(10));
-		Inv = Transform.InverseFromWarp(TArray<uint32> { 1 });
+		Inv = Transform.InverseFromLoop(TArray<uint32> { 1 });
 		TestEqual("Inverse time 8", FFrameNumber(40) * Inv, FFrameTime(25));
 	}
 
@@ -369,9 +369,9 @@ bool FMovieSceneSubSectionCoreWarpingAndScalingTransformsTest::RunTest(const FSt
 		TestEqual("Transform time 9", FFrameNumber(13) * Transform, FFrameTime(20));
 		TestEqual("Transform time 10", FFrameNumber(21) * Transform, FFrameTime(6));
 
-		FMovieSceneTimeTransform Inv = Transform.InverseFromWarp(TArray<uint32> { 0 });
+		FMovieSceneSequenceTransform Inv = Transform.InverseFromLoop(TArray<uint32> { 0 });
 		TestEqual("Inverse time 11", FFrameNumber(20) * Inv, FFrameTime(13));
-		Inv = Transform.InverseFromWarp(TArray<uint32> { 1 });
+		Inv = Transform.InverseFromLoop(TArray<uint32> { 1 });
 		TestEqual("Inverse time 12", FFrameNumber(20) * Inv, FFrameTime(28));
 	}
 
@@ -382,9 +382,9 @@ bool FMovieSceneSubSectionCoreWarpingAndScalingTransformsTest::RunTest(const FSt
 		TestEqual("Transform time 13", FFrameNumber(13) * Transform, FFrameTime(40));
 		TestEqual("Transform time 14", FFrameNumber(21) * Transform, FFrameTime(26));
 
-		FMovieSceneTimeTransform Inv = Transform.InverseFromWarp(TArray<uint32> { 0 });
+		FMovieSceneSequenceTransform Inv = Transform.InverseFromLoop(TArray<uint32> { 0 });
 		TestEqual("Inverse time 15", FFrameNumber(40) * Inv, FFrameTime(13));
-		Inv = Transform.InverseFromWarp(TArray<uint32> { 1 });
+		Inv = Transform.InverseFromLoop(TArray<uint32> { 1 });
 		TestEqual("Inverse time 16", FFrameNumber(40) * Inv, FFrameTime(28));
 	}
 
@@ -396,9 +396,9 @@ bool FMovieSceneSubSectionCoreWarpingAndScalingTransformsTest::RunTest(const FSt
 		TestEqual("Transform time 17", FFrameNumber(18) * Transform, FFrameTime(4));
 		TestEqual("Transform time 18", FFrameNumber(55) * Transform, FFrameTime(18));
 
-		FMovieSceneTimeTransform Inv = Transform.InverseFromWarp(TArray<uint32> { 0, FMovieSceneTimeWarping::InvalidWarpCount });
+		FMovieSceneSequenceTransform Inv = Transform.InverseFromLoop(TArray<uint32> { 0, FMovieSceneTimeWarping::InvalidWarpCount });
 		TestEqual("Inverse time 17", FFrameNumber(4) * Inv, FFrameTime(18));
-		Inv = Transform.InverseFromWarp(TArray<uint32> { 1, FMovieSceneTimeWarping::InvalidWarpCount });
+		Inv = Transform.InverseFromLoop(TArray<uint32> { 1, FMovieSceneTimeWarping::InvalidWarpCount });
 		TestEqual("Inverse time 18", FFrameNumber(18) * Inv, FFrameTime(55));
 	}
 
@@ -410,9 +410,9 @@ bool FMovieSceneSubSectionCoreWarpingAndScalingTransformsTest::RunTest(const FSt
 		TestEqual("Transform time 17", FFrameNumber(19) * Transform, FFrameTime(6));
 		TestEqual("Transform time 18", FFrameNumber(32) * Transform, FFrameTime(4));
 
-		FMovieSceneTimeTransform Inv = Transform.InverseFromWarp(TArray<uint32> { FMovieSceneTimeWarping::InvalidWarpCount, 0 });
+		FMovieSceneSequenceTransform Inv = Transform.InverseFromLoop(TArray<uint32> { FMovieSceneTimeWarping::InvalidWarpCount, 0 });
 		TestEqual("Inverse time 17", FFrameNumber(6) * Inv, FFrameTime(19));
-		Inv = Transform.InverseFromWarp(TArray<uint32> { FMovieSceneTimeWarping::InvalidWarpCount, 2 });
+		Inv = Transform.InverseFromLoop(TArray<uint32> { FMovieSceneTimeWarping::InvalidWarpCount, 2 });
 		TestEqual("Inverse time 18", FFrameNumber(4) * Inv, FFrameTime(32));
 	}
 
@@ -424,12 +424,74 @@ bool FMovieSceneSubSectionCoreWarpingAndScalingTransformsTest::RunTest(const FSt
 		TestEqual("Transform time 19", FFrameNumber(113) * Transform, FFrameTime(40));
 		TestEqual("Transform time 20", FFrameNumber(121) * Transform, FFrameTime(26));
 
-		FMovieSceneTimeTransform Inv = Transform.InverseFromWarp(TArray<uint32> { 0 });
+		FMovieSceneSequenceTransform Inv = Transform.InverseFromLoop(TArray<uint32> { 0 });
 		TestEqual("Inverse time 21", FFrameNumber(40) * Inv, FFrameTime(113));
-		Inv = Transform.InverseFromWarp(TArray<uint32> { 1 });
+		Inv = Transform.InverseFromLoop(TArray<uint32> { 1 });
 		TestEqual("Inverse time 22", FFrameNumber(40) * Inv, FFrameTime(128));
 	}
 
+	{
+		// Zero-timescale transform on a sub-sequence. Any frame numbers transformed in should be equal to the frame offset
+		FMovieSceneSequenceTransform Transform;
+		Transform.NestedTransforms.Add(FMovieSceneTimeTransform(0)); // no outer offset
+		Transform.NestedTransforms.Add(FMovieSceneTimeTransform(0, 0)); // 0 timescale
+		Transform.NestedTransforms.Add(FMovieSceneTimeTransform(30)); // 30 inner frame offset
+
+		TestEqual("Outer time 40 through 0 timescale with 30 offset", FFrameNumber(40) * Transform, FFrameTime(30));
+		TestEqual("Outer time 0 through 0 timescale with 30 offset", FFrameNumber(0) * Transform, FFrameTime(30));
+		TestEqual("Outer time 173 through 0 timescale with 30 offset", FFrameNumber(173) * Transform, FFrameTime(30));
+	}
+
+	{
+		// Zero-timescale transform on a sub-sequence. Same as previous, but we also will invert this transform and ensure timescale is correctly infinite
+		// and any transforms by that infinite transform. Anything transformed out should just be equal to the outer offset
+		FMovieSceneSequenceTransform Transform;
+		Transform.NestedTransforms.Add(FMovieSceneTimeTransform(-10)); // 10 outer offset
+		Transform.NestedTransforms.Add(FMovieSceneTimeTransform(0, 0)); // 0 timescale
+		Transform.NestedTransforms.Add(FMovieSceneTimeTransform(30)); // 30 inner frame offset
+
+		FMovieSceneSequenceTransform Inv = Transform.InverseNoLooping();
+		if (FMath::IsFinite(Inv.GetTimeScale()))
+		{
+			AddError(FString::Printf(TEXT("Inverse of a transform with zero timescale is not correctly infinite")));
+		}
+
+		TestEqual("Inner time 40 through inf timescale with 10 outer offset", FFrameNumber(40) * Inv, FFrameTime(10));
+		TestEqual("Inner time 0 through inf timescale with 10 outer offset", FFrameNumber(0)* Inv, FFrameTime(10));
+		TestEqual("Inner time 173 through inf timescale with 10 outer offset", FFrameNumber(173)* Inv, FFrameTime(10));
+
+		// Re-invert the inverse transform. This should be equivalent to the original transform and we shouldn't have lost anything.
+		FMovieSceneSequenceTransform InvInv = Inv.InverseNoLooping();
+		TestEqual("Doubly inverted zero-timescale transform should be equal to original", InvInv, Transform);
+	}
+
+	{
+		// Multiple levels of sub sequences with zero-timescale thrown in
+		FMovieSceneSequenceTransform OuterTransform;
+		OuterTransform.NestedTransforms.Add(FMovieSceneTimeTransform(-10)); // 10 outer offset
+		OuterTransform.NestedTransforms.Add(FMovieSceneTimeTransform(0, 0)); // 0 timescale
+		OuterTransform.NestedTransforms.Add(FMovieSceneTimeTransform(30)); // 30 inner frame offset
+
+		FMovieSceneSequenceTransform InnerTransform;
+		InnerTransform.LinearTransform.Offset = FFrameNumber(5); // An inner frame offset of 5
+
+		FMovieSceneSequenceTransform CompleteTransform = InnerTransform * OuterTransform;
+
+		TestEqual("Subsequence frame through zero timescale transform", FFrameNumber(40)* CompleteTransform, FFrameTime(35));
+		TestEqual("Subsequence frame through zero timescale transform", FFrameNumber(0)* CompleteTransform, FFrameTime(35));
+		TestEqual("Subsequence frame through zero timescale transform", FFrameNumber(173)* CompleteTransform, FFrameTime(35));
+
+		FMovieSceneSequenceTransform InvCompleteTransform = CompleteTransform.InverseNoLooping();
+
+		TestEqual("Inner time 40 through inf timescale with 10 outer offset", FFrameNumber(40)* InvCompleteTransform, FFrameTime(10));
+		TestEqual("Inner time 0 through inf timescale with 10 outer offset", FFrameNumber(0)* InvCompleteTransform, FFrameTime(10));
+		TestEqual("Inner time 173 through inf timescale with 10 outer offset", FFrameNumber(173)* InvCompleteTransform, FFrameTime(10));
+
+		// Re-invert the inverse transform. This should be equivalent to the original transform and we shouldn't have lost anything.
+		FMovieSceneSequenceTransform InvInv = InvCompleteTransform.InverseNoLooping();
+		TestEqual("Doubly inverted zero-timescale transform should be equal to original", InvInv, CompleteTransform);
+	}
+
 	return true;
 }
 
diff --git a/Engine/Source/Runtime/MovieScene/Public/Evaluation/MovieScenePlayback.h b/Engine/Source/Runtime/MovieScene/Public/Evaluation/MovieScenePlayback.h
--- a/Engine/Source/Runtime/MovieScene/Public/Evaluation/MovieScenePlayback.h
+++ b/Engine/Source/Runtime/MovieScene/Public/Evaluation/MovieScenePlayback.h
@@ -273,9 +273,24 @@ struct FMovieSceneContext : FMovieSceneEvaluationRange
 	/**
 	 * Get the inverse transform of the current sub sequence, to transform local times back to root times.
 	 */
+	UE_DEPRECATED(5.4, "Please use GetSequenceToRootSequenceTransform instead.")
 	FORCEINLINE FMovieSceneTimeTransform GetSequenceToRootTransform() const
 	{
-		return RootToSequenceTransform.InverseFromWarp(RootToSequenceWarpCounter);
+		FMovieSceneSequenceTransform SequenceTransform = RootToSequenceTransform.InverseFromLoop(RootToSequenceWarpCounter);
+		FMovieSceneTimeTransform ReturnTransform = SequenceTransform.LinearTransform;
+		for (int i = 0; i < SequenceTransform.NestedTransforms.Num(); ++i)
+		{
+			ReturnTransform = ReturnTransform * SequenceTransform.NestedTransforms[i].LinearTransform;
+		}
+		return ReturnTransform;
+	}
+
+	/**
+	 * Get the inverse sequence transform of the current sub sequence, to transform local times back to root times.
+	 */
+	FORCEINLINE FMovieSceneSequenceTransform GetSequenceToRootSequenceTransform() const
+	{
+		return RootToSequenceTransform.InverseFromLoop(RootToSequenceWarpCounter);
 	}
 
 	/**
@@ -345,7 +360,7 @@ public:
 		NewContext.RootToSequenceTransform.TransformTime(GetTime(), TransformedTime_Unused, WarpCounter);
 		NewContext.RootToSequenceWarpCounter = WarpCounter;
 
-		if (InTransform.IsWarping())
+		if (InTransform.IsLooping())
 		{
 			// If we have some looping, the transformed range might extend past the end of a loop and into
 			// the beginning of another. In that case, technically, the evaluation range ends up being a
@@ -356,7 +371,7 @@ public:
 			//       actually make contexts have an array of evaluation ranges.
 			TRangeBound<FFrameTime> UpperEvalutionRangeBound = NewContext.EvaluationRange.GetUpperBound();
 			const FMovieSceneNestedSequenceTransform& LeafTransform = InTransform.NestedTransforms.Last();
-			if (UpperEvalutionRangeBound.IsClosed() && LeafTransform.IsWarping())
+			if (UpperEvalutionRangeBound.IsClosed() && LeafTransform.IsLooping())
 			{
 				const FFrameNumber LeafWarpLength = LeafTransform.Warping.Length();
 				// Below: use strictly greater than comparison so that if the evalution range ends on the
diff --git a/Engine/Source/Runtime/MovieScene/Public/Evaluation/MovieSceneSequenceTransform.h b/Engine/Source/Runtime/MovieScene/Public/Evaluation/MovieSceneSequenceTransform.h
--- a/Engine/Source/Runtime/MovieScene/Public/Evaluation/MovieSceneSequenceTransform.h
+++ b/Engine/Source/Runtime/MovieScene/Public/Evaluation/MovieSceneSequenceTransform.h
@@ -3,6 +3,7 @@
 #pragma once
 
 #include "Algo/AllOf.h"
+#include "Algo/AnyOf.h"
 #include "Containers/Array.h"
 #include "Containers/ArrayView.h"
 #include "Containers/UnrealString.h"
@@ -81,15 +82,10 @@ struct FMovieSceneWarpCounter
 		{
 			WarpCounts.Add(WarpCount);
 		}
-		else if (WarpCounts.Num() > 0)
+		else
 		{
-			if (WarpCounts.Last() != FMovieSceneTimeWarping::InvalidWarpCount)
-			{
-				WarpCounts.Add(FMovieSceneTimeWarping::InvalidWarpCount);
-			}
+			WarpCounts.Add(FMovieSceneTimeWarping::InvalidWarpCount);
 		}
-		// else: ignore, non-warping sub-sequences have their transforms compressed into a single
-		//       linear transform until the point we meet the first looping sub-sequence.
 	}
 
 	void AddNonWarpingLevel()
@@ -146,9 +142,14 @@ struct FMovieSceneNestedSequenceTransform
 	bool IsIdentity() const { return LinearTransform.IsIdentity() && !IsWarping(); }
 
 	/**
-	 * Returns whether this transform is warping.
+	 * Returns whether this transform is warping. This includes both looping as well as zero-timescale.
 	 */
-	bool IsWarping() const { return Warping.IsValid(); }
+	bool IsWarping() const { return Warping.IsValid() || FMath::IsNearlyZero(LinearTransform.TimeScale); }
+
+	/**
+	 * Returns whether this transform is warping specifically from looping. This doesn't include zero-timescale cases.
+	 */
+	bool IsLooping() const { return Warping.IsValid(); }
 
 	/**
 	 * Linear time transform for this sub-sequence.
@@ -157,7 +158,7 @@ struct FMovieSceneNestedSequenceTransform
 	FMovieSceneTimeTransform LinearTransform;
 
 	/**
-	 * Time warping information for this sub-sequence.
+	 * Time warping information for this sub-sequence used for looping.
 	 */
 	UPROPERTY()
 	FMovieSceneTimeWarping Warping;
@@ -174,11 +175,20 @@ struct FMovieSceneNestedSequenceTransform
 
 	FMovieSceneTimeTransform InverseLinearOnly() const
 	{
+		// We don't support storing both an offset and a zero or infinite timescale in a single nested transform.
+		// These should be stored as separate nested transforms.
+		ensureMsgf(FMath::IsNearlyZero(LinearTransform.Offset.AsDecimal()) || (FMath::IsFinite(LinearTransform.TimeScale) && !FMath::IsNearlyZero(LinearTransform.TimeScale)), TEXT("Error- nested transform with both an offset and a non-deterministic timescale."));
+
 		return LinearTransform.Inverse();
 	}
 
 	FMovieSceneTimeTransform InverseFromWarp(uint32 WarpCount) const
 	{
+		// We don't support storing both an offset and a zero or infinite timescale in a single nested transform.
+		// We also don't support storing a zero or infinite timescale in addition to looping information.
+		// These should be stored as separate nested transforms.
+		ensureMsgf(FMath::IsNearlyZero(LinearTransform.Offset.AsDecimal()) || ((FMath::IsFinite(LinearTransform.TimeScale) && !FMath::IsNearlyZero(LinearTransform.TimeScale))), TEXT("Error- nested transform with both an offset and a non-deterministic timescale."));
+		ensureMsgf(!IsLooping() || (FMath::IsFinite(LinearTransform.TimeScale) && !FMath::IsNearlyZero(LinearTransform.TimeScale)), TEXT("Error- nested transform with both looping and non-deterministic timescale."));
 		return LinearTransform.Inverse() * Warping.InverseFromWarp(WarpCount);
 	}
 };
@@ -250,6 +260,13 @@ struct FMovieSceneSequenceTransform
 		return NestedTransforms.Num() > 0;
 	}
 
+
+	bool IsLooping() const
+	{
+		return NestedTransforms.Num() > 0 &&
+			Algo::AnyOf(NestedTransforms, &FMovieSceneNestedSequenceTransform::IsLooping);
+	}
+
 	/**
 	 * Returns whether this sequence transform is purely linear (i.e. doesn't involve time warping).
 	 */
@@ -301,7 +318,7 @@ struct FMovieSceneSequenceTransform
 			{
 				OutTime = OutTime * NestedTransform.LinearTransform;
 
-				if (NestedTransform.IsWarping())
+				if (NestedTransform.IsLooping())
 				{
 					uint32 WarpIndex;
 					NestedTransform.Warping.TransformTime(OutTime, OutTime, WarpIndex);
@@ -325,7 +342,7 @@ struct FMovieSceneSequenceTransform
 		for (const FMovieSceneNestedSequenceTransform& NestedTransform : NestedTransforms)
 		{
 			Result = Result * NestedTransform.LinearTransform;
-			if (NestedTransform.IsWarping())
+			if (NestedTransform.IsLooping())
 			{
 				Result = NestedTransform.Warping.TransformRangePure(Result);
 			}
@@ -343,7 +360,7 @@ struct FMovieSceneSequenceTransform
 		for (const FMovieSceneNestedSequenceTransform& NestedTransform : NestedTransforms)
 		{
 			Result = Result * NestedTransform.LinearTransform;
-			if (NestedTransform.IsWarping())
+			if (NestedTransform.IsLooping())
 			{
 				Result = NestedTransform.Warping.TransformRangeUnwarped(Result);
 			}
@@ -361,7 +378,7 @@ struct FMovieSceneSequenceTransform
 		for (const FMovieSceneNestedSequenceTransform& NestedTransform : NestedTransforms)
 		{
 			Result = Result * NestedTransform.LinearTransform;
-			if (NestedTransform.IsWarping())
+			if (NestedTransform.IsLooping())
 			{
 				Result = NestedTransform.Warping.TransformRangeConstrained(Result);
 			}
@@ -405,15 +422,49 @@ struct FMovieSceneSequenceTransform
 	/**
 	 * Retrieve the inverse of the linear part of this transform.
 	 */
+	UE_DEPRECATED(5.4, "Please use InverseNoLooping instead.")
 	FMovieSceneTimeTransform InverseLinearOnly() const
 	{
+		ensureMsgf(!FMath::IsNearlyZero(LinearTransform.TimeScale), TEXT("Inverse of a zero timescale transform is undefined in a FMovieSceneTimeTransform. Please use InverseNoLooping for proper behavior."));
 		return LinearTransform.Inverse();
 	}
 
+	/**
+	 * Retrieve the inverse of the linear part of this transform without any looping.
+	 * Note that if nested transforms exist, we will return a sequence transform with nested transforms inverted, but looping information will be lost.
+	 */
+	FMovieSceneSequenceTransform InverseNoLooping() const
+	{
+		ensureMsgf(FMath::IsFinite(LinearTransform.TimeScale) && !FMath::IsNearlyZero(LinearTransform.TimeScale), TEXT("Warning: Sequence LinearTransform has a zero or infinite timescale."));
+		if (NestedTransforms.Num() > 0)
+		{
+			// Start accumulating the inverse transforms in reverse order.
+			FMovieSceneSequenceTransform Result;
+
+			for(int i = NestedTransforms.Num() - 1; i >= 0; --i)
+			{
+				Result.NestedTransforms.Add(NestedTransforms[i].LinearTransform.Inverse());
+			}
+
+			// Add the inverse of the main linear transform as a further nested transform if not identity
+			if (!LinearTransform.IsIdentity())
+			{
+				Result.NestedTransforms.Add(LinearTransform.Inverse());
+			}
+
+			return Result;
+		}
+		else
+		{
+			return LinearTransform.Inverse();
+		}
+	}
+
 	/**
 	 * Retrieve the inverse of this transform assuming the local times would all belong to
 	 * the first warp inside each nested time range.
 	 */
+	UE_DEPRECATED(5.4, "Please use InverseFromAllFirstLoops instead.")
 	FMovieSceneTimeTransform InverseFromAllFirstWarps() const
 	{
 		const size_t NestedTransformsSize = NestedTransforms.Num();
@@ -423,13 +474,36 @@ struct FMovieSceneSequenceTransform
 		}
 		else
 		{
-			TArray<uint32> WarpCounts;
-			WarpCounts.Reserve(NestedTransforms.Num());
+			FMovieSceneSequenceTransform SequenceTransform = InverseFromAllFirstLoops();
+			FMovieSceneTimeTransform ReturnTransform = SequenceTransform.LinearTransform;
+			for (int i = 0; i < SequenceTransform.NestedTransforms.Num(); ++i)
+			{
+				ReturnTransform = ReturnTransform * SequenceTransform.NestedTransforms[i].LinearTransform;
+			}
+			return ReturnTransform;
+		}
+	}
+
+	/**
+	 * Retrieve the inverse of this transform assuming the local times would all belong to the first loop inside each nested time range.
+	 */
+	FMovieSceneSequenceTransform InverseFromAllFirstLoops() const 
+	{
+		const size_t NestedTransformsSize = NestedTransforms.Num();
+		if (NestedTransformsSize == 0)
+		{
+			ensureMsgf(FMath::IsFinite(LinearTransform.TimeScale) && !FMath::IsNearlyZero(LinearTransform.TimeScale), TEXT("Warning: Sequence LinearTransform has a zero or infinite timescale."));
+			return LinearTransform.Inverse();
+		}
+		else
+		{
+			TArray<uint32> LoopCounts;
+			LoopCounts.Reserve(NestedTransforms.Num());
 			for (const FMovieSceneNestedSequenceTransform& NestedTransform : NestedTransforms)
 			{
-				WarpCounts.Add(NestedTransform.IsWarping() ? 0 : FMovieSceneTimeWarping::InvalidWarpCount);
+				LoopCounts.Add(NestedTransform.IsLooping() ? 0 : FMovieSceneTimeWarping::InvalidWarpCount);
 			}
-			return InverseFromWarp(WarpCounts);
+			return InverseFromLoop(LoopCounts);
 		}
 	}
 
@@ -438,9 +512,16 @@ struct FMovieSceneSequenceTransform
 	 * n'th warp inside the time range. As such, this method requires that the warp counts
 	 * in the counter object is equal to the number of nested loops in this transform.
 	 */
+	UE_DEPRECATED(5.4, "Please use InverseFromLoop instead.")
 	FMovieSceneTimeTransform InverseFromWarp(const FMovieSceneWarpCounter& WarpCounter) const
 	{
-		return InverseFromWarp(WarpCounter.WarpCounts);
+		FMovieSceneSequenceTransform SequenceTransform = InverseFromLoop(WarpCounter.WarpCounts);
+		FMovieSceneTimeTransform ReturnTransform = SequenceTransform.LinearTransform;
+		for (int i = 0; i < SequenceTransform.NestedTransforms.Num(); ++i)
+		{
+			ReturnTransform = ReturnTransform * SequenceTransform.NestedTransforms[i].LinearTransform;
+		}
+		return ReturnTransform;
 	}
 
 	/**
@@ -448,6 +529,7 @@ struct FMovieSceneSequenceTransform
 	 * n'th warp inside the time range. As such, this method takes as parameter an array with
 	 * as many items as there currently have nested loops in this transform.
 	 */
+	UE_DEPRECATED(5.4, "Please use InverseFromLoop instead.")
 	FMovieSceneTimeTransform InverseFromWarp(const TArrayView<const uint32>& WarpCounts) const
 	{
 		const size_t NestedTransformsSize = NestedTransforms.Num();
@@ -480,6 +562,117 @@ struct FMovieSceneSequenceTransform
 		}
 	}
 
+	/**
+	 * Retrieve the inverse of this transform assuming the local times would belong to the
+	 * n'th warp inside the time range. As such, this method requires that the warp counts
+	 * in the counter object is equal to the number of nested loops in this transform.
+	 * Correctly handles zero and infinite timescales.
+	 */
+	FMovieSceneSequenceTransform InverseFromLoop(const FMovieSceneWarpCounter& LoopCounter) const
+	{
+		return InverseFromLoop(LoopCounter.WarpCounts);
+	}
+
+	/**
+	 * Retrieve the inverse of this transform assuming the local times would belong to the
+	 * n'th warp inside the time range. As such, this method takes as parameter an array with
+	 * as many items as there currently have nested loops in this transform.
+	 * Correctly handles zero and infinite timescales.
+	 */
+	FMovieSceneSequenceTransform InverseFromLoop(const TArrayView<const uint32>& LoopCounts) const
+	{
+		const size_t NestedTransformsSize = NestedTransforms.Num();
+		check(LoopCounts.Num() == NestedTransformsSize);
+
+		if (NestedTransformsSize > 0)
+		{
+			FMovieSceneSequenceTransform Result;
+			for(int i = NestedTransformsSize - 1; i >=0; --i)
+			{
+				const uint32 CurLoopCount = LoopCounts[i];
+				Result.NestedTransforms.Add(NestedTransforms[i].InverseFromWarp(CurLoopCount));
+			}
+
+			// Add the inverse of the main linear transform as a further nested transform if not identity
+			if (!LinearTransform.IsIdentity())
+			{
+				Result.NestedTransforms.Add(LinearTransform.Inverse());
+			}
+
+			return Result;
+		}
+		else
+		{
+			ensureMsgf(FMath::IsFinite(LinearTransform.TimeScale) && !FMath::IsNearlyZero(LinearTransform.TimeScale), TEXT("Warning: Sequence LinearTransform has a zero or infinite timescale."));
+			return LinearTransform.Inverse();
+		}
+	}
+
+	/**
+	 * Multiply 2 transforms together, resulting in a single transform that gets from RHS parent to LHS space
+	 * @note Transforms apply from right to left
+	 */
+	FMovieSceneSequenceTransform operator*(const FMovieSceneSequenceTransform& RHS) const
+	{
+		if (!IsWarping())
+		{
+			if (!RHS.IsWarping())
+			{
+				// None of the transforms are warping... we can combine them into another linear transform.
+				return FMovieSceneSequenceTransform(LinearTransform * RHS.LinearTransform);
+			}
+			else
+			{
+				// LHS is linear, but RHS is warping. Since transforms are supposed to apply from right to left,
+				// we need to append LHS at the "bottom" of RHS, i.e. add a new nested transform that's LHS. However
+				// if LHS is identity, we have nothing to do, and if both LHS and RHS' deeper transform are linear,
+				// we can combine both.
+				FMovieSceneSequenceTransform Result(RHS);
+				if (!LinearTransform.IsIdentity())
+				{
+					const int32 NumNestedTransforms = RHS.NestedTransforms.Num();
+					check(NumNestedTransforms > 0);
+					const FMovieSceneNestedSequenceTransform& LastNested = RHS.NestedTransforms[NumNestedTransforms - 1];
+					if (LastNested.IsWarping())
+					{
+						Result.NestedTransforms.Add(FMovieSceneNestedSequenceTransform(LinearTransform));
+					}
+					else
+					{
+						Result.NestedTransforms[NumNestedTransforms - 1] = LinearTransform * LastNested.LinearTransform;
+					}
+				}
+				return Result;
+			}
+		}
+		else // LHS is warping
+		{
+			if (!RHS.IsWarping())
+			{
+				// RHS isn't warping, but LHS is, so we combine the linear transform parts and start looping
+				// from there.
+				FMovieSceneSequenceTransform Result;
+				Result.LinearTransform = LinearTransform * RHS.LinearTransform;
+				Result.NestedTransforms = NestedTransforms;
+				return Result;
+			}
+			else
+			{
+				// Both are looping, we need to combine them. Usually, a warping transform doesn't use its linear part,
+				// because whatever linear placement/scaling it has would be in the linear part of the nested transform
+				// struct.
+				FMovieSceneSequenceTransform Result(RHS);
+				const bool bHasOnlyNested = LinearTransform.IsIdentity();
+				if (!bHasOnlyNested)
+				{
+					Result.NestedTransforms.Add(FMovieSceneNestedSequenceTransform(LinearTransform));
+				}
+				Result.NestedTransforms.Append(NestedTransforms);
+				return Result;
+			}
+		}
+	}
+
 	UPROPERTY()
 	FMovieSceneTimeTransform LinearTransform;
 
@@ -505,7 +698,7 @@ inline FFrameTime operator*(FFrameTime InTime, const FMovieSceneSequenceTransfor
 		FFrameTime OutTime = InTime * RHS.LinearTransform;
 		for (const FMovieSceneNestedSequenceTransform& NestedTransform : RHS.NestedTransforms)
 		{
-			if (NestedTransform.IsWarping())
+			if (NestedTransform.IsLooping())
 			{
 				OutTime = OutTime * NestedTransform.LinearTransform * NestedTransform.Warping;
 			}
@@ -543,71 +736,6 @@ TRange<T>& operator*=(TRange<T>& LHS, const FMovieSceneSequenceTransform& RHS)
 	return LHS;
 }
 
-/**
- * Multiply 2 transforms together, resulting in a single transform that gets from RHS parent to LHS space
- * @note Transforms apply from right to left
- */
-inline FMovieSceneSequenceTransform operator*(const FMovieSceneSequenceTransform& LHS, const FMovieSceneSequenceTransform& RHS)
-{
-	if (!LHS.IsWarping())
-	{
-		if (!RHS.IsWarping())
-		{
-			// None of the transforms are warping... we can combine them into another linear transform.
-			return FMovieSceneSequenceTransform(LHS.LinearTransform * RHS.LinearTransform);
-		}
-		else
-		{
-			// LHS is linear, but RHS is warping. Since transforms are supposed to apply from right to left,
-			// we need to append LHS at the "bottom" of RHS, i.e. add a new nested transform that's LHS. However
-			// if LHS is identity, we have nothing to do, and if both LHS and RHS' deeper transform are linear,
-			// we can combine both.
-			FMovieSceneSequenceTransform Result(RHS);
-			if (!LHS.LinearTransform.IsIdentity())
-			{
-				const int32 NumNestedTransforms = RHS.NestedTransforms.Num();
-				check(NumNestedTransforms > 0);
-				const FMovieSceneNestedSequenceTransform& LastNested = RHS.NestedTransforms[NumNestedTransforms - 1];
-				if (LastNested.IsWarping())
-				{
-					Result.NestedTransforms.Add(FMovieSceneNestedSequenceTransform(LHS.LinearTransform));
-				}
-				else
-				{
-					Result.NestedTransforms[NumNestedTransforms - 1] = LHS.LinearTransform * LastNested.LinearTransform;
-				}
-			}
-			return Result;
-		}
-	}
-	else // LHS is warping
-	{
-		if (!RHS.IsWarping())
-		{
-			// RHS isn't warping, but LHS is, so we combine the linear transform parts and start looping
-			// from there.
-			FMovieSceneSequenceTransform Result;
-			Result.LinearTransform = LHS.LinearTransform * RHS.LinearTransform;
-			Result.NestedTransforms = LHS.NestedTransforms;
-			return Result;
-		}
-		else
-		{
-			// Both are looping, we need to combine them. Usually, a warping transform doesn't use its linear part,
-			// because whatever linear placement/scaling it has would be in the linear part of the nested transform
-			// struct.
-			FMovieSceneSequenceTransform Result(RHS);
-			const bool bHasOnlyNested = LHS.LinearTransform.IsIdentity();
-			if (!bHasOnlyNested)
-			{
-				Result.NestedTransforms.Add(FMovieSceneNestedSequenceTransform(LHS.LinearTransform));
-			}
-			Result.NestedTransforms.Append(LHS.NestedTransforms);
-			return Result;
-		}
-	}
-}
-
 /** Convert a FMovieSceneSequenceTransform into a string */
 FString LexToString(const FMovieSceneSequenceTransform& InTransform);
 
diff --git a/Engine/Source/Runtime/MovieScene/Public/Evaluation/MovieSceneTimeArray.h b/Engine/Source/Runtime/MovieScene/Public/Evaluation/MovieSceneTimeArray.h
--- a/Engine/Source/Runtime/MovieScene/Public/Evaluation/MovieSceneTimeArray.h
+++ b/Engine/Source/Runtime/MovieScene/Public/Evaluation/MovieSceneTimeArray.h
@@ -135,37 +135,22 @@ private:
 		FMovieSceneWarpCounter LoopCounter;
 		FMovieSceneSequenceTransform RootToLocalTransform;
 
+		// TODO: It seems a bit silly to store this as a TransformStack when we need to do the same work to inverse it.
+
 		for (int32 Index = 0; Index < TransformStack.Num(); ++Index)
 		{
-			const FTransformStep& Step(TransformStack[Index]);
-			FMovieSceneSequenceTransform StepTransform;
-			StepTransform.LinearTransform = Step.LinearTransform;
-			if (Step.Warping.IsValid())
+			RootToLocalTransform.NestedTransforms.Add(FMovieSceneNestedSequenceTransform(TransformStack[Index].LinearTransform, TransformStack[Index].Warping));
+			if (TransformStack[Index].Warping.IsValid())
 			{
-				StepTransform.NestedTransforms.Add(FMovieSceneNestedSequenceTransform(Step.Warping));
+				LoopCounter.AddWarpingLevel(TransformStack[Index].WarpCounter);
 			}
-
-			const bool bWasRootToLocalTransformWarping = RootToLocalTransform.IsWarping();
-			RootToLocalTransform = StepTransform * RootToLocalTransform;
-
-			ensure(Step.WarpCounter == FMovieSceneTimeWarping::InvalidWarpCount || Step.Warping.IsValid());
-
-			// The linear part of StepTransform will be added as a nested transform only if it's 
-			// actually doing something (non identity), and if RootToLocalTransform was warping already
-			// (meaning that StepTransform's linear part couldn't be "merged" into an existing linear
-			// transform).
-			if (!StepTransform.LinearTransform.IsIdentity() && bWasRootToLocalTransformWarping)
+			else
 			{
 				LoopCounter.AddNonWarpingLevel();
 			}
-			// If StepTransform is warping, we need to add the loop index for it.
-			if (Step.Warping.IsValid())
-			{
-				LoopCounter.AddWarpingLevel(Step.WarpCounter);
-			}
 		}
 
-		CachedInverseTransform = RootToLocalTransform.InverseFromWarp(LoopCounter);
+		CachedInverseTransform = RootToLocalTransform.InverseFromLoop(LoopCounter);
 	}
 
 private:
@@ -182,7 +167,7 @@ private:
 	};
 	TArray<FTransformStep> TransformStack;
 
-	FMovieSceneTimeTransform CachedInverseTransform;
+	FMovieSceneSequenceTransform CachedInverseTransform;
 };
 
 template<typename DataType>
diff --git a/Engine/Source/Runtime/MovieScene/Public/Evaluation/MovieSceneTimeTransform.h b/Engine/Source/Runtime/MovieScene/Public/Evaluation/MovieSceneTimeTransform.h
--- a/Engine/Source/Runtime/MovieScene/Public/Evaluation/MovieSceneTimeTransform.h
+++ b/Engine/Source/Runtime/MovieScene/Public/Evaluation/MovieSceneTimeTransform.h
@@ -65,8 +65,8 @@ struct FMovieSceneTimeTransform
 	 */
 	FMovieSceneTimeTransform Inverse() const
 	{
-		const FFrameTime NewOffset = -Offset / TimeScale;
-		return FMovieSceneTimeTransform(NewOffset, 1.f / TimeScale);
+		const FFrameTime NewOffset = FMath::IsNearlyZero(TimeScale) ? -Offset : -Offset / TimeScale;
+		return FMovieSceneTimeTransform(NewOffset, FMath::IsNearlyZero(TimeScale) ? std::numeric_limits<float>::infinity() : 1.f / TimeScale);
 	}
 
 	/** The sequence's time scale (or play rate) */
@@ -91,6 +91,11 @@ inline FFrameTime operator*(FFrameTime InTime, const FMovieSceneTimeTransform& R
 	{
 		return InTime + RHS.Offset;
 	}
+	else if (!FMath::IsFinite(RHS.TimeScale))
+	{
+		// Special case infinite timescale (inverse of a zero-timescale) and just return the RHS.Offset
+		return RHS.Offset;
+	}
 
 	return InTime * RHS.TimeScale + RHS.Offset;
 }
@@ -116,11 +121,16 @@ inline FFrameTime& operator*=(FFrameTime& InTime, const FMovieSceneTimeTransform
 template<typename T>
 TRange<T> operator*(const TRange<T>& LHS, const FMovieSceneTimeTransform& RHS)
 {
+	//  Special case large (inf) timescale to return open range.
+	if (!FMath::IsFinite(RHS.TimeScale))
+		return TRange<T>(TRangeBound<T>::Open(), TRangeBound<T>::Open());
+
+	// Special case 0-timescale below. We force multiplies with zero timescale to always have inclusive boundaries to prevent no frames from being present.
 	TRangeBound<T> SourceLower = LHS.GetLowerBound();
 	TRangeBound<T> TransformedLower =
 		SourceLower.IsOpen() ? 
 			TRangeBound<T>() : 
-			SourceLower.IsInclusive() ?
+			SourceLower.IsInclusive() || FMath::IsNearlyZero(RHS.TimeScale)  ?
 				TRangeBound<T>::Inclusive(SourceLower.GetValue() * RHS) :
 				TRangeBound<T>::Exclusive(SourceLower.GetValue() * RHS);
 
@@ -128,7 +138,7 @@ TRange<T> operator*(const TRange<T>& LHS, const FMovieSceneTimeTransform& RHS)
 	TRangeBound<T> TransformedUpper =
 		SourceUpper.IsOpen() ? 
 			TRangeBound<T>() : 
-			SourceUpper.IsInclusive() ?
+			SourceUpper.IsInclusive() || FMath::IsNearlyZero(RHS.TimeScale) ?
 				TRangeBound<T>::Inclusive(SourceUpper.GetValue() * RHS) :
 				TRangeBound<T>::Exclusive(SourceUpper.GetValue() * RHS);
 
@@ -137,11 +147,18 @@ TRange<T> operator*(const TRange<T>& LHS, const FMovieSceneTimeTransform& RHS)
 
 inline TRange<FFrameNumber> operator*(const TRange<FFrameNumber>& LHS, const FMovieSceneTimeTransform& RHS)
 {
+	// Special case large (inf) timescale to return open range.
+	if (!FMath::IsFinite(RHS.TimeScale))
+	{
+		return TRange<FFrameNumber>(TRangeBound<FFrameNumber>::Open(), TRangeBound<FFrameNumber>::Open());
+	}
+	
+	// Special case 0-timescale below. We force multiplies with zero timescale to always have inclusive boundaries to prevent no frames from being present.
 	TRangeBound<FFrameNumber> SourceLower = LHS.GetLowerBound();
 	TRangeBound<FFrameNumber> TransformedLower =
 		SourceLower.IsOpen() ? 
 			TRangeBound<FFrameNumber>() : 
-			SourceLower.IsInclusive() ?
+			SourceLower.IsInclusive() || FMath::IsNearlyZero(RHS.TimeScale) ?
 				TRangeBound<FFrameNumber>::Inclusive((SourceLower.GetValue() * RHS).FloorToFrame()) :
 				TRangeBound<FFrameNumber>::Exclusive((SourceLower.GetValue() * RHS).FloorToFrame());
 
@@ -149,7 +166,7 @@ inline TRange<FFrameNumber> operator*(const TRange<FFrameNumber>& LHS, const FMo
 	TRangeBound<FFrameNumber> TransformedUpper =
 		SourceUpper.IsOpen() ? 
 			TRangeBound<FFrameNumber>() : 
-			SourceUpper.IsInclusive() ?
+			SourceUpper.IsInclusive() || FMath::IsNearlyZero(RHS.TimeScale) ?
 				TRangeBound<FFrameNumber>::Inclusive((SourceUpper.GetValue() * RHS).FloorToFrame()) :
 				TRangeBound<FFrameNumber>::Exclusive((SourceUpper.GetValue() * RHS).FloorToFrame());
 
@@ -180,7 +197,7 @@ inline FMovieSceneTimeTransform operator*(const FMovieSceneTimeTransform& LHS, c
 	// | TimeScaleA	, OffsetA	|	.	| TimeScaleB, OffsetB	|
 	// | 0			, 1			|		| 0			, 1			|
 	//
-	const FFrameTime ScaledOffsetRHS = (LHS.TimeScale == 1.f) ? RHS.Offset : RHS.Offset * LHS.TimeScale;
+	const FFrameTime ScaledOffsetRHS = (LHS.TimeScale == 1.f || !FMath::IsFinite(LHS.TimeScale)) ? RHS.Offset : RHS.Offset * LHS.TimeScale;
 	return FMovieSceneTimeTransform(
 		LHS.Offset + ScaledOffsetRHS,		// New Offset
 		RHS.TimeScale * LHS.TimeScale		// New TimeScale
diff --git a/Engine/Source/Runtime/MovieScene/Public/MovieSceneSequenceVisitor.h b/Engine/Source/Runtime/MovieScene/Public/MovieSceneSequenceVisitor.h
--- a/Engine/Source/Runtime/MovieScene/Public/MovieSceneSequenceVisitor.h
+++ b/Engine/Source/Runtime/MovieScene/Public/MovieSceneSequenceVisitor.h
@@ -37,7 +37,7 @@ struct FMovieSceneBinding;
 *					}
 *	
 *					UMovieSceneCameraCutSection* CameraCutSection = CastChecked<UMovieSceneCameraCutSection>(Section);
-*					TRange<FFrameNumber> RootCameraRange = TRange<FFrameNumber>::Intersection(LocalSpace.RootClampRange, CameraCutSection->GetRange() * LocalSpace.RootToSequenceTransform.InverseLinearOnly());
+*					TRange<FFrameNumber> RootCameraRange = TRange<FFrameNumber>::Intersection(LocalSpace.RootClampRange, CameraCutSection->GetRange() * LocalSpace.RootToSequenceTransform.InverseNoLooping());
 *					if (!RootCameraRange.IsEmpty())
 *					{
 *						CameraCutTree.Add(RootCameraRange, MakeTuple(CameraCutSection, LocalSpace.SequenceID, LocalSpace.HierarchicalBias));
diff --git a/Engine/Source/Runtime/MovieSceneTracks/Private/Evaluation/MovieSceneCameraShakeSourceTriggerTemplate.cpp b/Engine/Source/Runtime/MovieSceneTracks/Private/Evaluation/MovieSceneCameraShakeSourceTriggerTemplate.cpp
--- a/Engine/Source/Runtime/MovieSceneTracks/Private/Evaluation/MovieSceneCameraShakeSourceTriggerTemplate.cpp
+++ b/Engine/Source/Runtime/MovieSceneTracks/Private/Evaluation/MovieSceneCameraShakeSourceTriggerTemplate.cpp
@@ -129,7 +129,7 @@ void FMovieSceneCameraShakeSourceTriggerSectionTemplate::EvaluateSwept(const FMo
 	TArray<FMovieSceneCameraShakeSourceTrigger> ShakesToTrigger;
 
 	const int32 TriggerCount = TriggerTimes.Num();
-	const float PositionInSeconds = Context.GetTime() * Context.GetRootToSequenceTransform().InverseLinearOnly() / Context.GetFrameRate();
+	const float PositionInSeconds = Context.GetTime() * Context.GetRootToSequenceTransform().InverseNoLooping() / Context.GetFrameRate();
 
 	for (int32 KeyIndex = 0; KeyIndex < TriggerCount; ++KeyIndex)
 	{
diff --git a/Engine/Source/Runtime/MovieSceneTracks/Private/Evaluation/MovieSceneEventTemplate.cpp b/Engine/Source/Runtime/MovieSceneTracks/Private/Evaluation/MovieSceneEventTemplate.cpp
--- a/Engine/Source/Runtime/MovieSceneTracks/Private/Evaluation/MovieSceneEventTemplate.cpp
+++ b/Engine/Source/Runtime/MovieSceneTracks/Private/Evaluation/MovieSceneEventTemplate.cpp
@@ -176,7 +176,7 @@ void FMovieSceneEventSectionTemplate::EvaluateSwept(const FMovieSceneEvaluationO
 	const int32 Last = bBackwards ? 0 : KeyTimes.Num() - 1;
 	const int32 Inc = bBackwards ? -1 : 1;
 
-	const float PositionInSeconds = Context.GetTime() * Context.GetRootToSequenceTransform().InverseLinearOnly() / Context.GetFrameRate();
+	const float PositionInSeconds = Context.GetTime() * Context.GetRootToSequenceTransform().InverseNoLooping() / Context.GetFrameRate();
 
 	if (bBackwards)
 	{
diff --git a/Engine/Source/Runtime/MovieSceneTracks/Private/Sections/MovieSceneEventRepeaterSection.cpp b/Engine/Source/Runtime/MovieSceneTracks/Private/Sections/MovieSceneEventRepeaterSection.cpp
--- a/Engine/Source/Runtime/MovieSceneTracks/Private/Sections/MovieSceneEventRepeaterSection.cpp
+++ b/Engine/Source/Runtime/MovieSceneTracks/Private/Sections/MovieSceneEventRepeaterSection.cpp
@@ -65,7 +65,7 @@ void UMovieSceneEventRepeaterSection::ImportEntityImpl(UMovieSceneEntitySystemLi
 		Event.Ptrs,
 		Params.GetObjectBindingID(),
 		ThisInstance.GetSequenceID(),
-		Context.GetTime() * Context.GetSequenceToRootTransform()
+		Context.GetTime() * Context.GetSequenceToRootSequenceTransform()
 	};
 
 	EventSystem->AddEvent(ThisInstance.GetRootInstanceHandle(), TriggerData);
diff --git a/Engine/Source/Runtime/MovieSceneTracks/Private/Sections/MovieSceneEventTriggerSection.cpp b/Engine/Source/Runtime/MovieSceneTracks/Private/Sections/MovieSceneEventTriggerSection.cpp
--- a/Engine/Source/Runtime/MovieSceneTracks/Private/Sections/MovieSceneEventTriggerSection.cpp
+++ b/Engine/Source/Runtime/MovieSceneTracks/Private/Sections/MovieSceneEventTriggerSection.cpp
@@ -84,7 +84,7 @@ void UMovieSceneEventTriggerSection::ImportEntityImpl(UMovieSceneEntitySystemLin
 		Events[EventIndex].Ptrs,
 		Params.GetObjectBindingID(),
 		ThisInstance.GetSequenceID(),
-		Times[EventIndex] * Context.GetSequenceToRootTransform()
+		Times[EventIndex] * Context.GetSequenceToRootSequenceTransform()
 	};
 
 	EventSystem->AddEvent(ThisInstance.GetRootInstanceHandle(), TriggerData);
diff --git a/Engine/Source/Runtime/MovieSceneTracks/Private/TrackInstances/MovieSceneCameraCutTrackInstance.cpp b/Engine/Source/Runtime/MovieSceneTracks/Private/TrackInstances/MovieSceneCameraCutTrackInstance.cpp
--- a/Engine/Source/Runtime/MovieSceneTracks/Private/TrackInstances/MovieSceneCameraCutTrackInstance.cpp
+++ b/Engine/Source/Runtime/MovieSceneTracks/Private/TrackInstances/MovieSceneCameraCutTrackInstance.cpp
@@ -247,7 +247,7 @@ void UMovieSceneCameraCutTrackInstance::OnAnimate()
 		else
 		{
 			const UMovieSceneCameraCutTrack* Track = Section->GetTypedOuter<UMovieSceneCameraCutTrack>();
-			const FMovieSceneTimeTransform SequenceToRootTransform = Context.GetSequenceToRootTransform();
+			const FMovieSceneSequenceTransform SequenceToRootTransform = Context.GetSequenceToRootSequenceTransform();
 
 			FBlendedCameraCut Params(Input, CameraBindingID, SequenceInstance.GetSequenceID());
 			Params.bCanBlend = Track->bCanBlend;
@@ -264,13 +264,13 @@ void UMovieSceneCameraCutTrackInstance::OnAnimate()
 			if (Section->HasStartFrame() && Section->Easing.GetEaseInDuration() > 0)
 			{
 				Params.LocalEaseInEndTime = Params.LocalStartTime + Section->Easing.GetEaseInDuration();
-				const float RootEaseInTime = SequenceToRootTransform.TimeScale * Context.GetFrameRate().AsSeconds(FFrameNumber(Section->Easing.GetEaseInDuration()));
+				const float RootEaseInTime = SequenceToRootTransform.GetTimeScale() * Context.GetFrameRate().AsSeconds(FFrameNumber(Section->Easing.GetEaseInDuration()));
 				Params.EaseIn = FBlendedCameraCutEasingInfo(RootEaseInTime, Section->Easing.EaseIn);
 			}
 			if (Section->HasEndFrame() && Section->Easing.GetEaseOutDuration() > 0)
 			{
 				Params.LocalEaseOutStartTime = Params.LocalEndTime - Section->Easing.GetEaseOutDuration();
-				const float RootEaseOutTime = SequenceToRootTransform.TimeScale * Context.GetFrameRate().AsSeconds(FFrameNumber(Section->Easing.GetEaseOutDuration()));
+				const float RootEaseOutTime = SequenceToRootTransform.GetTimeScale() * Context.GetFrameRate().AsSeconds(FFrameNumber(Section->Easing.GetEaseOutDuration()));
 				Params.EaseOut = FBlendedCameraCutEasingInfo(RootEaseOutTime, Section->Easing.EaseOut);
 			}
 
