Sequencer: Fixed race condition caused by async children being able to signal the main thread after completion while a worker thread was finishing iterating its children

The race here stems from the fact that a child task signals its parent for completion if it was the last one. In doing so, it could also signal the main thread to finish waiting on tasks, if it was the last task to be executed. This is problematic because there is still a thread iterating the child tasks at the same time as the last one is completed, resulting in a race between the iteration finishing successfully, and the main thread exiting out of our loop and potentially reallocating the array.

The fix is to prevent children from completing their parents while the loop is running by adding an artificial count for the iteration. If there are no children outstanding after the loop, CompleteTask is called for the Parent.

diff --git a/Engine/Source/Runtime/MovieScene/Private/EntitySystem/MovieSceneTaskScheduler.cpp b/Engine/Source/Runtime/MovieScene/Private/EntitySystem/MovieSceneTaskScheduler.cpp
--- a/Engine/Source/Runtime/MovieScene/Private/EntitySystem/MovieSceneTaskScheduler.cpp
+++ b/Engine/Source/Runtime/MovieScene/Private/EntitySystem/MovieSceneTaskScheduler.cpp
@@ -132,10 +132,25 @@ void FScheduledTask::Run(const FEntitySystemScheduler* Scheduler, FTaskExecution
 	//     this happens in FEntitySystemScheduler::CompleteTask if Parent is valid
 	if (!ChildTasks.IsEmpty())
 	{
-		 for (int32 Index : ChildTasks)
-		 {
-		 	Scheduler->PrerequisiteCompleted(FTaskID(Index), nullptr);
-		 }
+		// Increment the child completion count to protect CompleteTask being called for _this_ parent task
+		// while the loop over ChildTasks is running. This prevents a race condition where the final child can end up being
+		// the last task altogether, which triggers OnAllTasksFinished, potentially allowing the waiting thread to continue
+		// and destroy or  otherwise mutate the contents of FEntitySystemScheduler resulting in a crash.
+		//
+		// Once our loop has finished we check the child complete count to see if this was the last one
+		ChildCompleteCount.fetch_add(1);
+
+		for (int32 Index : ChildTasks)
+		{
+			Scheduler->PrerequisiteCompleted(FTaskID(Index), nullptr);
+		}
+
+		// Subtract our count added on ln 205. If this is the last count, complete this task (all children have completed)
+		const int32 PreviousCompleteCount = ChildCompleteCount.fetch_sub(1);
+		if (PreviousCompleteCount == 1)
+		{
+			Scheduler->CompleteTask(this, InFlags);
+		}
 	}
 	else
 	{
