TObjectPtr compilation fixes ahead of enabling GC barrier

diff --git a/Engine/Source/Runtime/MovieScene/Public/EntitySystem/MovieSceneOverlappingEntityTracker.h b/Engine/Source/Runtime/MovieScene/Public/EntitySystem/MovieSceneOverlappingEntityTracker.h
--- a/Engine/Source/Runtime/MovieScene/Public/EntitySystem/MovieSceneOverlappingEntityTracker.h
+++ b/Engine/Source/Runtime/MovieScene/Public/EntitySystem/MovieSceneOverlappingEntityTracker.h
@@ -49,7 +49,10 @@ struct FGarbageTraits
 	{
 		return FBuiltInComponentTypes::IsBoundObjectGarbage(InObject);
 	}
-
+	//FORCEINLINE static bool IsGarbage(FObjectComponent& InComponent)
+	//{
+	//	return FBuiltInComponentTypes::IsBoundObjectGarbage(InComponent.GetObject());
+	//}
 	
 	template<typename T>
 	FORCEINLINE static void AddReferencedObjects(FReferenceCollector& ReferenceCollector, T* In)
