Sequencer: Prevent evaluation phase prerequisites from being propagated to the scheduling phase if the custom scheduler is enabled

diff --git a/Engine/Source/Runtime/MovieScene/Private/EntitySystem/MovieSceneEntitySystemGraphs.cpp b/Engine/Source/Runtime/MovieScene/Private/EntitySystem/MovieSceneEntitySystemGraphs.cpp
--- a/Engine/Source/Runtime/MovieScene/Private/EntitySystem/MovieSceneEntitySystemGraphs.cpp
+++ b/Engine/Source/Runtime/MovieScene/Private/EntitySystem/MovieSceneEntitySystemGraphs.cpp
@@ -382,6 +382,8 @@ void FMovieSceneEntitySystemGraph::ExecutePhase(UE::MovieScene::ESystemPhase Pha
 {
 	using namespace UE::MovieScene;
 
+	const bool bCustomSchedulingEnabled = FEntitySystemScheduler::IsCustomSchedulingEnabled();
+
 	Linker->EntityManager.UpdateThreadingModel();
 
 	const EEntityThreadingModel ThreadingModel = Linker->EntityManager.GetThreadingModel();
@@ -462,7 +464,8 @@ void FMovieSceneEntitySystemGraph::ExecutePhase(UE::MovieScene::ESystemPhase Pha
 				{
 					FMovieSceneEntitySystemGraphNode& ToNode = Nodes.Array[*ToNodeID];
 					if (EnumHasAnyFlags(ToNode.System->GetPhase(), Phase)
-						|| (Phase == ESystemPhase::Evaluation && EnumHasAnyFlags(ToNode.System->GetPhase(), ESystemPhase::Scheduling))
+						// If custom scheduling is disabled, allow propagation between evaluation/scheduling phase as well
+						|| (!bCustomSchedulingEnabled && Phase == ESystemPhase::Evaluation && EnumHasAnyFlags(ToNode.System->GetPhase(), ESystemPhase::Scheduling))
 						)
 					{
 						if (!ToNode.Prerequisites)
@@ -474,7 +477,7 @@ void FMovieSceneEntitySystemGraph::ExecutePhase(UE::MovieScene::ESystemPhase Pha
 				}
 			}
 
-			// Done wtih subsequents now
+			// Done with subsequents now
 			DownstreamTasks.Subsequents->Empty();
 		}
 	}
