Sequencer: Fixed dissection algorithm where multiple different instances provided dissected updates for the same frame.

The algorithm assumes a single array sorted by 'order', where the most recent updates were ordered in the array last. Given sequences A, B and C that provided 3, 3, and 2 dissections respectively, we might expect to see an array of:

[ A0, B0, C0, A1, B1, A_maxint, B_maxint, C_maxint ], which we would evaluate grouped by their order.

Unfortunately when there are multiple sequences that provide dissection, this array ends up sorted by Sequence first _then_ order. This means we actually see the following:

[ A0, A1, A_maxint, B0, B1, B_maxint, C0, C_maxint]

This can result in excessive updates when you have lots of sequences that provide dissection on the same frame, because the update loop early exits as soon as it finds a mismatching order (we essentially flush for every single entry in that array).

The fix is to ensure the dissected update array is sorted correctly. Ultimately we might look at refactoring this code to be more performant, but for now this is a robust stop-gap.

diff --git a/Engine/Source/Runtime/MovieScene/Private/EntitySystem/MovieSceneEntitySystemRunner.cpp b/Engine/Source/Runtime/MovieScene/Private/EntitySystem/MovieSceneEntitySystemRunner.cpp
--- a/Engine/Source/Runtime/MovieScene/Private/EntitySystem/MovieSceneEntitySystemRunner.cpp
+++ b/Engine/Source/Runtime/MovieScene/Private/EntitySystem/MovieSceneEntitySystemRunner.cpp
@@ -11,6 +11,7 @@
 #include "IMovieScenePlayer.h"
 #include "MovieSceneSequence.h"
 #include "Algo/Reverse.h"
+#include "Algo/Sort.h"
 #include "ProfilingDebugging/CountersTrace.h"
 
 DECLARE_CYCLE_STAT(TEXT("Runner Flush"), 				MovieSceneEval_RunnerFlush, 				STATGROUP_MovieSceneEval);
@@ -671,6 +672,8 @@ bool FMovieSceneEntitySystemRunner::GameThread_UpdateSequenceInstances(UMovieSce
 				MarkForUpdate(Request.Params.InstanceHandle, Request.Params.UpdateFlags);
 			}
 		}
+
+		Algo::SortBy(DissectedUpdates, &FDissectedUpdate::Order);
 	}
 	else
 	{
@@ -1101,6 +1104,7 @@ void FMovieSceneEntitySystemRunner::GameThread_PostEvaluationPhase(UMovieSceneEn
 	if (TmpDissectedUpdates.Num() > 0)
 	{
 		DissectedUpdates.Append(TmpDissectedUpdates);
+		// Important: we do not sort the Dissected updates here to ensure that the previously populated entries are evaluated first
 	}
 
 	// If we have any pending updates, we need to run another evaluation.
