Sequencer: Add GC references to linkers from the tick manager

Ordinarily these linkers are always kept alive by the tick clients that use them so we don't encounter a dangling ptr until the only tick client in the group is using blocking evaluation. In this instance it creates and maintains its own linker, and causes the one from the tick manager to get collected.

I made this change rather than having to check for blocking evaluation and have a potentially null linker in the group since that would create a lot of additional complexity for a feature that is rarely used (and one we don't really want to support longer-term)

diff --git a/Engine/Source/Runtime/MovieScene/Private/MovieSceneSequenceTickManager.cpp b/Engine/Source/Runtime/MovieScene/Private/MovieSceneSequenceTickManager.cpp
--- a/Engine/Source/Runtime/MovieScene/Private/MovieSceneSequenceTickManager.cpp
+++ b/Engine/Source/Runtime/MovieScene/Private/MovieSceneSequenceTickManager.cpp
@@ -60,6 +60,15 @@ UMovieSceneSequenceTickManager::UMovieSceneSequenceTickManager(const FObjectInit
 	PendingActorOperations = nullptr;
 }
 
+void UMovieSceneSequenceTickManager::AddReferencedObjects(UObject* InThis, FReferenceCollector& Collector)
+{
+	UMovieSceneSequenceTickManager* This = CastChecked<UMovieSceneSequenceTickManager>(InThis);
+	for (FLinkerGroup& Group : This->LinkerGroups)
+	{
+		Collector.AddReferencedObject(Group.Linker);
+	}
+}
+
 UMovieSceneEntitySystemLinker* UMovieSceneSequenceTickManager::GetLinker(const FMovieSceneSequenceTickInterval& TickInterval)
 {
 	const int32 RoundedTickIntervalMs = TickInterval.RoundTickIntervalMs();
diff --git a/Engine/Source/Runtime/MovieScene/Public/MovieSceneSequenceTickManager.h b/Engine/Source/Runtime/MovieScene/Public/MovieSceneSequenceTickManager.h
--- a/Engine/Source/Runtime/MovieScene/Public/MovieSceneSequenceTickManager.h
+++ b/Engine/Source/Runtime/MovieScene/Public/MovieSceneSequenceTickManager.h
@@ -98,6 +98,8 @@ public:
 	MOVIESCENE_API void UnregisterSequenceActor(AActor* InActor, TScriptInterface<IMovieSceneSequenceTickManagerClient> InActorInterface);
 	/*~ -------------------------------------------------------------- */
 
+	static MOVIESCENE_API void AddReferencedObjects(UObject* InThis, FReferenceCollector& Collector);
+
 private:
 
 	/*~ UObject interface */
@@ -122,7 +124,7 @@ private:
 	struct FLinkerGroup
 	{
 		/** The linker that owns all the entity data, systems and instances */
-		UMovieSceneEntitySystemLinker* Linker;
+		TObjectPtr<UMovieSceneEntitySystemLinker> Linker;
 		/** Runner responsible for evaluating each phase of the pipeline */
 		TSharedPtr<FMovieSceneEntitySystemRunner> Runner;
 
