Sequencer: Ensure that evaluation template instances are fully torn down on EndPlay

  - The only meaningful change here is to call TearDown (formerly BeginDestroy) from UMovieSceneSequencePlayer::TearDown.
  - Renamed BeginDestroy to TearDown to give a better indication of its application

diff --git a/Engine/Plugins/MovieScene/TemplateSequence/Source/TemplateSequence/Private/CameraAnimationSequencePlayer.cpp b/Engine/Plugins/MovieScene/TemplateSequence/Source/TemplateSequence/Private/CameraAnimationSequencePlayer.cpp
--- a/Engine/Plugins/MovieScene/TemplateSequence/Source/TemplateSequence/Private/CameraAnimationSequencePlayer.cpp
+++ b/Engine/Plugins/MovieScene/TemplateSequence/Source/TemplateSequence/Private/CameraAnimationSequencePlayer.cpp
@@ -255,7 +255,7 @@ UCameraAnimationSequencePlayer::~UCameraAnimationSequencePlayer()
 
 void UCameraAnimationSequencePlayer::BeginDestroy()
 {
-	RootTemplateInstance.BeginDestroy();
+	RootTemplateInstance.TearDown();
 
 	Super::BeginDestroy();
 }
diff --git a/Engine/Source/Runtime/MovieScene/Private/Evaluation/MovieSceneEvaluationTemplateInstance.cpp b/Engine/Source/Runtime/MovieScene/Private/Evaluation/MovieSceneEvaluationTemplateInstance.cpp
--- a/Engine/Source/Runtime/MovieScene/Private/Evaluation/MovieSceneEvaluationTemplateInstance.cpp
+++ b/Engine/Source/Runtime/MovieScene/Private/Evaluation/MovieSceneEvaluationTemplateInstance.cpp
@@ -30,8 +30,9 @@ FMovieSceneRootEvaluationTemplateInstance::FMovieSceneRootEvaluationTemplateInst
 #endif
 }
 
-void FMovieSceneRootEvaluationTemplateInstance::BeginDestroy()
+void FMovieSceneRootEvaluationTemplateInstance::TearDown()
 {
+	// Avoid redundant work if the linker is being destroyed anyway
 	if (EntitySystemLinker && IsValidChecked(EntitySystemLinker) && !EntitySystemLinker->IsUnreachable() && !EntitySystemLinker->HasAnyFlags(RF_BeginDestroyed))
 	{
 		if (TSharedPtr<FMovieSceneEntitySystemRunner> Runner = WeakRunner.Pin())
@@ -50,7 +51,7 @@ void FMovieSceneRootEvaluationTemplateInstance::BeginDestroy()
 
 FMovieSceneRootEvaluationTemplateInstance::~FMovieSceneRootEvaluationTemplateInstance()
 {
-	BeginDestroy();
+	TearDown();
 }
 
 UMovieSceneEntitySystemLinker* FMovieSceneRootEvaluationTemplateInstance::ConstructEntityLinker(IMovieScenePlayer& Player)
diff --git a/Engine/Source/Runtime/MovieScene/Private/MovieSceneSequencePlayer.cpp b/Engine/Source/Runtime/MovieScene/Private/MovieSceneSequencePlayer.cpp
--- a/Engine/Source/Runtime/MovieScene/Private/MovieSceneSequencePlayer.cpp
+++ b/Engine/Source/Runtime/MovieScene/Private/MovieSceneSequencePlayer.cpp
@@ -886,7 +886,7 @@ void UMovieSceneSequencePlayer::Initialize(UMovieSceneSequence* InSequence)
 	// If we were already registered with a different Tick Interval we need to re-register with the new one, which involves tearing everything down and setting up a new instance
 	else if (RegisteredTickInterval.GetValue() != TickInterval)
 	{
-		RootTemplateInstance.BeginDestroy();
+		RootTemplateInstance.TearDown();
 		TickManager->UnregisterTickClient(this);
 
 		TickManager->RegisterTickClient(RegisteredTickInterval.GetValue(), this);
@@ -1249,6 +1249,8 @@ void UMovieSceneSequencePlayer::UpdateMovieSceneInstance(FMovieSceneEvaluationRa
 
 void UMovieSceneSequencePlayer::TearDown()
 {
+	RootTemplateInstance.TearDown();
+
 	if (TickManager)
 	{
 		TickManager->UnregisterTickClient(this);
@@ -1870,7 +1872,7 @@ float UMovieSceneSequencePlayer::GetPing() const
 
 void UMovieSceneSequencePlayer::BeginDestroy()
 {
-	RootTemplateInstance.BeginDestroy();
+	RootTemplateInstance.TearDown();
 
 	TearDown();
 
diff --git a/Engine/Source/Runtime/MovieScene/Public/Evaluation/MovieSceneEvaluationTemplateInstance.h b/Engine/Source/Runtime/MovieScene/Public/Evaluation/MovieSceneEvaluationTemplateInstance.h
--- a/Engine/Source/Runtime/MovieScene/Public/Evaluation/MovieSceneEvaluationTemplateInstance.h
+++ b/Engine/Source/Runtime/MovieScene/Public/Evaluation/MovieSceneEvaluationTemplateInstance.h
@@ -123,7 +123,13 @@ public:
 
 	MOVIESCENE_API void CopyActuators(FMovieSceneBlendingAccumulator& Accumulator) const;
 
-	MOVIESCENE_API void BeginDestroy();
+	UE_DEPRECATED(5.3, "Please call TearDown instead.")
+	void BeginDestroy()
+	{
+		TearDown();
+	}
+
+	MOVIESCENE_API void TearDown();
 
 	MOVIESCENE_API void EnableGlobalPreAnimatedStateCapture();
 
diff --git a/Engine/Source/Runtime/UMG/Private/Animation/UMGSequencePlayer.cpp b/Engine/Source/Runtime/UMG/Private/Animation/UMGSequencePlayer.cpp
--- a/Engine/Source/Runtime/UMG/Private/Animation/UMGSequencePlayer.cpp
+++ b/Engine/Source/Runtime/UMG/Private/Animation/UMGSequencePlayer.cpp
@@ -574,12 +574,12 @@ void UUMGSequencePlayer::RemoveEvaluationData()
 
 void UUMGSequencePlayer::TearDown()
 {
-	RootTemplateInstance.BeginDestroy();
+	RootTemplateInstance.TearDown();
 }
 
 void UUMGSequencePlayer::BeginDestroy()
 {
-	RootTemplateInstance.BeginDestroy();
+	RootTemplateInstance.TearDown();
 
 	// Remove any latent actions added by this player.
 	if (CVarUserWidgetUseParallelAnimation.GetValueOnGameThread())
