Sequencer: Destroying a spawned object should mark cleaning up so that the sequence package doesn't mark itself dirty

This is similar to 24472100 #jira UE-178446 which was to address the package getting marked dirty when closing the sequence.

- Tested that making modifications to a spawnable, still invoke OnObjectModified and cause the sequence to be marked dirty.
- Tested that making no modifications to spawnables, closing the sequence or saving the sequence, do NOT cause the sequence to be marked dirty.

diff --git a/Engine/Source/Runtime/MovieScene/Private/MovieSceneSpawnRegister.cpp b/Engine/Source/Runtime/MovieScene/Private/MovieSceneSpawnRegister.cpp
--- a/Engine/Source/Runtime/MovieScene/Private/MovieSceneSpawnRegister.cpp
+++ b/Engine/Source/Runtime/MovieScene/Private/MovieSceneSpawnRegister.cpp
@@ -76,6 +76,8 @@ UObject* FMovieSceneSpawnRegister::SpawnObject(const FGuid& BindingId, UMovieSce
 
 bool FMovieSceneSpawnRegister::DestroySpawnedObject(const FGuid& BindingId, FMovieSceneSequenceIDRef TemplateID, IMovieScenePlayer& Player)
 {
+	TGuardValue<bool> CleaningUp(bCleaningUp, true);
+
 	FMovieSceneSpawnRegisterKey Key(TemplateID, BindingId);
 	
 	FSpawnedObject* Existing = Register.Find(Key);
