Sequencer: don't run sequence updates when a client late-joins a level with stopped sequences

diff --git a/Engine/Source/Runtime/MovieScene/Private/MovieSceneSequencePlayer.cpp b/Engine/Source/Runtime/MovieScene/Private/MovieSceneSequencePlayer.cpp
--- a/Engine/Source/Runtime/MovieScene/Private/MovieSceneSequencePlayer.cpp
+++ b/Engine/Source/Runtime/MovieScene/Private/MovieSceneSequencePlayer.cpp
@@ -1740,14 +1740,17 @@ void UMovieSceneSequencePlayer::PostNetReceive()
 					bUpdateNetSync = true;
 				}
 			}
-			else if (Status == EMovieScenePlayerStatus::Stopped)
-			{
-				SetPlaybackPosition(FMovieSceneSequencePlaybackParams(NetSyncProps.LastKnownPosition, EUpdatePositionMethod::Jump));
-			}
 			else if (Status == EMovieScenePlayerStatus::Scrubbing)
 			{
+				// Scrub to the new position.
 				SetPlaybackPosition(FMovieSceneSequencePlaybackParams(NetSyncProps.LastKnownPosition, EUpdatePositionMethod::Scrub));
 			}
+			else if (Status == EMovieScenePlayerStatus::Stopped)
+			{
+				// Both client and server are stopped so just update our (client) position to match the server's.
+				UpdatePlayPosition(PlayPosition, NetSyncProps.LastKnownPosition, EUpdatePositionMethod::Jump);
+				TimeController->Reset(GetCurrentTime());
+			}
 
 			bIsAsyncUpdate = false;
 		}
