Sequencer: be more lenient in editor about invalid dynamic binding resolvers

Undo operations can cause immediate re-evaluation of the sequence, including resolving object bindings. This will always happen before the user can re-compile the Blueprint where the dynamic binding resolvers are, so we might get invalid method signatures there.

diff --git a/Engine/Source/Runtime/MovieScene/Private/MovieSceneDynamicBindingInvoker.cpp b/Engine/Source/Runtime/MovieScene/Private/MovieSceneDynamicBindingInvoker.cpp
--- a/Engine/Source/Runtime/MovieScene/Private/MovieSceneDynamicBindingInvoker.cpp
+++ b/Engine/Source/Runtime/MovieScene/Private/MovieSceneDynamicBindingInvoker.cpp
@@ -67,9 +67,16 @@ FMovieSceneDynamicBindingResolveResult FMovieSceneDynamicBindingInvoker::Resolve
 
 FMovieSceneDynamicBindingResolveResult FMovieSceneDynamicBindingInvoker::InvokeDynamicBinding(UObject* DirectorInstance, const FMovieSceneDynamicBinding& DynamicBinding, const FMovieSceneDynamicBindingResolveParams& ResolveParams)
 {
-	// Parse all function parameters.
+	FMovieSceneDynamicBindingResolveResult Result;
+
+	// Do some basic checks.
 	UFunction* DynamicBindingFunc = DynamicBinding.Function.Get();
-	check(DynamicBindingFunc);
+	if (!ensure(DynamicBindingFunc))
+	{ 
+		return Result;
+	}
+
+	// Parse all function parameters.
 	uint8* Parameters = (uint8*)FMemory_Alloca(DynamicBindingFunc->ParmsSize + DynamicBindingFunc->MinAlignment);
 	Parameters = Align(Parameters, DynamicBindingFunc->MinAlignment);
 
@@ -100,14 +107,19 @@ FMovieSceneDynamicBindingResolveResult FMovieSceneDynamicBindingInvoker::InvokeD
 		ResolveParamsProp->SetValue_InContainer(Parameters, &ResolveParams);
 	}
 
-	// Invoke the function.
-	DirectorInstance->ProcessEvent(DynamicBindingFunc, Parameters);
-
-	// Grab the result value.
-	FMovieSceneDynamicBindingResolveResult Result;
+#if WITH_EDITOR
+	// In the editor we need to be more forgiving, because we might have temporarily invalid states, such as
+	// when undo-ing operations.
+	if (ReturnProp != nullptr && ReturnProp->Struct == FMovieSceneDynamicBindingResolveResult::StaticStruct())
+#else
 	if (ensureMsgf(ReturnProp != nullptr && ReturnProp->Struct == FMovieSceneDynamicBindingResolveResult::StaticStruct(),
-			TEXT("The dynamic binding resolver function has no return value of type FMovieSceneDynamicBindingResolveResult")))
+		TEXT("The dynamic binding resolver function has no return value of type FMovieSceneDynamicBindingResolveResult")))
+#endif
 	{
+		// Invoke the function.
+		DirectorInstance->ProcessEvent(DynamicBindingFunc, Parameters);
+
+		// Grab the result value.
 		ReturnProp->GetValue_InContainer(Parameters, static_cast<void*>(&Result));
 	}
 
