Fix CIS

Sequencer: Outer the entity system linker to the transient package

Add SetWorld/GetWorld on the linker so that systems would get the player context world instead of the linker world, which is the transient package.

#rnx

diff --git a/Engine/Source/Runtime/MovieScene/Public/EntitySystem/MovieSceneEntitySystemLinker.h b/Engine/Source/Runtime/MovieScene/Public/EntitySystem/MovieSceneEntitySystemLinker.h
--- a/Engine/Source/Runtime/MovieScene/Public/EntitySystem/MovieSceneEntitySystemLinker.h
+++ b/Engine/Source/Runtime/MovieScene/Public/EntitySystem/MovieSceneEntitySystemLinker.h
@@ -4,6 +4,7 @@
 
 #include "MovieSceneEntityIDs.h"
 #include "MovieSceneSequenceID.h"
+#include "Engine/World.h"
 #include "Evaluation/MovieScenePlayback.h"
 #include "EntitySystem/MovieSceneEntityManager.h"
 #include "EntitySystem/MovieSceneInstanceRegistry.h"
