Sequencer: Added SetCompletionModeOverride and GetCompletionModeOverride

diff --git a/Engine/Source/Runtime/MovieScene/Private/MovieSceneSequencePlayer.cpp b/Engine/Source/Runtime/MovieScene/Private/MovieSceneSequencePlayer.cpp
--- a/Engine/Source/Runtime/MovieScene/Private/MovieSceneSequencePlayer.cpp
+++ b/Engine/Source/Runtime/MovieScene/Private/MovieSceneSequencePlayer.cpp
@@ -644,6 +644,21 @@ void UMovieSceneSequencePlayer::RestoreState()
 	RestorePreAnimatedState();
 }
 
+void UMovieSceneSequencePlayer::SetCompletionModeOverride(EMovieSceneCompletionModeOverride CompletionModeOverride)
+{
+	if (IsPlaying() && PlaybackSettings.FinishCompletionStateOverride != EMovieSceneCompletionModeOverride::ForceRestoreState && CompletionModeOverride == EMovieSceneCompletionModeOverride::ForceRestoreState)
+	{
+		UE_LOG(LogMovieScene, Warning, TEXT("Attempting to set completion mode override to force restore state while the sequence is already playing. Force restore state must be set before starting playback."));
+	}
+
+	PlaybackSettings.FinishCompletionStateOverride = CompletionModeOverride;
+}
+
+EMovieSceneCompletionModeOverride UMovieSceneSequencePlayer::GetCompletionModeOverride() const
+{
+	return PlaybackSettings.FinishCompletionStateOverride;
+}
+
 bool UMovieSceneSequencePlayer::IsPlaying() const
 {
 	return Status == EMovieScenePlayerStatus::Playing;
diff --git a/Engine/Source/Runtime/MovieScene/Public/MovieSceneSequencePlayer.h b/Engine/Source/Runtime/MovieScene/Public/MovieSceneSequencePlayer.h
--- a/Engine/Source/Runtime/MovieScene/Public/MovieSceneSequencePlayer.h
+++ b/Engine/Source/Runtime/MovieScene/Public/MovieSceneSequencePlayer.h
@@ -13,6 +13,7 @@
 #include "MovieSceneTimeController.h"
 #include "Evaluation/MovieScenePlayback.h"
 #include "Evaluation/MovieScenePlayback.h"
+#include "MovieSceneSequencePlaybackSettings.h"
 #include "MovieSceneSequenceTickManagerClient.h"
 #include "MovieSceneSequencePlaybackSettings.h"
 #include "MovieSceneLatentActionManager.h"
@@ -328,9 +329,17 @@ public:
 	/**
 	 * Restore any changes made by this player to their original state
 	 */
-	UFUNCTION(BlueprintCallable, Category="Game|Cinematic")
+	UFUNCTION(BlueprintCallable, Category = "Game|Cinematic")
 	MOVIESCENE_API void RestoreState();
 
+	/** Set the state of the completion mode override. Note, setting the state to force restore state will only take effect if the sequence hasn't started playing */
+	UFUNCTION(BlueprintCallable, Category = "Game|Cinematic")
+	MOVIESCENE_API void SetCompletionModeOverride(EMovieSceneCompletionModeOverride CompletionModeOverride);
+
+	/** Get the state of the completion mode override */
+	UFUNCTION(BlueprintCallable, Category = "Game|Cinematic")
+	MOVIESCENE_API EMovieSceneCompletionModeOverride GetCompletionModeOverride() const;
+
 public:
 
 	/** Check whether the sequence is actively playing. */
