CISUnity Fix, missing headers
#trivial

diff --git a/Engine/Source/Runtime/MovieScene/Public/Channels/MovieSceneBoolChannel.h b/Engine/Source/Runtime/MovieScene/Public/Channels/MovieSceneBoolChannel.h
--- a/Engine/Source/Runtime/MovieScene/Public/Channels/MovieSceneBoolChannel.h
+++ b/Engine/Source/Runtime/MovieScene/Public/Channels/MovieSceneBoolChannel.h
@@ -18,6 +18,7 @@
 #include "Templates/UnrealTemplate.h"
 #include "UObject/Class.h"
 #include "UObject/ObjectMacros.h"
+#include "Curves/RealCurve.h"
 
 #include "MovieSceneBoolChannel.generated.h"
 
diff --git a/Engine/Source/Runtime/MovieScene/Public/Channels/MovieSceneIntegerChannel.h b/Engine/Source/Runtime/MovieScene/Public/Channels/MovieSceneIntegerChannel.h
--- a/Engine/Source/Runtime/MovieScene/Public/Channels/MovieSceneIntegerChannel.h
+++ b/Engine/Source/Runtime/MovieScene/Public/Channels/MovieSceneIntegerChannel.h
@@ -19,6 +19,7 @@
 #include "Templates/UnrealTemplate.h"
 #include "UObject/Class.h"
 #include "UObject/ObjectMacros.h"
+#include "Curves/RealCurve.h"
 
 #include "MovieSceneIntegerChannel.generated.h"
 
