Increase the maximum possible size of TInlineValue<T>

diff --git a/Engine/Source/Runtime/Core/Public/Misc/InlineValue.h b/Engine/Source/Runtime/Core/Public/Misc/InlineValue.h
--- a/Engine/Source/Runtime/Core/Public/Misc/InlineValue.h
+++ b/Engine/Source/Runtime/Core/Public/Misc/InlineValue.h
@@ -17,7 +17,7 @@
  *
  * Can be viewed as a TUniquePtr with a small allocation optimization.
  */
-template<typename BaseType, uint8 DesiredMaxInlineSize=64, uint8 DefaultAlignment=8>
+template<typename BaseType, uint16 DesiredMaxInlineSize=64, uint8 DefaultAlignment=8>
 class TInlineValue
 {
 public:
diff --git a/Engine/Source/Runtime/MovieScene/Private/Evaluation/MovieSceneEvalTemplateSerializer.h b/Engine/Source/Runtime/MovieScene/Private/Evaluation/MovieSceneEvalTemplateSerializer.h
--- a/Engine/Source/Runtime/MovieScene/Private/Evaluation/MovieSceneEvalTemplateSerializer.h
+++ b/Engine/Source/Runtime/MovieScene/Private/Evaluation/MovieSceneEvalTemplateSerializer.h
@@ -11,7 +11,7 @@
 struct FMovieSceneEmptyStruct;
 
 /** Serialize an inline value that has a GetScriptStruct method */
-template<typename T, uint8 N>
+template<typename T, uint16 N>
 bool SerializeInlineValue(TInlineValue<T, N>& Impl, FArchive& Ar, bool bWarnOnError)
 {
 	Ar.UsingCustomVersion(FMovieSceneEvaluationCustomVersion::GUID);
