Sequencer: add grouping system to help with the common pattern of multiple contributors being grouped when animated the same target

The goal of this change is to replace the ad-hoc grouping and tracking that the property instantiator does with some first-class concept of "grouping", handled by a new system. In turn, this lets the hierarhical bias system look at those generic groups, instead of specifically looking at property blend channels. In the long run, this will let the hierarchical bias system also work with groups of, say, material parameter tracks and material switchers and so on.

diff --git a/Engine/Source/Runtime/MovieScene/Private/EntitySystem/BuiltInComponentTypes.cpp b/Engine/Source/Runtime/MovieScene/Private/EntitySystem/BuiltInComponentTypes.cpp
--- a/Engine/Source/Runtime/MovieScene/Private/EntitySystem/BuiltInComponentTypes.cpp
+++ b/Engine/Source/Runtime/MovieScene/Private/EntitySystem/BuiltInComponentTypes.cpp
@@ -72,6 +72,8 @@ FBuiltInComponentTypes::FBuiltInComponentTypes()
 	ComponentRegistry->NewComponentType(&EvalTime,              TEXT("Eval Time"));
 	ComponentRegistry->NewComponentType(&EvalSeconds,           TEXT("Eval Seconds"));
 
+	ComponentRegistry->NewComponentType(&Group,					TEXT("Entity Group"));
+
 	ComponentRegistry->NewComponentType(&BoundObjectKey,        TEXT("Bound Object Key"));
 	// Intentionally hidden from the reference graph because they are always accompanied by a BoundObjectKey which is used for garbage collection
 	ComponentRegistry->NewComponentTypeNoAddReferencedObjects(&BoundObject, TEXT("Bound Object"));
diff --git a/Engine/Source/Runtime/MovieScene/Private/EntitySystem/MovieSceneEntityGroupingSystem.cpp b/Engine/Source/Runtime/MovieScene/Private/EntitySystem/MovieSceneEntityGroupingSystem.cpp
new file mode 100644
--- /dev/null
+++ b/Engine/Source/Runtime/MovieScene/Private/EntitySystem/MovieSceneEntityGroupingSystem.cpp
@@ -0,0 +1,324 @@
+// Copyright Epic Games, Inc. All Rights Reserved.
+
+#include "EntitySystem/MovieSceneEntityGroupingSystem.h"
+
+#include "Containers/ArrayView.h"
+#include "EntitySystem/BuiltInComponentTypes.h"
+#include "EntitySystem/MovieSceneBoundObjectInstantiator.h"
+#include "EntitySystem/MovieSceneBoundSceneComponentInstantiator.h"
+#include "EntitySystem/MovieSceneEntityIDs.h"
+#include "EntitySystem/MovieSceneEntityMutations.h"
+#include "EntitySystem/MovieSceneEntitySystemLinker.h"
+UE_DISABLE_OPTIMIZATION
+namespace UE::MovieScene
+{
+
+FEntityGroupBuilder::FEntityGroupBuilder(UMovieSceneEntityGroupingSystem* InOwner, FEntityGroupingPolicyKey InPolicyKey)
+	: Owner(InOwner)
+	, PolicyKey(InPolicyKey)
+{
+}
+
+void FEntityGroupBuilder::AddEntityToGroup(const FMovieSceneEntityID& InEntity, const FEntityGroupID& InNewGroupID)
+{
+	if (ensure(InNewGroupID.IsValid()))
+	{
+		// Add the entity to the group.
+		UMovieSceneEntityGroupingSystem::FEntityGroupInfo& GroupInfo = Owner->Groups.FindOrAdd(InNewGroupID);
+		GroupInfo.Entities.Add(InEntity);
+	}
+}
+
+bool FEntityGroupBuilder::RemoveEntityFromGroup(const FMovieSceneEntityID& InEntity, const FEntityGroupID& InPreviousGroupID)
+{
+	if (ensure(InPreviousGroupID.IsValid()))
+	{
+		// Remove the entity from the group. We should find that group, and find that entity inside it.
+		UMovieSceneEntityGroupingSystem::FEntityGroupInfo* PreviousGroup = Owner->Groups.Find(InPreviousGroupID);
+		if (ensure(PreviousGroup))
+		{
+			using FEntityGroupingHandlerInfo = UMovieSceneEntityGroupingSystem::FEntityGroupingHandlerInfo;
+
+			const int32 NumRemoved = PreviousGroup->Entities.Remove(InEntity);
+			ensure(NumRemoved == 1);
+			if (PreviousGroup->Entities.Num() == 0)
+			{
+				// The group is now empty! Let's remove it.
+				Owner->Groups.Remove(InPreviousGroupID);
+				return true;
+			}
+		}
+	}
+	return false;
+}
+
+struct FAddGroupMutation : IMovieSceneEntityMutation
+{
+	UMovieSceneEntityGroupingSystem* System;
+	FBuiltInComponentTypes* BuiltInComponents;
+
+	FAddGroupMutation(UMovieSceneEntityGroupingSystem* InSystem)
+		: System(InSystem)
+	{
+		BuiltInComponents = FBuiltInComponentTypes::Get();
+	}
+
+	virtual void CreateMutation(FEntityManager* EntityManager, FComponentMask* InOutEntityComponentTypes) const override
+	{
+		for (const UMovieSceneEntityGroupingSystem::FEntityGroupingHandlerInfo& HandlerInfo : System->GroupHandlers)
+		{
+			if (InOutEntityComponentTypes->ContainsAll(HandlerInfo.ComponentMask))
+			{
+				InOutEntityComponentTypes->Set(BuiltInComponents->Group);
+				break;
+			}
+		}
+	}
+
+	virtual void InitializeAllocation(FEntityAllocation* Allocation, const FComponentMask& AllocationType) const override
+	{
+		FEntityAllocationWriteContext WriteContext = FEntityAllocationWriteContext::NewAllocation();
+		TComponentWriter<FEntityGroupID> GroupIDs = Allocation->WriteComponents(BuiltInComponents->Group, WriteContext);
+
+		// Find the policy key for this allocation. We'll initialize all group IDs to an invalid group but
+		// with a valid policy key.
+		FEntityGroupingPolicyKey PolicyKey;
+		for (auto It = System->GroupHandlers.CreateConstIterator(); It; ++It)
+		{
+			if (AllocationType.ContainsAll(It->ComponentMask))
+			{
+				PolicyKey = FEntityGroupingPolicyKey(It.GetIndex());
+				break;
+			}
+		}
+		ensure(PolicyKey.IsValid());
+
+		const int32 Num = Allocation->Num();
+		for (int32 Index = 0; Index < Num; ++Index)
+		{
+			GroupIDs[Index] = FEntityGroupID(PolicyKey, INDEX_NONE);
+		}
+	}
+};
+
+struct FUpdateGroupsTask
+{
+	using FEntityGroupInfo = UMovieSceneEntityGroupingSystem::FEntityGroupInfo;
+	using FEntityGroupingHandlerInfo = UMovieSceneEntityGroupingSystem::FEntityGroupingHandlerInfo;
+
+	UMovieSceneEntityGroupingSystem* System;
+	FBuiltInComponentTypes* BuiltInComponents;
+	TBitArray<> ModifiedGroups;
+
+	FUpdateGroupsTask(UMovieSceneEntityGroupingSystem* InSystem)
+		: System(InSystem)
+		, BuiltInComponents(nullptr)
+	{
+	}
+
+	void PreTask()
+	{
+		BuiltInComponents = FBuiltInComponentTypes::Get();
+
+		// Run the handlers' pre-task callback.
+		for (const UMovieSceneEntityGroupingSystem::FEntityGroupingHandlerInfo& HandlerInfo : System->GroupHandlers)
+		{
+			HandlerInfo.Handler->PreTask();
+		}
+	}
+
+	void ForEachAllocation(FEntityAllocationIteratorItem Item, FReadEntityIDs EntityIDs, TWrite<FEntityGroupID> GroupComponents)
+	{
+		TBitArray<> MatchingHandlers;
+		const FComponentMask& AllocationType = Item.GetAllocationType();
+		GatherMatchingGroupingHandlers(AllocationType, MatchingHandlers);
+
+		ensureMsgf(
+				MatchingHandlers.CountSetBits() <= 1, 
+				TEXT("Found more than one matching grouping handler for an entity allocation. "
+					 "Entities cannot belong to more than one group, so we will only process the first one!"));
+
+		const int32 HandlerIndex = MatchingHandlers.Find(true);
+		if (ensureMsgf(
+					HandlerIndex != INDEX_NONE,
+					TEXT("No matching gropuing handling for entity allocation even though it has a group ID component!")))
+		{
+			FEntityGroupingHandlerInfo& HandlerInfo = System->GroupHandlers[HandlerIndex];
+			FEntityGroupBuilder Builder(System, FEntityGroupingPolicyKey{ HandlerIndex });
+			HandlerInfo.Handler->ProcessAllocation(Item, EntityIDs, GroupComponents, &Builder);
+		}
+	}
+
+	void GatherMatchingGroupingHandlers(const FComponentMask& AllocationType, TBitArray<>& OutMatchingHandlers)
+	{
+		for (auto It = System->GroupHandlers.CreateConstIterator(); It; ++It)
+		{
+			if (AllocationType.ContainsAll(It->ComponentMask))
+			{
+				int32 Index = It.GetIndex();
+				OutMatchingHandlers.PadToNum(Index + 1, false);
+				OutMatchingHandlers[Index] = true;
+			}
+		}
+	}
+
+	void PostTask()
+	{
+		// Run the handlers' post-task callback.
+		for (const UMovieSceneEntityGroupingSystem::FEntityGroupingHandlerInfo& HandlerInfo : System->GroupHandlers)
+		{
+			HandlerInfo.Handler->PostTask();
+		}
+	}
+};
+
+} // namespace UE::MovieScene
+
+UMovieSceneEntityGroupingSystem::UMovieSceneEntityGroupingSystem(const FObjectInitializer& ObjInit)
+	: Super(ObjInit)
+{
+	using namespace UE::MovieScene;
+
+	FBuiltInComponentTypes* BuiltInComponents = FBuiltInComponentTypes::Get();
+
+	SystemCategories = EEntitySystemCategory::Core;
+	Phase = ESystemPhase::Instantiation;
+	RelevantComponent = BuiltInComponents->Group;
+
+	// We know that (as the time of this writing) we only have two systems that support grouping:
+	// object properties and materials.
+	GroupHandlers.Reserve(2);
+
+	if (HasAnyFlags(RF_ClassDefaultObject))
+	{
+		// We produce the group component. It's in the name.
+		DefineComponentProducer(GetClass(), BuiltInComponents->Group);
+
+		// This isn't something we *really* need, but pretty much all use-cases we have will want
+		// to group things using bound objects and/or bound scene components, so let's run after the
+		// systems that set those up.
+		DefineImplicitPrerequisite(UMovieSceneGenericBoundObjectInstantiator::StaticClass(), GetClass());
+		DefineImplicitPrerequisite(UMovieSceneBoundSceneComponentInstantiator::StaticClass(), GetClass());
+	}
+}
+
+bool UMovieSceneEntityGroupingSystem::IsRelevantImpl(UMovieSceneEntitySystemLinker* InLinker) const
+{
+	// We are relevant if we have any groupings to do.
+	return !GroupHandlers.IsEmpty();
+}
+
+void UMovieSceneEntityGroupingSystem::OnRun(FSystemTaskPrerequisites& InPrerequisites, FSystemSubsequentTasks& Subsequents)
+{
+	using namespace UE::MovieScene;
+
+	FBuiltInComponentTypes* BuiltInComponents = FBuiltInComponentTypes::Get();
+
+	// Mutate any new allocation that fits any of our grouping policies by adding the group component.
+	FEntityComponentFilter BroadFilter;
+	BroadFilter.Any({ BuiltInComponents->Tags.NeedsLink });
+	FAddGroupMutation Mutation(this);
+	Linker->EntityManager.MutateAll(BroadFilter, Mutation);
+
+	// Go over all the entities and update their groups.
+	FUpdateGroupsTask GroupTask(this);
+	FEntityTaskBuilder()
+	.ReadEntityIDs()
+	.Write(BuiltInComponents->Group)
+	.FilterAny({ BuiltInComponents->Tags.NeedsLink, BuiltInComponents->Tags.NeedsUnlink })
+	.RunInline_PerAllocation(&Linker->EntityManager, GroupTask);
+
+#if !UE_BUILD_SHIPPING
+	for (const TPair<FEntityGroupID, FEntityGroupInfo>& Pair : Groups)
+	{
+		for (FMovieSceneEntityID EntityID : Pair.Value.Entities)
+		{
+			TReadOptional<FEntityGroupID> GroupComponent = Linker->EntityManager.ReadComponent(EntityID, BuiltInComponents->Group);
+			ensureMsgf(
+					GroupComponent.IsValid() && *GroupComponent == Pair.Key,
+					TEXT("Found mismatch between group cache and group component!"));
+		}
+	}
+#endif
+}
+
+void UMovieSceneEntityGroupingSystem::OnLink()
+{
+#if WITH_EDITOR
+	FCoreUObjectDelegates::OnObjectsReplaced.AddUObject(this, &UMovieSceneEntityGroupingSystem::OnObjectsReplaced);
+#endif
+}
+
+void UMovieSceneEntityGroupingSystem::OnUnlink()
+{
+	const bool bIsEmpty = GroupHandlers.Num() == 0 && Groups.IsEmpty();
+	if (!ensure(bIsEmpty))
+	{
+		GroupHandlers.Empty();
+		Groups.Empty();
+	}
+
+#if WITH_EDITOR
+	FCoreUObjectDelegates::OnObjectsReplaced.RemoveAll(this);
+#endif
+}
+
+void UMovieSceneEntityGroupingSystem::OnCleanTaggedGarbage()
+{
+	using namespace UE::MovieScene;
+
+	// Garbage has been tagged with NeedsUnlink, so visit those and remove them from their groups.
+	// In theory, a group with garbage in its group key should get emptied, because we assume that
+	// a group key only has garbage in it if the components used to derive it also have garbage in
+	// them. And in that case, their entities would have been flagged, and removed from that
+	// group.
+	FBuiltInComponentTypes* BuiltInComponents = FBuiltInComponentTypes::Get();
+
+	FUpdateGroupsTask GroupTask(this);
+	FEntityTaskBuilder()
+	.ReadEntityIDs()
+	.Write(BuiltInComponents->Group)
+	.FilterAny({ BuiltInComponents->Tags.NeedsUnlink })
+	.RunInline_PerAllocation(&Linker->EntityManager, GroupTask);
+}
+
+#if WITH_EDITOR
+void UMovieSceneEntityGroupingSystem::OnObjectsReplaced(const TMap<UObject*, UObject*>& ReplacementMap)
+{
+	for (const FEntityGroupingHandlerInfo& HandlerInfo : GroupHandlers)
+	{
+		HandlerInfo.Handler->OnObjectsReplaced(ReplacementMap);
+	}
+}
+#endif
+
+void UMovieSceneEntityGroupingSystem::RemoveGrouping(FEntityGroupingPolicyKey InPolicyKey)
+{
+	using namespace UE::MovieScene;
+
+	check(InPolicyKey.IsValid());
+
+	// Get the list of existing groups using this policy, and clean them up.
+	TArray<FEntityGroupID> ExistingGroupIDs;
+	Groups.GetKeys(ExistingGroupIDs);
+	for (const FEntityGroupID& ExistingGroupID : ExistingGroupIDs)
+	{
+		if (!ensureMsgf(ExistingGroupID.PolicyKey != InPolicyKey, TEXT("Found leftover group from policy being removed")))
+		{
+			Groups.Remove(ExistingGroupID);
+		}
+	}
+
+	// Remove the handler itself.
+	GroupHandlers.RemoveAt(InPolicyKey.Index);
+}
+
+TArrayView<const UE::MovieScene::FMovieSceneEntityID> UMovieSceneEntityGroupingSystem::GetGroup(const FEntityGroupID& InGroupID) const
+{
+	if (const FEntityGroupInfo* GroupInfo = Groups.Find(InGroupID))
+	{
+		return MakeArrayView(GroupInfo->Entities);
+	}
+	return TArrayView<const UE::MovieScene::FMovieSceneEntityID>();
+}
+UE_ENABLE_OPTIMIZATION
diff --git a/Engine/Source/Runtime/MovieScene/Public/EntitySystem/BuiltInComponentTypes.h b/Engine/Source/Runtime/MovieScene/Public/EntitySystem/BuiltInComponentTypes.h
--- a/Engine/Source/Runtime/MovieScene/Public/EntitySystem/BuiltInComponentTypes.h
+++ b/Engine/Source/Runtime/MovieScene/Public/EntitySystem/BuiltInComponentTypes.h
@@ -326,6 +326,68 @@ private:
 	FObjectKey ObjectKey;
 };
 
+/**
+ * The key to a grouping policy registered on the grouping system (see UMovieSceneEntityGroupingSystem)
+ */
+struct FEntityGroupingPolicyKey
+{
+	static const FEntityGroupingPolicyKey Invalid() { return FEntityGroupingPolicyKey(); }
+
+	FEntityGroupingPolicyKey() {}
+	FEntityGroupingPolicyKey(int32 FromIndex) : Index(FromIndex) {}
+
+	bool IsValid() const
+	{
+		return Index != INDEX_NONE;
+	}
+
+	friend uint32 GetTypeHash(const FEntityGroupingPolicyKey& Key)
+	{
+		return GetTypeHash(Key.Index);
+	}
+
+	friend bool operator==(const FEntityGroupingPolicyKey& A, const FEntityGroupingPolicyKey& B)
+	{
+		return A.Index == B.Index;
+	}
+
+	int32 Index = INDEX_NONE;
+};
+
+/**
+ * The component data for describing what group an entity belongs to (see UMovieSceneEntityGroupingSystem)
+ */
+struct FEntityGroupID
+{
+	static const FEntityGroupID Invalid() { return FEntityGroupID(); }
+
+	FEntityGroupingPolicyKey PolicyKey;
+	int32 GroupIndex = INDEX_NONE;
+
+	FEntityGroupID() {}
+	FEntityGroupID(const FEntityGroupingPolicyKey InPolicyKey, int32 InGroupIndex) : PolicyKey(InPolicyKey), GroupIndex(InGroupIndex) {}
+
+	bool IsValid() const
+	{
+		return PolicyKey.IsValid() && GroupIndex != INDEX_NONE;
+	}
+
+	bool HasGroup() const
+	{
+		return GroupIndex != INDEX_NONE;
+	}
+
+	friend uint32 GetTypeHash(const FEntityGroupID& GroupID)
+	{
+		return HashCombine(GetTypeHash(GroupID.PolicyKey.Index), GetTypeHash(GroupID.GroupIndex));
+	}
+
+	friend bool operator==(const FEntityGroupID& A, const FEntityGroupID& B)
+	{
+		return A.PolicyKey == B.PolicyKey && A.GroupIndex == B.GroupIndex;
+	}
+};
+
 /**
  * Pre-defined built in component types
  */
@@ -364,6 +426,8 @@ public:
 
 	TComponentTypeID<double>              EvalSeconds;
 
+	TComponentTypeID<FEntityGroupID>	  Group;
+
 public:
 
 	TComponentTypeID<FMovieSceneBlendChannelID> BlendChannelInput;
@@ -571,6 +635,9 @@ private:
 	TMap<FComponentTypeID, FComponentTypeID> ResultToBase;
 };
 
+#if UE_MOVIESCENE_ENTITY_DEBUG
+template<> struct TComponentDebugType<FEntityGroupID> { static const EComponentDebugType Type = EComponentDebugType::GroupID; };
+#endif
 
 } // namespace MovieScene
 } // namespace UE
diff --git a/Engine/Source/Runtime/MovieScene/Public/EntitySystem/MovieSceneComponentDebug.h b/Engine/Source/Runtime/MovieScene/Public/EntitySystem/MovieSceneComponentDebug.h
--- a/Engine/Source/Runtime/MovieScene/Public/EntitySystem/MovieSceneComponentDebug.h
+++ b/Engine/Source/Runtime/MovieScene/Public/EntitySystem/MovieSceneComponentDebug.h
@@ -38,6 +38,7 @@ enum class EComponentDebugType
 	SequenceID,
 	EntityID,
 	Name,
+	GroupID,
 	BlendChannelID,
 };
 
diff --git a/Engine/Source/Runtime/MovieScene/Public/EntitySystem/MovieSceneEntityGroupingSystem.h b/Engine/Source/Runtime/MovieScene/Public/EntitySystem/MovieSceneEntityGroupingSystem.h
new file mode 100644
--- /dev/null
+++ b/Engine/Source/Runtime/MovieScene/Public/EntitySystem/MovieSceneEntityGroupingSystem.h
@@ -0,0 +1,501 @@
+// Copyright Epic Games, Inc. All Rights Reserved.
+
+#pragma once
+
+#include "Containers/Map.h"
+#include "Containers/SparseArray.h"
+#include "CoreTypes.h"
+#include "EntitySystem/BuiltInComponentTypes.h"
+#include "EntitySystem/MovieSceneComponentTypeInfo.h"
+#include "EntitySystem/MovieSceneEntityIDs.h"
+#include "EntitySystem/MovieSceneEntitySystem.h"
+#include "EntitySystem/MovieSceneEntitySystemTask.h"
+#include "EntitySystem/MovieSceneEntitySystemTypes.h"
+#include "Templates/PointerIsConvertibleFromTo.h"
+#include "UObject/ObjectMacros.h"
+#include "UObject/UObjectGlobals.h"
+
+#include "MovieSceneEntityGroupingSystem.generated.h"
+
+class UMovieSceneEntityGroupingSystem;
+
+namespace UE::MovieScene
+{
+
+struct FAddGroupMutation;
+struct FUpdateGroupsTask;
+
+/**
+ * Utility class used by the grouping system's policies (see below) to manage groups.
+ */
+struct MOVIESCENE_API FEntityGroupBuilder
+{
+	FEntityGroupBuilder(UMovieSceneEntityGroupingSystem* InOwner, FEntityGroupingPolicyKey InPolicyKey);
+
+	/** Make a full group ID from an existing group index */
+	FEntityGroupID MakeGroupID(int32 GroupIndex) const { return FEntityGroupID(PolicyKey, GroupIndex); }
+	/** Make an invalid group ID that is associated with the policy key */
+	FEntityGroupID MakeInvalidGroupID() const { return FEntityGroupID(PolicyKey, INDEX_NONE); }
+	/** Add the entity to the given group. The entity must already have the group ID component. */
+	void AddEntityToGroup(const FMovieSceneEntityID& InEntity, const FEntityGroupID& InNewGroupID);
+	/** Remove the entity from the given group. The entity must already have the group ID component. */
+	bool RemoveEntityFromGroup(const FMovieSceneEntityID& InEntity, const FEntityGroupID& InPreviousGroupID);
+
+private:
+	UMovieSceneEntityGroupingSystem* Owner;
+	FEntityGroupingPolicyKey PolicyKey;
+};
+
+/**
+ * Base class for grouping handlers, used by the grouping system (see below).
+ */
+struct IEntityGroupingHandler
+{
+	virtual ~IEntityGroupingHandler() {}
+	virtual void PreTask() {}
+	virtual void ProcessAllocation(FEntityAllocationIteratorItem Item, FReadEntityIDs EntityIDs, TWrite<FEntityGroupID> GroupIDs, FEntityGroupBuilder* Builder) = 0;
+	virtual void PostTask() {}
+
+#if WITH_EDITOR
+	virtual void OnObjectsReplaced(const TMap<UObject*, UObject*>& ReplacementMap) = 0;
+#endif
+};
+
+/**
+ * Strongly-typed grouping handler class, which knows about the exact components to look for, and how
+ * to use them to group entities.
+ */
+template<typename GroupingPolicy, typename ComponentIndices, typename ...ComponentTypes>
+struct TEntityGroupingHandlerImpl;
+
+template<typename GroupingPolicy, typename ...ComponentTypes>
+struct TEntityGroupingHandler 
+	: TEntityGroupingHandlerImpl<GroupingPolicy, TMakeIntegerSequence<int, sizeof...(ComponentTypes)>, ComponentTypes...>
+{
+	TEntityGroupingHandler(GroupingPolicy&& InPolicy, TComponentTypeID<ComponentTypes>... InComponents)
+		: TEntityGroupingHandlerImpl<GroupingPolicy, TMakeIntegerSequence<int, sizeof...(ComponentTypes)>, ComponentTypes...>(
+				MoveTemp(InPolicy), InComponents...)
+	{
+	}
+};
+
+template<typename GroupingPolicy, int ...ComponentIndices, typename ...ComponentTypes>
+struct TEntityGroupingHandlerImpl<GroupingPolicy, TIntegerSequence<int, ComponentIndices...>, ComponentTypes...> : IEntityGroupingHandler
+{
+	using GroupKeyType = typename GroupingPolicy::GroupKeyType;
+
+	/** The grouping policy */
+	GroupingPolicy Policy;
+
+	/** The components that are required for making up a group key */
+	TTuple<TComponentTypeID<ComponentTypes>...> Components;
+
+	/** The group keys that we know about, mapped to their corresponding group index */
+	TMap<GroupKeyType, int32> GroupKeyToIndex;
+
+	/** The list of group indices in use */
+	TBitArray<> AllocatedGroupIndices;
+	/** The transient list of groups freed this frame */
+	TBitArray<> FreedGroupIndices;
+
+	TEntityGroupingHandlerImpl(GroupingPolicy&& InPolicy, TComponentTypeID<ComponentTypes>... InComponents)
+		: Policy(MoveTemp(InPolicy))
+		, Components(InComponents...)
+	{
+	}
+
+	int32 GetOrAllocateGroupIndex(typename TCallTraits<GroupKeyType>::ParamType InGroupKey)
+	{
+		int32& GroupIndex = GroupKeyToIndex.FindOrAdd(InGroupKey, INDEX_NONE);
+		if (GroupIndex == INDEX_NONE)
+		{
+			// This group key isn't known to us... let's allocate a new group index for it.
+			// Try to find an available index first. Otherwise use a new high index.
+			int32 NewGroupIndex = AllocatedGroupIndices.Find(false);
+			if (NewGroupIndex == INDEX_NONE)
+			{
+				NewGroupIndex = AllocatedGroupIndices.Add(true);
+			}
+			else
+			{
+				AllocatedGroupIndices[NewGroupIndex] = true;
+			}
+			GroupIndex = NewGroupIndex;
+		}
+		else
+		{
+			// We know this group key, so we'll return the group index we already have
+			// associated with it. We just need to "revive" it in case it was scheduled
+			// for being freed.
+			if (FreedGroupIndices.IsValidIndex(GroupIndex))
+			{
+				FreedGroupIndices[GroupIndex] = false;
+			}
+		}
+		return GroupIndex;
+	}
+
+	void FreeGroupIndex(int32 InGroupIndex)
+	{
+		if (ensure(AllocatedGroupIndices.IsValidIndex(InGroupIndex) && AllocatedGroupIndices[InGroupIndex]))
+		{
+			AllocatedGroupIndices[InGroupIndex] = false;
+		}
+	}
+
+	/** Callback on the grouping policy invoked before all grouping happens this frame */
+	static void PreTaskImpl(void*, ...) {}
+	template <typename T> static void PreTaskImpl(T* InPolicy, decltype(&T::PreTask)* = 0)
+	{
+		InPolicy->PreTask();
+	}
+
+	virtual void PreTask() override
+	{
+		FreedGroupIndices.Reset();
+		PreTaskImpl(&Policy);
+	}
+
+	/**
+	 * Callback on the grouping policy invoked before grouping happens for a given allocation.
+	 * If the callback returns false, this allocation will be skipped.
+	 */
+	static bool PreAllocationCallbackImpl(void*, ...) { return true; }
+	template <typename T> static bool PreAllocationCallbackImpl(T* InPolicy, FEntityAllocationIteratorItem Item, decltype(&T::PreProcessGroups)* = 0)
+	{
+		return InPolicy->PreProcessGroups(Item);
+	}
+
+	/** Process an allocation and group the entities found therein */
+	virtual void ProcessAllocation(FEntityAllocationIteratorItem Item,  FReadEntityIDs EntityIDs, TWrite<FEntityGroupID> GroupIDs, FEntityGroupBuilder* Builder) override
+	{
+		// Check if we should process this allocation.
+		const bool bProcessAllocation = PreAllocationCallbackImpl(&Policy, Item);
+		if (!bProcessAllocation)
+		{
+			return;
+		}
+
+		const FBuiltInComponentTypes* BuiltInComponents = FBuiltInComponentTypes::Get();
+		const FComponentMask& AllocationType = Item.GetAllocationType();
+		const bool bNeedsLink = AllocationType.Contains(BuiltInComponents->Tags.NeedsLink);
+		const bool bNeedsUnlink = AllocationType.Contains(BuiltInComponents->Tags.NeedsUnlink);
+		ensure(bNeedsLink || bNeedsUnlink);
+
+		if (bNeedsLink)
+		{
+			VisitLinkedEntities(Item, EntityIDs, GroupIDs, Builder);
+		}
+		else if (bNeedsUnlink)
+		{
+			VisitUnlinkedEntities(Item, EntityIDs, GroupIDs, Builder);
+		}
+
+		PostAllocationCallbackImpl(&Policy, Item);
+	}
+	
+	void VisitLinkedEntities(FEntityAllocationIteratorItem Item, FReadEntityIDs EntityIDs, TWrite<FEntityGroupID> GroupIDs, FEntityGroupBuilder* Builder)
+	{
+		const FEntityAllocation* Allocation = Item.GetAllocation();
+		const int32 Num = Allocation->Num();
+
+		const FEntityGroupID InvalidGroupID = Builder->MakeInvalidGroupID();
+
+		TTuple<TComponentReader<ComponentTypes>...> ComponentReaders(
+				Allocation->ReadComponents(Components.template Get<ComponentIndices>())...);
+
+		for (int32 Index = 0; Index < Num; ++Index)
+		{
+			GroupKeyType GroupKey;
+			const bool bValidGroupKey = Policy.GetGroupKey(ComponentReaders.template Get<ComponentIndices>()[Index]..., GroupKey);
+
+			const FMovieSceneEntityID EntityID(EntityIDs[Index]);
+			FEntityGroupID& GroupID(GroupIDs[Index]);
+
+			if (bValidGroupKey)
+			{
+				// Find or create the appropriate group and put the entity in it.
+				int32 NewGroupIndex = GetOrAllocateGroupIndex(GroupKey);
+				FEntityGroupID NewGroupID = Builder->MakeGroupID(NewGroupIndex);
+				Builder->AddEntityToGroup(EntityID, NewGroupID);
+				GroupID = NewGroupID;
+			}
+			else
+			{
+				// This entity doesn't belong to any group.
+				// Let's assign an invalid group ID that nonetheless has a valid policy key
+				// pointing to this grouping.
+				GroupID = InvalidGroupID;
+			}
+		}
+	}
+
+	void VisitUnlinkedEntities(FEntityAllocationIteratorItem Item, FReadEntityIDs EntityIDs, TWrite<FEntityGroupID> GroupIDs, FEntityGroupBuilder* Builder)
+	{
+		const FEntityAllocation* Allocation = Item.GetAllocation();
+		const int32 Num = Allocation->Num();
+
+		for (int32 Index = 0; Index < Num; ++Index)
+		{
+			const FMovieSceneEntityID EntityID(EntityIDs[Index]);
+			FEntityGroupID& GroupID(GroupIDs[Index]);
+
+			if (GroupID.IsValid())
+			{
+				const bool bIsGroupEmpty = Builder->RemoveEntityFromGroup(EntityID, GroupID);
+				if (bIsGroupEmpty)
+				{
+					// The group is now empty, we can re-use its index for a new group... but we don't
+					// want to re-use it until later, because we could end up in two situations we
+					// want to avoid:
+					//
+					// 1) At the same time that a group is emptied of all its entities, new entities
+					//    come in that belong to that group because they generate the exact same key.
+					//    We want that group to effectively "persist" with the same index.
+					//
+					// 2) If we freed the index right away, we couldn't tell the difference between
+					//	  the above situation, and a brand new group that just happens to re-use the
+					//	  recently freed index.
+					//
+					FreedGroupIndices.PadToNum(GroupID.GroupIndex + 1, false);
+					FreedGroupIndices[GroupID.GroupIndex] = true;
+				}
+				// Leave the GroupID on the entity so that downstream systems can use it to track
+				// that this entity is leaving its group.
+			}
+		}
+	}
+	
+	/** Callback on the grouping policy invoked after a given allocation has been processed */
+	static void PostAllocationCallbackImpl(void*, ...) {}
+	template <typename T> static void PostAllocationCallbackImpl(T* InPolicy, FEntityAllocationIteratorItem Item, decltype(&T::PostProcessGroups)* = 0)
+	{
+		InPolicy->PostProcessGroups(Item);
+	}
+
+	static void PostTaskImpl(void*, ...){}
+	template <typename T> static void PostTaskImpl(T* InPolicy, decltype(&T::PostTask)* = 0)
+	{
+		InPolicy->PostTask();
+	}
+	
+	virtual void PostTask() override
+	{
+		if (FreedGroupIndices.CountSetBits() > 0)
+		{
+			// Build a reverse lookup map to figure out which group keys we don't need anymore
+			// based on the group indices we have freed.
+			TMap<int32, GroupKeyType> GroupIndexToKey;
+			for (const TPair<GroupKeyType, int32>& Pair : GroupKeyToIndex)
+			{
+				GroupIndexToKey.Add(Pair.Value, Pair.Key);
+			}
+
+			// Free the indices we don't use anymore, and free the group key lookup entry too.
+			for (TConstSetBitIterator<> It(FreedGroupIndices); It; ++It)
+			{
+				const GroupKeyType* GroupKey = GroupIndexToKey.Find(It.GetIndex());
+				if (ensure(GroupKey))
+				{
+					const int NumRemoved = GroupKeyToIndex.Remove(*GroupKey);
+					ensure(NumRemoved == 1);
+				}
+				FreeGroupIndex(It.GetIndex());
+			}
+			FreedGroupIndices.Reset();
+		}
+
+		PostTaskImpl(&Policy);
+	}
+
+#if WITH_EDITOR
+	virtual void OnObjectsReplaced(const TMap<UObject*, UObject*>& ReplacementMap) override
+	{
+		// Get a list of keys that contain replaced objects.
+		TMap<GroupKeyType, GroupKeyType> ReplacedKeys;
+		for (const TPair<GroupKeyType, int32>& Pair : GroupKeyToIndex)
+		{
+			GroupKeyType NewKey = Pair.Key;
+			if (Policy.OnObjectsReplaced(NewKey, ReplacementMap))
+			{
+				ReplacedKeys.Add(Pair.Key, NewKey);
+			}
+		}
+		// Replace the keys but keep the group indices.
+		for (const TPair<GroupKeyType, GroupKeyType>& Pair : ReplacedKeys)
+		{
+			int32 GroupIndex;
+			const bool bRemoved = GroupKeyToIndex.RemoveAndCopyValue(Pair.Key, GroupIndex);
+			if (ensure(bRemoved))
+			{
+				GroupKeyToIndex.Add(Pair.Value, GroupIndex);
+			}
+		}
+	}
+#endif
+};
+
+namespace Private
+{
+	template<typename T>
+	bool ReplaceGroupKeyObjectElement(T&& InElem, const TMap<UObject*, UObject*>& ReplacementMap)
+	{
+		return false;
+	}
+
+	template<typename T>
+	typename TEnableIf<TPointerIsConvertibleFromTo<T, UObject>::Value, bool>::Type ReplaceGroupKeyObjectElement(T& InOutElem, const TMap<UObject*, UObject*>& ReplacementMap)
+	{
+		UObject* CurrentObject = InOutElem;
+		if (UObject* const* NewObject = ReplacementMap.Find(CurrentObject))
+		{
+			InOutElem = NewObject;
+			return true;
+		}
+		return false;
+	}
+}
+
+/**
+ * A simple grouping policy that uses a tuple of component values as the group key.
+ */
+template<typename... ComponentTypes>
+struct TTupleGroupingPolicy
+{
+	using GroupKeyType = TTuple<ComponentTypes...>;
+
+	bool GetGroupKey(ComponentTypes... InComponents, GroupKeyType& OutGroupKey)
+	{
+		OutGroupKey = MakeTuple(InComponents...);
+		return true;
+	}
+
+	bool OnObjectsReplaced(GroupKeyType& InOutKey, const TMap<UObject*, UObject*>& ReplacementMap)
+	{
+		bool bReplaced = false;
+		VisitTupleElements([&ReplacementMap, &bReplaced](auto& InElem)
+			{
+				bReplaced |= Private::ReplaceGroupKeyObjectElement(InElem, ReplacementMap);
+			},
+			InOutKey);
+		return bReplaced;
+	}
+};
+
+} // namespace UE::MovieScene
+
+
+UCLASS()
+class MOVIESCENE_API UMovieSceneEntityGroupingSystem : public UMovieSceneEntitySystem
+{
+public:
+
+	using FEntityGroupID = UE::MovieScene::FEntityGroupID;
+	using FEntityGroupingPolicyKey = UE::MovieScene::FEntityGroupingPolicyKey;
+	using IEntityGroupingHandler = UE::MovieScene::IEntityGroupingHandler;
+	using FMovieSceneEntityID = UE::MovieScene::FMovieSceneEntityID;
+
+	GENERATED_BODY()
+
+	UMovieSceneEntityGroupingSystem(const FObjectInitializer& ObjInit);
+
+	/**
+	 * Add a new grouping policy that will use the given components to make up a group key.
+	 *
+	 * Grouping policies must be simple structs that can be copied and owned by the grouping system, and
+	 * that implement the following members:
+	 *
+	 * - GroupKeyType [mandatory]
+	 *		A typedef or alias to the group key type to use to group entities together.
+	 *
+	 * - GetGroupKey(ComponentTypes... InComponents, GroupKeyType& OutGroupKey) [mandatory]
+	 *		A function that creates a group key used to group entities together.
+	 *		Returns true to indicate the key is valid, false otherwise. If false, the entity
+	 *		corresponding to the given components will not be grouped.
+	 *
+	 *	- OnObjectsReplaced(GroupKeyType& InOut, const TMap<UObject*, UObject*>&) [mandatory, sadly]
+	 *		Potentially changes a key if it contains a pointer to a replaced object. Should
+	 *		return true if any replacement occured.
+	 *
+	 *	- PreTask() [optional]
+	 *		A function called before any grouping is done during an instantiation phase.
+	 *
+	 *	- PostTask() [optional]
+	 *		A function called after any grouping is done dunring an instantiation phase.
+	 *
+	 *	- PreProcessGroups() [optional]
+	 *
+	 *	- PostProcessGroups() [optional]
+	 */
+	template<typename GroupingPolicy, typename ...ComponentTypes>
+	FEntityGroupingPolicyKey AddGrouping(GroupingPolicy&& InPolicy, TComponentTypeID<ComponentTypes>... InComponents)
+	{
+		using namespace UE::MovieScene;
+
+		using NewGroupHandlerType = TEntityGroupingHandler<GroupingPolicy, ComponentTypes...>;
+
+		static_assert(sizeof(NewGroupHandlerType) <= 256, "Handler type too big! Please increase the TInlineValue size.");
+
+		const int32 HandlerIndex = GroupHandlers.Emplace();
+		FEntityGroupingHandlerInfo& HandlerInfo = GroupHandlers[HandlerIndex];
+		HandlerInfo.Handler.Emplace<NewGroupHandlerType>(MoveTemp(InPolicy), InComponents...);
+		HandlerInfo.ComponentMask = FComponentMask({ InComponents... });
+
+		FEntityGroupingPolicyKey NewPolicyKey{ HandlerIndex };
+		return NewPolicyKey;
+	}
+
+	/**
+	 * Add a new grouping policy that will make a key that is a tuple of the given components' values.
+	 */
+	template<typename ...ComponentTypes>
+	FEntityGroupingPolicyKey AddGrouping(TComponentTypeID<ComponentTypes>... InComponents)
+	{
+		UE::MovieScene::TTupleGroupingPolicy<ComponentTypes...> TuplePolicy;
+		return AddGrouping(MoveTemp(TuplePolicy), InComponents...);
+	}
+
+	/**
+	 * Remove a previously added grouping policy.
+	 */
+	void RemoveGrouping(FEntityGroupingPolicyKey InPolicyKey);
+
+	/**
+	 * Get the list of entities in the given group.
+	 */
+	TArrayView<const UE::MovieScene::FMovieSceneEntityID> GetGroup(const FEntityGroupID& InGroupID) const;
+
+private:
+
+	virtual bool IsRelevantImpl(UMovieSceneEntitySystemLinker* InLinker) const override;
+	virtual void OnRun(FSystemTaskPrerequisites& InPrerequisites, FSystemSubsequentTasks& Subsequents) override;
+	virtual void OnLink() override;
+	virtual void OnUnlink() override;
+	virtual void OnCleanTaggedGarbage() override;
+
+#if WITH_EDITOR
+	void OnObjectsReplaced(const TMap<UObject*, UObject*>& ReplacementMap);
+#endif
+
+private:
+
+	struct FEntityGroupInfo
+	{
+		TArray<UE::MovieScene::FMovieSceneEntityID> Entities;
+	};
+
+	TMap<FEntityGroupID, FEntityGroupInfo> Groups;
+
+	struct FEntityGroupingHandlerInfo
+	{
+		TInlineValue<IEntityGroupingHandler, 256> Handler;
+		FComponentMask ComponentMask;
+	};
+
+	TSparseArray<FEntityGroupingHandlerInfo> GroupHandlers;
+
+	friend struct UE::MovieScene::FAddGroupMutation;
+	friend struct UE::MovieScene::FUpdateGroupsTask;
+	friend struct UE::MovieScene::FEntityGroupBuilder;
+};
+
diff --git a/Engine/Source/Runtime/MovieSceneTracks/Private/Systems/MovieSceneHierarchicalBiasSystem.cpp b/Engine/Source/Runtime/MovieSceneTracks/Private/Systems/MovieSceneHierarchicalBiasSystem.cpp
--- a/Engine/Source/Runtime/MovieSceneTracks/Private/Systems/MovieSceneHierarchicalBiasSystem.cpp
+++ b/Engine/Source/Runtime/MovieSceneTracks/Private/Systems/MovieSceneHierarchicalBiasSystem.cpp
@@ -5,6 +5,7 @@
 #include "EntitySystem/BuiltInComponentTypes.h"
 #include "EntitySystem/MovieSceneEntitySystemTask.h"
 #include "EntitySystem/MovieSceneEntitySystemLinker.h"
+#include "EntitySystem/MovieScenePreAnimatedStateSystem.h"
 
 #include UE_INLINE_GENERATED_CPP_BY_NAME(MovieSceneHierarchicalBiasSystem)
 
@@ -14,17 +15,17 @@ namespace UE
 namespace MovieScene
 {
 
-struct FBlendChannelSequenceKey
+struct FEntityGroupSequenceKey
 {
 	FRootInstanceHandle RootInstance;
-	FMovieSceneBlendChannelID BlendChannelID;
-	friend uint32 GetTypeHash(const FBlendChannelSequenceKey& In)
+	FEntityGroupID GroupID;
+	friend uint32 GetTypeHash(const FEntityGroupSequenceKey& In)
 	{
-		return HashCombine(GetTypeHash(In.RootInstance), GetTypeHash(In.BlendChannelID));
+		return HashCombine(GetTypeHash(In.RootInstance), GetTypeHash(In.GroupID));
 	}
-	friend bool operator==(const FBlendChannelSequenceKey& A, const FBlendChannelSequenceKey& B)
+	friend bool operator==(const FEntityGroupSequenceKey& A, const FEntityGroupSequenceKey& B)
 	{
-		return A.RootInstance == B.RootInstance && A.BlendChannelID == B.BlendChannelID;
+		return A.RootInstance == B.RootInstance && A.GroupID == B.GroupID;
 	}
 };
 
@@ -34,17 +35,17 @@ struct FHierarchicalBiasTask
 		: Linker(InLinker)
 	{}
 
-	void InitializeChannel(FRootInstanceHandle RootInstanceHandle, FMovieSceneBlendChannelID BlendChannel)
+	void InitializeGroup(FRootInstanceHandle RootInstanceHandle, FEntityGroupID GroupID)
 	{
-		MaxBiasByChannel.FindOrAdd(FBlendChannelSequenceKey{ RootInstanceHandle, BlendChannel }, MIN_int16);
+		MaxBiasByGroup.FindOrAdd(FEntityGroupSequenceKey{ RootInstanceHandle, GroupID }, MIN_int16);
 	}
 
 	bool HasAnyWork() const
 	{
-		return MaxBiasByChannel.Num() != 0;
+		return MaxBiasByGroup.Num() != 0;
 	}
 
-	void ForEachAllocation(FEntityAllocationIteratorItem Iterator, TRead<FMovieSceneEntityID> EntityIDs, TRead<FRootInstanceHandle> RootInstanceHandles, TRead<FMovieSceneBlendChannelID> BlendChannels, TReadOptional<int16> OptHBiases)
+	void ForEachAllocation(FEntityAllocationIteratorItem Iterator, TRead<FMovieSceneEntityID> EntityIDs, TRead<FRootInstanceHandle> RootInstanceHandles, TRead<FEntityGroupID> GroupIDs, TReadOptional<int16> OptHBiases)
 	{
 		const int32 Num = Iterator.GetAllocation()->Num();
 		const FComponentMask& AllocationType = Iterator.GetAllocationType();
@@ -55,22 +56,22 @@ struct FHierarchicalBiasTask
 		{
 			for (int32 Index = 0; Index < Num; ++Index)
 			{
-				FBlendChannelSequenceKey Key{ RootInstanceHandles[Index], BlendChannels[Index] };
-				ActiveContributorsByChannel.Add(Key, EntityIDs[Index]);
+				FEntityGroupSequenceKey Key{ RootInstanceHandles[Index], GroupIDs[Index] };
+				ActiveContributorsByGroup.Add(Key, EntityIDs[Index]);
 			}
 		}
 		else if (OptHBiases)
 		{
 			for (int32 Index = 0; Index < Num; ++Index)
 			{
-				VisitChannel(EntityIDs[Index], RootInstanceHandles[Index], BlendChannels[Index], OptHBiases[Index]);
+				VisitGroup(EntityIDs[Index], RootInstanceHandles[Index], GroupIDs[Index], OptHBiases[Index]);
 			}
 		}
 		else
 		{
 			for (int32 Index = 0; Index < Num; ++Index)
 			{
-				VisitChannel(EntityIDs[Index], RootInstanceHandles[Index], BlendChannels[Index], 0);
+				VisitGroup(EntityIDs[Index], RootInstanceHandles[Index], GroupIDs[Index], 0);
 			}
 		}
 	}
@@ -79,12 +80,12 @@ struct FHierarchicalBiasTask
 	{
 		FBuiltInComponentTypes* Components = FBuiltInComponentTypes::Get();
 
-		for (auto It = ActiveContributorsByChannel.CreateIterator(); It; ++It)
+		for (auto It = ActiveContributorsByGroup.CreateIterator(); It; ++It)
 		{
 			Linker->EntityManager.RemoveComponent(It.Value(), Components->Tags.Ignored);
 		}
 
-		for (auto It = InactiveContributorsByChannel.CreateIterator(); It; ++It)
+		for (auto It = InactiveContributorsByGroup.CreateIterator(); It; ++It)
 		{
 			Linker->EntityManager.AddComponent(It.Value(), Components->Tags.Ignored);
 		}
@@ -92,40 +93,40 @@ struct FHierarchicalBiasTask
 
 private:
 
-	void VisitChannel(FMovieSceneEntityID EntityID, FRootInstanceHandle RootInstanceHandle, FMovieSceneBlendChannelID BlendChannel, int16 HBias)
+	void VisitGroup(FMovieSceneEntityID EntityID, FRootInstanceHandle RootInstanceHandle, FEntityGroupID GroupID, int16 HBias)
 	{
-		FBlendChannelSequenceKey Key{ RootInstanceHandle, BlendChannel };
+		FEntityGroupSequenceKey Key{ RootInstanceHandle, GroupID };
 
-		// If this channel hasn't changed at all (ie InitializeChannel was not called for it) do nothing
-		if (int16* ExistingBias = MaxBiasByChannel.Find(Key))
+		// If this group hasn't changed at all (ie InitializeGroup was not called for it) do nothing
+		if (int16* ExistingBias = MaxBiasByGroup.Find(Key))
 		{
 			if (HBias > *ExistingBias)
 			{
-				for (auto It = ActiveContributorsByChannel.CreateKeyIterator(Key); It; ++It)
+				for (auto It = ActiveContributorsByGroup.CreateKeyIterator(Key); It; ++It)
 				{
-					InactiveContributorsByChannel.Add(Key, It.Value());
+					InactiveContributorsByGroup.Add(Key, It.Value());
 					It.RemoveCurrent();
 				}
 
 				*ExistingBias = HBias;
-				ActiveContributorsByChannel.Add(Key, EntityID);
+				ActiveContributorsByGroup.Add(Key, EntityID);
 			}
 			else if (HBias == *ExistingBias)
 			{
-				ActiveContributorsByChannel.Add(Key, EntityID);
+				ActiveContributorsByGroup.Add(Key, EntityID);
 			}
 			else
 			{
-				InactiveContributorsByChannel.Add(Key, EntityID);
+				InactiveContributorsByGroup.Add(Key, EntityID);
 			}
 		}
 	}
 
-	TMap<FBlendChannelSequenceKey, int16> MaxBiasByChannel;
+	TMap<FEntityGroupSequenceKey, int16> MaxBiasByGroup;
 
-	TMultiMap<FBlendChannelSequenceKey, FMovieSceneEntityID> InactiveContributorsByChannel;
+	TMultiMap<FEntityGroupSequenceKey, FMovieSceneEntityID> InactiveContributorsByGroup;
 
-	TMultiMap<FBlendChannelSequenceKey, FMovieSceneEntityID> ActiveContributorsByChannel;
+	TMultiMap<FEntityGroupSequenceKey, FMovieSceneEntityID> ActiveContributorsByGroup;
 
 	UMovieSceneEntitySystemLinker* Linker;
 };
@@ -143,7 +144,11 @@ UMovieSceneHierarchicalBiasSystem::UMovieSceneHierarchicalBiasSystem(const FObje
 
 	if (HasAnyFlags(RF_ClassDefaultObject))
 	{
-		DefineComponentConsumer(GetClass(), FBuiltInComponentTypes::Get()->BlendChannelInput);
+		DefineComponentConsumer(GetClass(), FBuiltInComponentTypes::Get()->Group);
+
+		// Don't flag things with the Ignore tag (due to hierarchical biases) until all systems have
+		// had a chance to take them into account for pre-animated state.
+		DefineImplicitPrerequisite(UMovieSceneCachePreAnimatedStateSystem::StaticClass(), GetClass());
 	}
 }
 
@@ -152,7 +157,7 @@ bool UMovieSceneHierarchicalBiasSystem::IsRelevantImpl(UMovieSceneEntitySystemLi
 	using namespace UE::MovieScene;
 
 	FBuiltInComponentTypes* Components = FBuiltInComponentTypes::Get();
-	return InLinker->EntityManager.ContainsAllComponents({ Components->BlendChannelInput, Components->HierarchicalBias });
+	return InLinker->EntityManager.ContainsAllComponents({ Components->Group, Components->HierarchicalBias });
 }
 
 void UMovieSceneHierarchicalBiasSystem::OnRun(FSystemTaskPrerequisites& InPrerequisites, FSystemSubsequentTasks& Subsequents)
@@ -163,19 +168,20 @@ void UMovieSceneHierarchicalBiasSystem::OnRun(FSystemTaskPrerequisites& InPrereq
 
 	FHierarchicalBiasTask Task(Linker);
 
-	// First, add all the channels that have changed to the map
+	// First, add all the groups that have changed to the map
 	FEntityTaskBuilder()
 	.Read(Components->RootInstanceHandle)
-	.Read(Components->BlendChannelInput)
+	.Read(Components->Group)
 	.FilterAny({ Components->Tags.NeedsLink, Components->Tags.NeedsUnlink })
-	.Iterate_PerEntity(&Linker->EntityManager, [&Task](FRootInstanceHandle RootInstanceHandle, FMovieSceneBlendChannelID BlendChannel){ Task.InitializeChannel(RootInstanceHandle, BlendChannel); });
+	.Iterate_PerEntity(&Linker->EntityManager, [&Task](FRootInstanceHandle RootInstanceHandle, FEntityGroupID GroupID)
+			{ Task.InitializeGroup(RootInstanceHandle, GroupID); });
 
 	if (Task.HasAnyWork())
 	{
 		FEntityTaskBuilder()
 		.ReadEntityIDs()
 		.Read(Components->RootInstanceHandle)
-		.Read(Components->BlendChannelInput)
+		.Read(Components->Group)
 		.ReadOptional(Components->HierarchicalBias)
 		.FilterNone({ Components->Tags.NeedsUnlink })
 		.RunInline_PerAllocation(&Linker->EntityManager, Task);
diff --git a/Engine/Source/Runtime/MovieSceneTracks/Private/Systems/MovieScenePropertyInstantiator.cpp b/Engine/Source/Runtime/MovieSceneTracks/Private/Systems/MovieScenePropertyInstantiator.cpp
--- a/Engine/Source/Runtime/MovieSceneTracks/Private/Systems/MovieScenePropertyInstantiator.cpp
+++ b/Engine/Source/Runtime/MovieSceneTracks/Private/Systems/MovieScenePropertyInstantiator.cpp
@@ -1,25 +1,53 @@
 // Copyright Epic Games, Inc. All Rights Reserved.
 
 #include "Systems/MovieScenePropertyInstantiator.h"
+
 #include "Algo/AllOf.h"
+#include "Algo/IndexOf.h"
+#include "EntitySystem/Interrogation/MovieSceneInterrogationLinker.h"
+#include "EntitySystem/MovieSceneBlenderSystem.h"
 #include "EntitySystem/MovieSceneEntityBuilder.h"
-#include "EntitySystem/MovieScenePropertyBinding.h"
+#include "EntitySystem/MovieSceneEntityGroupingSystem.h"
+#include "EntitySystem/MovieSceneEntityIDs.h"
 #include "EntitySystem/MovieSceneEntitySystemLinker.h"
-#include "EntitySystem/MovieSceneBlenderSystem.h"
+#include "EntitySystem/MovieScenePropertyBinding.h"
 #include "EntitySystem/MovieScenePropertyRegistry.h"
-#include "Systems/MovieScenePiecewiseDoubleBlenderSystem.h"
-#include "EntitySystem/Interrogation/MovieSceneInterrogationLinker.h"
-
-#include "Algo/AllOf.h"
-#include "Algo/IndexOf.h"
 #include "ProfilingDebugging/CountersTrace.h"
+#include "Systems/MovieScenePiecewiseDoubleBlenderSystem.h"
 
 #include UE_INLINE_GENERATED_CPP_BY_NAME(MovieScenePropertyInstantiator)
 
 DECLARE_CYCLE_STAT(TEXT("DiscoverInvalidatedProperties"), MovieSceneEval_DiscoverInvalidatedProperties, STATGROUP_MovieSceneECS);
 DECLARE_CYCLE_STAT(TEXT("ProcessInvalidatedProperties"), MovieSceneEval_ProcessInvalidatedProperties, STATGROUP_MovieSceneECS);
 DECLARE_CYCLE_STAT(TEXT("InitializePropertyMetaData"), MovieSceneEval_InitializePropertyMetaData, STATGROUP_MovieSceneECS);
+UE_DISABLE_OPTIMIZATION
+namespace UE::MovieScene
+{
 
+struct FPropertyInstantiatorGroupingPolicy
+{
+	using GroupKeyType = TTuple<UObject*, FName>;
+
+	bool GetGroupKey(UObject* Object, const FMovieScenePropertyBinding& PropertyBinding, GroupKeyType& OutGroupKey)
+	{
+		OutGroupKey = MakeTuple(Object, PropertyBinding.PropertyPath);
+		return true;
+	}
+
+#if WITH_EDITOR
+	bool OnObjectsReplaced(GroupKeyType& InOutKey, const TMap<UObject*, UObject*>& ReplacementMap)
+	{
+		if (UObject* const * NewObject = ReplacementMap.Find(InOutKey.Key))
+		{
+			InOutKey.Key = *NewObject;
+			return true;
+		}
+		return false;
+	}
+#endif
+};
+
+} // namespace UE::MovieScene
 
 void UMovieScenePropertyInstantiatorSystem::FHierarchicalMetaData::CombineWith(const FHierarchicalMetaData& Other)
 {
@@ -60,6 +88,8 @@ UMovieScenePropertyInstantiatorSystem::UMovieScenePropertyInstantiatorSystem(con
 	if (HasAnyFlags(RF_ClassDefaultObject))
 	{
 		DefineComponentConsumer(GetClass(), BuiltInComponents->BoundObject);
+		DefineComponentConsumer(GetClass(), BuiltInComponents->Group);
+
 		DefineComponentProducer(GetClass(), BuiltInComponents->BlendChannelInput);
 		DefineComponentProducer(GetClass(), BuiltInComponents->HierarchicalBlendTarget);
 		DefineComponentProducer(GetClass(), BuiltInComponents->SymbolicTags.CreatesEntities);
@@ -79,10 +109,17 @@ UE::MovieScene::FPropertyStats UMovieScenePropertyInstantiatorSystem::GetStatsFo
 
 void UMovieScenePropertyInstantiatorSystem::OnLink()
 {
+	using namespace UE::MovieScene;
+
 	CleanFastPathMask.Reset();
 	CleanFastPathMask.SetAll({ BuiltInComponents->FastPropertyOffset, BuiltInComponents->SlowProperty, BuiltInComponents->CustomPropertyIndex });
 	CleanFastPathMask.CombineWithBitwiseOR(Linker->EntityManager.GetComponents()->GetMigrationMask(), EBitwiseOperatorFlags::MaxSize);
 
+	UMovieSceneEntityGroupingSystem* GroupingSystem = Linker->LinkSystem<UMovieSceneEntityGroupingSystem>();
+	PropertyGroupingKey = GroupingSystem->AddGrouping(
+			FPropertyInstantiatorGroupingPolicy(),
+			BuiltInComponents->BoundObject, BuiltInComponents->PropertyBinding);
+			
 #if WITH_EDITOR
 	FCoreUObjectDelegates::OnObjectsReplaced.AddUObject(this, &UMovieScenePropertyInstantiatorSystem::OnObjectsReplaced);
 #endif
@@ -95,16 +132,12 @@ void UMovieScenePropertyInstantiatorSystem::OnUnlink()
 	const bool bAllPropertiesClean = (
 				ResolvedProperties.Num() == 0 &&
 				Contributors.Num() == 0 &&
-				NewContributors.Num() == 0 &&
-				EntityToProperty.Num() == 0 &&
-				ObjectPropertyToResolvedIndex.Num() == 0);
+				NewContributors.Num() == 0);
 	if (!ensure(bAllPropertiesClean))
 	{
 		ResolvedProperties.Reset();
 		Contributors.Reset();
 		NewContributors.Reset();
-		EntityToProperty.Reset();
-		ObjectPropertyToResolvedIndex.Reset();
 		PropertyStats.Reset();
 	}
 
@@ -127,6 +160,13 @@ void UMovieScenePropertyInstantiatorSystem::OnUnlink()
 		SaveGlobalStateTasks.Reset();
 	}
 
+	UMovieSceneEntityGroupingSystem* GroupingSystem = Linker->FindSystem<UMovieSceneEntityGroupingSystem>();
+	if (ensure(GroupingSystem))
+	{
+		GroupingSystem->RemoveGrouping(PropertyGroupingKey);
+	}
+	PropertyGroupingKey = FEntityGroupingPolicyKey();
+	
 #if WITH_EDITOR
 	FCoreUObjectDelegates::OnObjectsReplaced.RemoveAll(this);
 #endif
@@ -177,19 +217,6 @@ void UMovieScenePropertyInstantiatorSystem::OnCleanTaggedGarbage()
 void UMovieScenePropertyInstantiatorSystem::OnObjectsReplaced(const TMap<UObject*, UObject*>& ReplacementMap)
 {
 #if WITH_EDITOR
-	TArray<TTuple<UObject*, FName>> ObjectPropertyKeys;
-	ObjectPropertyToResolvedIndex.GetKeys(ObjectPropertyKeys);
-	for (const TTuple<UObject*, FName>& Key : ObjectPropertyKeys)
-	{
-		if (UObject* const* NewObject = ReplacementMap.Find(Key.Key))
-		{
-			FName PropertyName = Key.Value;
-			int32 ResolvedIndex;
-			ObjectPropertyToResolvedIndex.RemoveAndCopyValue(Key, ResolvedIndex);
-			ObjectPropertyToResolvedIndex.Add(TTuple<UObject*, FName>(*NewObject, PropertyName), ResolvedIndex);
-		}
-	}
-
 	for (auto It = ResolvedProperties.CreateIterator(); It; ++It)
 	{
 		FObjectPropertyInfo& ResolvedProperty = (*It);
@@ -220,9 +247,6 @@ void UMovieScenePropertyInstantiatorSystem::OnRun(FSystemTaskPrerequisites& InPr
 	{
 		InitializePropertyMetaData(InPrerequisites, Subsequents);
 	}
-
-	ObjectPropertyToResolvedIndex.Compact();
-	EntityToProperty.Compact();
 }
 
 void UMovieScenePropertyInstantiatorSystem::DiscoverInvalidatedProperties(TBitArray<>& OutInvalidatedProperties)
@@ -231,7 +255,7 @@ void UMovieScenePropertyInstantiatorSystem::DiscoverInvalidatedProperties(TBitAr
 
 	MOVIESCENE_DETAILED_SCOPE_CYCLE_COUNTER(MovieSceneEval_DiscoverInvalidatedProperties);
 
-	TArrayView<const FPropertyDefinition> Properties = this->BuiltInComponents->PropertyRegistry.GetProperties();
+	TArrayView<const FPropertyDefinition> Properties = BuiltInComponents->PropertyRegistry.GetProperties();
 
 	PropertyStats.SetNum(Properties.Num());
 
@@ -245,7 +269,7 @@ void UMovieScenePropertyInstantiatorSystem::DiscoverNewProperties(TBitArray<>&Ou
 
 	TArrayView<const FPropertyDefinition> Properties = this->BuiltInComponents->PropertyRegistry.GetProperties();
 
-	auto VisitNewProperties = [this, Properties, &OutInvalidatedProperties](FEntityAllocationIteratorItem AllocationItem, const FMovieSceneEntityID* EntityIDs, UObject* const * ObjectPtrs, const FMovieScenePropertyBinding* PropertyPtrs, const int16* HierarchicalBiases)
+	auto VisitNewProperties = [this, Properties, &OutInvalidatedProperties](FEntityAllocationIteratorItem AllocationItem, const FMovieSceneEntityID* EntityIDs, UObject* const * ObjectPtrs, const FMovieScenePropertyBinding* PropertyPtrs, const FEntityGroupID* GroupIDs, const int16* HierarchicalBiases)
 	{
 		const FEntityAllocation* Allocation     = AllocationItem.GetAllocation();
 		const FComponentMask&    AllocationType = AllocationItem.GetAllocationType();
@@ -271,15 +295,11 @@ void UMovieScenePropertyInstantiatorSystem::DiscoverNewProperties(TBitArray<>&Ou
 
 		for (int32 Index = 0; Index < Allocation->Num(); ++Index)
 		{
-			const int32 PropertyIndex = this->ResolveProperty(CustomAccessors, ObjectPtrs[Index], PropertyPtrs[Index], PropertyDefinitionIndex);
-			
-			// If the property did not resolve, we still add it to the LUT
-			// So that the ensure inside VisitExpiredEntities only fires
-			// for genuine link/unlink disparities
-			this->EntityToProperty.Add(EntityIDs[Index], PropertyIndex);
-
-			if (PropertyIndex != INDEX_NONE)
+			const bool bResolved = this->ResolveProperty(CustomAccessors, ObjectPtrs[Index], PropertyPtrs[Index], GroupIDs[Index], PropertyDefinitionIndex);
+			if (bResolved)
 			{
+				const int32 PropertyIndex = GroupIDs[Index].GroupIndex;
+
 				FContributorKey Key { PropertyIndex, HierarchicalBiases ? HierarchicalBiases[Index] : DefaultHierarchicalBias };
 				this->Contributors.Add(Key, EntityIDs[Index]);
 				this->NewContributors.Add(Key, EntityIDs[Index]);
@@ -294,6 +314,7 @@ void UMovieScenePropertyInstantiatorSystem::DiscoverNewProperties(TBitArray<>&Ou
 	.ReadEntityIDs()
 	.Read(BuiltInComponents->BoundObject)
 	.Read(BuiltInComponents->PropertyBinding)
+	.Read(BuiltInComponents->Group)
 	.ReadOptional(BuiltInComponents->HierarchicalBias)
 	.FilterNone({ BuiltInComponents->BlendChannelOutput })
 	.FilterAll({ BuiltInComponents->Tags.NeedsLink })
@@ -304,27 +325,22 @@ void UMovieScenePropertyInstantiatorSystem::DiscoverExpiredProperties(TBitArray<
 {
 	using namespace UE::MovieScene;
 
-	auto VisitExpiredEntities = [this, &OutInvalidatedProperties](FMovieSceneEntityID EntityID)
+	auto VisitExpiredEntities = [this, &OutInvalidatedProperties](FMovieSceneEntityID EntityID, const FEntityGroupID& GroupID)
 	{
-		const int32* PropertyIndexPtr = this->EntityToProperty.Find(EntityID);
-		if (ensureMsgf(PropertyIndexPtr, TEXT("Could not find entity to clean up from linker entity ID - this indicates VisitNewProperties never got called for this entity, or a garbage collection has somehow destroyed the entity without flushing the ecs.")))
+		const int32 PropertyIndex = GroupID.GroupIndex;
+		if (PropertyIndex != INDEX_NONE)
 		{
-			const int32 PropertyIndex = *PropertyIndexPtr;
-			if (PropertyIndex != INDEX_NONE)
-			{
-				OutInvalidatedProperties.PadToNum(PropertyIndex + 1, false);
-				OutInvalidatedProperties[PropertyIndex] = true;
+			OutInvalidatedProperties.PadToNum(PropertyIndex + 1, false);
+			OutInvalidatedProperties[PropertyIndex] = true;
 
-				this->Contributors.Remove(PropertyIndex, EntityID);
-			}
-
-			// Always remove the entity ID from the LUT
-			this->EntityToProperty.Remove(EntityID);
+			this->Contributors.Remove(PropertyIndex, EntityID);
 		}
+
 	};
 
 	FEntityTaskBuilder()
 	.ReadEntityIDs()
+	.Read(BuiltInComponents->Group)
 	.FilterNone({ BuiltInComponents->BlendChannelOutput })
 	.FilterAll({ BuiltInComponents->BoundObject, BuiltInComponents->PropertyBinding, BuiltInComponents->Tags.NeedsUnlink })
 	.Iterate_PerEntity(&Linker->EntityManager, VisitExpiredEntities);
@@ -543,16 +559,6 @@ void UMovieScenePropertyInstantiatorSystem::DestroyStaleProperty(int32 PropertyI
 
 void UMovieScenePropertyInstantiatorSystem::PostDestroyStaleProperties()
 {
-	// @todo: If perf is a real issue with this look, we could call ObjectPropertyToResolvedIndex.Remove(MakeTuple(PropertyInfo->BoundObject, PropertyInfo->PropertyBinding.PropertyPath));
-	// In the loop above, but it is possible that BoundObject no longer relates to a valid object at that point
-	for (auto It = ObjectPropertyToResolvedIndex.CreateIterator(); It; ++It)
-	{
-		if (!ResolvedProperties.IsAllocated(It.Value()))
-		{
-			It.RemoveCurrent();
-		}
-	}
-
 	ResolvedProperties.Shrink();
 }
 
@@ -1098,55 +1104,55 @@ void UMovieScenePropertyInstantiatorSystem::InitializeBlendPath(const FPropertyP
 	check(!Linker->EntityManager.HasComponent(Params.PropertyInfo->FinalBlendOutputID, BuiltInComponents->BlendChannelInput));
 }
 
-int32 UMovieScenePropertyInstantiatorSystem::ResolveProperty(UE::MovieScene::FCustomAccessorView CustomAccessors, UObject* Object, const FMovieScenePropertyBinding& PropertyBinding, int32 PropertyDefinitionIndex)
+bool UMovieScenePropertyInstantiatorSystem::ResolveProperty(UE::MovieScene::FCustomAccessorView CustomAccessors, UObject* Object, const FMovieScenePropertyBinding& PropertyBinding, const UE::MovieScene::FEntityGroupID& GroupID, int32 PropertyDefinitionIndex)
 {
 	using namespace UE::MovieScene;
 
-	TTuple<UObject*, FName> Key = MakeTuple(Object, PropertyBinding.PropertyPath);
-	if (const int32* ExistingPropertyIndex = ObjectPropertyToResolvedIndex.Find(Key))
+	if (ResolvedProperties.IsValidIndex(GroupID.GroupIndex))
 	{
-		return *ExistingPropertyIndex;
+#if !UE_BUILD_SHIPPING
+		const FObjectPropertyInfo& ResolvedProperty = ResolvedProperties[GroupID.GroupIndex];
+		ensure(ResolvedProperty.BoundObject == Object);
+		ensure(ResolvedProperty.PropertyBinding == PropertyBinding);
+#endif
+		return true;
 	}
 
 	TOptional<FResolvedProperty> ResolvedProperty = FPropertyRegistry::ResolveProperty(Object, PropertyBinding, CustomAccessors);
 	if (!ResolvedProperty.IsSet())
 	{
 		UE_LOG(LogMovieScene, Warning, TEXT("Unable to resolve property '%s' from '%s' instance '%s'"), *PropertyBinding.PropertyPath.ToString(), *Object->GetClass()->GetName(), *Object->GetName());
-		return INDEX_NONE;
+		return false;
 	}
 
-	FObjectPropertyInfo NewInfo(MoveTemp(ResolvedProperty.GetValue()));
+	ResolvedProperties.EmplaceAt(GroupID.GroupIndex, MoveTemp(ResolvedProperty.GetValue()));
+	FObjectPropertyInfo& NewInfo = ResolvedProperties[GroupID.GroupIndex];
 
 	NewInfo.BoundObject = Object;
 	NewInfo.PropertyBinding = PropertyBinding;
 	NewInfo.PropertyDefinitionIndex = PropertyDefinitionIndex;
 
-	const int32 NewPropertyIndex = ResolvedProperties.Add(NewInfo);
-
-	ObjectPropertyToResolvedIndex.Add(Key, NewPropertyIndex);
-
 	++PropertyStats[PropertyDefinitionIndex].NumProperties;
 
-	return NewPropertyIndex;
+	return true;
 }
 
 UE::MovieScene::FPropertyRecomposerPropertyInfo UMovieScenePropertyInstantiatorSystem::FindPropertyFromSource(FMovieSceneEntityID EntityID, UObject* Object) const
 {
 	using namespace UE::MovieScene;
 
-	TOptionalComponentReader<int16>                      HBiasComponent  = Linker->EntityManager.ReadComponent(EntityID, BuiltInComponents->HierarchicalBias);
 	TOptionalComponentReader<FMovieScenePropertyBinding> PropertyBinding = Linker->EntityManager.ReadComponent(EntityID, BuiltInComponents->PropertyBinding);
-	if (!PropertyBinding)
+	TOptionalComponentReader<FEntityGroupID>             GroupID         = Linker->EntityManager.ReadComponent(EntityID, BuiltInComponents->Group);
+	if (!PropertyBinding || !GroupID)
 	{
 		return FPropertyRecomposerPropertyInfo::Invalid();
 	}
 
-	const int16 HBias = HBiasComponent ? *HBiasComponent : 0;
-	TTuple<UObject*, FName> Key = MakeTuple(Object, PropertyBinding->PropertyPath);
-	if (const int32* PropertyIndex = ObjectPropertyToResolvedIndex.Find(Key))
+	const int32 PropertyIndex = GroupID->GroupIndex;
+	if (PropertyIndex != INDEX_NONE && ensure(ResolvedProperties.IsValidIndex(PropertyIndex)))
 	{
-		const uint16 BlendChannel = ResolvedProperties[*PropertyIndex].BlendChannel;
-		const FObjectPropertyInfo& PropertyInfo = ResolvedProperties[*PropertyIndex];
+		const uint16 BlendChannel = ResolvedProperties[PropertyIndex].BlendChannel;
+		const FObjectPropertyInfo& PropertyInfo = ResolvedProperties[PropertyIndex];
 		return FPropertyRecomposerPropertyInfo { BlendChannel, PropertyInfo.Blender.Get(), PropertyInfo.FinalBlendOutputID };
 	}
 
@@ -1208,3 +1214,4 @@ void UMovieScenePropertyInstantiatorSystem::FPropertyParameters::MakeOutputCompo
 		OutComponentType.Remove(FBuiltInComponentTypes::Get()->Tags.RestoreState);
 	}
 }
+UE_ENABLE_OPTIMIZATION
diff --git a/Engine/Source/Runtime/MovieSceneTracks/Public/Systems/MovieScenePropertyInstantiator.h b/Engine/Source/Runtime/MovieSceneTracks/Public/Systems/MovieScenePropertyInstantiator.h
--- a/Engine/Source/Runtime/MovieSceneTracks/Public/Systems/MovieScenePropertyInstantiator.h
+++ b/Engine/Source/Runtime/MovieSceneTracks/Public/Systems/MovieScenePropertyInstantiator.h
@@ -2,13 +2,14 @@
 
 #pragma once
 
+#include "EntitySystem/BuiltInComponentTypes.h"
+#include "EntitySystem/EntityAllocationIterator.h"
+#include "EntitySystem/MovieSceneDecompositionQuery.h"
 #include "EntitySystem/MovieSceneEntityIDs.h"
-#include "EntitySystem/MovieSceneEntitySystem.h"
 #include "EntitySystem/MovieSceneEntityInstantiatorSystem.h"
-#include "EntitySystem/MovieSceneDecompositionQuery.h"
-#include "EntitySystem/BuiltInComponentTypes.h"
-#include "EntitySystem/MovieScenePropertyRegistry.h"
+#include "EntitySystem/MovieSceneEntitySystem.h"
 #include "EntitySystem/MovieScenePropertyComponentHandler.h"
+#include "EntitySystem/MovieScenePropertyRegistry.h"
 
 #include "Misc/TVariant.h"
 #include "Misc/Optional.h"
@@ -50,7 +51,6 @@ public:
 	 */
 	MOVIESCENETRACKS_API UE::MovieScene::FPropertyStats GetStatsForProperty(UE::MovieScene::FCompositePropertyTypeID PropertyID) const;
 
-
 	/**
 	 * Recompose a value from the constituent parts specified in InQuery, taking into account the weightings of each channel.
 	 * For instance, if a property comprises 3 additive values (a:1, b:2, c:3), and we recompose 'a' with an InCurrentValue of 10, the result for a would be 5.
@@ -245,7 +245,7 @@ private:
 
 	FSetupBlenderSystemResult SetupBlenderSystem(const FPropertyParameters& Params);
 
-	int32 ResolveProperty(UE::MovieScene::FCustomAccessorView CustomAccessors, UObject* Object, const FMovieScenePropertyBinding& PropertyBinding, int32 PropertyDefinitionIndex);
+	bool ResolveProperty(UE::MovieScene::FCustomAccessorView CustomAccessors, UObject* Object, const FMovieScenePropertyBinding& PropertyBinding, const UE::MovieScene::FEntityGroupID& GroupID, int32 PropertyDefinitionIndex);
 
 	UE::MovieScene::FPropertyRecomposerPropertyInfo FindPropertyFromSource(FMovieSceneEntityID EntityID, UObject* Object) const;
 
@@ -255,15 +255,12 @@ private:
 
 	static constexpr uint16 INVALID_BLEND_CHANNEL = uint16(-1);
 
+	UE::MovieScene::FEntityGroupingPolicyKey PropertyGroupingKey;
+
 	TSparseArray<FObjectPropertyInfo> ResolvedProperties;
 	TMultiMap<FContributorKey, FMovieSceneEntityID> Contributors;
 	TMultiMap<FContributorKey, FMovieSceneEntityID> NewContributors;
 
-	/** Reverse lookup from an entity to the index within ResolvedProperties that it animates.
-	 * @note: can contain INDEX_NONE for properties that have not resolved. */
-	TMap<FMovieSceneEntityID, int32> EntityToProperty;
-	TMap< TTuple<UObject*, FName>, int32 > ObjectPropertyToResolvedIndex;
-
 	TArray<UE::MovieScene::FPropertyStats> PropertyStats;
 
 	UE::MovieScene::FComponentMask CleanFastPathMask;
