MovieScene: For budgeted evaluations, split ConditionalRecompile from SequenceInstanceUpdate phase.

diff --git a/Engine/Source/Runtime/MovieScene/Private/EntitySystem/MovieSceneEntitySystemRunner.cpp b/Engine/Source/Runtime/MovieScene/Private/EntitySystem/MovieSceneEntitySystemRunner.cpp
--- a/Engine/Source/Runtime/MovieScene/Private/EntitySystem/MovieSceneEntitySystemRunner.cpp
+++ b/Engine/Source/Runtime/MovieScene/Private/EntitySystem/MovieSceneEntitySystemRunner.cpp
@@ -345,7 +345,14 @@ bool FMovieSceneEntitySystemRunner::FlushNext(UMovieSceneEntitySystemLinker* Lin
 		return StartEvaluation(Linker);
 	}
 
-	// Step 2: Update sequence instances and import entities
+	// Step 2: Execute any conditional recompiles for dirtied sequences.
+	if (EnumHasAnyFlags(FlushState, ERunnerFlushState::ConditionalRecompile))
+	{
+		EnterFlushState(ERunnerFlushState::ConditionalRecompile);
+		return GameThread_ConditionalRecompile(Linker);
+	}
+
+	// Step 3: Update sequence instances and import entities
 	//         This is the entry-point for for each iteration
 	if (EnumHasAnyFlags(FlushState, ERunnerFlushState::Import))
 	{
@@ -353,7 +360,7 @@ bool FMovieSceneEntitySystemRunner::FlushNext(UMovieSceneEntitySystemLinker* Lin
 		return GameThread_UpdateSequenceInstances(Linker);
 	}
 
-	// Step 3: Conditionally run the spawn phase of the system graph
+	// Step 4: Conditionally run the spawn phase of the system graph
 	//
 	if (EnumHasAnyFlags(FlushState, ERunnerFlushState::Spawn))
 	{
@@ -361,7 +368,7 @@ bool FMovieSceneEntitySystemRunner::FlushNext(UMovieSceneEntitySystemLinker* Lin
 		return GameThread_SpawnPhase(Linker);
 	}
 
-	// Step 4: Conditionally run the instantiation phase of the system graph
+	// Step 5: Conditionally run the instantiation phase of the system graph
 	//
 	if (EnumHasAnyFlags(FlushState, ERunnerFlushState::Instantiation))
 	{
@@ -369,7 +376,7 @@ bool FMovieSceneEntitySystemRunner::FlushNext(UMovieSceneEntitySystemLinker* Lin
 		return GameThread_InstantiationPhase(Linker);
 	}
 
-	// Step 5: Run the evaluation phase of the system graph
+	// Step 6: Run the evaluation phase of the system graph
 	//
 	if (EnumHasAnyFlags(FlushState, ERunnerFlushState::Evaluation))
 	{
@@ -377,7 +384,7 @@ bool FMovieSceneEntitySystemRunner::FlushNext(UMovieSceneEntitySystemLinker* Lin
 		return GameThread_EvaluationPhase(Linker);
 	}
 
-	// Step 6: Run the finalization phase of the system graph, including legacy templates
+	// Step 7: Run the finalization phase of the system graph, including legacy templates
 	//
 	if (EnumHasAnyFlags(FlushState, ERunnerFlushState::Finalization))
 	{
@@ -386,7 +393,7 @@ bool FMovieSceneEntitySystemRunner::FlushNext(UMovieSceneEntitySystemLinker* Lin
 		return true;
 	}
 
-	// Step 7: Trigger events
+	// Step 8: Trigger events
 	//
 	if (EnumHasAnyFlags(FlushState, ERunnerFlushState::EventTriggers))
 	{
@@ -395,7 +402,7 @@ bool FMovieSceneEntitySystemRunner::FlushNext(UMovieSceneEntitySystemLinker* Lin
 		return true;
 	}
 
-	// Step 7: Call PostEvaluation on all current sequence instances
+	// Step 9: Call PostEvaluation on all current sequence instances
 	//
 	if (EnumHasAnyFlags(FlushState, ERunnerFlushState::PostEvaluation))
 	{
@@ -404,8 +411,8 @@ bool FMovieSceneEntitySystemRunner::FlushNext(UMovieSceneEntitySystemLinker* Lin
 		return true;
 	}
 
-	// Step 7: Perform any clean up and broadcast 'end eval' events
-	//         NOTE: Only ever called once regardless of how many iterations we perform
+	// Step 10: Perform any clean up and broadcast 'end eval' events
+	//          NOTE: Only ever called once regardless of how many iterations we perform
 	if (EnumHasAnyFlags(FlushState, ERunnerFlushState::End))
 	{
 		EnterFlushState(ERunnerFlushState::End);
@@ -504,14 +511,14 @@ void FMovieSceneEntitySystemRunner::ResetFlushState()
 	using namespace UE::MovieScene;
 
 	constexpr ERunnerFlushState StatesThatTriggerReset = ERunnerFlushState::Instantiation | ERunnerFlushState::Evaluation | ERunnerFlushState::Finalization | ERunnerFlushState::EventTriggers;
-	constexpr ERunnerFlushState StatesThatDontTriggerReset = ERunnerFlushState::Import | ERunnerFlushState::Spawn;
+	constexpr ERunnerFlushState StatesThatDontTriggerReset = ERunnerFlushState::ConditionalRecompile | ERunnerFlushState::Import | ERunnerFlushState::Spawn;
 
 	// We only need to reset the flush state if we are far enough through an evaluation
 	if (EnumHasAnyFlags(FlushState, StatesThatTriggerReset) && !EnumHasAnyFlags(FlushState, StatesThatDontTriggerReset))
 	{
 		// When resetting - we don't need to (or want to) re-import anything, we just want to re-run the current
 		// frame of updates 
-		FlushState = UE::MovieScene::ERunnerFlushState::LoopEval & ~UE::MovieScene::ERunnerFlushState::Import;
+		FlushState = UE::MovieScene::ERunnerFlushState::LoopEval & ~(UE::MovieScene::ERunnerFlushState::ConditionalRecompile | UE::MovieScene::ERunnerFlushState::Import);
 	}
 }
 
@@ -556,6 +563,34 @@ void FMovieSceneEntitySystemRunner::SkipFlushState(UE::MovieScene::ERunnerFlushS
 	FlushState &= ~FlushStateToSkip;
 }
 
+bool FMovieSceneEntitySystemRunner::GameThread_ConditionalRecompile(UMovieSceneEntitySystemLinker* Linker)
+{
+	using namespace UE::MovieScene;
+
+	TRACE_CPUPROFILER_EVENT_SCOPE(FMovieSceneEntitySystemRunner::GameThread_ConditionalRecompile);
+
+	FInstanceRegistry* InstanceRegistry = GetInstanceRegistry();
+
+	for (int32 UpdateIndex = 0; UpdateIndex < UpdateQueue.Num(); ++UpdateIndex)
+	{
+		const FUpdateParamsAndContext& Request = UpdateQueue[UpdateIndex];
+		if (!InstanceRegistry->IsHandleValid(Request.Params.InstanceHandle))
+		{
+			continue;
+		}
+
+		FSequenceInstance& Instance = InstanceRegistry->MutateInstance(Request.Params.InstanceHandle);
+		if (!Instance.IsRootSequence())
+		{
+			continue;
+		}
+
+		Instance.ConditionalRecompile(Linker);
+	}
+
+	return true;
+}
+
 bool FMovieSceneEntitySystemRunner::GameThread_UpdateSequenceInstances(UMovieSceneEntitySystemLinker* Linker)
 {
 	using namespace UE::MovieScene;
@@ -623,7 +658,6 @@ bool FMovieSceneEntitySystemRunner::GameThread_UpdateSequenceInstances(UMovieSce
 				continue;
 			}
 
-			Instance.ConditionalRecompile(Linker);
 			Instance.DissectContext(Linker, Request.Context, Dissections);
 
 			if (Dissections.Num() != 0)
diff --git a/Engine/Source/Runtime/MovieScene/Public/EntitySystem/MovieSceneEntitySystemRunner.h b/Engine/Source/Runtime/MovieScene/Public/EntitySystem/MovieSceneEntitySystemRunner.h
--- a/Engine/Source/Runtime/MovieScene/Public/EntitySystem/MovieSceneEntitySystemRunner.h
+++ b/Engine/Source/Runtime/MovieScene/Public/EntitySystem/MovieSceneEntitySystemRunner.h
@@ -28,19 +28,20 @@ struct FInstanceRegistry;
 /** Bit-mask enumeration that defines tasks that need running */
 enum class ERunnerFlushState
 {
-	None           = 0,				// Signifies no evaluation is currently underway
-	Start          = 1 << 0,		// Sets up initial evaluation flags for external players and listeners
-	Import         = 1 << 1,		// Update sequence instances and import entities into the entity manager
-	Spawn          = 1 << 2,		// Perform the Spawn phase in the Entity System Graph
-	Instantiation  = 1 << 3,		// Perform the Instantiation phase in the Entity System Graph
-	Evaluation     = 1 << 4,		// Perform the Evaluation phase in the Entity System Graph
-	Finalization   = 1 << 5,		// Perform the Finalization phase in the Entity System Graph and trigger any external events
-	EventTriggers  = 1 << 6,		// (re-entrant) Triggers any bound event triggers - skipped by Finalization if the delegate is not bound
-	PostEvaluation = 1 << 7,		// (re-entrant) Call post evaluation callbacks on sequence instances
-	End            = 1 << 8,		// Counterpart for Start - resets external players and listeners
+	None                 = 0,			// Signifies no evaluation is currently underway
+	Start                = 1 << 0,		// Sets up initial evaluation flags for external players and listeners
+	ConditionalRecompile = 1 << 1,		// Conditional recompile of dirtied sequences.
+	Import               = 1 << 2,		// Update sequence instances and import entities into the entity manager
+	Spawn                = 1 << 3,		// Perform the Spawn phase in the Entity System Graph
+	Instantiation        = 1 << 4,		// Perform the Instantiation phase in the Entity System Graph
+	Evaluation           = 1 << 5,		// Perform the Evaluation phase in the Entity System Graph
+	Finalization         = 1 << 6,		// Perform the Finalization phase in the Entity System Graph and trigger any external events
+	EventTriggers        = 1 << 7,		// (re-entrant) Triggers any bound event triggers - skipped by Finalization if the delegate is not bound
+	PostEvaluation       = 1 << 8,		// (re-entrant) Call post evaluation callbacks on sequence instances
+	End                  = 1 << 9,		// Counterpart for Start - resets external players and listeners
 
 	// Signifies that, during the Finalization task, there were still outstanding tasks and we need to perform another iteration
-	LoopEval       = Import | Spawn | Instantiation | Evaluation | Finalization | EventTriggers | PostEvaluation | End,
+	LoopEval       = ConditionalRecompile | Import | Spawn | Instantiation | Evaluation | Finalization | EventTriggers | PostEvaluation | End,
 
 	// Initial flush state
 	Everything     = Start | LoopEval,
@@ -192,6 +193,8 @@ private:
 	 */
 	MOVIESCENE_API bool StartEvaluation(UMovieSceneEntitySystemLinker* Linker);
 
+	/** Execute any pending conditional recompiles on the currently queued update requests */
+	MOVIESCENE_API bool GameThread_ConditionalRecompile(UMovieSceneEntitySystemLinker* Linker);
 	/** Update sequence instances based on currently queued update requests, or outstanding dissected updates */
 	MOVIESCENE_API bool GameThread_UpdateSequenceInstances(UMovieSceneEntitySystemLinker* Linker);
 	/** Execute the spawn phase of the entity system graph, if there is anything to do */
