Sequencer: Do not invalidate ledger entities for sub sequences that have not changed

diff --git a/Engine/Source/Runtime/MovieScene/Private/EntitySystem/MovieSceneSequenceInstance.cpp b/Engine/Source/Runtime/MovieScene/Private/EntitySystem/MovieSceneSequenceInstance.cpp
--- a/Engine/Source/Runtime/MovieScene/Private/EntitySystem/MovieSceneSequenceInstance.cpp
+++ b/Engine/Source/Runtime/MovieScene/Private/EntitySystem/MovieSceneSequenceInstance.cpp
@@ -182,7 +182,7 @@ void FSequenceInstance::InvalidateCachedData(UMovieSceneEntitySystemLinker* Link
 
 		ISequenceUpdater::FactoryInstance(SequenceUpdater, CompiledDataManager, RootCompiledDataID);
 
-		SequenceUpdater->InvalidateCachedData(Linker);
+		SequenceUpdater->InvalidateCachedData(Linker, RootInstanceHandle);
 		SequenceUpdater->PopulateUpdateFlags(Linker, Player, UpdateFlags);
 
 		if (LegacyEvaluator)
diff --git a/Engine/Source/Runtime/MovieScene/Private/EntitySystem/MovieSceneSequenceUpdaters.cpp b/Engine/Source/Runtime/MovieScene/Private/EntitySystem/MovieSceneSequenceUpdaters.cpp
--- a/Engine/Source/Runtime/MovieScene/Private/EntitySystem/MovieSceneSequenceUpdaters.cpp
+++ b/Engine/Source/Runtime/MovieScene/Private/EntitySystem/MovieSceneSequenceUpdaters.cpp
@@ -41,7 +41,7 @@ struct FSequenceUpdater_Flat : ISequenceUpdater
 	virtual void Update(UMovieSceneEntitySystemLinker* Linker, FRootInstanceHandle InstanceHandle, IMovieScenePlayer* InPlayer, const FMovieSceneContext& Context) override;
 	virtual bool CanFinishImmediately(UMovieSceneEntitySystemLinker* InLinker, FRootInstanceHandle RootInstanceHandle) const override;
 	virtual void Finish(UMovieSceneEntitySystemLinker* Linker, FRootInstanceHandle InstanceHandle, IMovieScenePlayer* InPlayer) override;
-	virtual void InvalidateCachedData(UMovieSceneEntitySystemLinker* Linker) override;
+	virtual void InvalidateCachedData(UMovieSceneEntitySystemLinker* Linker, FRootInstanceHandle RootInstanceHandle) override;
 	virtual void Destroy(UMovieSceneEntitySystemLinker* Linker) override;
 	virtual TUniquePtr<ISequenceUpdater> MigrateToHierarchical() override;
 	virtual FInstanceHandle FindSubInstance(FMovieSceneSequenceID SubSequenceID) const override { return FInstanceHandle(); }
@@ -70,7 +70,7 @@ struct FSequenceUpdater_Hierarchical : ISequenceUpdater
 	virtual void Update(UMovieSceneEntitySystemLinker* Linker, FRootInstanceHandle InstanceHandle, IMovieScenePlayer* InPlayer, const FMovieSceneContext& Context) override;
 	virtual bool CanFinishImmediately(UMovieSceneEntitySystemLinker* InLinker, FRootInstanceHandle RootInstanceHandle) const override;
 	virtual void Finish(UMovieSceneEntitySystemLinker* Linker, FRootInstanceHandle InstanceHandle, IMovieScenePlayer* InPlayer) override;
-	virtual void InvalidateCachedData(UMovieSceneEntitySystemLinker* Linker) override;
+	virtual void InvalidateCachedData(UMovieSceneEntitySystemLinker* Linker, FRootInstanceHandle RootInstanceHandle) override;
 	virtual void Destroy(UMovieSceneEntitySystemLinker* Linker) override;
 	virtual TUniquePtr<ISequenceUpdater> MigrateToHierarchical() override { return nullptr; }
 	virtual FInstanceHandle FindSubInstance(FMovieSceneSequenceID SubSequenceID) const override { return SequenceInstances.FindRef(SubSequenceID).Handle; }
@@ -80,7 +80,7 @@ private:
 
 	TRange<FFrameNumber> UpdateEntitiesForSequence(const FMovieSceneEntityComponentField* ComponentField, FFrameTime SequenceTime, FMovieSceneEvaluationFieldEntitySet& OutEntities);
 
-	FInstanceHandle GetOrCreateSequenceInstance(IMovieScenePlayer* InPlayer, const FMovieSceneSequenceHierarchy* Hierarchy, FInstanceRegistry* InstanceRegistry, FRootInstanceHandle RootInstanceHandle, FMovieSceneSequenceID SequenceID);
+	FInstanceHandle GetOrCreateSequenceInstance(IMovieScenePlayer* InPlayer, UMovieSceneSequence* SubSequence, const FMovieSceneSequenceHierarchy* Hierarchy, FInstanceRegistry* InstanceRegistry, FRootInstanceHandle RootInstanceHandle, FMovieSceneSequenceID SequenceID);
 
 private:
 
@@ -88,6 +88,7 @@ private:
 
 	struct FSubInstanceData
 	{
+		FGuid SequenceSignature;
 		FInstanceHandle Handle;
 		bool bNeedsDestroy = false;
 	};
@@ -323,14 +324,14 @@ bool FSequenceUpdater_Flat::CanFinishImmediately(UMovieSceneEntitySystemLinker*
 
 void FSequenceUpdater_Flat::Finish(UMovieSceneEntitySystemLinker* Linker, FRootInstanceHandle InstanceHandle, IMovieScenePlayer* InPlayer)
 {
-	InvalidateCachedData(Linker);
+	InvalidateCachedData(Linker, InstanceHandle);
 }
 
 void FSequenceUpdater_Flat::Destroy(UMovieSceneEntitySystemLinker* Linker)
 {
 }
 
-void FSequenceUpdater_Flat::InvalidateCachedData(UMovieSceneEntitySystemLinker* Linker)
+void FSequenceUpdater_Flat::InvalidateCachedData(UMovieSceneEntitySystemLinker* Linker, FRootInstanceHandle RootInstanceHandle)
 {
 	CachedEntityRange = TRange<FFrameNumber>::Empty();
 	CachedDeterminismFences.Reset();
@@ -457,7 +458,7 @@ void FSequenceUpdater_Hierarchical::DissectContext(UMovieSceneEntitySystemLinker
 	}
 }
 
-FInstanceHandle FSequenceUpdater_Hierarchical::GetOrCreateSequenceInstance(IMovieScenePlayer* InPlayer, const FMovieSceneSequenceHierarchy* Hierarchy, FInstanceRegistry* InstanceRegistry, FRootInstanceHandle RootInstanceHandle, FMovieSceneSequenceID SequenceID)
+FInstanceHandle FSequenceUpdater_Hierarchical::GetOrCreateSequenceInstance(IMovieScenePlayer* InPlayer, UMovieSceneSequence* SubSequence, const FMovieSceneSequenceHierarchy* Hierarchy, FInstanceRegistry* InstanceRegistry, FRootInstanceHandle RootInstanceHandle, FMovieSceneSequenceID SequenceID)
 {
 	check(SequenceID != MovieSceneSequenceID::Root);
 
@@ -477,13 +478,13 @@ FInstanceHandle FSequenceUpdater_Hierarchical::GetOrCreateSequenceInstance(IMovi
 	{
 		ParentInstance = RootInstanceHandle;
 	}
-	else
+	else if (UMovieSceneSequence* ParentSequence = Hierarchy->FindSubData(Node->ParentID)->GetSequence())
 	{
-		ParentInstance = GetOrCreateSequenceInstance(InPlayer, Hierarchy, InstanceRegistry, RootInstanceHandle, Node->ParentID);
+		ParentInstance = GetOrCreateSequenceInstance(InPlayer, ParentSequence, Hierarchy, InstanceRegistry, RootInstanceHandle, Node->ParentID);
 	}
 
 	FInstanceHandle InstanceHandle = InstanceRegistry->AllocateSubInstance(InPlayer, SequenceID, RootInstanceHandle, ParentInstance);
-	SequenceInstances.Add(SequenceID, FSubInstanceData { InstanceHandle });
+	SequenceInstances.Add(SequenceID, FSubInstanceData { SubSequence->GetSignature(), InstanceHandle });
 
 	return InstanceHandle;
 }
@@ -502,7 +503,7 @@ void FSequenceUpdater_Hierarchical::OverrideRootSequence(UMovieSceneEntitySystem
 			InstanceRegistry->MutateInstance(RootInstanceHandle).Ledger.UnlinkEverything(InLinker);
 		}
 
-		InvalidateCachedData(InLinker);
+		InvalidateCachedData(InLinker, RootInstanceHandle);
 		RootOverrideSequenceID = NewRootOverrideSequenceID;
 	}
 }
@@ -533,7 +534,7 @@ void FSequenceUpdater_Hierarchical::Update(UMovieSceneEntitySystemLinker* Linker
 		UMovieSceneSequence*                RootSequence    = SubData ? SubData->GetSequence() : nullptr;
 		if (ensure(RootSequence))
 		{
-			RootInstanceHandle = GetOrCreateSequenceInstance(InPlayer, RootHierarchy, InstanceRegistry, InstanceHandle, RootOverrideSequenceID);
+			RootInstanceHandle = GetOrCreateSequenceInstance(InPlayer, RootSequence, RootHierarchy, InstanceRegistry, InstanceHandle, RootOverrideSequenceID);
 			RootCompiledDataID = CompiledDataManager->GetDataID(RootSequence);
 			RootContext        = Context.Transform(SubData->RootToSequenceTransform, SubData->TickResolution);
 
@@ -642,7 +643,7 @@ void FSequenceUpdater_Hierarchical::Update(UMovieSceneEntitySystemLinker* Linker
 				FMovieSceneCompiledDataID SubDataID = CompiledDataManager->GetDataID(SubSequence);
 
 				// Set the context for the root sequence instance
-				FInstanceHandle    SubSequenceHandle = GetOrCreateSequenceInstance(InPlayer, RootHierarchy, InstanceRegistry, InstanceHandle, SequenceIDFromRoot);
+				FInstanceHandle    SubSequenceHandle = GetOrCreateSequenceInstance(InPlayer, SubSequence, RootHierarchy, InstanceRegistry, InstanceHandle, SequenceIDFromRoot);
 				FSequenceInstance& SubSequenceInstance = InstanceRegistry->MutateInstance(SubSequenceHandle);
 
 				// Update the sub sequence's context
@@ -783,7 +784,7 @@ void FSequenceUpdater_Hierarchical::Finish(UMovieSceneEntitySystemLinker* Linker
 		InstanceRegistry->MutateInstance(Pair.Value.Handle).Finish(Linker);
 	}
 
-	InvalidateCachedData(Linker);
+	InvalidateCachedData(Linker, InstanceHandle);
 }
 
 void FSequenceUpdater_Hierarchical::Destroy(UMovieSceneEntitySystemLinker* Linker)
@@ -796,16 +797,28 @@ void FSequenceUpdater_Hierarchical::Destroy(UMovieSceneEntitySystemLinker* Linke
 	}
 }
 
-void FSequenceUpdater_Hierarchical::InvalidateCachedData(UMovieSceneEntitySystemLinker* Linker)
+void FSequenceUpdater_Hierarchical::InvalidateCachedData(UMovieSceneEntitySystemLinker* Linker, FRootInstanceHandle RootInstanceHandle)
 {
 	bDynamicWeighting.Reset();
 	CachedEntityRange = TRange<FFrameNumber>::Empty();
 
 	FInstanceRegistry* InstanceRegistry = Linker->GetInstanceRegistry();
 
-	for (TPair<FMovieSceneSequenceID, FSubInstanceData> Pair : SequenceInstances)
+	FMovieSceneRootEvaluationTemplateInstance& RootTemplate = InstanceRegistry->GetInstance(RootInstanceHandle).GetPlayer()->GetEvaluationTemplate();
+
+	for (TPair<FMovieSceneSequenceID, FSubInstanceData>& Pair : SequenceInstances)
 	{
-		InstanceRegistry->MutateInstance(Pair.Value.Handle).Ledger.Invalidate();
+		UMovieSceneSequence* Sequence = RootTemplate.GetSequence(Pair.Key);
+		if (!Sequence)
+		{
+			Pair.Value.SequenceSignature = FGuid();
+			InstanceRegistry->MutateInstance(Pair.Value.Handle).Ledger.Invalidate();
+		}
+		else if (Pair.Value.SequenceSignature != Sequence->GetSignature())
+		{
+			Pair.Value.SequenceSignature = Sequence->GetSignature();
+			InstanceRegistry->MutateInstance(Pair.Value.Handle).Ledger.Invalidate();
+		}
 	}
 }
 
diff --git a/Engine/Source/Runtime/MovieScene/Public/EntitySystem/MovieSceneSequenceUpdaters.h b/Engine/Source/Runtime/MovieScene/Public/EntitySystem/MovieSceneSequenceUpdaters.h
--- a/Engine/Source/Runtime/MovieScene/Public/EntitySystem/MovieSceneSequenceUpdaters.h
+++ b/Engine/Source/Runtime/MovieScene/Public/EntitySystem/MovieSceneSequenceUpdaters.h
@@ -117,8 +117,9 @@ public:
 	 * Invalidate any cached information that depends on the compiled sequence data due to the compiled data changing
 	 *
 	 * @param InLinker         The linker that is evaluating this sequence
+	 * @param InInstanceHandle The instance handle for the top level sequence instance that this updater belongs to
 	 */
-	virtual void InvalidateCachedData(UMovieSceneEntitySystemLinker* InLinker) = 0;
+	virtual void InvalidateCachedData(UMovieSceneEntitySystemLinker* InLinker, FRootInstanceHandle InInstanceHandle) = 0;
 
 
 	/**
