Sequencer: Ensure the garbage collector is able to to collect references in MovieSceneEvalTemplates owned objects.

#rnx

diff --git a/Engine/Source/Runtime/MovieScene/Private/Evaluation/MovieSceneEvalTemplateSerializer.h b/Engine/Source/Runtime/MovieScene/Private/Evaluation/MovieSceneEvalTemplateSerializer.h
--- a/Engine/Source/Runtime/MovieScene/Private/Evaluation/MovieSceneEvalTemplateSerializer.h
+++ b/Engine/Source/Runtime/MovieScene/Private/Evaluation/MovieSceneEvalTemplateSerializer.h
@@ -52,7 +52,7 @@ bool SerializeInlineValue(TInlineValue<T, N>& Impl, FArchive& Ar, bool bWarnOnEr
 
 		return true;
 	}
-	else if (Ar.IsSaving())
+	else if (Ar.IsSaving() || Ar.IsObjectReferenceCollector())
 	{
 		if (Impl.IsValid())
 		{
