Sequencer: Help the GC reachability analysis peek into the overloaded FMovieSceneEvalTemplate structs.

#rnx

diff --git a/Engine/Source/Runtime/MovieScene/Private/Evaluation/MovieSceneEvalTemplate.cpp b/Engine/Source/Runtime/MovieScene/Private/Evaluation/MovieSceneEvalTemplate.cpp
--- a/Engine/Source/Runtime/MovieScene/Private/Evaluation/MovieSceneEvalTemplate.cpp
+++ b/Engine/Source/Runtime/MovieScene/Private/Evaluation/MovieSceneEvalTemplate.cpp
@@ -3,6 +3,9 @@
 #include "Evaluation/MovieSceneEvalTemplate.h"
 #include "Evaluation/MovieSceneEvalTemplateSerializer.h"
 
+#include "UObject/Class.h"
+
+
 #include UE_INLINE_GENERATED_CPP_BY_NAME(MovieSceneEvalTemplate)
 
 float FMovieSceneEvalTemplate::EvaluateEasing(FFrameTime CurrentTime) const
@@ -16,3 +19,15 @@ bool FMovieSceneEvalTemplatePtr::Serialize(FArchive& Ar)
 	return SerializeInlineValue(*this, Ar, bShouldWarn);
 }
 
+
+void FMovieSceneEvalTemplatePtr::AddStructReferencedObjects(FReferenceCollector& Collector)
+{
+	if (IsValid())
+	{
+		FMovieSceneEvalTemplate* Template = &GetValue();
+		TObjectPtr<UScriptStruct> Struct(&Template->GetScriptStruct());
+
+		Collector.AddReferencedObject(Struct);
+		Collector.AddPropertyReferencesWithStructARO(Struct, Template);
+	}
+}
diff --git a/Engine/Source/Runtime/MovieScene/Public/Evaluation/MovieSceneEvalTemplate.h b/Engine/Source/Runtime/MovieScene/Public/Evaluation/MovieSceneEvalTemplate.h
--- a/Engine/Source/Runtime/MovieScene/Public/Evaluation/MovieSceneEvalTemplate.h
+++ b/Engine/Source/Runtime/MovieScene/Public/Evaluation/MovieSceneEvalTemplate.h
@@ -261,9 +261,11 @@ struct FMovieSceneEvalTemplatePtr
 
 	/** Serialize the template */
 	MOVIESCENE_API bool Serialize(FArchive& Ar);
+
+	MOVIESCENE_API void AddStructReferencedObjects(class FReferenceCollector& InCollector);
 };
 
 template<> struct TStructOpsTypeTraits<FMovieSceneEvalTemplatePtr> : public TStructOpsTypeTraitsBase2<FMovieSceneEvalTemplatePtr>
 {
-	enum { WithSerializer = true, WithCopy = true };
+	enum { WithSerializer = true, WithCopy = true, WithAddStructReferencedObjects = true };
 };
