Sequencer: Fixed AddReferencedObjectsForComponent overload with the wrong signature

diff --git a/Engine/Source/Runtime/MovieScene/Private/EntitySystem/BuiltInComponentTypes.cpp b/Engine/Source/Runtime/MovieScene/Private/EntitySystem/BuiltInComponentTypes.cpp
--- a/Engine/Source/Runtime/MovieScene/Private/EntitySystem/BuiltInComponentTypes.cpp
+++ b/Engine/Source/Runtime/MovieScene/Private/EntitySystem/BuiltInComponentTypes.cpp
@@ -31,11 +31,11 @@ UObject* FObjectComponent::GetObject() const
 	return ObjectKey.ResolveObjectPtr();
 }
 
-void AddReferencedObjectForComponent(FReferenceCollector& ReferenceCollector, FObjectComponent* ComponentData)
+void AddReferencedObjectForComponent(FReferenceCollector* ReferenceCollector, FObjectComponent* ComponentData)
 {
 	if (ComponentData->IsStrongReference())
 	{
-		ReferenceCollector.AddReferencedObject(ComponentData->ObjectPtr);
+		ReferenceCollector->AddReferencedObject(ComponentData->ObjectPtr);
 	}
 }
 
diff --git a/Engine/Source/Runtime/MovieScene/Public/EntitySystem/BuiltInComponentTypes.h b/Engine/Source/Runtime/MovieScene/Public/EntitySystem/BuiltInComponentTypes.h
--- a/Engine/Source/Runtime/MovieScene/Public/EntitySystem/BuiltInComponentTypes.h
+++ b/Engine/Source/Runtime/MovieScene/Public/EntitySystem/BuiltInComponentTypes.h
@@ -308,7 +308,7 @@ struct FObjectComponent
 	MOVIESCENE_API UObject* GetObject() const;
 
 	/** Conditionally add a reference for the specified component data based on whether it is strongly referenced or not */
-	friend void AddReferencedObjectForComponent(FReferenceCollector& ReferenceCollector, FObjectComponent* ComponentData);
+	MOVIESCENE_API friend void AddReferencedObjectForComponent(FReferenceCollector* ReferenceCollector, FObjectComponent* ComponentData);
 
 private:
 
