Sequencer: Make MoveSection virtual so that sections can implement how they want to be moved

diff --git a/Engine/Source/Runtime/MovieScene/Private/MovieSceneSection.cpp b/Engine/Source/Runtime/MovieScene/Private/MovieSceneSection.cpp
--- a/Engine/Source/Runtime/MovieScene/Private/MovieSceneSection.cpp
+++ b/Engine/Source/Runtime/MovieScene/Private/MovieSceneSection.cpp
@@ -215,7 +215,7 @@ TSharedPtr<FStructOnScope> UMovieSceneSection::GetKeyStruct(TArrayView<const FKe
 	return nullptr;
 }
 
-void UMovieSceneSection::MoveSection(FFrameNumber DeltaFrame)
+void UMovieSceneSection::MoveSectionImpl(FFrameNumber DeltaFrame)
 {
 	if (TryModify())
 	{
@@ -239,6 +239,11 @@ void UMovieSceneSection::MoveSection(FFrameNumber DeltaFrame)
 	}
 }
 
+void UMovieSceneSection::MoveSection(FFrameNumber DeltaFrame)
+{
+	MoveSectionImpl(DeltaFrame);
+}
+
 TRange<FFrameNumber> UMovieSceneSection::ComputeEffectiveRange() const
 {
 	if (!SectionRange.Value.GetLowerBound().IsOpen() && !SectionRange.Value.GetUpperBound().IsOpen())
@@ -527,7 +532,7 @@ void UMovieSceneSection::InitialPlacement(const TArray<UMovieSceneSection*>& Sec
 			TRange<FFrameNumber> OtherRange = OverlappedSection->GetRange();
 			if (OtherRange.GetUpperBound().IsClosed())
 			{
-				MoveSection(OtherRange.GetUpperBoundValue() - InStartTime);
+				MoveSectionImpl(OtherRange.GetUpperBoundValue() - InStartTime);
 			}
 			else
 			{
diff --git a/Engine/Source/Runtime/MovieScene/Public/MovieSceneSection.h b/Engine/Source/Runtime/MovieScene/Public/MovieSceneSection.h
--- a/Engine/Source/Runtime/MovieScene/Public/MovieSceneSection.h
+++ b/Engine/Source/Runtime/MovieScene/Public/MovieSceneSection.h
@@ -445,7 +445,7 @@ public:
 	 *
 	 * @param DeltaTime	The distance in time to move the curve
 	 */
-	MOVIESCENE_API void MoveSection(FFrameNumber DeltaTime);
+	MOVIESCENE_API virtual void MoveSection(FFrameNumber DeltaTime);
 
 	/**
 	 * Return the range within which this section is effective. Used for automatic calculation of sequence bounds.
@@ -703,6 +703,7 @@ private:
 	 */
 	MOVIESCENE_API virtual EMovieSceneChannelProxyType CacheChannelProxy();
 
+	void MoveSectionImpl(FFrameNumber DeltaTime);
 
 public:
 
