Sequencer: Ensure that recompilations can be handled even if we're part way through a budgeted (time-sliced) evaluation

  - If a modification happened to a sequence while it was partially evaluated, we would previously wait until that evaluation had completely finished before recompiling. This was generally stable (even if it produced undesireable results), except if the modification was caused by a GC and destruction of a sub sequence. In this event the compiled data would be removed entirely from the compiled data manager, but the root template field would not be recompiled.

  - The solution is to check whether there is any additional work required when we are about to early-return from a budgeted evaluation. If there is meaningful work remaining (ie, we're not right at the end), we set the ConditionalRecompile flush state again, so that we re-check the sequences for volatility before evaluating the rest of the flush states on the next frame.

  - If we detect a compilation at any point, we re-run all the flush states from scratch, including re-updating any sequences we were currently evaluating (in case they have changed).

diff --git a/Engine/Source/Runtime/MovieScene/Private/EntitySystem/MovieSceneEntitySystemRunner.cpp b/Engine/Source/Runtime/MovieScene/Private/EntitySystem/MovieSceneEntitySystemRunner.cpp
--- a/Engine/Source/Runtime/MovieScene/Private/EntitySystem/MovieSceneEntitySystemRunner.cpp
+++ b/Engine/Source/Runtime/MovieScene/Private/EntitySystem/MovieSceneEntitySystemRunner.cpp
@@ -26,6 +26,22 @@ DECLARE_CYCLE_STAT(TEXT("Evaluation Phase"), 			MovieSceneEval_EvaluationPhase,
 DECLARE_CYCLE_STAT(TEXT("Finalization Phase"),          MovieSceneEval_FinalizationPhase,       	STATGROUP_MovieSceneECS);
 DECLARE_CYCLE_STAT(TEXT("Post Evaluation Phase"),       MovieSceneEval_PostEvaluationPhase,     	STATGROUP_MovieSceneECS);
 
+namespace UE::MovieScene::FlushState
+{
+	// Signifies that, during the Finalization task, there were still outstanding tasks and we need to perform another iteration
+	ERunnerFlushState LoopEval = ERunnerFlushState::ConditionalRecompile
+		| ERunnerFlushState::Import
+		| ERunnerFlushState::Spawn
+		| ERunnerFlushState::Instantiation
+		| ERunnerFlushState::Evaluation
+		| ERunnerFlushState::Finalization
+		| ERunnerFlushState::EventTriggers
+		| ERunnerFlushState::PostEvaluation
+		| ERunnerFlushState::End;
+
+	ERunnerFlushState Everything = ERunnerFlushState::Start | LoopEval;
+}
+
 
 /**
  * Structure for making it possible to make re-entrant evaluation on a linker.
@@ -245,7 +261,7 @@ bool FMovieSceneEntitySystemRunner::QueueFinalUpdateImpl(FInstanceHandle InInsta
 	// 3. the instance has no current updates
 	//
 	const bool bCanFinishImmediately = Instance.CanFinishImmediately(Linker);
-	const ERunnerFlushState UnsafeDestroyMask = ERunnerFlushState::Everything & ~(ERunnerFlushState::PostEvaluation | ERunnerFlushState::End);
+	const ERunnerFlushState UnsafeDestroyMask = FlushState::Everything & ~(ERunnerFlushState::PostEvaluation | ERunnerFlushState::End);
 	const bool bSafeToDestroyNow = !EnumHasAnyFlags(FlushState, UnsafeDestroyMask);
 	if (bCanFinishImmediately && bSafeToDestroyNow && !HasQueuedUpdates(InInstanceHandle))
 	{
@@ -360,6 +376,13 @@ bool FMovieSceneEntitySystemRunner::FlushNext(UMovieSceneEntitySystemLinker* Lin
 		return GameThread_UpdateSequenceInstances(Linker);
 	}
 
+	// Step 4: Maybe re-update sequence instances if a recompile has ocurred after they were last updated, but before evaluation finished
+	if (EnumHasAnyFlags(FlushState, ERunnerFlushState::ReimportAfterCompile))
+	{
+		EnterFlushState(ERunnerFlushState::ReimportAfterCompile);
+		return GameThread_ReimportSequenceInstances(Linker);
+	}
+
 	// Step 4: Conditionally run the spawn phase of the system graph
 	//
 	if (EnumHasAnyFlags(FlushState, ERunnerFlushState::Spawn))
@@ -430,7 +453,7 @@ void FMovieSceneEntitySystemRunner::Flush(double BudgetMs)
 	// If we're not currently evaluating, start by flushing everything
 	if (FlushState == ERunnerFlushState::None)
 	{
-		FlushState = ERunnerFlushState::Everything;
+		FlushState = FlushState::Everything;
 	}
 
 	FlushOutstanding(BudgetMs);
@@ -476,21 +499,30 @@ void FMovieSceneEntitySystemRunner::FlushOutstanding(double BudgetMs, UE::MovieS
 	// For the purposes of this function, None means that everything must be flushed
 	if (TargetState == ERunnerFlushState::None)
 	{
-		TargetState = ERunnerFlushState::Everything;
+		TargetState = FlushState::Everything;
 	}
 
 	const double BudgetSeconds = BudgetMs / 1000.f;
 	if (!bRequireFullFlush && BudgetSeconds > 0.0)
 	{
 		double StartTime = FPlatformTime::Seconds();
-		do
+		while (EnumHasAnyFlags(FlushState, TargetState))
 		{
 			if (!FlushNext(Linker))
 			{
 				break;
 			}
+
+			if (FPlatformTime::Seconds() - StartTime >= BudgetSeconds)
+			{
+				// We've reached our budget - we must check next time whether we need to recompile again
+				if (FlushState != ERunnerFlushState::End)
+				{
+					FlushState |= ERunnerFlushState::ConditionalRecompile;
+				}
+				break;
+			}
 		}
-		while (FPlatformTime::Seconds() - StartTime < BudgetSeconds && EnumHasAnyFlags(FlushState, TargetState));
 	}
 	else
 	{
@@ -518,7 +550,7 @@ void FMovieSceneEntitySystemRunner::ResetFlushState()
 	{
 		// When resetting - we don't need to (or want to) re-import anything, we just want to re-run the current
 		// frame of updates 
-		FlushState = UE::MovieScene::ERunnerFlushState::LoopEval & ~(UE::MovieScene::ERunnerFlushState::ConditionalRecompile | UE::MovieScene::ERunnerFlushState::Import);
+		FlushState = UE::MovieScene::FlushState::LoopEval & ~(UE::MovieScene::ERunnerFlushState::ConditionalRecompile | UE::MovieScene::ERunnerFlushState::Import);
 	}
 }
 
@@ -571,6 +603,7 @@ bool FMovieSceneEntitySystemRunner::GameThread_ConditionalRecompile(UMovieSceneE
 
 	FInstanceRegistry* InstanceRegistry = GetInstanceRegistry();
 
+	bool bAnyRecompile = false;
 	for (int32 UpdateIndex = 0; UpdateIndex < UpdateQueue.Num(); ++UpdateIndex)
 	{
 		const FUpdateParamsAndContext& Request = UpdateQueue[UpdateIndex];
@@ -585,7 +618,22 @@ bool FMovieSceneEntitySystemRunner::GameThread_ConditionalRecompile(UMovieSceneE
 			continue;
 		}
 
-		Instance.ConditionalRecompile(Linker);
+		if (Instance.ConditionalRecompile(Linker))
+		{
+			bAnyRecompile = true;
+		}
+	}
+
+	if (bAnyRecompile)
+	{
+		if (!EnumHasAnyFlags(FlushState, ERunnerFlushState::Import))
+		{
+			// If we have already imported everything, we have to run a reimport to make sure everything is imported correctly
+			FlushState |= ERunnerFlushState::ReimportAfterCompile;
+		}
+
+		// If a recompile has occurred, we need to make sure we run everything again (except for checking for another recompile!)
+		FlushState |= (FlushState::LoopEval & ~ERunnerFlushState::ConditionalRecompile);
 	}
 
 	return true;
@@ -816,6 +864,26 @@ bool FMovieSceneEntitySystemRunner::GameThread_UpdateSequenceInstances(UMovieSce
 	return true;
 }
 
+bool FMovieSceneEntitySystemRunner::GameThread_ReimportSequenceInstances(UMovieSceneEntitySystemLinker* Linker)
+{
+	using namespace UE::MovieScene;
+
+	FInstanceRegistry* InstanceRegistry = WeakLinker.Get()->GetInstanceRegistry();
+
+	// Only called after a sequence has been recompiled after we have already updated the current instances
+	// This allows us to re-update all the sequence instances in case anything has changed
+	for (const FQueuedUpdateParams& UpdatedInstance : CurrentInstances)
+	{
+		if (!EnumHasAnyFlags(UpdatedInstance.UpdateFlags, ERunnerUpdateFlags::Destroy | ERunnerUpdateFlags::Finish))
+		{
+			FSequenceInstance& SequenceInstance = InstanceRegistry->MutateInstance(UpdatedInstance.InstanceHandle);
+			SequenceInstance.Update(Linker, SequenceInstance.GetContext());
+		}
+	}
+
+	return true;
+}
+
 bool FMovieSceneEntitySystemRunner::GameThread_SpawnPhase(UMovieSceneEntitySystemLinker* Linker)
 {
 	using namespace UE::MovieScene;
@@ -1146,7 +1214,7 @@ void FMovieSceneEntitySystemRunner::GameThread_PostEvaluationPhase(UMovieSceneEn
 	// which will pick up the next updates
 	if (UpdateQueue.Num() > 0 || DissectedUpdates.Num() > 0)
 	{
-		FlushState = ERunnerFlushState::LoopEval;
+		FlushState = FlushState::LoopEval;
 	}
 }
 
diff --git a/Engine/Source/Runtime/MovieScene/Private/EntitySystem/MovieSceneSequenceInstance.cpp b/Engine/Source/Runtime/MovieScene/Private/EntitySystem/MovieSceneSequenceInstance.cpp
--- a/Engine/Source/Runtime/MovieScene/Private/EntitySystem/MovieSceneSequenceInstance.cpp
+++ b/Engine/Source/Runtime/MovieScene/Private/EntitySystem/MovieSceneSequenceInstance.cpp
@@ -198,7 +198,7 @@ void FSequenceInstance::InvalidateCachedData(UMovieSceneEntitySystemLinker* Link
 	}
 }
 
-void FSequenceInstance::ConditionalRecompile(UMovieSceneEntitySystemLinker* Linker)
+bool FSequenceInstance::ConditionalRecompile(UMovieSceneEntitySystemLinker* Linker)
 {
 	if (VolatilityManager)
 	{
@@ -210,8 +210,11 @@ void FSequenceInstance::ConditionalRecompile(UMovieSceneEntitySystemLinker* Link
 		if (VolatilityManager->ConditionalRecompile(*Player, RootCompiledDataID, CompiledDataManager))
 		{
 			InvalidateCachedData(Linker);
+			return true;
 		}
 	}
+
+	return false;
 }
 
 void FSequenceInstance::DissectContext(UMovieSceneEntitySystemLinker* Linker, const FMovieSceneContext& InContext, TArray<TRange<FFrameTime>>& OutDissections)
diff --git a/Engine/Source/Runtime/MovieScene/Public/EntitySystem/MovieSceneEntitySystemRunner.h b/Engine/Source/Runtime/MovieScene/Public/EntitySystem/MovieSceneEntitySystemRunner.h
--- a/Engine/Source/Runtime/MovieScene/Public/EntitySystem/MovieSceneEntitySystemRunner.h
+++ b/Engine/Source/Runtime/MovieScene/Public/EntitySystem/MovieSceneEntitySystemRunner.h
@@ -32,19 +32,14 @@ enum class ERunnerFlushState
 	Start                = 1 << 0,		// Sets up initial evaluation flags for external players and listeners
 	ConditionalRecompile = 1 << 1,		// Conditional recompile of dirtied sequences.
 	Import               = 1 << 2,		// Update sequence instances and import entities into the entity manager
-	Spawn                = 1 << 3,		// Perform the Spawn phase in the Entity System Graph
-	Instantiation        = 1 << 4,		// Perform the Instantiation phase in the Entity System Graph
-	Evaluation           = 1 << 5,		// Perform the Evaluation phase in the Entity System Graph
-	Finalization         = 1 << 6,		// Perform the Finalization phase in the Entity System Graph and trigger any external events
-	EventTriggers        = 1 << 7,		// (re-entrant) Triggers any bound event triggers - skipped by Finalization if the delegate is not bound
-	PostEvaluation       = 1 << 8,		// (re-entrant) Call post evaluation callbacks on sequence instances
-	End                  = 1 << 9,		// Counterpart for Start - resets external players and listeners
-
-	// Signifies that, during the Finalization task, there were still outstanding tasks and we need to perform another iteration
-	LoopEval       = ConditionalRecompile | Import | Spawn | Instantiation | Evaluation | Finalization | EventTriggers | PostEvaluation | End,
-
-	// Initial flush state
-	Everything     = Start | LoopEval,
+	ReimportAfterCompile = 1 << 3,		// Re-update sequence instances after a recompile occurred on a partially evaluated sequence. Not normally run.
+	Spawn                = 1 << 4,		// Perform the Spawn phase in the Entity System Graph
+	Instantiation        = 1 << 5,		// Perform the Instantiation phase in the Entity System Graph
+	Evaluation           = 1 << 6,		// Perform the Evaluation phase in the Entity System Graph
+	Finalization         = 1 << 7,		// Perform the Finalization phase in the Entity System Graph and trigger any external events
+	EventTriggers        = 1 << 8,		// (re-entrant) Triggers any bound event triggers - skipped by Finalization if the delegate is not bound
+	PostEvaluation       = 1 << 9,		// (re-entrant) Call post evaluation callbacks on sequence instances
+	End                  = 1 << 10,		// Counterpart for Start - resets external players and listeners
 };
 ENUM_CLASS_FLAGS(ERunnerFlushState)
 
@@ -197,6 +192,8 @@ private:
 	MOVIESCENE_API bool GameThread_ConditionalRecompile(UMovieSceneEntitySystemLinker* Linker);
 	/** Update sequence instances based on currently queued update requests, or outstanding dissected updates */
 	MOVIESCENE_API bool GameThread_UpdateSequenceInstances(UMovieSceneEntitySystemLinker* Linker);
+	/** Re-update sequence instances after a recompile on a partially evaluated sequence */
+	MOVIESCENE_API bool GameThread_ReimportSequenceInstances(UMovieSceneEntitySystemLinker* Linker);
 	/** Execute the spawn phase of the entity system graph, if there is anything to do */
 	MOVIESCENE_API bool GameThread_SpawnPhase(UMovieSceneEntitySystemLinker* Linker);
 	/** Execute the instantiation phase of the entity system graph, if there is anything to do */
diff --git a/Engine/Source/Runtime/MovieScene/Public/EntitySystem/MovieSceneSequenceInstance.h b/Engine/Source/Runtime/MovieScene/Public/EntitySystem/MovieSceneSequenceInstance.h
--- a/Engine/Source/Runtime/MovieScene/Public/EntitySystem/MovieSceneSequenceInstance.h
+++ b/Engine/Source/Runtime/MovieScene/Public/EntitySystem/MovieSceneSequenceInstance.h
@@ -59,8 +59,9 @@ public:
 	 * Conditionally recompile this sequence if it needs to be
 	 *
 	 * @param Linker     The linker that owns this sequence instance
+	 * @return true if a recompile has ocurred, false otherwise
 	 */
-	MOVIESCENE_API void ConditionalRecompile(UMovieSceneEntitySystemLinker* Linker);
+	MOVIESCENE_API bool ConditionalRecompile(UMovieSceneEntitySystemLinker* Linker);
 
 	/**
 	 * Called only for top-level sequence instances before any updates to it with the specified context.
