Sequencer- Improve performance of searching for items in large sequences by keeping the Spawnables, Possessables, and ObjectBindings arrays sorted by FGuid and using BinarySearch where searching is necessary. Testing was performed on adds and searches in these arrays to ensure improvement (to searches). Adds are a bit more expensive, but should be alright in the vast majority of use cases. Iimprovement could be made if necessary to bulk additions by adding bulk operation functionality.

diff --git a/Engine/Plugins/Animation/ControlRig/Source/ControlRigEditor/Private/ControlRigSequencerEditorLibrary.cpp b/Engine/Plugins/Animation/ControlRig/Source/ControlRigEditor/Private/ControlRigSequencerEditorLibrary.cpp
--- a/Engine/Plugins/Animation/ControlRig/Source/ControlRigEditor/Private/ControlRigSequencerEditorLibrary.cpp
+++ b/Engine/Plugins/Animation/ControlRig/Source/ControlRigEditor/Private/ControlRigSequencerEditorLibrary.cpp
@@ -276,54 +276,50 @@ UMovieSceneTrack* UControlRigSequencerEditorLibrary::FindOrCreateControlRigTrack
 	UMovieSceneTrack* BaseTrack = nullptr;
 	if (LevelSequence && MovieScene && InBinding.BindingID.IsValid())
 	{
-		const TArray<FMovieSceneBinding>& Bindings = MovieScene->GetBindings();
-		for (const FMovieSceneBinding& Binding : Bindings)
+		if (const FMovieSceneBinding* Binding = MovieScene->FindBinding(InBinding.BindingID))
 		{
-			if (Binding.GetObjectGuid() == InBinding.BindingID)
+			TArray<UMovieSceneTrack*> Tracks = MovieScene->FindTracks(UMovieSceneControlRigParameterTrack::StaticClass(), Binding->GetObjectGuid(), NAME_None);
+			for (UMovieSceneTrack* AnyOleTrack : Tracks)
 			{
-				TArray<UMovieSceneTrack*> Tracks = MovieScene->FindTracks(UMovieSceneControlRigParameterTrack::StaticClass(), Binding.GetObjectGuid(), NAME_None);
-				for (UMovieSceneTrack* AnyOleTrack : Tracks)
+				UMovieSceneControlRigParameterTrack* Track = Cast<UMovieSceneControlRigParameterTrack>(AnyOleTrack);
+				if (Track && Track->GetControlRig() && Track->GetControlRig()->GetClass() == ControlRigClass)
 				{
-					UMovieSceneControlRigParameterTrack* Track = Cast<UMovieSceneControlRigParameterTrack>(AnyOleTrack);
-					if (Track && Track->GetControlRig() && Track->GetControlRig()->GetClass() == ControlRigClass)
-					{
-						return Track;
-					}
+					return Track;
 				}
+			}
+
+			TArray<UObject*, TInlineAllocator<1>> Result;
+			UObject* Context = nullptr;
+			ALevelSequenceActor* OutActor = nullptr;
+			ULevelSequencePlayer* OutPlayer = nullptr;
+			Result = GetBoundObjects(World, LevelSequence, InBinding, &OutPlayer, &OutActor);
+			if (Result.Num() > 0 && Result[0])
+			{
+				UObject* BoundObject = Result[0];
+				USkeleton* Skeleton = nullptr;
+				USkeletalMeshComponent* SkeletalMeshComponent = nullptr;
+				AcquireSkeletonAndSkelMeshCompFromObject(BoundObject, &Skeleton, &SkeletalMeshComponent);
 
-				TArray<UObject*, TInlineAllocator<1>> Result;
-				UObject* Context = nullptr;
-				ALevelSequenceActor* OutActor = nullptr;
-				ULevelSequencePlayer* OutPlayer = nullptr;
-				Result = GetBoundObjects(World, LevelSequence, InBinding, &OutPlayer, &OutActor);
-				if (Result.Num() > 0 && Result[0])
+				UControlRig* ControlRig = nullptr;
+				if (Skeleton && SkeletalMeshComponent)
 				{
-					UObject* BoundObject = Result[0];
-					USkeleton* Skeleton = nullptr;
-					USkeletalMeshComponent* SkeletalMeshComponent = nullptr;
-					AcquireSkeletonAndSkelMeshCompFromObject(BoundObject, &Skeleton, &SkeletalMeshComponent);
+					UMovieSceneControlRigParameterTrack* Track = AddControlRig(LevelSequence, ControlRigClass, SkeletalMeshComponent, InBinding.BindingID, nullptr);
 
-					UControlRig* ControlRig = nullptr;
-					if (Skeleton && SkeletalMeshComponent)
+					if (Track)
 					{
-						UMovieSceneControlRigParameterTrack* Track = AddControlRig(LevelSequence, ControlRigClass, SkeletalMeshComponent, InBinding.BindingID, nullptr);
-
-						if (Track)
-						{
-							BaseTrack = Track;								
-						}
+						BaseTrack = Track;								
 					}
 				}
+			}
 
-				if (OutPlayer)
-				{
-					OutPlayer->Stop();
-				}
+			if (OutPlayer)
+			{
+				OutPlayer->Stop();
+			}
 
-				if (OutActor)
-				{
-					World->DestroyActor(OutActor);
-				}
+			if (OutActor)
+			{
+				World->DestroyActor(OutActor);
 			}
 		}
 	}
diff --git a/Engine/Plugins/FX/Niagara/Source/NiagaraEditor/Private/Sequencer/LevelSequence/NiagaraSystemTrackEditor.cpp b/Engine/Plugins/FX/Niagara/Source/NiagaraEditor/Private/Sequencer/LevelSequence/NiagaraSystemTrackEditor.cpp
--- a/Engine/Plugins/FX/Niagara/Source/NiagaraEditor/Private/Sequencer/LevelSequence/NiagaraSystemTrackEditor.cpp
+++ b/Engine/Plugins/FX/Niagara/Source/NiagaraEditor/Private/Sequencer/LevelSequence/NiagaraSystemTrackEditor.cpp
@@ -41,16 +41,13 @@ bool FNiagaraSystemTrackEditor::SupportsType(TSubclassOf<UMovieSceneTrack> Type)
 
 bool HasLifeCycleTrack(UMovieScene& MovieScene, FGuid ObjectBinding)
 {
-	for (const FMovieSceneBinding& Binding : MovieScene.GetBindings())
+	if (const FMovieSceneBinding* Binding = MovieScene.FindBinding(ObjectBinding))
 	{
-		if (Binding.GetObjectGuid() == ObjectBinding)
+		for (const UMovieSceneTrack* Track : Binding->GetTracks())
 		{
-			for (UMovieSceneTrack* Track : Binding.GetTracks())
+			if (Track->IsA<UMovieSceneNiagaraSystemTrack>())
 			{
-				if (Track->IsA<UMovieSceneNiagaraSystemTrack>())
-				{
-					return true;
-				}
+				return true;
 			}
 		}
 	}
@@ -59,19 +56,15 @@ bool HasLifeCycleTrack(UMovieScene& MovieScene, FGuid ObjectBinding)
 
 void GetAnimatedParameterNames(UMovieScene& MovieScene, FGuid ObjectBinding, TSet<FName>& AnimatedParameterNames)
 {
-	for (const FMovieSceneBinding& Binding : MovieScene.GetBindings())
+	if (const FMovieSceneBinding* Binding = MovieScene.FindBinding(ObjectBinding))
 	{
-		if (Binding.GetObjectGuid() == ObjectBinding)
+		for (const UMovieSceneTrack* Track : Binding->GetTracks())
 		{
-			for (UMovieSceneTrack* Track : Binding.GetTracks())
+			const UMovieSceneNiagaraParameterTrack* ParameterTrack = Cast<UMovieSceneNiagaraParameterTrack>(Track);
+			if (ParameterTrack != nullptr)
 			{
-				UMovieSceneNiagaraParameterTrack* ParameterTrack = Cast<UMovieSceneNiagaraParameterTrack>(Track);
-				if (ParameterTrack != nullptr)
-				{
-					AnimatedParameterNames.Add(ParameterTrack->GetParameter().GetName());
-				}
+				AnimatedParameterNames.Add(ParameterTrack->GetParameter().GetName());
 			}
-			break;
 		}
 	}
 }
diff --git a/Engine/Plugins/Importers/USDImporter/Source/USDStage/Private/USDLevelSequenceHelper.cpp b/Engine/Plugins/Importers/USDImporter/Source/USDStage/Private/USDLevelSequenceHelper.cpp
--- a/Engine/Plugins/Importers/USDImporter/Source/USDStage/Private/USDLevelSequenceHelper.cpp
+++ b/Engine/Plugins/Importers/USDImporter/Source/USDStage/Private/USDLevelSequenceHelper.cpp
@@ -118,14 +118,9 @@ namespace UsdLevelSequenceHelperImpl
 			return nullptr;
 		}
 
-		for (const FMovieSceneBinding& Binding : MovieScene->GetBindings())
+		if (const FMovieSceneBinding* Binding = MovieScene->FindBinding(Guid))
 		{
-			if (Binding.GetObjectGuid() != Guid)
-			{
-				continue;
-			}
-
-			for (UMovieSceneTrack* Track : Binding.GetTracks())
+			for (UMovieSceneTrack* Track : Binding->GetTracks())
 			{
 				if (TrackType* CastTrack = Cast<TrackType>(Track))
 				{
diff --git a/Engine/Plugins/MovieScene/SequencerScripting/Source/SequencerScripting/Private/ExtensionLibraries/MovieSceneBindingExtensions.cpp b/Engine/Plugins/MovieScene/SequencerScripting/Source/SequencerScripting/Private/ExtensionLibraries/MovieSceneBindingExtensions.cpp
--- a/Engine/Plugins/MovieScene/SequencerScripting/Source/SequencerScripting/Private/ExtensionLibraries/MovieSceneBindingExtensions.cpp
+++ b/Engine/Plugins/MovieScene/SequencerScripting/Source/SequencerScripting/Private/ExtensionLibraries/MovieSceneBindingExtensions.cpp
@@ -15,7 +15,7 @@ bool UMovieSceneBindingExtensions::IsValid(const FMovieSceneBindingProxy& InBind
 	UMovieScene* MovieScene = InBinding.Sequence ? InBinding.Sequence->GetMovieScene() : nullptr;
 	if (MovieScene && InBinding.BindingID.IsValid())
 	{
-		return Algo::FindBy(MovieScene->GetBindings(), InBinding.BindingID, &FMovieSceneBinding::GetObjectGuid) != nullptr;
+		return MovieScene->FindBinding(InBinding.BindingID) != nullptr;
 	}
 
 	return false;
@@ -97,8 +97,7 @@ int32 UMovieSceneBindingExtensions::GetSortingOrder(const FMovieSceneBindingProx
 	UMovieScene* MovieScene = InBinding.GetMovieScene();
 	if (MovieScene)
 	{
-		const FMovieSceneBinding* Binding = Algo::FindBy(MovieScene->GetBindings(), InBinding.BindingID, &FMovieSceneBinding::GetObjectGuid);
-		if (Binding)
+		if (const FMovieSceneBinding* Binding = MovieScene->FindBinding(InBinding.BindingID))
 		{
 			return Binding->GetSortingOrder();
 		}
@@ -131,8 +130,7 @@ TArray<UMovieSceneTrack*> UMovieSceneBindingExtensions::GetTracks(const FMovieSc
 	UMovieScene* MovieScene = InBinding.GetMovieScene();
 	if (MovieScene)
 	{
-		const FMovieSceneBinding* Binding = Algo::FindBy(MovieScene->GetBindings(), InBinding.BindingID, &FMovieSceneBinding::GetObjectGuid);
-		if (Binding)
+		if (const FMovieSceneBinding* Binding = MovieScene->FindBinding(InBinding.BindingID))
 		{
 			return Binding->GetTracks();
 		}
@@ -174,8 +172,7 @@ TArray<UMovieSceneTrack*> UMovieSceneBindingExtensions::FindTracksByType(const F
 
 	if (MovieScene)
 	{
-		const FMovieSceneBinding* Binding = Algo::FindBy(MovieScene->GetBindings(), InBinding.BindingID, &FMovieSceneBinding::GetObjectGuid);
-		if (Binding)
+		if (const FMovieSceneBinding* Binding = MovieScene->FindBinding(InBinding.BindingID))
 		{
 			bool bExactMatch = false;
 			return UMovieSceneSequenceExtensions::FilterTracks(Binding->GetTracks(), DesiredClass, bExactMatch);
@@ -191,8 +188,7 @@ TArray<UMovieSceneTrack*> UMovieSceneBindingExtensions::FindTracksByExactType(co
 
 	if (MovieScene)
 	{
-		const FMovieSceneBinding* Binding = Algo::FindBy(MovieScene->GetBindings(), InBinding.BindingID, &FMovieSceneBinding::GetObjectGuid);
-		if (Binding)
+		if (const FMovieSceneBinding* Binding = MovieScene->FindBinding(InBinding.BindingID))
 		{
 			bool bExactMatch = true;
 			return UMovieSceneSequenceExtensions::FilterTracks(Binding->GetTracks(), DesiredClass, bExactMatch);
@@ -208,7 +204,7 @@ UMovieSceneTrack* UMovieSceneBindingExtensions::AddTrack(const FMovieSceneBindin
 
 	if (MovieScene)
 	{
-		const bool bBindingExists = Algo::FindBy(MovieScene->GetBindings(), InBinding.BindingID, &FMovieSceneBinding::GetObjectGuid) != nullptr;
+		const bool bBindingExists = MovieScene->FindBinding(InBinding.BindingID) != nullptr;
 		if (bBindingExists)
 		{
 			UMovieSceneTrack* NewTrack = NewObject<UMovieSceneTrack>(MovieScene, DesiredClass, NAME_None, RF_Transactional);
diff --git a/Engine/Plugins/MovieScene/SequencerScripting/Source/SequencerScripting/Private/ExtensionLibraries/MovieSceneSequenceExtensions.cpp b/Engine/Plugins/MovieScene/SequencerScripting/Source/SequencerScripting/Private/ExtensionLibraries/MovieSceneSequenceExtensions.cpp
--- a/Engine/Plugins/MovieScene/SequencerScripting/Source/SequencerScripting/Private/ExtensionLibraries/MovieSceneSequenceExtensions.cpp
+++ b/Engine/Plugins/MovieScene/SequencerScripting/Source/SequencerScripting/Private/ExtensionLibraries/MovieSceneSequenceExtensions.cpp
@@ -680,7 +680,7 @@ FMovieSceneBindingProxy UMovieSceneSequenceExtensions::FindBindingById(UMovieSce
 	UMovieScene* MovieScene = GetMovieScene(Sequence);
 	if (MovieScene)
 	{
-		const FMovieSceneBinding* Binding = Algo::FindBy(MovieScene->GetBindings(), BindingId, &FMovieSceneBinding::GetObjectGuid);
+		const FMovieSceneBinding* Binding = MovieScene->FindBinding(BindingId);
 		if (Binding)
 		{
 			return FMovieSceneBindingProxy(Binding->GetObjectGuid(), Sequence);
diff --git a/Engine/Source/Editor/MovieSceneTools/Private/TrackEditors/CustomPrimitiveDataTrackEditor.cpp b/Engine/Source/Editor/MovieSceneTools/Private/TrackEditors/CustomPrimitiveDataTrackEditor.cpp
--- a/Engine/Source/Editor/MovieSceneTools/Private/TrackEditors/CustomPrimitiveDataTrackEditor.cpp
+++ b/Engine/Source/Editor/MovieSceneTools/Private/TrackEditors/CustomPrimitiveDataTrackEditor.cpp
@@ -515,7 +515,7 @@ void FCustomPrimitiveDataTrackEditor::HandleAddCustomPrimitiveDataTrackExecute(U
 	for (UActorComponent* ActorComponent : ActorComponents)
 	{
 		FGuid ObjectHandle = SequencerPtr->GetHandleToObject(ActorComponent);
-		const FMovieSceneBinding* Binding = Algo::FindBy(MovieScene->GetBindings(), ObjectHandle, &FMovieSceneBinding::GetObjectGuid);
+		const FMovieSceneBinding* Binding = MovieScene->FindBinding(ObjectHandle);
 		
 		const bool bAlreadyExists = MovieScene->FindTrack<UMovieSceneCustomPrimitiveDataTrack>(ObjectHandle) != nullptr;
 		if (!bAlreadyExists)
diff --git a/Engine/Source/Editor/MovieSceneTools/Private/TrackEditors/MaterialTrackEditor.cpp b/Engine/Source/Editor/MovieSceneTools/Private/TrackEditors/MaterialTrackEditor.cpp
--- a/Engine/Source/Editor/MovieSceneTools/Private/TrackEditors/MaterialTrackEditor.cpp
+++ b/Engine/Source/Editor/MovieSceneTools/Private/TrackEditors/MaterialTrackEditor.cpp
@@ -509,7 +509,7 @@ void FComponentMaterialTrackEditor::ConstructObjectBindingTrackMenu(FMenuBuilder
 	}
 
 	const UMovieScene* MovieScene = GetFocusedMovieScene();
-	const FMovieSceneBinding* Binding = Algo::FindBy(MovieScene->GetBindings(), ObjectBindings[0], &FMovieSceneBinding::GetObjectGuid);
+	const FMovieSceneBinding* Binding = MovieScene->FindBinding(ObjectBindings[0]);
 
 	if (UPrimitiveComponent* PrimitiveComponent = Cast<UPrimitiveComponent>(SceneComponent))
 	{
@@ -623,7 +623,7 @@ void FComponentMaterialTrackEditor::HandleAddComponentMaterialActionExecute(USce
 	for (UActorComponent* ActorComponent : ActorComponents)
 	{
 		FGuid ObjectHandle = SequencerPtr->GetHandleToObject(ActorComponent);
-		const FMovieSceneBinding* Binding = Algo::FindBy(MovieScene->GetBindings(), ObjectHandle, &FMovieSceneBinding::GetObjectGuid);
+		const FMovieSceneBinding* Binding = MovieScene->FindBinding(ObjectHandle);
 
 		const bool bAlreadyExists = Algo::FindBy(Binding->GetTracks(), MaterialInfo, GetMaterialInfoForTrack) != nullptr;
 		if (!bAlreadyExists)
diff --git a/Engine/Source/Editor/MovieSceneTools/Private/TrackEditors/PrimitiveMaterialTrackEditor.cpp b/Engine/Source/Editor/MovieSceneTools/Private/TrackEditors/PrimitiveMaterialTrackEditor.cpp
--- a/Engine/Source/Editor/MovieSceneTools/Private/TrackEditors/PrimitiveMaterialTrackEditor.cpp
+++ b/Engine/Source/Editor/MovieSceneTools/Private/TrackEditors/PrimitiveMaterialTrackEditor.cpp
@@ -52,7 +52,7 @@ void FPrimitiveMaterialTrackEditor::ConstructObjectBindingTrackMenu(FMenuBuilder
 	}	
 
 	const UMovieScene* MovieScene = GetFocusedMovieScene();
-	const FMovieSceneBinding* Binding = Algo::FindBy(MovieScene->GetBindings(), ObjectBindings[0], &FMovieSceneBinding::GetObjectGuid);
+	const FMovieSceneBinding* Binding = MovieScene->FindBinding(ObjectBindings[0]);
 
 	if (UPrimitiveComponent* PrimitiveComponent = Cast<UPrimitiveComponent>(SceneComponent))
 	{
diff --git a/Engine/Source/Editor/Sequencer/Private/MVVM/ViewModels/ObjectBindingModel.cpp b/Engine/Source/Editor/Sequencer/Private/MVVM/ViewModels/ObjectBindingModel.cpp
--- a/Engine/Source/Editor/Sequencer/Private/MVVM/ViewModels/ObjectBindingModel.cpp
+++ b/Engine/Source/Editor/Sequencer/Private/MVVM/ViewModels/ObjectBindingModel.cpp
@@ -1126,10 +1126,7 @@ FSortingKey FObjectBindingModel::GetSortingKey() const
 	if (OwnerModel)
 	{
 		UMovieScene* MovieScene = OwnerModel->GetMovieScene();
-		const FMovieSceneBinding* MovieSceneBinding = MovieScene->GetBindings().FindByPredicate([&](FMovieSceneBinding& Binding)
-		{
-			return Binding.GetObjectGuid() == ObjectBindingID;
-		});
+		const FMovieSceneBinding* MovieSceneBinding = MovieScene->FindBinding(ObjectBindingID);
 
 		if (MovieSceneBinding)
 		{
diff --git a/Engine/Source/Editor/Sequencer/Private/SequencerUtilities.cpp b/Engine/Source/Editor/Sequencer/Private/SequencerUtilities.cpp
--- a/Engine/Source/Editor/Sequencer/Private/SequencerUtilities.cpp
+++ b/Engine/Source/Editor/Sequencer/Private/SequencerUtilities.cpp
@@ -541,7 +541,7 @@ TArray<FGuid> ExpandMultiplePossessableBindings(TSharedRef<ISequencer> Sequencer
 	Sequence->Modify();
 	MovieScene->Modify();
 
-	FMovieSceneBinding* PossessableBinding = (FMovieSceneBinding*)MovieScene->GetBindings().FindByPredicate([&](FMovieSceneBinding& Binding) { return Binding.GetObjectGuid() == PossessableGuid; });
+	FMovieSceneBinding* PossessableBinding = MovieScene->FindBinding(PossessableGuid);
 
 	// First gather the children
 	TArray<FGuid> ChildPossessableGuids;
@@ -587,7 +587,7 @@ TArray<FGuid> ExpandMultiplePossessableBindings(TSharedRef<ISequencer> Sequencer
 		FMovieScenePossessable* NewPossessable = MovieScene->FindPossessable(NewPossessableGuid);
 		if (NewPossessable)
 		{
-			FMovieSceneBinding* NewPossessableBinding = (FMovieSceneBinding*)MovieScene->GetBindings().FindByPredicate([&](FMovieSceneBinding& Binding) { return Binding.GetObjectGuid() == NewPossessableGuid; });
+			FMovieSceneBinding* NewPossessableBinding = MovieScene->FindBinding(NewPossessableGuid);
 
 			if (ParentObject)
 			{
@@ -1070,7 +1070,7 @@ TArray<FMovieSceneSpawnable*> FSequencerUtilities::ConvertToSpawnable(TSharedRef
 			// Remap all the spawnable's tracks and child bindings onto the new possessable
 			MovieScene->MoveBindingContents(PossessableGuid, SpawnableGuid);
 
-			FMovieSceneBinding* PossessableBinding = (FMovieSceneBinding*)MovieScene->GetBindings().FindByPredicate([&](FMovieSceneBinding& Binding) { return Binding.GetObjectGuid() == PossessableGuid; });
+			FMovieSceneBinding* PossessableBinding = MovieScene->FindBinding(PossessableGuid);
 			check(PossessableBinding);
 
 			for (UMovieSceneFolder* Folder : MovieScene->GetRootFolders())
@@ -1087,7 +1087,7 @@ TArray<FMovieSceneSpawnable*> FSequencerUtilities::ConvertToSpawnable(TSharedRef
 			{
 				Sequence->UnbindPossessableObjects(PossessableGuid);
 
-				FMovieSceneBinding* SpawnableBinding = (FMovieSceneBinding*)MovieScene->GetBindings().FindByPredicate([&](FMovieSceneBinding& Binding) { return Binding.GetObjectGuid() == SpawnableGuid; });
+				FMovieSceneBinding* SpawnableBinding = MovieScene->FindBinding(SpawnableGuid);
 				check(SpawnableBinding);
 
 				SpawnableBinding->SetSortingOrder(SortingOrder);
@@ -1558,19 +1558,15 @@ bool FSequencerUtilities::PasteTracks(const FString& TextToImport, FMovieScenePa
 				ResetCopiedTracksFlags(NewTrack);
 
 				// Remove tracks with the same name before adding
-				for (const FMovieSceneBinding& Binding : MovieScene->GetBindings())
+				if (const FMovieSceneBinding* Binding = MovieScene->FindBinding(ObjectBinding.BindingID))
 				{
-					if (Binding.GetObjectGuid() == ObjectBinding.BindingID)
+					for (UMovieSceneTrack* Track : Binding->GetTracks())
 					{
-						// Tracks of the same class should be unique per name.
-						for (UMovieSceneTrack* Track : Binding.GetTracks())
+						if (Track->GetClass() == NewTrack->GetClass() && Track->GetTrackName() == NewTrack->GetTrackName() && Track->GetDisplayName().IdenticalTo(NewTrack->GetDisplayName()))
 						{
-							if (Track->GetClass() == NewTrack->GetClass() && Track->GetTrackName() == NewTrack->GetTrackName() && Track->GetDisplayName().IdenticalTo(NewTrack->GetDisplayName()))
-							{
-								// If a track of the same class and name exists, remove it so the new track replaces it
-								MovieScene->RemoveTrack(*Track);
-								break;
-							}
+							// If a track of the same class and name exists, remove it so the new track replaces it
+							MovieScene->RemoveTrack(*Track);
+							break;
 						}
 					}
 				}
diff --git a/Engine/Source/Runtime/Core/Public/UObject/FortniteMainBranchObjectVersion.h b/Engine/Source/Runtime/Core/Public/UObject/FortniteMainBranchObjectVersion.h
--- a/Engine/Source/Runtime/Core/Public/UObject/FortniteMainBranchObjectVersion.h
+++ b/Engine/Source/Runtime/Core/Public/UObject/FortniteMainBranchObjectVersion.h
@@ -366,6 +366,9 @@ struct FFortniteMainBranchObjectVersion
 		// World partition actor data layers activation logic operator support defaults for old maps
 		WorldPartitionDataLayersLogicOperatorAdded,
 
+		// Started sorting Possessables, Spawnables, and MovieSceneBindings for better search performance.
+		MovieSceneSortedBindings,
+
 		// -----<new versions can be added above this line>-------------------------------------------------
 		VersionPlusOne,
 		LatestVersion = VersionPlusOne - 1
diff --git a/Engine/Source/Runtime/MovieScene/Private/MovieScene.cpp b/Engine/Source/Runtime/MovieScene/Private/MovieScene.cpp
--- a/Engine/Source/Runtime/MovieScene/Private/MovieScene.cpp
+++ b/Engine/Source/Runtime/MovieScene/Private/MovieScene.cpp
@@ -15,6 +15,7 @@
 #include "Misc/FrameRate.h"
 #include "UObject/ObjectSaveContext.h"
 #include "UObject/UObjectHash.h"
+#include "UObject/FortniteMainBranchObjectVersion.h"
 
 #include UE_INLINE_GENERATED_CPP_BY_NAME(MovieScene)
 #include UE_INLINE_GENERATED_CPP_BY_NAME(MovieSceneMarkedFrame)
@@ -119,6 +120,7 @@ void UMovieScene::Serialize( FArchive& Ar )
 {
 	Ar.UsingCustomVersion(FMovieSceneEvaluationCustomVersion::GUID);
 	Ar.UsingCustomVersion(FSequencerObjectVersion::GUID);
+	Ar.UsingCustomVersion(FFortniteMainBranchObjectVersion::GUID);
 
 	// Serialize the MovieScene
 	Super::Serialize(Ar);
@@ -202,6 +204,17 @@ void UMovieScene::Serialize( FArchive& Ar )
 		}
 	}
 #endif
+
+	if (Ar.IsLoading())
+	{
+		const bool bSortBindings = Ar.CustomVer(FFortniteMainBranchObjectVersion::GUID) < FFortniteMainBranchObjectVersion::MovieSceneSortedBindings;
+		if (bSortBindings)
+		{
+			Possessables.Sort();
+			Spawnables.Sort();
+			ObjectBindings.Sort();
+		}
+	}
 }
 
 #if WITH_EDITOR
@@ -213,6 +226,15 @@ void UMovieScene::PostEditUndo()
 }
 #endif
 
+template <typename RangeType, typename ValueType>
+auto InsertSorted(RangeType& Range, ValueType&& Value) -> decltype(GetNum(Range))
+{
+	auto InsertIndex = Algo::LowerBound(Range, Value);
+	check(InsertIndex >= 0 && InsertIndex <= GetNum(Range));
+	Range.Insert(Value, InsertIndex);
+	return InsertIndex;
+}
+
 // @todo sequencer: Some of these methods should only be used by tools, and should probably move out of MovieScene!
 FGuid UMovieScene::AddSpawnable( const FString& Name, UObject& ObjectTemplate )
 {
@@ -220,14 +242,18 @@ FGuid UMovieScene::AddSpawnable( const FString& Name, UObject& ObjectTemplate )
 
 	FMovieSceneSpawnable NewSpawnable( Name, ObjectTemplate );
 	NewSpawnable.AutoSetNetAddressableName();
-	Spawnables.Add( NewSpawnable );
 
-	// Add a new binding so that tracks can be added to it
-	new (ObjectBindings) FMovieSceneBinding( NewSpawnable.GetGuid(), NewSpawnable.GetName() );
+	// Insert new spawnable sorted into the Spawnables array
+	int32 NewSpawnableIndex = InsertSorted(Spawnables, MoveTemp(NewSpawnable));
+
+	check(Spawnables.IsValidIndex(NewSpawnableIndex));
 
-	EventHandlers.Trigger(&UE::MovieScene::ISequenceDataEventHandler::OnBindingAdded, ObjectBindings.Last());
+	// Add a new binding (sorted) so that tracks can be added to it
+	int32 NewBindingIndex = InsertSorted(ObjectBindings, FMovieSceneBinding(Spawnables[NewSpawnableIndex].GetGuid(), Spawnables[NewSpawnableIndex].GetName()));
 
-	return NewSpawnable.GetGuid();
+	EventHandlers.Trigger(&UE::MovieScene::ISequenceDataEventHandler::OnBindingAdded, ObjectBindings[NewBindingIndex]);
+
+	return Spawnables[NewSpawnableIndex].GetGuid();
 }
 
 void UMovieScene::AddSpawnable(const FMovieSceneSpawnable& InNewSpawnable, const FMovieSceneBinding& InNewBinding)
@@ -237,16 +263,18 @@ void UMovieScene::AddSpawnable(const FMovieSceneSpawnable& InNewSpawnable, const
 	FMovieSceneSpawnable NewSpawnable;
 	NewSpawnable = InNewSpawnable;
 	NewSpawnable.AutoSetNetAddressableName();
-	Spawnables.Add(NewSpawnable);
+
+	// Insert new spawnable sorted into the Spawnables array
+	InsertSorted(Spawnables, MoveTemp(NewSpawnable));
 
 	FMovieSceneBinding NewBinding = InNewBinding;
 	for (auto Track : NewBinding.GetTracks())
 	{
 		Track->Rename(nullptr, this);
 	}
-	ObjectBindings.Add(NewBinding);
+	int32 NewBindingIndex = InsertSorted(ObjectBindings, MoveTemp(NewBinding));
 
-	EventHandlers.Trigger(&UE::MovieScene::ISequenceDataEventHandler::OnBindingAdded, ObjectBindings.Last());
+	EventHandlers.Trigger(&UE::MovieScene::ISequenceDataEventHandler::OnBindingAdded, ObjectBindings[NewBindingIndex]);
 }
 
 bool UMovieScene::RemoveSpawnable( const FGuid& Guid )
@@ -254,21 +282,16 @@ bool UMovieScene::RemoveSpawnable( const FGuid& Guid )
 	bool bAnythingRemoved = false;
 	if( ensure( Guid.IsValid() ) )
 	{
-		for( auto SpawnableIter( Spawnables.CreateIterator() ); SpawnableIter; ++SpawnableIter )
+		int32 Index = IndexOfSpawnable(Guid);
+		if (Index != INDEX_NONE)
 		{
-			auto& CurSpawnable = *SpawnableIter;
-			if( CurSpawnable.GetGuid() == Guid )
-			{
-				Modify();
-				RemoveBinding( Guid );
-
-				Spawnables.RemoveAt( SpawnableIter.GetIndex() );
+			Modify();
+			RemoveBinding(Guid);
+			Spawnables.RemoveAt(Index);
 
-				EventHandlers.Trigger(&UE::MovieScene::ISequenceDataEventHandler::OnBindingRemoved, Guid);
+			EventHandlers.Trigger(&UE::MovieScene::ISequenceDataEventHandler::OnBindingRemoved, Guid);
 
-				bAnythingRemoved = true;
-				break;
-			}
+			bAnythingRemoved = true;
 		}
 	}
 
@@ -292,9 +315,12 @@ int32 UMovieScene::GetSpawnableCount() const
 
 FMovieSceneSpawnable* UMovieScene::FindSpawnable( const FGuid& Guid )
 {
-	return Spawnables.FindByPredicate([&](FMovieSceneSpawnable& Spawnable) {
-		return Spawnable.GetGuid() == Guid;
-	});
+	int32 Index = IndexOfSpawnable(Guid);
+	if (Index != INDEX_NONE)
+	{
+		return &Spawnables[Index];
+	}
+	return nullptr;
 }
 
 
@@ -302,15 +328,14 @@ FGuid UMovieScene::AddPossessable( const FString& Name, UClass* Class )
 {
 	Modify();
 
-	FMovieScenePossessable NewPossessable( Name, Class );
-	Possessables.Add( NewPossessable );
-
+	int32 NewPossessableIndex = InsertSorted(Possessables, FMovieScenePossessable(Name, Class));
+	check(Possessables.IsValidIndex(NewPossessableIndex));
 	// Add a new binding so that tracks can be added to it
-	new (ObjectBindings) FMovieSceneBinding( NewPossessable.GetGuid(), NewPossessable.GetName() );
-
-	EventHandlers.Trigger(&UE::MovieScene::ISequenceDataEventHandler::OnBindingAdded, ObjectBindings.Last());
+	int32 NewBindingIndex = InsertSorted(ObjectBindings, FMovieSceneBinding(Possessables[NewPossessableIndex].GetGuid(), Possessables[NewPossessableIndex].GetName()));
+	check(ObjectBindings.IsValidIndex(NewBindingIndex));
+	EventHandlers.Trigger(&UE::MovieScene::ISequenceDataEventHandler::OnBindingAdded, ObjectBindings[NewBindingIndex]);
 
-	return NewPossessable.GetGuid();
+	return Possessables[NewPossessableIndex].GetGuid();
 }
 
 
@@ -320,16 +345,17 @@ void UMovieScene::AddPossessable(const FMovieScenePossessable& InNewPossessable,
 
 	FMovieScenePossessable NewPossessable;
 	NewPossessable = InNewPossessable;
-	Possessables.Add(NewPossessable);
+	InsertSorted(Possessables, MoveTemp(NewPossessable));
 
 	FMovieSceneBinding NewBinding = InNewBinding;
 	for (auto Track : NewBinding.GetTracks())
 	{
 		Track->Rename(nullptr, this);
 	}
-	ObjectBindings.Add(NewBinding);
+	int32 NewBindingIndex = InsertSorted(ObjectBindings, MoveTemp(NewBinding));
+	check(ObjectBindings.IsValidIndex(NewBindingIndex));
 
-	EventHandlers.Trigger(&UE::MovieScene::ISequenceDataEventHandler::OnBindingAdded, ObjectBindings.Last());
+	EventHandlers.Trigger(&UE::MovieScene::ISequenceDataEventHandler::OnBindingAdded, ObjectBindings[NewBindingIndex]);
 }
 
 
@@ -337,34 +363,30 @@ bool UMovieScene::RemovePossessable( const FGuid& PossessableGuid )
 {
 	bool bAnythingRemoved = false;
 
-	for( auto PossesableIter( Possessables.CreateIterator() ); PossesableIter; ++PossesableIter )
+	int32 Index = IndexOfPossessable(PossessableGuid);
+	if (Index != INDEX_NONE)
 	{
-		auto& CurPossesable = *PossesableIter;
-
-		if( CurPossesable.GetGuid() == PossessableGuid )
-		{	
-			Modify();
+		const FMovieScenePossessable& FoundPossessable = Possessables[Index];
+		Modify();
 
-			// Remove the parent-child link for a parent spawnable/child possessable if necessary
-			if (CurPossesable.GetParent().IsValid())
+		// Remove the parent-child link for a parent spawnable/child possessable if necessary
+		if (FoundPossessable.GetParent().IsValid())
+		{
+			FMovieSceneSpawnable* ParentSpawnable = FindSpawnable(FoundPossessable.GetParent());
+			if (ParentSpawnable)
 			{
-				FMovieSceneSpawnable* ParentSpawnable = FindSpawnable(CurPossesable.GetParent());
-				if (ParentSpawnable)
-				{
-					ParentSpawnable->RemoveChildPossessable(PossessableGuid);
-				}
+				ParentSpawnable->RemoveChildPossessable(PossessableGuid);
 			}
+		}
 
-			// Found it!
-			Possessables.RemoveAt( PossesableIter.GetIndex() );
+		// Found it!
+		Possessables.RemoveAt(Index);
 
-			RemoveBinding( PossessableGuid );
+		RemoveBinding(PossessableGuid);
 
-			EventHandlers.Trigger(&UE::MovieScene::ISequenceDataEventHandler::OnBindingRemoved, PossessableGuid);
+		EventHandlers.Trigger(&UE::MovieScene::ISequenceDataEventHandler::OnBindingRemoved, PossessableGuid);
 
-			bAnythingRemoved = true;
-			break;
-		}
+		bAnythingRemoved = true;
 	}
 
 	return bAnythingRemoved;
@@ -375,36 +397,31 @@ bool UMovieScene::ReplacePossessable( const FGuid& OldGuid, const FMovieScenePos
 {
 	bool bAnythingReplaced = false;
 
-	for (auto& Possessable : Possessables)
+	if (FMovieScenePossessable* OldPossessable = FindPossessable(OldGuid))
 	{
-		if (Possessable.GetGuid() == OldGuid)
-		{	
-			Modify();
+		Modify();
 
-			bool bNullPossessedObjectClass = true;
+		bool bNullPossessedObjectClass = true;
 #if WITH_EDITORONLY_DATA
-			bNullPossessedObjectClass = InNewPosessable.GetPossessedObjectClass() == nullptr;
+		bNullPossessedObjectClass = InNewPosessable.GetPossessedObjectClass() == nullptr;
 #endif
 
-			// Found it!
-			if (bNullPossessedObjectClass)
-			{
-				// @todo: delete this when
-				// bool ReplacePossessable(const FGuid& OldGuid, const FGuid& NewGuid, const FString& Name)
-				// is removed
-				Possessable.SetGuid(InNewPosessable.GetGuid());
-				Possessable.SetName(InNewPosessable.GetName());
-			}
-			else
-			{
-				Possessable = InNewPosessable;
-			}
-
-			ReplaceBinding( OldGuid, InNewPosessable.GetGuid(), InNewPosessable.GetName() );
-			bAnythingReplaced = true;
-
-			break;
+		// Found it!
+		if (bNullPossessedObjectClass)
+		{
+			// @todo: delete this when
+			// bool ReplacePossessable(const FGuid& OldGuid, const FGuid& NewGuid, const FString& Name)
+			// is removed
+			OldPossessable->SetGuid(InNewPosessable.GetGuid());
+			OldPossessable->SetName(InNewPosessable.GetName());
+		}
+		else
+		{
+			*OldPossessable = InNewPosessable;
 		}
+
+		ReplaceBinding( OldGuid, InNewPosessable.GetGuid(), InNewPosessable.GetName() );
+		bAnythingReplaced = true;
 	}
 
 	return bAnythingReplaced;
@@ -413,14 +430,11 @@ bool UMovieScene::ReplacePossessable( const FGuid& OldGuid, const FMovieScenePos
 
 FMovieScenePossessable* UMovieScene::FindPossessable( const FGuid& Guid )
 {
-	for (auto& Possessable : Possessables)
+	int32 Index = IndexOfPossessable(Guid);
+	if (Index != INDEX_NONE)
 	{
-		if (Possessable.GetGuid() == Guid)
-		{
-			return &Possessable;
-		}
+		return &Possessables[Index];
 	}
-
 	return nullptr;
 }
 
@@ -999,15 +1013,10 @@ TArray<UMovieSceneSection*> UMovieScene::GetAllSections() const
 UMovieSceneTrack* UMovieScene::FindTrack(TSubclassOf<UMovieSceneTrack> TrackClass, const FGuid& ObjectGuid, const FName& TrackName) const
 {
 	check( ObjectGuid.IsValid() );
-	
-	for (const auto& Binding : ObjectBindings)
-	{
-		if (Binding.GetObjectGuid() != ObjectGuid) 
-		{
-			continue;
-		}
 
-		for (const auto& Track : Binding.GetTracks())
+	if (const FMovieSceneBinding* Binding = FindBinding(ObjectGuid))
+	{
+		for (const auto& Track : Binding->GetTracks())
 		{
 			if (TrackClass.GetDefaultObject() == nullptr || Track->GetClass()->IsChildOf(TrackClass))
 			{
@@ -1026,14 +1035,9 @@ TArray<UMovieSceneTrack*> UMovieScene::FindTracks(TSubclassOf<UMovieSceneTrack>
 {
 	check(ObjectGuid.IsValid());
 	TArray<UMovieSceneTrack*> MovieSceneTracks;
-	for (const auto& Binding : ObjectBindings)
+	if (const FMovieSceneBinding* Binding = FindBinding(ObjectGuid))
 	{
-		if (Binding.GetObjectGuid() != ObjectGuid)
-		{
-			continue;
-		}
-
-		for (const auto& Track : Binding.GetTracks())
+		for (const auto& Track : Binding->GetTracks())
 		{
 			if (TrackClass.GetDefaultObject() == nullptr || Track->GetClass()->IsChildOf(TrackClass))
 			{
@@ -1059,18 +1063,15 @@ UMovieSceneTrack* UMovieScene::AddTrack( TSubclassOf<UMovieSceneTrack> TrackClas
 
 	UMovieSceneTrack* CreatedType = nullptr;
 
-	check( ObjectGuid.IsValid() )
+	check(ObjectGuid.IsValid());
 
-	for (auto& Binding : ObjectBindings)
+	if (FMovieSceneBinding* Binding = FindBinding(ObjectGuid))
 	{
-		if( Binding.GetObjectGuid() == ObjectGuid ) 
-		{
-			Modify();
+		Modify();
 
-			CreatedType = NewObject<UMovieSceneTrack>(this, TrackClass, NAME_None, RF_Transactional);
-			check(CreatedType);
-			Binding.AddTrack( *CreatedType, this );
-		}
+		CreatedType = NewObject<UMovieSceneTrack>(this, TrackClass, NAME_None, RF_Transactional);
+		check(CreatedType);
+		Binding->AddTrack(*CreatedType, this);
 	}
 
 	return CreatedType;
@@ -1089,14 +1090,11 @@ bool UMovieScene::AddGivenTrack(UMovieSceneTrack* InTrack, const FGuid& ObjectGu
 	check(InTrack);
 
 	Modify();
-	for (auto& Binding : ObjectBindings)
+	if (FMovieSceneBinding* Binding = FindBinding(ObjectGuid))
 	{
-		if (Binding.GetObjectGuid() == ObjectGuid)
-		{
-			InTrack->Rename(nullptr, this);
-			Binding.AddTrack(*InTrack, this);
-			return true;
-		}
+		InTrack->Rename(nullptr, this);
+		Binding->AddTrack(*InTrack, this);
+		return true;
 	}
 
 	return false;
@@ -1489,36 +1487,44 @@ void UMovieScene::RemoveBinding(const FGuid& Guid)
 	// WARNING: This function intentionally does not trigger events to ensure
 	// that events are triggered when all processing is complete (ie, when removing a spawnable or posessable)
 
-	for (int32 BindingIndex = 0; BindingIndex < ObjectBindings.Num(); ++BindingIndex)
+	int32 Index = IndexOfBinding(Guid);
+	if (Index != INDEX_NONE)
 	{
-		if (ObjectBindings[BindingIndex].GetObjectGuid() == Guid)
-		{
-			ObjectBindings.RemoveAt(BindingIndex);
-			break;
-		}
+		ObjectBindings.RemoveAt(Index);
 	}
 }
 
+int32 UMovieScene::IndexOfBinding(const FGuid& ForGuid) const
+{
+	return Algo::BinarySearch(ObjectBindings, ForGuid);
+}
+
+int32 UMovieScene::IndexOfSpawnable(const FGuid& ForGuid) const
+{
+	return Algo::BinarySearch(Spawnables, ForGuid);
+}
+
+int32 UMovieScene::IndexOfPossessable(const FGuid& ForGuid) const
+{
+	return Algo::BinarySearch(Possessables, ForGuid);
+}
+
 void UMovieScene::ReplaceBinding(const FGuid& OldGuid, const FGuid& NewGuid, const FString& Name)
 {
-	for (auto& Binding : ObjectBindings)
+	if (FMovieSceneBinding* Binding = FindBinding(OldGuid))
 	{
-		if (Binding.GetObjectGuid() == OldGuid)
-		{
-			Binding.SetObjectGuid(NewGuid);
-			Binding.SetName(Name);
-
-			// Changing a binding guid invalidates any tracks contained within the binding
-			// Make sure they are written into the transaction buffer by calling modify
-			for (UMovieSceneTrack* Track : Binding.GetTracks())
-			{
-				Track->Modify();
-			}
+		Binding->SetObjectGuid(NewGuid);
+		Binding->SetName(Name);
 
-			EventHandlers.Trigger(&UE::MovieScene::ISequenceDataEventHandler::OnBindingRemoved, OldGuid);
-			EventHandlers.Trigger(&UE::MovieScene::ISequenceDataEventHandler::OnBindingAdded, Binding);
-			break;
+		// Changing a binding guid invalidates any tracks contained within the binding
+		// Make sure they are written into the transaction buffer by calling modify
+		for (UMovieSceneTrack* Track : Binding->GetTracks())
+		{
+			Track->Modify();
 		}
+
+		EventHandlers.Trigger(&UE::MovieScene::ISequenceDataEventHandler::OnBindingRemoved, OldGuid);
+		EventHandlers.Trigger(&UE::MovieScene::ISequenceDataEventHandler::OnBindingAdded, *Binding);
 	}
 }
 
@@ -1526,11 +1532,9 @@ void UMovieScene::ReplaceBinding(const FGuid& BindingToReplaceGuid, const FMovie
 {
 	Modify();
 
-	FMovieSceneBinding* Binding = ObjectBindings.FindByPredicate([BindingToReplaceGuid](const FMovieSceneBinding& CheckedBinding) { return CheckedBinding.GetObjectGuid() == BindingToReplaceGuid; });
-	if (Binding)
+	if (FMovieSceneBinding* Binding = FindBinding(BindingToReplaceGuid))
 	{
 		*Binding = NewBinding;
-
 		// We also need to change the track's owners to be the MovieScene.
 		for (auto Track : Binding->GetTracks())
 		{
@@ -1544,25 +1548,8 @@ void UMovieScene::ReplaceBinding(const FGuid& BindingToReplaceGuid, const FMovie
 
 void UMovieScene::MoveBindingContents(const FGuid& SourceBindingId, const FGuid& DestinationBindingId)
 {
-	FMovieSceneBinding* SourceBinding = nullptr;
-	FMovieSceneBinding* DestinationBinding = nullptr;
-
-	for (FMovieSceneBinding& Binding : ObjectBindings)
-	{
-		if (Binding.GetObjectGuid() == SourceBindingId)
-		{
-			SourceBinding = &Binding;
-		}
-		else if (Binding.GetObjectGuid() == DestinationBindingId)
-		{
-			DestinationBinding = &Binding;
-		}
-
-		if (SourceBinding && DestinationBinding)
-		{
-			break;
-		}
-	}
+	FMovieSceneBinding* SourceBinding = FindBinding(SourceBindingId);
+	FMovieSceneBinding* DestinationBinding = FindBinding(DestinationBindingId);
 
 	if (SourceBinding && DestinationBinding)
 	{
@@ -1618,6 +1605,26 @@ void UMovieScene::MoveBindingContents(const FGuid& SourceBindingId, const FGuid&
 	}
 }
 
+FMovieSceneBinding* UMovieScene::FindBinding(const FGuid& ForGuid)
+{
+	int32 Index = IndexOfBinding(ForGuid);
+	if (Index != INDEX_NONE)
+	{
+		return &ObjectBindings[Index];
+	}
+	return nullptr;
+}
+
+const FMovieSceneBinding* UMovieScene::FindBinding(const FGuid& ForGuid) const
+{
+	int32 Index = IndexOfBinding(ForGuid);
+	if (Index != INDEX_NONE)
+	{
+		return &ObjectBindings[Index];
+	}
+	return nullptr;
+}
+
 TSharedPtr<FMovieSceneTimeController> UMovieScene::MakeCustomTimeController(UObject* PlaybackContext)
 {
 	return MakeShared<FMovieSceneTimeController_Custom>(CustomClockSourcePath, PlaybackContext);
diff --git a/Engine/Source/Runtime/MovieScene/Public/MovieScene.h b/Engine/Source/Runtime/MovieScene/Public/MovieScene.h
--- a/Engine/Source/Runtime/MovieScene/Public/MovieScene.h
+++ b/Engine/Source/Runtime/MovieScene/Public/MovieScene.h
@@ -714,11 +714,16 @@ public:
 	 * @param ForGuid	The binding's Guid to look for.
 	 * @return			Pointer to the binding, otherwise nullptr.
 	 */
-	FMovieSceneBinding* FindBinding(const FGuid& ForGuid)
-	{
-		return ObjectBindings.FindByPredicate([ForGuid](const FMovieSceneBinding& Binding) { return Binding.GetObjectGuid() == ForGuid; });
-	}
+	MOVIESCENE_API FMovieSceneBinding* FindBinding(const FGuid& ForGuid);
 
+	/**
+	* Tries to find an FMovieSceneBinding for the specified Guid.
+	*
+	* @param ForGuid	The binding's Guid to look for.
+	* @return			Pointer to the binding, otherwise nullptr.
+	*/
+	MOVIESCENE_API const FMovieSceneBinding* FindBinding(const FGuid& ForGuid) const;
+	
 public:
 
 	// @todo sequencer: the following methods really shouldn't be here
@@ -1194,6 +1199,32 @@ protected:
 	 */
 	MOVIESCENE_API void RemoveBinding(const FGuid& Guid);
 
+
+	/**
+	 * Tries to find an FMovieSceneBinding for the specified Guid.
+	 *
+	 * @param ForGuid	The binding's Guid to look for.
+	 * @return			Index of the binding in ObjectBindings array, otherwise INDEX_NONE.
+	 */
+	int32 IndexOfBinding(const FGuid& ForGuid) const;
+
+	/**
+	* Tries to find an FMovieSceneSpawnable for the specified Guid.
+	*
+	* @param ForGuid	The spawnable's Guid to look for.
+	* @return			Index of the binding in Spawnables array, otherwise INDEX_NONE.
+	*/
+	int32 IndexOfSpawnable(const FGuid& ForGuid) const;
+
+
+	/**
+	* Tries to find an FMovieScenePossessable for the specified Guid.
+	*
+	* @param ForGuid	The possessable's Guid to look for.
+	* @return			Index of the binding in Possessables array, otherwise INDEX_NONE.
+	*/
+	int32 IndexOfPossessable(const FGuid& ForGuid) const;
+
 protected:
 
 	/** Called before this object is being deserialized. */
diff --git a/Engine/Source/Runtime/MovieScene/Public/MovieSceneBinding.h b/Engine/Source/Runtime/MovieScene/Public/MovieSceneBinding.h
--- a/Engine/Source/Runtime/MovieScene/Public/MovieSceneBinding.h
+++ b/Engine/Source/Runtime/MovieScene/Public/MovieSceneBinding.h
@@ -132,6 +132,11 @@ struct FMovieSceneBinding
 	 */
 	MOVIESCENE_API void SetTracks(TArray<UMovieSceneTrack*>&& InTracks, UMovieScene* Owner);
 
+	/* For sorts so we can search quickly by Guid */
+	FORCEINLINE bool operator<(const FMovieSceneBinding& RHS) const { return ObjectGuid < RHS.ObjectGuid; }
+	FORCEINLINE bool operator<(const FGuid& InGuid) const { return ObjectGuid < InGuid; }
+	FORCEINLINE friend bool operator<(const FGuid& InGuid, const FMovieSceneBinding& RHS) { return InGuid < RHS.GetObjectGuid(); }
+
 #if WITH_EDITORONLY_DATA
 	/**
 	* Get this folder's desired sorting order
diff --git a/Engine/Source/Runtime/MovieScene/Public/MovieScenePossessable.h b/Engine/Source/Runtime/MovieScene/Public/MovieScenePossessable.h
--- a/Engine/Source/Runtime/MovieScene/Public/MovieScenePossessable.h
+++ b/Engine/Source/Runtime/MovieScene/Public/MovieScenePossessable.h
@@ -163,6 +163,11 @@ public:
 	/* Bind the potential spawnable object to this possessable by setting the ObjectBindingID */
 	MOVIESCENE_API bool BindSpawnableObject(FMovieSceneSequenceID SequenceID, UObject* Object, IMovieScenePlayer* Player);
 
+	/* For sorts so we can search quickly by Guid */
+	FORCEINLINE bool operator<(const FMovieScenePossessable& RHS) const { return Guid < RHS.Guid; }
+	FORCEINLINE friend bool operator<(const FGuid& InGuid, const FMovieScenePossessable& RHS) { return InGuid < RHS.GetGuid(); }
+	bool operator<(const FGuid& InGuid) const { return Guid < InGuid; }
+
 private:
 
 	/** Unique identifier of the possessable object. */
diff --git a/Engine/Source/Runtime/MovieScene/Public/MovieSceneSpawnable.h b/Engine/Source/Runtime/MovieScene/Public/MovieSceneSpawnable.h
--- a/Engine/Source/Runtime/MovieScene/Public/MovieSceneSpawnable.h
+++ b/Engine/Source/Runtime/MovieScene/Public/MovieSceneSpawnable.h
@@ -237,6 +237,11 @@ public:
 	 */
 	MOVIESCENE_API void AutoSetNetAddressableName();
 
+	/* For sorts and BinarySearch so we can search quickly by Guid */
+	FORCEINLINE bool operator<(const FMovieSceneSpawnable& RHS) const { return Guid < RHS.Guid; }
+	FORCEINLINE friend bool operator<(const FGuid& InGuid, const FMovieSceneSpawnable& RHS) { return InGuid < RHS.GetGuid(); }
+	FORCEINLINE bool operator<(const FGuid& InGuid) const { return Guid < InGuid; }
+
 	/** Array of tags that can be used for grouping and categorizing. */
 	UPROPERTY(EditAnywhere, AdvancedDisplay, Category=Actor)
 	TArray<FName> Tags;
diff --git a/Engine/Source/Runtime/MovieSceneTracks/Private/Sections/MovieSceneCameraCutSection.cpp b/Engine/Source/Runtime/MovieSceneTracks/Private/Sections/MovieSceneCameraCutSection.cpp
--- a/Engine/Source/Runtime/MovieSceneTracks/Private/Sections/MovieSceneCameraCutSection.cpp
+++ b/Engine/Source/Runtime/MovieSceneTracks/Private/Sections/MovieSceneCameraCutSection.cpp
@@ -137,17 +137,14 @@ void UMovieSceneCameraCutSection::ComputeInitialCameraCutTransform()
 		UMovieScene* MovieScene = GetTypedOuter<UMovieScene>();
 		check(MovieScene);
 
-		for (const FMovieSceneBinding& Binding : MovieScene->GetBindings())
+		if (const FMovieSceneBinding* Binding = MovieScene->FindBinding(CameraBindingID.GetGuid()))
 		{
-			if (Binding.GetObjectGuid() == CameraBindingID.GetGuid())
+			for (UMovieSceneTrack* Track : Binding->GetTracks())
 			{
-				for (UMovieSceneTrack* Track : Binding.GetTracks())
+				CameraTransformTrack = Cast<UMovieScene3DTransformTrack>(Track);
+				if (CameraTransformTrack)
 				{
-					CameraTransformTrack = Cast<UMovieScene3DTransformTrack>(Track);
-					if (CameraTransformTrack)
-					{
-						break;
-					}
+					break;
 				}
 			}
 		}
