Sequencer: Improved handling of updates after a compile during a partial evaluation

  - Operate on a copy of the CurrentInstances array since it is possible that this can change when sequences are being updated
  - Only call Update on root-level sequences

diff --git a/Engine/Source/Runtime/MovieScene/Private/EntitySystem/MovieSceneEntitySystemRunner.cpp b/Engine/Source/Runtime/MovieScene/Private/EntitySystem/MovieSceneEntitySystemRunner.cpp
--- a/Engine/Source/Runtime/MovieScene/Private/EntitySystem/MovieSceneEntitySystemRunner.cpp
+++ b/Engine/Source/Runtime/MovieScene/Private/EntitySystem/MovieSceneEntitySystemRunner.cpp
@@ -868,16 +868,23 @@ bool FMovieSceneEntitySystemRunner::GameThread_ReimportSequenceInstances(UMovieS
 {
 	using namespace UE::MovieScene;
 
-	FInstanceRegistry* InstanceRegistry = WeakLinker.Get()->GetInstanceRegistry();
-
 	// Only called after a sequence has been recompiled after we have already updated the current instances
 	// This allows us to re-update all the sequence instances in case anything has changed
-	for (const FQueuedUpdateParams& UpdatedInstance : CurrentInstances)
+
+	FInstanceRegistry* InstanceRegistry = WeakLinker.Get()->GetInstanceRegistry();
+
+	// Operate on a copy of the update params since it is possible CurrentInstances can change if sequence instances are destroyed
+	TArray<FQueuedUpdateParams> CurrentInstancesCopy = CurrentInstances;
+	for (const FQueuedUpdateParams& UpdatedInstance : CurrentInstancesCopy)
 	{
-		if (!EnumHasAnyFlags(UpdatedInstance.UpdateFlags, ERunnerUpdateFlags::Destroy | ERunnerUpdateFlags::Finish))
+		if (InstanceRegistry->IsHandleValid(UpdatedInstance.InstanceHandle) && !EnumHasAnyFlags(UpdatedInstance.UpdateFlags, ERunnerUpdateFlags::Destroy | ERunnerUpdateFlags::Finish))
 		{
 			FSequenceInstance& SequenceInstance = InstanceRegistry->MutateInstance(UpdatedInstance.InstanceHandle);
-			SequenceInstance.Update(Linker, SequenceInstance.GetContext());
+			// Only update root instances
+			if (SequenceInstance.IsRootSequence())
+			{
+				SequenceInstance.Update(Linker, SequenceInstance.GetContext());
+			}
 		}
 	}
 
