Sequencer: Ensure that output indices that are destroyed by GC do not remain tagged as invalidated

diff --git a/Engine/Source/Runtime/MovieScene/Public/EntitySystem/MovieSceneOverlappingEntityTracker.h b/Engine/Source/Runtime/MovieScene/Public/EntitySystem/MovieSceneOverlappingEntityTracker.h
--- a/Engine/Source/Runtime/MovieScene/Public/EntitySystem/MovieSceneOverlappingEntityTracker.h
+++ b/Engine/Source/Runtime/MovieScene/Public/EntitySystem/MovieSceneOverlappingEntityTracker.h
@@ -592,6 +592,13 @@ struct TOverlappingEntityTracker_WithGarbage : TOverlappingEntityTrackerImpl<Out
 
 				this->Outputs.RemoveAt(Index, 1);
 
+				// Make sure this output is not flagged as invalidated because it is being destroyed.
+				// This prevents us from blindly processing it in ProcessInvalidatedOutputs
+				if (this->InvalidatedOutputs.IsValidIndex(OutputIndex))
+				{
+					this->InvalidatedOutputs[OutputIndex] = false;
+				}
+
 				for (auto It = this->OutputToEntity.CreateKeyIterator(OutputIndex); It; ++It)
 				{
 					this->EntityToOutput.Remove(It.Value());
