Sequencer: start tracking pre-animated state on new track instance inputs before stopping tracking on old ones

When two track instance inputs correspond to two adjacent sections, we see the input of the ending section be removed, and the input of the starting section be added. This would potentially mean that the first section gets removed from pre-animated state tracking before the new section gets started with the same tracking. If the two sections happen to be tied to the same "thing", and are the only contributors, we would restore pre-animated state between the two for nothing. Now we correctly carry pre-animated state over to the second section.

diff --git a/Engine/Source/Runtime/MovieScene/Private/EntitySystem/TrackInstance/MovieSceneTrackInstance.cpp b/Engine/Source/Runtime/MovieScene/Private/EntitySystem/TrackInstance/MovieSceneTrackInstance.cpp
--- a/Engine/Source/Runtime/MovieScene/Private/EntitySystem/TrackInstance/MovieSceneTrackInstance.cpp
+++ b/Engine/Source/Runtime/MovieScene/Private/EntitySystem/TrackInstance/MovieSceneTrackInstance.cpp
@@ -78,6 +78,8 @@ void UMovieSceneTrackInstance::UpdateInputs(TArray<FMovieSceneTrackInstanceInput
 	const int32 OldNum = Inputs.Num();
 	const int32 NewNum = InNewInputs.Num();
 
+	TArray<int32> StopTrackingIndices;
+
 	for ( ; OldIndex < OldNum || NewIndex < NewNum; )
 	{
 		while (OldIndex < OldNum && NewIndex < NewNum && Inputs[OldIndex].IsSameInput(InNewInputs[NewIndex]))
@@ -87,6 +89,9 @@ void UMovieSceneTrackInstance::UpdateInputs(TArray<FMovieSceneTrackInstanceInput
 				// This input is the same as one already existing in the track instance
 				// But it is being reimported - remove it and add it again
 				OnInputRemoved(Inputs[OldIndex]);
+				// We don't delay ending tracking for this input, because it will most likely get tracked
+				// in the call to OnInputAdded, which won't do anything if it's already tracked. And then
+				// later in the delayed end-tracking we would remove it. We don't want that.
 				if (InputMetaData)
 				{
 					InputMetaData->StopTrackingCaptureSource(Inputs[OldIndex]);
@@ -112,10 +117,7 @@ void UMovieSceneTrackInstance::UpdateInputs(TArray<FMovieSceneTrackInstanceInput
 			{
 				// Out with the old
 				OnInputRemoved(Inputs[OldIndex]);
-				if (InputMetaData)
-				{
-					InputMetaData->StopTrackingCaptureSource(Inputs[OldIndex]);
-				}
+				StopTrackingIndices.Add(OldIndex);
 				++OldIndex;
 			}
 			else
@@ -132,10 +134,7 @@ void UMovieSceneTrackInstance::UpdateInputs(TArray<FMovieSceneTrackInstanceInput
 		{
 			// Out with the old
 			OnInputRemoved(Inputs[OldIndex]);
-			if (InputMetaData)
-			{
-				InputMetaData->StopTrackingCaptureSource(Inputs[OldIndex]);
-			}
+			StopTrackingIndices.Add(OldIndex);
 			++OldIndex;
 		}
 		else if (ensure(NewIndex < NewNum))
@@ -149,6 +148,14 @@ void UMovieSceneTrackInstance::UpdateInputs(TArray<FMovieSceneTrackInstanceInput
 		}
 	}
 
+	if (InputMetaData)
+	{
+		for (int32 Index : StopTrackingIndices)
+		{
+			InputMetaData->StopTrackingCaptureSource(Inputs[Index]);
+		}
+	}
+
 	Swap(Inputs, InNewInputs);
 
 	OnEndUpdateInputs();
