Sequencer: Add Modify() in SetSequence

Previously, calliing SetSequence wouldn't modify the signature, causing the UI to be stale. This resulted in an ensure when trying to dive into a subsequence since the sequence ID didn't exist in the stale hierarchy.

diff --git a/Engine/Source/Runtime/MovieScene/Private/Sections/MovieSceneSubSection.cpp b/Engine/Source/Runtime/MovieScene/Private/Sections/MovieSceneSubSection.cpp
--- a/Engine/Source/Runtime/MovieScene/Private/Sections/MovieSceneSubSection.cpp
+++ b/Engine/Source/Runtime/MovieScene/Private/Sections/MovieSceneSubSection.cpp
@@ -220,6 +220,11 @@ void UMovieSceneSubSection::ImportEntityImpl(UMovieSceneEntitySystemLinker* Enti
 
 void UMovieSceneSubSection::SetSequence(UMovieSceneSequence* Sequence)
 {
+	if (!TryModify())
+	{
+		return;
+	}
+
 	SubSequence = Sequence;
 
 #if WITH_EDITOR
