Sequencer: add shared state structure for sequence hierarchies, with new "playback capabilities" container

This lets external code (like level sequence actors or UMG widgets) specify what a given sequence instance can provide in terms of contextual information or logic.

This is a first step towards removing the necessity for IMovieScenePlayer in core sequencer evaluation code.

diff --git a/Engine/Source/Runtime/MovieScene/Private/EntitySystem/MovieSceneInstanceRegistry.cpp b/Engine/Source/Runtime/MovieScene/Private/EntitySystem/MovieSceneInstanceRegistry.cpp
--- a/Engine/Source/Runtime/MovieScene/Private/EntitySystem/MovieSceneInstanceRegistry.cpp
+++ b/Engine/Source/Runtime/MovieScene/Private/EntitySystem/MovieSceneInstanceRegistry.cpp
@@ -6,6 +6,7 @@
 #include "EntitySystem/MovieSceneEntitySystemLinker.h"
 #include "EntitySystem/MovieSceneEntityInstantiatorSystem.h"
 #include "EntitySystem/MovieSceneEntitySystemTask.h"
+#include "EntitySystem/MovieSceneSharedPlaybackState.h"
 #include "EntitySystem/EntityAllocationIterator.h"
 #include "EntitySystem/BuiltInComponentTypes.h"
 
@@ -63,7 +64,9 @@ FRootInstanceHandle FInstanceRegistry::AllocateRootInstance(IMovieScenePlayer* P
 	FSparseArrayAllocationInfo NewAllocation = Instances.AddUninitialized();
 	FRootInstanceHandle InstanceHandle { (uint16)NewAllocation.Index, InstanceSerial };
 
-	new (NewAllocation) FSequenceInstance(Linker, Player, InstanceHandle);
+	TSharedRef<FSharedPlaybackState> NewPlaybackState = MakeShared<FSharedPlaybackState>();
+
+	new (NewAllocation) FSequenceInstance(Linker, Player, NewPlaybackState, InstanceHandle);
 
 	return InstanceHandle;
 }
@@ -75,8 +78,10 @@ FInstanceHandle FInstanceRegistry::AllocateSubInstance(IMovieScenePlayer* Player
 	const uint16 InstanceSerial = InstanceSerialNumber++;
 	FSparseArrayAllocationInfo NewAllocation = Instances.AddUninitialized();
 	FInstanceHandle InstanceHandle { (uint16)NewAllocation.Index, InstanceSerial };
+	
+	TSharedRef<FSharedPlaybackState> PlaybackState = GetInstance(RootInstanceHandle).GetSharedPlaybackState();
 
-	new (NewAllocation) FSequenceInstance(Linker, Player, InstanceHandle, ParentInstanceHandle, RootInstanceHandle, SequenceID);
+	new (NewAllocation) FSequenceInstance(Linker, Player, PlaybackState, InstanceHandle, ParentInstanceHandle, RootInstanceHandle, SequenceID);
 
 	return InstanceHandle;
 }
@@ -86,6 +91,7 @@ void FInstanceRegistry::DestroyInstance(FInstanceHandle InstanceHandle)
 	if (ensureMsgf(Instances.IsValidIndex(InstanceHandle.InstanceID) && Instances[InstanceHandle.InstanceID].GetSerialNumber() == InstanceHandle.InstanceSerial, TEXT("Attempting to destroy an instance an invalid instance handle.")))
 	{
 		FSequenceInstance& Instance = Instances[InstanceHandle.InstanceID];
+		const bool bIsRootInstance = Instance.IsRootSequence();
 		const bool bHasFinished = (GExitPurge || Instance.HasFinished());
 		if (!bHasFinished)
 		{
diff --git a/Engine/Source/Runtime/MovieScene/Private/EntitySystem/MovieSceneSequenceInstance.cpp b/Engine/Source/Runtime/MovieScene/Private/EntitySystem/MovieSceneSequenceInstance.cpp
--- a/Engine/Source/Runtime/MovieScene/Private/EntitySystem/MovieSceneSequenceInstance.cpp
+++ b/Engine/Source/Runtime/MovieScene/Private/EntitySystem/MovieSceneSequenceInstance.cpp
@@ -4,7 +4,7 @@
 #include "EntitySystem/MovieSceneEntitySystemLinker.h"
 #include "EntitySystem/MovieSceneEntitySystem.h"
 #include "EntitySystem/MovieSceneSequenceUpdaters.h"
-#include "Evaluation/PreAnimatedState/MovieScenePreAnimatedStateExtension.h"
+#include "EntitySystem/MovieSceneSharedPlaybackState.h"
 
 #include "Compilation/MovieSceneCompiledVolatilityManager.h"
 #include "Compilation/MovieSceneCompiledDataManager.h"
@@ -12,6 +12,7 @@
 #include "Evaluation/MovieSceneEvaluationTemplateInstance.h"
 #include "Evaluation/Instances/MovieSceneTrackEvaluator.h"
 #include "Evaluation/MovieSceneRootOverridePath.h"
+#include "Evaluation/PreAnimatedState/MovieScenePreAnimatedStateExtension.h"
 
 #include "IMovieScenePlayer.h"
 #include "MovieSceneSequence.h"
@@ -67,8 +68,9 @@ void PurgeStaleTrackTemplates(UMovieSceneCompiledDataManager* CompiledDataManage
 
 
 
-FSequenceInstance::FSequenceInstance(UMovieSceneEntitySystemLinker* Linker, IMovieScenePlayer* Player, FRootInstanceHandle InInstanceHandle)
-	: SequenceID(MovieSceneSequenceID::Root)
+FSequenceInstance::FSequenceInstance(UMovieSceneEntitySystemLinker* Linker, IMovieScenePlayer* Player, TSharedRef<FSharedPlaybackState> PlaybackState, FRootInstanceHandle InInstanceHandle)
+	: SharedPlaybackState(PlaybackState)
+	, SequenceID(MovieSceneSequenceID::Root)
 	, RootOverrideSequenceID(MovieSceneSequenceID::Root)
 	, PlayerIndex(Player->GetUniqueIndex())
 	, InstanceHandle(InInstanceHandle)
@@ -93,8 +95,9 @@ FSequenceInstance::FSequenceInstance(UMovieSceneEntitySystemLinker* Linker, IMov
 	InvalidateCachedData(Linker);
 }
 
-FSequenceInstance::FSequenceInstance(UMovieSceneEntitySystemLinker* Linker, IMovieScenePlayer* Player, FInstanceHandle InInstanceHandle, FInstanceHandle InParentInstanceHandle, FRootInstanceHandle InRootInstanceHandle, FMovieSceneSequenceID InSequenceID)
-	: SequenceID(InSequenceID)
+FSequenceInstance::FSequenceInstance(UMovieSceneEntitySystemLinker* Linker, IMovieScenePlayer* Player, TSharedRef<FSharedPlaybackState> PlaybackState, FInstanceHandle InInstanceHandle, FInstanceHandle InParentInstanceHandle, FRootInstanceHandle InRootInstanceHandle, FMovieSceneSequenceID InSequenceID)
+	: SharedPlaybackState(PlaybackState)
+	, SequenceID(InSequenceID)
 	, RootOverrideSequenceID(MovieSceneSequenceID::Invalid)
 	, PlayerIndex(Player->GetUniqueIndex())
 	, InstanceHandle(InInstanceHandle)
@@ -168,6 +171,8 @@ void FSequenceInstance::InvalidateCachedData(UMovieSceneEntitySystemLinker* Link
 
 	if (SequenceID == MovieSceneSequenceID::Root)
 	{
+		SharedPlaybackState->InvalidateCachedData(Linker);
+
 		Player->State.AssignSequence(SequenceID, *RootTemplate.GetRootSequence(), *Player);
 
 		// Try and recreate the volatility manager if this sequence is now volatile
diff --git a/Engine/Source/Runtime/MovieScene/Private/Evaluation/IMovieScenePlaybackCapability.cpp b/Engine/Source/Runtime/MovieScene/Private/Evaluation/IMovieScenePlaybackCapability.cpp
new file mode 100644
--- /dev/null
+++ b/Engine/Source/Runtime/MovieScene/Private/Evaluation/IMovieScenePlaybackCapability.cpp
@@ -0,0 +1,19 @@
+// Copyright Epic Games, Inc. All Rights Reserved.
+
+#include "Evaluation/IMovieScenePlaybackCapability.h"
+
+namespace UE::MovieScene
+{
+
+FPlaybackCapabilityID FPlaybackCapabilityID::Register()
+{
+	static int32 StaticID = 0;
+	FPlaybackCapabilityID NewID{ StaticID++ };
+	// The ID is a bit index inside FPlaybackCapabilitiesImpl's AllCapabilities uint32 bitmask,
+	// so we don't want to overflow this.
+	checkf((size_t)StaticID < sizeof(uint32) * CHAR_BIT, TEXT("Exceeded the maximum possible amount of playback capabilities!"));
+	return NewID;
+}
+
+} // namespace UE::MovieScene
+
diff --git a/Engine/Source/Runtime/MovieScene/Private/Evaluation/MovieScenePlaybackCapabilities.cpp b/Engine/Source/Runtime/MovieScene/Private/Evaluation/MovieScenePlaybackCapabilities.cpp
new file mode 100644
--- /dev/null
+++ b/Engine/Source/Runtime/MovieScene/Private/Evaluation/MovieScenePlaybackCapabilities.cpp
@@ -0,0 +1,79 @@
+// Copyright Epic Games, Inc. All Rights Reserved.
+
+#include "Evaluation/MovieScenePlaybackCapabilities.h"
+
+namespace UE::MovieScene
+{
+
+FPlaybackCapabilities::FPlaybackCapabilities(FPlaybackCapabilities&& RHS)
+{
+	Memory = RHS.Memory;
+	Alignment = RHS.Alignment;
+	Capacity = RHS.Capacity;
+	Num = RHS.Num;
+	AllCapabilities = RHS.AllCapabilities;
+
+	new (&RHS) FPlaybackCapabilities();
+}
+
+FPlaybackCapabilities& FPlaybackCapabilities::operator=(FPlaybackCapabilities&& RHS)
+{
+	Destroy();
+
+	Memory = RHS.Memory;
+	Alignment = RHS.Alignment;
+	Capacity = RHS.Capacity;
+	Num = RHS.Num;
+	AllCapabilities = RHS.AllCapabilities;
+
+	new (&RHS) FPlaybackCapabilities();
+
+	return *this;
+}
+
+FPlaybackCapabilities::~FPlaybackCapabilities()
+{
+	Destroy();
+}
+
+void FPlaybackCapabilities::InvalidateCachedData(UMovieSceneEntitySystemLinker* Linker)
+{
+	TArrayView<const FPlaybackCapabilityHeader> Headers = GetHeaders();
+	for (int32 Index = 0; Index < Headers.Num(); ++Index)
+	{
+		const FPlaybackCapabilityHeader& Header = Headers[Index];
+		const FPlaybackCapabilityHelpers& ThisHelpers = Helpers[Index];
+		check(ThisHelpers.InterfaceCast != nullptr);
+		{
+			void* Ptr = Header.Capability.Resolve(Memory);
+			if (IPlaybackCapability* Interface = (*ThisHelpers.InterfaceCast)(Ptr))
+			{
+				Interface->InvalidateCachedData(Linker);
+			}
+		}
+	}
+}
+
+void FPlaybackCapabilities::Destroy()
+{
+	check((Memory == nullptr) == (Num == 0));
+
+	if (Memory)
+	{
+		TArrayView<const FPlaybackCapabilityHeader> Headers = GetHeaders();
+		for (int32 Index = 0; Index < Headers.Num(); ++Index)
+		{
+			const FPlaybackCapabilityHeader& Header = Headers[Index];
+			const FPlaybackCapabilityHelpers& ThisHelpers = Helpers[Index];
+			check(ThisHelpers.Destructor != nullptr);
+			{
+				void* Ptr = Header.Capability.Resolve(Memory);
+				(*ThisHelpers.Destructor)(Ptr);
+			}
+		}
+
+		FMemory::Free(Memory);
+	}
+}
+
+} // namespace UE::MovieScene
diff --git a/Engine/Source/Runtime/MovieScene/Private/Tests/MovieScenePlaybackCapabilitiesTests.cpp b/Engine/Source/Runtime/MovieScene/Private/Tests/MovieScenePlaybackCapabilitiesTests.cpp
new file mode 100644
--- /dev/null
+++ b/Engine/Source/Runtime/MovieScene/Private/Tests/MovieScenePlaybackCapabilitiesTests.cpp
@@ -0,0 +1,328 @@
+// Copyright Epic Games, Inc. All Rights Reserved.
+
+#include "CoreTypes.h"
+#include "MovieSceneFwd.h"
+#include "Misc/AutomationTest.h"
+#include "Evaluation/MovieScenePlaybackCapabilities.h"
+
+#if WITH_DEV_AUTOMATION_TESTS
+
+namespace UE::MovieScene::Tests
+{
+
+// A couple of simple test capabilities.
+struct FTestSimpleCapability
+{
+	static int32 TimesDestroyed;
+
+	int32 IntValue;
+	FString StringValue;
+
+	FTestSimpleCapability() 
+		: IntValue(0)
+	{}
+	FTestSimpleCapability(int32 InIntValue, const FString& InStringValue) 
+		: IntValue(InIntValue)
+		, StringValue(InStringValue)
+	{}
+	~FTestSimpleCapability()
+	{
+		++TimesDestroyed;
+	}
+
+	static TPlaybackCapabilityID<FTestSimpleCapability> ID;
+};
+
+TPlaybackCapabilityID<FTestSimpleCapability> FTestSimpleCapability::ID = TPlaybackCapabilityID<FTestSimpleCapability>::Register();
+
+int32 FTestSimpleCapability::TimesDestroyed = 0;
+
+struct FTestOtherSimpleCapability
+{
+	float FloatValue = 0.f;
+
+	static TPlaybackCapabilityID<FTestOtherSimpleCapability> ID;
+};
+
+TPlaybackCapabilityID<FTestOtherSimpleCapability> FTestOtherSimpleCapability::ID = TPlaybackCapabilityID<FTestOtherSimpleCapability>::Register();
+
+// Interface capability and sample implementation.
+struct FTestCapabilityBase
+{
+	static int32 TimesDestroyedBase;
+
+	float FloatValue = 0.f;
+
+	virtual ~FTestCapabilityBase() { ++TimesDestroyedBase; }
+
+	virtual int32 GetIntValue() { return 0; }
+	virtual FString GetStringValue() { return FString(); }
+
+	static TPlaybackCapabilityID<FTestCapabilityBase> ID;
+};
+
+TPlaybackCapabilityID<FTestCapabilityBase> FTestCapabilityBase::ID = TPlaybackCapabilityID<FTestCapabilityBase>::Register();
+
+int32 FTestCapabilityBase::TimesDestroyedBase = 0;
+
+struct FTestCapabilityDerived : FTestCapabilityBase
+{
+	static int32 TimesDestroyedDerived;
+
+	int32 IntValue;
+	FString StringValue;
+
+	FTestCapabilityDerived(int32 InIntValue, const FString& InStringValue)
+		: IntValue(InIntValue)
+		, StringValue(InStringValue)
+	{}
+	virtual ~FTestCapabilityDerived() { ++TimesDestroyedDerived; }
+
+	virtual int32 GetIntValue() override { return IntValue; }
+	virtual FString GetStringValue() override { return StringValue; }
+};
+
+int32 FTestCapabilityDerived::TimesDestroyedDerived = 0;
+
+// Base capability that implements the IPlaybackCapability interface.
+struct FTestInvalidatableCapability : IPlaybackCapability
+{
+	int32 TimesInvalidated = 0;
+
+	virtual void InvalidateCachedData(UMovieSceneEntitySystemLinker* Linker) override
+	{
+		++TimesInvalidated;
+	}
+
+	static TPlaybackCapabilityID<FTestInvalidatableCapability> ID;
+};
+
+TPlaybackCapabilityID<FTestInvalidatableCapability> FTestInvalidatableCapability::ID = TPlaybackCapabilityID<FTestInvalidatableCapability>::Register();
+
+// Derived version of the simple test capability, but with interface
+struct FTestSimpleCapabilityDerivedWithInterface : FTestSimpleCapability, IPlaybackCapability
+{
+	int32 TimesInvalidated = 0;
+
+	virtual void InvalidateCachedData(UMovieSceneEntitySystemLinker* Linker) override
+	{
+		++TimesInvalidated;
+	}
+};
+
+} // namespace UE::MovieScene::Tests
+
+IMPLEMENT_SIMPLE_AUTOMATION_TEST(FMovieScenePlaybackCapabilitiesSimpleTest, 
+		"System.Engine.Sequencer.Capabilities.Simple", 
+		EAutomationTestFlags::EditorContext | EAutomationTestFlags::EngineFilter)
+bool FMovieScenePlaybackCapabilitiesSimpleTest::RunTest(const FString& Parameters)
+{
+	using namespace UE::MovieScene;
+	using namespace UE::MovieScene::Tests;
+
+	FTestSimpleCapability::TimesDestroyed = 0;
+	{
+		FPlaybackCapabilities Caps;
+		UTEST_FALSE("No capability", Caps.HasCapability(FTestSimpleCapability::ID));
+
+		FTestSimpleCapability& SimpleCap = Caps.AddCapability(FTestSimpleCapability::ID, 42, TEXT("Just a test"));
+		UTEST_EQUAL("Int value", SimpleCap.IntValue, 42);
+		UTEST_EQUAL("String value", SimpleCap.StringValue, TEXT("Just a test"));
+
+		UTEST_TRUE("Found capability", Caps.HasCapability(FTestSimpleCapability::ID));
+
+		FTestSimpleCapability* SimpleCap2 = Caps.FindCapability(FTestSimpleCapability::ID);
+		UTEST_NOT_NULL("Retrived test capacity", SimpleCap2);
+		UTEST_EQUAL("Same capacity?", SimpleCap2, &SimpleCap);
+	}
+	UTEST_EQUAL("Destroyed", FTestSimpleCapability::TimesDestroyed, 1);
+
+	return true;
+}
+
+IMPLEMENT_SIMPLE_AUTOMATION_TEST(FMovieScenePlaybackCapabilitiesInlineInterfaceTest, 
+		"System.Engine.Sequencer.Capabilities.InlineInterface", 
+		EAutomationTestFlags::EditorContext | EAutomationTestFlags::EngineFilter)
+bool FMovieScenePlaybackCapabilitiesInlineInterfaceTest::RunTest(const FString& Parameters)
+{
+	using namespace UE::MovieScene;
+	using namespace UE::MovieScene::Tests;
+
+	FTestCapabilityBase::TimesDestroyedBase = 0;
+	FTestCapabilityDerived::TimesDestroyedDerived = 0;
+	{
+		FPlaybackCapabilities Caps;
+		UTEST_FALSE("No capability", Caps.HasCapability(FTestCapabilityBase::ID));
+
+		FTestCapabilityBase& InterfaceCap = Caps.AddCapabilityImplementation<FTestCapabilityDerived>(
+				FTestCapabilityBase::ID, 42, TEXT("Just a test"));
+		UTEST_EQUAL("Int value", InterfaceCap.GetIntValue(), 42);
+		UTEST_EQUAL("String value", InterfaceCap.GetStringValue(), TEXT("Just a test"));
+
+		UTEST_TRUE("Found capability", Caps.HasCapability(FTestCapabilityBase::ID));
+
+		FTestCapabilityBase* InterfaceCap2 = Caps.FindCapability(FTestCapabilityBase::ID);
+		UTEST_NOT_NULL("Retrived interface capacity", InterfaceCap2);
+		UTEST_EQUAL("Same capacity?", InterfaceCap2, &InterfaceCap);
+	}
+	UTEST_EQUAL("Times destroyed", FTestCapabilityBase::TimesDestroyedBase, 1);
+	UTEST_EQUAL("Times destroyed", FTestCapabilityDerived::TimesDestroyedDerived, 1);
+
+	return true;
+}
+
+IMPLEMENT_SIMPLE_AUTOMATION_TEST(FMovieScenePlaybackCapabilitiesRawPointerTest, 
+		"System.Engine.Sequencer.Capabilities.RawPointer", 
+		EAutomationTestFlags::EditorContext | EAutomationTestFlags::EngineFilter)
+bool FMovieScenePlaybackCapabilitiesRawPointerTest::RunTest(const FString& Parameters)
+{
+	using namespace UE::MovieScene;
+	using namespace UE::MovieScene::Tests;
+
+	FTestSimpleCapability::TimesDestroyed = 0;
+	{
+		FTestSimpleCapability ActualCap;
+		{
+			FPlaybackCapabilities Caps;
+			UTEST_FALSE("No capability", Caps.HasCapability(FTestSimpleCapability::ID));
+
+			FTestSimpleCapability& PtrCap = Caps.AddCapabilityRaw(FTestSimpleCapability::ID, &ActualCap);
+			UTEST_TRUE("Found capability", Caps.HasCapability(FTestSimpleCapability::ID));
+
+			FTestSimpleCapability* PtrCap2 = Caps.FindCapability(FTestSimpleCapability::ID);
+			UTEST_NOT_NULL("Retrived test capacity", PtrCap2);
+			UTEST_EQUAL("Same capacity?", PtrCap2, &PtrCap);
+			UTEST_EQUAL("Same capacity?", PtrCap2, &ActualCap);
+		}
+		UTEST_EQUAL("Not destroyed", FTestSimpleCapability::TimesDestroyed, 0);
+	}
+	UTEST_EQUAL("Destroyed", FTestSimpleCapability::TimesDestroyed, 1);
+
+	return true;
+}
+
+IMPLEMENT_SIMPLE_AUTOMATION_TEST(FMovieScenePlaybackCapabilitiesSharedTest, 
+		"System.Engine.Sequencer.Capabilities.Shared", 
+		EAutomationTestFlags::EditorContext | EAutomationTestFlags::EngineFilter)
+bool FMovieScenePlaybackCapabilitiesSharedTest::RunTest(const FString& Parameters)
+{
+	using namespace UE::MovieScene;
+	using namespace UE::MovieScene::Tests;
+
+	FTestSimpleCapability::TimesDestroyed = 0;
+	{
+		TSharedRef<FTestSimpleCapability> ActualCap = MakeShared<FTestSimpleCapability>();
+		UTEST_EQUAL("Ref count", ActualCap.GetSharedReferenceCount(), 1);
+		{
+			FPlaybackCapabilities Caps;
+			UTEST_FALSE("No capability", Caps.HasCapability(FTestSimpleCapability::ID));
+
+			FTestSimpleCapability& PtrCap = Caps.AddCapabilityShared(FTestSimpleCapability::ID, ActualCap);
+			UTEST_TRUE("Found capability", Caps.HasCapability(FTestSimpleCapability::ID));
+			UTEST_EQUAL("Ref count", ActualCap.GetSharedReferenceCount(), 2);
+
+			FTestSimpleCapability* PtrCap2 = Caps.FindCapability(FTestSimpleCapability::ID);
+			UTEST_NOT_NULL("Retrived test capacity", PtrCap2);
+			UTEST_EQUAL("Same capacity?", PtrCap2, &PtrCap);
+			UTEST_EQUAL("Same capacity?", PtrCap2, ActualCap.ToSharedPtr().Get());
+		}
+		UTEST_EQUAL("Not destroyed", FTestSimpleCapability::TimesDestroyed, 0);
+		UTEST_EQUAL("Ref count", ActualCap.GetSharedReferenceCount(), 1);
+	}
+	UTEST_EQUAL("Destroyed", FTestSimpleCapability::TimesDestroyed, 1);
+
+	return true;
+}
+
+IMPLEMENT_SIMPLE_AUTOMATION_TEST(FMovieScenePlaybackCapabilitiesInvalidateTest, 
+		"System.Engine.Sequencer.Capabilities.Invalidate", 
+		EAutomationTestFlags::EditorContext | EAutomationTestFlags::EngineFilter)
+bool FMovieScenePlaybackCapabilitiesInvalidateTest::RunTest(const FString& Parameters)
+{
+	using namespace UE::MovieScene;
+	using namespace UE::MovieScene::Tests;
+
+	// Test simple inline
+	{
+		FPlaybackCapabilities Caps;
+		FTestInvalidatableCapability& Cap = Caps.AddCapability(FTestInvalidatableCapability::ID);
+		UTEST_EQUAL("Not invalidated", Cap.TimesInvalidated, 0);
+		Caps.InvalidateCachedData(nullptr);
+		UTEST_EQUAL("Invalidated", Cap.TimesInvalidated, 1);
+	}
+
+	// Test inline subclass implementing interface when base class doesn't
+	{
+		FPlaybackCapabilities Caps;
+		FTestSimpleCapability& Cap = Caps.AddCapabilityImplementation<FTestSimpleCapabilityDerivedWithInterface>(FTestSimpleCapability::ID);
+		FTestSimpleCapabilityDerivedWithInterface* ActualCap = static_cast<FTestSimpleCapabilityDerivedWithInterface*>(&Cap);
+		UTEST_EQUAL("Not invalidated", ActualCap->TimesInvalidated, 0);
+		Caps.InvalidateCachedData(nullptr);
+		UTEST_EQUAL("Invalidated", ActualCap->TimesInvalidated, 1);
+	}
+
+	// Test raw pointer
+	{
+		FTestInvalidatableCapability ActualCap;
+
+		FPlaybackCapabilities Caps;
+		Caps.AddCapabilityRaw(FTestInvalidatableCapability::ID, &ActualCap);
+		UTEST_EQUAL("Not invalidated", ActualCap.TimesInvalidated, 0);
+		Caps.InvalidateCachedData(nullptr);
+		UTEST_EQUAL("Invalidated", ActualCap.TimesInvalidated, 1);
+	}
+
+	// Test shared pointer
+	{
+		TSharedRef<FTestInvalidatableCapability> ActualCap = MakeShared<FTestInvalidatableCapability>();
+
+		FPlaybackCapabilities Caps;
+		Caps.AddCapabilityShared(FTestInvalidatableCapability::ID, ActualCap);
+		UTEST_EQUAL("Not invalidated", ActualCap->TimesInvalidated, 0);
+		Caps.InvalidateCachedData(nullptr);
+		UTEST_EQUAL("Invalidated", ActualCap->TimesInvalidated, 1);
+	}
+
+	return true;
+}
+
+IMPLEMENT_SIMPLE_AUTOMATION_TEST(FMovieScenePlaybackCapabilitiesMultipleTest, 
+		"System.Engine.Sequencer.Capabilities.Multiple", 
+		EAutomationTestFlags::EditorContext | EAutomationTestFlags::EngineFilter)
+bool FMovieScenePlaybackCapabilitiesMultipleTest::RunTest(const FString& Parameters)
+{
+	using namespace UE::MovieScene;
+	using namespace UE::MovieScene::Tests;
+
+	FTestOtherSimpleCapability RawPtrCap;
+	TSharedRef<FTestInvalidatableCapability> SharedCap = MakeShared<FTestInvalidatableCapability>();
+	{
+		FPlaybackCapabilities Caps;
+
+		Caps.AddCapabilityRaw(FTestOtherSimpleCapability::ID, &RawPtrCap);
+		Caps.AddCapabilityShared(FTestInvalidatableCapability::ID, SharedCap);
+
+		UTEST_FALSE("Found capability", Caps.HasCapability(FTestSimpleCapability::ID));
+		UTEST_FALSE("Found capability", Caps.HasCapability(FTestCapabilityBase::ID));
+		UTEST_TRUE("Found capability", Caps.HasCapability(FTestOtherSimpleCapability::ID));
+		UTEST_TRUE("Found capability", Caps.HasCapability(FTestInvalidatableCapability::ID));
+
+		Caps.AddCapability(FTestSimpleCapability::ID, 8, TEXT("Other simple test"));
+		Caps.AddCapabilityImplementation<FTestCapabilityDerived>(
+				FTestCapabilityBase::ID, 12, TEXT("Another implementation test"));
+
+		UTEST_TRUE("Found capability", Caps.HasCapability(FTestSimpleCapability::ID));
+		UTEST_TRUE("Found capability", Caps.HasCapability(FTestCapabilityBase::ID));
+		UTEST_TRUE("Found capability", Caps.HasCapability(FTestOtherSimpleCapability::ID));
+		UTEST_TRUE("Found capability", Caps.HasCapability(FTestInvalidatableCapability::ID));
+
+		UTEST_EQUAL("Before invalidation", SharedCap->TimesInvalidated, 0);
+		Caps.InvalidateCachedData(nullptr);
+		UTEST_EQUAL("After invalidation", SharedCap->TimesInvalidated, 1);
+	}
+
+	return true;
+}
+
+#endif // WITH_DEV_AUTOMATION_TESTS
+
diff --git a/Engine/Source/Runtime/MovieScene/Public/EntitySystem/MovieSceneSequenceInstance.h b/Engine/Source/Runtime/MovieScene/Public/EntitySystem/MovieSceneSequenceInstance.h
--- a/Engine/Source/Runtime/MovieScene/Public/EntitySystem/MovieSceneSequenceInstance.h
+++ b/Engine/Source/Runtime/MovieScene/Public/EntitySystem/MovieSceneSequenceInstance.h
@@ -31,6 +31,7 @@ namespace MovieScene
 struct FCompiledDataVolatilityManager;
 struct FPreAnimatedStateExtension;
 struct FSequenceInstance;
+struct FSharedPlaybackState;
 struct FSubSequencePath;
 struct ISequenceUpdater;
 
@@ -130,6 +131,14 @@ public:
 
 public:
 
+	/**
+	 * Retrieve the shared playback state for this instance's hierarchy
+	 */
+	TSharedRef<FSharedPlaybackState> GetSharedPlaybackState() const
+	{
+		return SharedPlaybackState;
+	}
+
 	/**
 	 * Retrieve the IMovieScenePlayer that is playing back the top level sequence for this instance
 	 *
@@ -311,10 +320,10 @@ public:
 public:
 
 	/** Constructor for top level sequences */
-	MOVIESCENE_API explicit FSequenceInstance(UMovieSceneEntitySystemLinker* Linker, IMovieScenePlayer* Player, FRootInstanceHandle ThisInstanceHandle);
+	MOVIESCENE_API explicit FSequenceInstance(UMovieSceneEntitySystemLinker* Linker, IMovieScenePlayer* Player, TSharedRef<FSharedPlaybackState> PlaybackState, FRootInstanceHandle ThisInstanceHandle);
 
 	/** Constructor for sub sequences */
-	MOVIESCENE_API explicit FSequenceInstance(UMovieSceneEntitySystemLinker* Linker, IMovieScenePlayer* Player, FInstanceHandle ThisInstanceHandle, FInstanceHandle InParentInstanceHandle, FRootInstanceHandle RootInstanceHandle, FMovieSceneSequenceID InSequenceID);
+	MOVIESCENE_API explicit FSequenceInstance(UMovieSceneEntitySystemLinker* Linker, IMovieScenePlayer* Player, TSharedRef<FSharedPlaybackState> PlaybackState, FInstanceHandle ThisInstanceHandle, FInstanceHandle InParentInstanceHandle, FRootInstanceHandle RootInstanceHandle, FMovieSceneSequenceID InSequenceID);
 
 	/** Destructor */
 	MOVIESCENE_API ~FSequenceInstance();
@@ -347,6 +356,9 @@ private:
 	/** For top-level sequences only - a utility class that is used to ensure that volatile sequences are up to date. Only valid in editor, or for sequences that have the volatile flag. */
 	TUniquePtr<FCompiledDataVolatilityManager> VolatilityManager;
 
+	/** Playback state shared by the entire sequence hierarchy */
+	TSharedRef<FSharedPlaybackState> SharedPlaybackState;
+
 
 	/** Delegate Binding for when an object binding is invalidated in this instance . */
 	FDelegateHandle OnInvalidateObjectBindingHandle;
diff --git a/Engine/Source/Runtime/MovieScene/Public/EntitySystem/MovieSceneSharedPlaybackState.h b/Engine/Source/Runtime/MovieScene/Public/EntitySystem/MovieSceneSharedPlaybackState.h
new file mode 100644
--- /dev/null
+++ b/Engine/Source/Runtime/MovieScene/Public/EntitySystem/MovieSceneSharedPlaybackState.h
@@ -0,0 +1,51 @@
+// Copyright Epic Games, Inc. All Rights Reserved.
+
+#pragma once
+
+#include "CoreTypes.h"
+#include "Evaluation/MovieScenePlaybackCapabilities.h"
+
+namespace UE::MovieScene
+{
+
+struct FSharedPlaybackState
+{
+	/** Playback capabilities for the root sequence */
+	FPlaybackCapabilities Capabilities;
+
+public:
+
+	template<typename T>
+	T* FindCapability() const
+	{
+		return Capabilities.FindCapability<T>(T::ID);
+	}
+
+	template<typename T, typename ...ArgTypes>
+	T& AddCapability(ArgTypes&&... InArgs)
+	{
+		return Capabilities.AddCapability<T>(T::ID, Forward<ArgTypes>(InArgs)...);
+	}
+
+	template<typename T, typename ...ArgTypes>
+	T& AddCapabilityRaw(T* InPointer)
+	{
+		return Capabilities.AddCapabilityRaw<T>(T::ID, InPointer);
+	}
+
+	template<typename T, typename ...ArgTypes>
+	T& AddCapabilityShared(TSharedRef<T> InSharedRef)
+	{
+		return Capabilities.AddCapabilityShared<T>(T::ID, InSharedRef);
+	}
+
+public:
+
+	void InvalidateCachedData(UMovieSceneEntitySystemLinker* Linker)
+	{
+		Capabilities.InvalidateCachedData(Linker);
+	}
+};
+
+} // namespace UE::MovieScene
+
diff --git a/Engine/Source/Runtime/MovieScene/Public/Evaluation/IMovieScenePlaybackCapability.h b/Engine/Source/Runtime/MovieScene/Public/Evaluation/IMovieScenePlaybackCapability.h
new file mode 100644
--- /dev/null
+++ b/Engine/Source/Runtime/MovieScene/Public/Evaluation/IMovieScenePlaybackCapability.h
@@ -0,0 +1,65 @@
+// Copyright Epic Games, Inc. All Rights Reserved.
+
+#pragma once
+
+#include "CoreTypes.h"
+#include "Misc/AssertionMacros.h"
+#include "Misc/EnumClassFlags.h"
+
+class IMovieScenePlayer;
+class UMovieSceneEntitySystemLinker;
+
+namespace UE::MovieScene
+{
+
+/**
+ * An identifier for a playback capability.
+ */
+struct FPlaybackCapabilityID
+{
+	int32 Index = INDEX_NONE;
+
+	bool IsValid() const
+	{
+		return Index != INDEX_NONE;
+	}
+
+protected:
+	static MOVIESCENE_API FPlaybackCapabilityID Register();
+};
+
+/**
+ * A strongly-typed identifier for a specific playback capability class.
+ *
+ * The base capablity class must create a static ID member that returns its own typed ID. We will use
+ * this as a convention to quickly get the ID of a base capability class in the capabilities container.
+ */
+template<typename T>
+struct TPlaybackCapabilityID : FPlaybackCapabilityID
+{
+private:
+	// Only T should construct this (to ensure safe construction over DLL boundaries)
+	friend T;
+
+	TPlaybackCapabilityID()
+	{}
+
+	explicit TPlaybackCapabilityID(int32 InIndex)
+		: FPlaybackCapabilityID{ InIndex }
+	{}
+
+	static TPlaybackCapabilityID<T> Register()
+	{
+		FPlaybackCapabilityID StaticID = FPlaybackCapabilityID::Register();
+		return TPlaybackCapabilityID<T>(StaticID.Index);
+	}
+};
+
+struct IPlaybackCapability
+{
+	virtual ~IPlaybackCapability() {}
+
+	virtual void InvalidateCachedData(UMovieSceneEntitySystemLinker* Linker) {}
+};
+
+} // namespace UE::MovieScene
diff --git a/Engine/Source/Runtime/MovieScene/Public/Evaluation/MovieScenePlaybackCapabilities.h b/Engine/Source/Runtime/MovieScene/Public/Evaluation/MovieScenePlaybackCapabilities.h
new file mode 100644
--- /dev/null
+++ b/Engine/Source/Runtime/MovieScene/Public/Evaluation/MovieScenePlaybackCapabilities.h
@@ -0,0 +1,543 @@
+// Copyright Epic Games, Inc. All Rights Reserved.
+
+#pragma once
+
+#include "CoreTypes.h"
+#include "Evaluation/IMovieScenePlaybackCapability.h"
+#include "EntitySystem/RelativePtr.h"
+#include "UObject/WeakObjectPtr.h"
+#include "Misc/TVariant.h"
+#include "Templates/AlignmentTemplates.h"
+#include "Templates/PointerIsConvertibleFromTo.h"
+
+class IMovieScenePlayer;
+class IMovieScenePlaybackClient;
+
+namespace UE::MovieScene
+{
+
+enum class EPlaybackCapabilityStorageMode : uint8
+{
+	Inline = 0,
+	RawPointer,
+	SharedPointer
+};
+
+/**
+ * Return value for accessing raw capabalities.
+ */
+struct FPlaybackCapabilityPtr
+{
+	void* Ptr = 0;
+	EPlaybackCapabilityStorageMode StorageMode;
+
+	template<typename T>
+	T* ResolveOptional() const
+	{
+		if (Ptr)
+		{
+			switch (StorageMode)
+			{
+				case EPlaybackCapabilityStorageMode::Inline:
+					return static_cast<T*>(Ptr);
+				case EPlaybackCapabilityStorageMode::RawPointer:
+					return *static_cast<T**>(Ptr);
+				case EPlaybackCapabilityStorageMode::SharedPointer:
+					return static_cast<TSharedPtr<T>*>(Ptr)->Get();
+				default:
+					checkf(false, TEXT("Unexpected playback capability storage mode"));
+					return nullptr;
+			}
+		}
+		return nullptr;
+	}
+
+	template<typename T>
+	T& ResolveChecked() const
+	{
+		check(Ptr);
+		switch (StorageMode)
+		{
+			case EPlaybackCapabilityStorageMode::Inline:
+				return *static_cast<T*>(Ptr);
+			case EPlaybackCapabilityStorageMode::RawPointer:
+				return **static_cast<T**>(Ptr);
+			case EPlaybackCapabilityStorageMode::SharedPointer:
+				return *static_cast<TSharedPtr<T>*>(Ptr)->Get();
+			default:
+				checkf(false, TEXT("Unexpected playback capability storage mode"));
+				// Nothing reasonable to do, let's just proceed as if it was inline
+				return *static_cast<T*>(Ptr);
+		}
+	}
+};
+
+// Utility callbacks for the concrete capability objects we have in a container.
+using FPlaybackCapabilityInterfaceCastHelper = IPlaybackCapability*(*)(void* Ptr);
+using FPlaybackCapabilityDestructionHelper = void(*)(void* Ptr);
+
+// Callback for casting a stored capabilty object to an IPlaybackCapability pointer if possible.
+template<typename StorageType>
+struct TPlaybackCapabilityInterfaceCast
+{
+	static IPlaybackCapability* InterfaceCast(void* Ptr)
+	{
+		if constexpr (TPointerIsConvertibleFromTo<StorageType, IPlaybackCapability>::Value)
+		{
+			return static_cast<IPlaybackCapability*>((StorageType*)Ptr);
+		}
+		return nullptr;
+	}
+};
+template<typename PointedType>
+struct TPlaybackCapabilityInterfaceCast<PointedType*>
+{
+	static IPlaybackCapability* InterfaceCast(void* Ptr) 
+	{
+		if constexpr (TPointerIsConvertibleFromTo<PointedType, IPlaybackCapability>::Value)
+		{
+			PointedType* TypedPtr = *(PointedType**)Ptr;
+			return static_cast<IPlaybackCapability*>(TypedPtr);
+		}
+		return nullptr;
+	};
+};
+template<typename PointedType>
+struct TPlaybackCapabilityInterfaceCast<TSharedPtr<PointedType>>
+{
+	static IPlaybackCapability* InterfaceCast(void* Ptr) 
+	{
+		if constexpr (TPointerIsConvertibleFromTo<PointedType, IPlaybackCapability>::Value)
+		{
+			TSharedPtr<PointedType>& TypedPtr = *(TSharedPtr<PointedType>*)Ptr;
+			return static_cast<IPlaybackCapability*>(TypedPtr.Get());
+		}
+		return nullptr;
+	};
+};
+
+// Callback for destroying the stored capability object, whether that's the capability itself
+// (when stored inline), or a shared pointer, or whatever else.
+template<typename StorageType>
+struct TPlaybackCapabilityDestructor
+{
+	static void Destroy(void* Ptr)
+	{
+		StorageType* StoragePtr = (StorageType*)Ptr;
+		StoragePtr->~StorageType();
+	}
+};
+
+// Helper callbacks for managing the lifetime of a capability.
+struct FPlaybackCapabilityHelpers
+{
+	FPlaybackCapabilityInterfaceCastHelper InterfaceCast = nullptr;
+	FPlaybackCapabilityDestructionHelper Destructor = nullptr;
+};
+
+template<typename StorageType>
+struct TPlaybackCapabilityHelpers
+{
+	static FPlaybackCapabilityHelpers GetHelpers()
+	{
+		FPlaybackCapabilityHelpers Helpers;
+		Helpers.InterfaceCast = &TPlaybackCapabilityInterfaceCast<StorageType>::InterfaceCast;
+		Helpers.Destructor = &TPlaybackCapabilityDestructor<StorageType>::Destroy;
+		return Helpers;
+	}
+};
+
+// Storage traits for playback capabilities, only for internal use.
+template<typename StorageType, typename CapabilityType, typename=void>
+struct TPlaybackCapabilityStorageTraits;
+
+template<typename StorageType, typename CapabilityType>
+struct TPlaybackCapabilityStorageTraits<StorageType, CapabilityType, typename TEnableIf<TPointerIsConvertibleFromTo<StorageType, CapabilityType>::Value>::Type>
+{
+	static EPlaybackCapabilityStorageMode GetStorageMode() { return EPlaybackCapabilityStorageMode::Inline; }
+
+	static uint8 ComputePointerOffset(StorageType* StoragePtr)
+	{
+		const uint64 PointerOffset =
+			reinterpret_cast<uint64>(static_cast<CapabilityType*>(StoragePtr))
+			-
+			reinterpret_cast<uint64>(StoragePtr);
+		check(PointerOffset <= TNumericLimits<uint8>::Max());
+		return static_cast<uint8>(PointerOffset);
+	}
+};
+
+template<typename CapabilityType>
+struct TPlaybackCapabilityStorageTraits<CapabilityType*, CapabilityType, void>
+{
+	static EPlaybackCapabilityStorageMode GetStorageMode() { return EPlaybackCapabilityStorageMode::RawPointer; }
+
+	static uint8 ComputePointerOffset(CapabilityType** StoragePtr) { return 0; }
+};
+
+template<typename CapabilityType>
+struct TPlaybackCapabilityStorageTraits<TSharedPtr<CapabilityType>, CapabilityType, void>
+{
+	static EPlaybackCapabilityStorageMode GetStorageMode() { return EPlaybackCapabilityStorageMode::SharedPointer; }
+
+	static uint8 ComputePointerOffset(TSharedPtr<CapabilityType>* StoragePtr) { return 0; }
+};
+
+/**
+ * Header that describes an entry in the capabilities buffer.
+ */
+struct FPlaybackCapabilityHeader
+{
+	// 2 Bytes
+	TRelativePtr<void, uint16> Capability;
+
+	// 1 Byte
+	uint8 Sizeof = 0;
+
+	// 1 Byte
+	uint8 Alignment = 0;
+
+	// 1 Byte
+	uint8 PointerOffset = 0;
+
+	// 1 Byte
+	EPlaybackCapabilityStorageMode StorageMode = EPlaybackCapabilityStorageMode::Inline;
+
+	/** Resolve the given raw pointer into a capability pointer */
+	FPlaybackCapabilityPtr Resolve(void* InMemory) const
+	{
+		void* DerivedPointer = Capability.Resolve(InMemory);
+		return FPlaybackCapabilityPtr{ (void*)((uint8*)DerivedPointer + PointerOffset), StorageMode };
+	}
+};
+
+/**
+ * Basic, mostly untyped, implementation of the capabilities buffer.
+ */
+struct FPlaybackCapabilitiesImpl
+{
+protected:
+
+	FPlaybackCapabilitiesImpl() = default;
+
+	FPlaybackCapabilitiesImpl(const FPlaybackCapabilitiesImpl&) = delete;
+	FPlaybackCapabilitiesImpl& operator=(const FPlaybackCapabilitiesImpl&) = delete;
+
+	FPlaybackCapabilitiesImpl(FPlaybackCapabilitiesImpl&&) = delete;
+	FPlaybackCapabilitiesImpl& operator=(FPlaybackCapabilitiesImpl&&) = delete;
+
+	bool HasCapability(uint32 CapabilityBit) const
+	{
+		return (AllCapabilities & CapabilityBit) != 0;
+	}
+
+	FPlaybackCapabilityPtr FindCapability(uint32 CapabilityBit) const
+	{
+		if (HasCapability(CapabilityBit))
+		{
+			const int32 Index = GetCapabilityIndex(CapabilityBit);
+			return GetHeader(Index).Resolve(Memory);
+		}
+
+		return FPlaybackCapabilityPtr();
+	}
+
+	FPlaybackCapabilityPtr GetCapabilityChecked(uint32 CapabilityBit) const
+	{
+		const int32 Index = GetCapabilityIndex(CapabilityBit);
+		check(Index != INDEX_NONE);
+		return GetHeader(Index).Resolve(Memory);
+	}
+
+	/**
+	 * Creates and stores a new capability object at the given bit.
+	 * 
+	 * - StorageType is what is actually constructed and stored in our memory buffer. It can be the same as
+	 *			CapabilityType, a subclass of it, a pointer to it, or a shared pointer to it.
+	 * 
+	 * - CapabilityType is the base capability type, used onlyfor inline storage when the StorageType is a subclass
+	 *			and we need to compute the pointer offset.
+	 */
+	template<typename StorageType, typename CapabilityType, typename ...ArgTypes>
+	FPlaybackCapabilityPtr AddCapability(uint32 CapabilityBit, ArgTypes&&... InArgs)
+	{
+		checkf((AllCapabilities & CapabilityBit) == 0, TEXT("Capability already exists!"));
+
+		// Add the enum entry
+		AllCapabilities |= CapabilityBit;
+
+		// Find the index of the new capability by counting how many bits are set before it
+		// For example, given CapabilityBit=0b00010000 and AllCapabilities=0b00011011:
+		//                      (CapabilityBit-1) = 0b00001111
+		//    AllCapabilities & (CapabilityBit-1) = 0b00001011
+		//                  CountBits(0b00001011) = 3
+		const int32 NewCapabilityIndex = static_cast<int32>(FMath::CountBits(AllCapabilities & (CapabilityBit-1u)));
+
+		const FPlaybackCapabilityHeader* ExistingHeaders = reinterpret_cast<const FPlaybackCapabilityHeader*>(Memory);
+
+		uint64 RequiredAlignment = FMath::Max(alignof(StorageType), alignof(FPlaybackCapabilityHeader));
+
+		const int32 ExistingNum = static_cast<int32>(Num);
+
+		// Compute our required alignment for the allocation
+		{
+			for (int32 Index = 0; Index < ExistingNum; ++Index)
+			{
+				RequiredAlignment = FMath::Max(RequiredAlignment, (uint64)ExistingHeaders[Index].Alignment);
+			}
+		}
+
+		uint64 RequiredSizeof = 0u;
+
+		// Compute the required size of our allocation
+		{
+			// Allocate space for headers
+			RequiredSizeof += (ExistingNum+1) * sizeof(FPlaybackCapabilityHeader);
+
+			int32 Index = 0;
+
+			// Count up the sizes and alignments of pre-existing capabilities that exist before this new entry
+			for (; Index < NewCapabilityIndex; ++Index)
+			{
+				RequiredSizeof = Align(RequiredSizeof, (uint64)ExistingHeaders[Index].Alignment);
+				RequiredSizeof += ExistingHeaders[Index].Sizeof;
+			}
+
+			// Count up the size and alignment for the new capability
+			RequiredSizeof = Align(RequiredSizeof, alignof(StorageType));
+			RequiredSizeof += sizeof(StorageType);
+			++Index;
+
+			// Now count up the sizes and alignments of pre-existing capabilities that exist after this new entry
+			for (; Index < ExistingNum+1; ++Index)
+			{
+				RequiredSizeof = Align(RequiredSizeof, (uint64)ExistingHeaders[Index-1].Alignment);
+				RequiredSizeof += ExistingHeaders[Index-1].Sizeof;
+			}
+		}
+
+		check( RequiredAlignment <= 0XFF );
+
+		/// ----------------------------------------
+
+		uint8* OldAllocation = Memory;
+
+		// Make a new allocation if necessary
+		const bool bNeedsReallocation = RequiredAlignment > Alignment || RequiredSizeof > Capacity;
+		if (bNeedsReallocation)
+		{
+			// Use the greater of the required size or double the current size to allow some additional capcity
+			Capacity = FMath::Max(RequiredSizeof, uint64(Capacity)*2);
+			Memory = reinterpret_cast<uint8*>(FMemory::Malloc(Capacity, RequiredAlignment));
+		}
+
+		// We now have an extra entry
+		++Num;
+
+		FPlaybackCapabilityHeader* CurrentHeader = reinterpret_cast<FPlaybackCapabilityHeader*>(Memory) + (Num-1);
+		uint8* CapabilityPtr = Memory + RequiredSizeof;
+
+		auto RelocateCapability = [this, ExistingHeaders, OldAllocation, &CurrentHeader, &CapabilityPtr](int32 OldIndex)
+		{
+			// Go back
+			CapabilityPtr -= ExistingHeaders[OldIndex].Sizeof;
+			CapabilityPtr = AlignDown(CapabilityPtr, ExistingHeaders[OldIndex].Alignment);
+
+			// Copy the header
+			new (CurrentHeader) FPlaybackCapabilityHeader(ExistingHeaders[OldIndex]);
+			CurrentHeader->Capability = TRelativePtr<void, uint16>(this->Memory, CapabilityPtr);
+
+			--CurrentHeader;
+
+			void* OldCapability = ExistingHeaders[OldIndex].Capability.Resolve(OldAllocation);
+			FMemory::Memmove(CapabilityPtr, OldCapability, ExistingHeaders[OldIndex].Sizeof);
+		};
+
+		int32 Index = static_cast<int32>(Num) - 1;
+		for (; Index > NewCapabilityIndex; --Index)
+		{
+			RelocateCapability(Index-1);
+		}
+
+		// Make the new entry
+		CapabilityPtr -= sizeof(StorageType);
+		CapabilityPtr = AlignDown(CapabilityPtr, alignof(StorageType));
+
+		StorageType* NewCapabilityPtr = reinterpret_cast<StorageType*>(CapabilityPtr);
+
+		// Allocate the new type
+		new (NewCapabilityPtr) StorageType (Forward<ArgTypes>(InArgs)...);
+
+		static_assert(alignof(StorageType) < 0x7F, "Required alignment of capability must fit in 7 bytes");
+
+		using FStorageTraits = TPlaybackCapabilityStorageTraits<StorageType, CapabilityType>;
+
+		// Construct the header
+		const FPlaybackCapabilityHeader* NewHeader = CurrentHeader;
+		CurrentHeader->Capability.Reset(Memory, NewCapabilityPtr);
+		CurrentHeader->Sizeof = sizeof(StorageType);
+		CurrentHeader->Alignment = alignof(StorageType);
+		CurrentHeader->StorageMode = FStorageTraits::GetStorageMode();
+		// The pointer offset is whatever we need to add to the capability pointer in order to get
+		// a pointer to capability type itself (in case the storage type is a sub-class of it)
+		// We only need it if the capability object is stored inline inside our memory buffer.
+		CurrentHeader->PointerOffset = FStorageTraits::ComputePointerOffset(NewCapabilityPtr);
+
+		--CurrentHeader;
+		--Index;
+
+		// Relocate the entries that are before the new one
+		for (; Index >= 0; --Index)
+		{
+			RelocateCapability(Index);
+		}
+
+		// Tidy up the old allocation. We do not call destructors here because we relocated everything.
+		if (bNeedsReallocation && OldAllocation)
+		{
+			FMemory::Free(OldAllocation);
+		}
+
+		// Insert the helpers for the new capability.
+		Helpers.Insert(TPlaybackCapabilityHelpers<StorageType>::GetHelpers(), NewCapabilityIndex);
+
+		// Return the new capability pointer. We call the header's Resolve method here because returning NewCapabilityPtr
+		// would return the derived type pointer. We want the base (capability) pointer.
+		return NewHeader->Resolve(Memory);
+	}
+
+	const FPlaybackCapabilityHeader& GetHeader(uint8 Index) const
+	{
+		check(Index < Num);
+		return reinterpret_cast<FPlaybackCapabilityHeader*>(Memory)[Index];
+	}
+
+	TArrayView<const FPlaybackCapabilityHeader> GetHeaders() const
+	{
+		return MakeArrayView(reinterpret_cast<FPlaybackCapabilityHeader*>(Memory), Num);
+	}
+
+	int32 GetCapabilityIndex(uint32 CapabilityBit) const
+	{
+		if ( (AllCapabilities & CapabilityBit) == 0)
+		{
+			return INDEX_NONE;
+		}
+		return FMath::CountBits(AllCapabilities & (CapabilityBit-1u));
+	}
+	
+	// Schema:
+	// [header_1|...|header_n|entry_0|...|entry_n]
+	uint8* Memory = nullptr;
+	uint16 Alignment = 0u;
+	uint16 Capacity = 0u;
+	uint8 Num = 0u;
+	uint32 AllCapabilities = 0u;
+
+	// Function pointers for invalidating, destroying, etc. capabilities.
+	TArray<FPlaybackCapabilityHelpers> Helpers;
+};
+
+/**
+ * Actual playback capabilities container.
+ */
+struct FPlaybackCapabilities : FPlaybackCapabilitiesImpl
+{
+	FPlaybackCapabilities() = default;
+
+	FPlaybackCapabilities(const FPlaybackCapabilities&) = delete;
+	FPlaybackCapabilities& operator=(const FPlaybackCapabilities&) = delete;
+
+	MOVIESCENE_API FPlaybackCapabilities(FPlaybackCapabilities&&);
+	MOVIESCENE_API FPlaybackCapabilities& operator=(FPlaybackCapabilities&&);
+
+	MOVIESCENE_API ~FPlaybackCapabilities();
+
+	/** Checks whether this container has the given capability */
+	bool HasCapability(FPlaybackCapabilityID CapabilityID) const
+	{
+		uint32 CapabilityBit = 1 << CapabilityID.Index;
+		return FPlaybackCapabilitiesImpl::HasCapability(CapabilityBit);
+	}
+
+	/** Finds the specified capability within the container, if present */
+	template<typename T>
+	T* FindCapability(TPlaybackCapabilityID<T> CapabilityID) const
+	{
+		uint32 CapabilityBit = 1 << CapabilityID.Index;
+		FPlaybackCapabilityPtr Ptr = FPlaybackCapabilitiesImpl::FindCapability(CapabilityBit);
+		return Ptr.ResolveOptional<T>();
+	}
+
+	/** Returns the specified capability within the container, asserts if not found */
+	template<typename T>
+	T& GetCapabilityChecked(TPlaybackCapabilityID<T> CapabilityID) const
+	{
+		uint32 CapabilityBit = 1 << CapabilityID.Index;
+		FPlaybackCapabilityPtr Ptr = FPlaybackCapabilitiesImpl::GetCapabilityChecked(CapabilityBit);
+		return Ptr.ResolveChecked<T>();
+	}
+
+	/**
+	 * Adds the specified capability to the container, using the supplied arguments to construct it
+	 * The capability object will be stored inline and owned by this container. It will be destroyed when the
+	 * container itself is destroyed.
+	 */
+	template<typename T, typename ...ArgTypes>
+	T& AddCapability(TPlaybackCapabilityID<T> CapabilityID, ArgTypes&&... InArgs)
+	{
+		return DoAddCapability<T, T>(CapabilityID, Forward<ArgTypes>(InArgs)...);
+	}
+
+	/**
+	 * Adds the specified capability to the container, using the supplied arguments to construct a sub-class of it 
+	 * The capability object of the specified sub-class will be stored inline and owned by this container. It will
+	 * be destroyed when the container itself is destroyed.
+	 */
+	template<typename Impl, typename T, typename ...ArgTypes>
+	T& AddCapabilityImplementation(TPlaybackCapabilityID<T> CapabilityID, ArgTypes&&... InArgs)
+	{
+		return DoAddCapability<Impl, T>(CapabilityID, Forward<ArgTypes>(InArgs)...);
+	}
+
+	/**
+	 * Adds the specified capability to the container, as a simple raw pointer
+	 * Ownership of the capability object being pointed to is the caller's responsability. This container will only
+	 * store a raw pointer.
+	 */
+	template<typename T>
+	T& AddCapabilityRaw(TPlaybackCapabilityID<T> CapabilityID, T* InPointer)
+	{
+		return DoAddCapability<T*, T>(CapabilityID, InPointer);
+	}
+	
+	/**
+	 * Adds the specified capability to the container, as a shared pointer
+	 * Ownership of the capability object being pointed to respects classic shared pointer semantics. This container
+	 * only maintains one such shared pointer until the container is destroyed.
+	 */
+	template<typename T>
+	T& AddCapabilityShared(TPlaybackCapabilityID<T> CapabilityID, TSharedRef<T> InSharedRef)
+	{
+		return DoAddCapability<TSharedPtr<T>, T>(CapabilityID, InSharedRef);
+	}
+
+	/**
+	 * Calls InvalidateCacheData on any capability that implements the IPlaybackCapability interface.
+	 */
+	void InvalidateCachedData(UMovieSceneEntitySystemLinker* Linker);
+
+private:
+
+	template<typename Impl, typename T, typename ...ArgTypes>
+	T& DoAddCapability(TPlaybackCapabilityID<T> CapabilityID, ArgTypes&&... InArgs)
+	{
+		uint32 CapabilityBit = 1 << CapabilityID.Index;
+		FPlaybackCapabilityPtr Ptr = FPlaybackCapabilitiesImpl::AddCapability<Impl, T>(CapabilityBit, Forward<ArgTypes>(InArgs)...);
+		return Ptr.ResolveChecked<T>();
+	}
+
+	void Destroy();
+};
+
+} // namespace UE::MovieScene
