Sequencer: Entity system runner is now more robust against compilations mid-evaluation

Several fixes here:

  - Added a way for a flush state to decide whether it allows budgeting or not. For conditional recompiles, if we did compile, we allow budgeting to leave the rest of the flush to the next frame (since compilation is an expensive process). If there was no compile necessary, we force the next step to continue so there is no possibility we get stuck in a loop continually running ConditionalRecompile if it breaches the budget threshold on its own.

  - MarkForUpdate now checks whether an update is for Finish or Destroy while we are already updating sequences and calls Finish on the instance immediately. Finish otherwise would not get called because only those updates on the DissectedUpdates would get called. This is required due to the changes in sub-instance destruction from the sequence updater.

 - Added a target flush state to FMovieSceneEntitySystemRunner::Flush in order to support the new test

  - Fixed FMovieSceneEntitySystemRunner::ResetFlushState() not resetting when ConditionalRecompile was present. It was an error that this was included since it doesn't have any material meaning on whether we need to reset the flush state or not (additionally, it is nearly always present if running a budgeted evaluation)

  - Fixed ConditionalRecompile only operating on the UpdateQueue. If CurrentInstances is populated, that means an evaluation is partially complete for those instances, and we need to check those instead. It is possible for CurrentInstances to be populated while UpdateQueue is empty.

  - DissectContext is no longer called for Finish updates

  - Fixed ReimportSequenceInstances being able to add duplicate entries into CurrentInstances, and not handling Finish states correctly

diff --git a/Engine/Source/Runtime/MovieScene/Private/EntitySystem/MovieSceneEntitySystemRunner.cpp b/Engine/Source/Runtime/MovieScene/Private/EntitySystem/MovieSceneEntitySystemRunner.cpp
--- a/Engine/Source/Runtime/MovieScene/Private/EntitySystem/MovieSceneEntitySystemRunner.cpp
+++ b/Engine/Source/Runtime/MovieScene/Private/EntitySystem/MovieSceneEntitySystemRunner.cpp
@@ -96,6 +96,7 @@ FMovieSceneEntitySystemRunner::FMovieSceneEntitySystemRunner()
 	, FlushState(UE::MovieScene::ERunnerFlushState::None)
 	, CurrentFlushState(UE::MovieScene::ERunnerFlushState::None)
 	, bRequireFullFlush(false)
+	, bIsUpdatingSequence(false)
 {
 }
 
@@ -175,6 +176,11 @@ bool FMovieSceneEntitySystemRunner::IsCurrentlyEvaluating() const
 	return FlushState != UE::MovieScene::ERunnerFlushState::None;
 }
 
+bool FMovieSceneEntitySystemRunner::IsUpdatingSequence() const
+{
+	return bIsUpdatingSequence;
+}
+
 int32 FMovieSceneEntitySystemRunner::GetQueuedUpdateCount() const
 {
 	return UpdateQueue.Num() + DissectedUpdates.Num();
@@ -300,7 +306,7 @@ void FMovieSceneEntitySystemRunner::AbandonAndDestroyInstance(FInstanceHandle In
 	}
 }
 
-bool FMovieSceneEntitySystemRunner::StartEvaluation(UMovieSceneEntitySystemLinker* Linker)
+UE::MovieScene::ERunnerFlushResult FMovieSceneEntitySystemRunner::StartEvaluation(UMovieSceneEntitySystemLinker* Linker)
 {
 	using namespace UE::MovieScene;
 
@@ -312,12 +318,12 @@ bool FMovieSceneEntitySystemRunner::StartEvaluation(UMovieSceneEntitySystemLinke
 	{
 		// If nothing has changed, and we have no pending updates, we don't run any evaluation
 		FlushState = ERunnerFlushState::None;
-		return false;
+		return ERunnerFlushResult::Break;
 	}
 
 	if (!Linker->StartEvaluation(*this))
 	{
-		return false;
+		return ERunnerFlushResult::Break;
 	}
 
 	EnterFlushState(ERunnerFlushState::Start);
@@ -330,7 +336,7 @@ bool FMovieSceneEntitySystemRunner::StartEvaluation(UMovieSceneEntitySystemLinke
 	EntityManager.SetDispatchThread(ENamedThreads::GameThread_Local);
 	EntityManager.SetGatherThread(ENamedThreads::GameThread_Local);
 
-	return true;
+	return ERunnerFlushResult::ContinueAllowBudget;
 }
 
 void FMovieSceneEntitySystemRunner::EndEvaluation(UMovieSceneEntitySystemLinker* Linker)
@@ -339,7 +345,7 @@ void FMovieSceneEntitySystemRunner::EndEvaluation(UMovieSceneEntitySystemLinker*
 	Linker->EndEvaluation(*this);
 }
 
-bool FMovieSceneEntitySystemRunner::FlushNext(UMovieSceneEntitySystemLinker* Linker)
+UE::MovieScene::ERunnerFlushResult FMovieSceneEntitySystemRunner::FlushNext(UMovieSceneEntitySystemLinker* Linker)
 {
 	using namespace UE::MovieScene;
 
@@ -383,7 +389,7 @@ bool FMovieSceneEntitySystemRunner::FlushNext(UMovieSceneEntitySystemLinker* Lin
 		return GameThread_ReimportSequenceInstances(Linker);
 	}
 
-	// Step 4: Conditionally run the spawn phase of the system graph
+	// Step 5: Conditionally run the spawn phase of the system graph
 	//
 	if (EnumHasAnyFlags(FlushState, ERunnerFlushState::Spawn))
 	{
@@ -391,7 +397,7 @@ bool FMovieSceneEntitySystemRunner::FlushNext(UMovieSceneEntitySystemLinker* Lin
 		return GameThread_SpawnPhase(Linker);
 	}
 
-	// Step 5: Conditionally run the instantiation phase of the system graph
+	// Step 6: Conditionally run the instantiation phase of the system graph
 	//
 	if (EnumHasAnyFlags(FlushState, ERunnerFlushState::Instantiation))
 	{
@@ -399,7 +405,7 @@ bool FMovieSceneEntitySystemRunner::FlushNext(UMovieSceneEntitySystemLinker* Lin
 		return GameThread_InstantiationPhase(Linker);
 	}
 
-	// Step 6: Run the evaluation phase of the system graph
+	// Step 7: Run the evaluation phase of the system graph
 	//
 	if (EnumHasAnyFlags(FlushState, ERunnerFlushState::Evaluation))
 	{
@@ -407,46 +413,46 @@ bool FMovieSceneEntitySystemRunner::FlushNext(UMovieSceneEntitySystemLinker* Lin
 		return GameThread_EvaluationPhase(Linker);
 	}
 
-	// Step 7: Run the finalization phase of the system graph, including legacy templates
+	// Step 8: Run the finalization phase of the system graph, including legacy templates
 	//
 	if (EnumHasAnyFlags(FlushState, ERunnerFlushState::Finalization))
 	{
 		EnterFlushState(ERunnerFlushState::Finalization);
 		GameThread_EvaluationFinalizationPhase(Linker);
-		return true;
+		return ERunnerFlushResult::ContinueAllowBudget;
 	}
 
-	// Step 8: Trigger events
+	// Step 9: Trigger events
 	//
 	if (EnumHasAnyFlags(FlushState, ERunnerFlushState::EventTriggers))
 	{
 		EnterFlushState(ERunnerFlushState::EventTriggers);
 		GameThread_EventTriggerPhase(Linker);
-		return true;
+		return ERunnerFlushResult::ContinueAllowBudget;
 	}
 
-	// Step 9: Call PostEvaluation on all current sequence instances
+	// Step 10: Call PostEvaluation on all current sequence instances
 	//
 	if (EnumHasAnyFlags(FlushState, ERunnerFlushState::PostEvaluation))
 	{
 		EnterFlushState(ERunnerFlushState::PostEvaluation);
 		GameThread_PostEvaluationPhase(Linker);
-		return true;
+		return ERunnerFlushResult::ContinueAllowBudget;
 	}
 
-	// Step 10: Perform any clean up and broadcast 'end eval' events
+	// Step 11: Perform any clean up and broadcast 'end eval' events
 	//          NOTE: Only ever called once regardless of how many iterations we perform
 	if (EnumHasAnyFlags(FlushState, ERunnerFlushState::End))
 	{
 		EnterFlushState(ERunnerFlushState::End);
 		EndEvaluation(Linker);
-		return true;
+		return ERunnerFlushResult::ContinueAllowBudget;
 	}
 
-	return false;
+	return ERunnerFlushResult::Break;
 }
 
-void FMovieSceneEntitySystemRunner::Flush(double BudgetMs)
+void FMovieSceneEntitySystemRunner::Flush(double BudgetMs, UE::MovieScene::ERunnerFlushState TargetState)
 {
 	using namespace UE::MovieScene;
 
@@ -455,8 +461,12 @@ void FMovieSceneEntitySystemRunner::Flush(double BudgetMs)
 	{
 		FlushState = FlushState::Everything;
 	}
+	if (TargetState == ERunnerFlushState::None)
+	{
+		TargetState = FlushState::Everything;
+	}
 
-	FlushOutstanding(BudgetMs);
+	FlushOutstanding(BudgetMs, TargetState);
 }
 
 void FMovieSceneEntitySystemRunner::FlushOutstanding(double BudgetMs, UE::MovieScene::ERunnerFlushState TargetState)
@@ -503,23 +513,31 @@ void FMovieSceneEntitySystemRunner::FlushOutstanding(double BudgetMs, UE::MovieS
 	}
 
 	const double BudgetSeconds = BudgetMs / 1000.f;
-	if (!bRequireFullFlush && BudgetSeconds > 0.0)
+	if (bRequireFullFlush)
+	{
+		while (EnumHasAnyFlags(FlushState, FlushState::Everything))
+		{
+			if (FlushNext(Linker) == ERunnerFlushResult::Break)
+			{
+				break;
+			}
+		}
+
+		bRequireFullFlush = false;
+	}
+	else if (BudgetSeconds > 0.0)
 	{
 		double StartTime = FPlatformTime::Seconds();
 		while (EnumHasAnyFlags(FlushState, TargetState))
 		{
-			if (!FlushNext(Linker))
+			ERunnerFlushResult Result = FlushNext(Linker);
+			if (Result == ERunnerFlushResult::Break)
 			{
 				break;
 			}
 
-			if (FPlatformTime::Seconds() - StartTime >= BudgetSeconds)
+			if (Result == ERunnerFlushResult::ContinueAllowBudget && FPlatformTime::Seconds() - StartTime >= BudgetSeconds)
 			{
-				// We've reached our budget - we must check next time whether we need to recompile again
-				if (FlushState != ERunnerFlushState::End)
-				{
-					FlushState |= ERunnerFlushState::ConditionalRecompile;
-				}
 				break;
 			}
 		}
@@ -528,14 +546,18 @@ void FMovieSceneEntitySystemRunner::FlushOutstanding(double BudgetMs, UE::MovieS
 	{
 		while (EnumHasAnyFlags(FlushState, TargetState))
 		{
-			if (!FlushNext(Linker))
+			if (FlushNext(Linker) == ERunnerFlushResult::Break)
 			{
 				break;
 			}
 		}
 	}
 
-	bRequireFullFlush = false;
+	// If we are not right at the end of evaluation, we must check for compilation next time
+	if (FlushState != ERunnerFlushState::None && FlushState != ERunnerFlushState::End)
+	{
+		FlushState |= ERunnerFlushState::ConditionalRecompile;
+	}
 }
 
 void FMovieSceneEntitySystemRunner::ResetFlushState()
@@ -543,14 +565,14 @@ void FMovieSceneEntitySystemRunner::ResetFlushState()
 	using namespace UE::MovieScene;
 
 	constexpr ERunnerFlushState StatesThatTriggerReset = ERunnerFlushState::Instantiation | ERunnerFlushState::Evaluation | ERunnerFlushState::Finalization | ERunnerFlushState::EventTriggers;
-	constexpr ERunnerFlushState StatesThatDontTriggerReset = ERunnerFlushState::ConditionalRecompile | ERunnerFlushState::Import | ERunnerFlushState::Spawn;
+	constexpr ERunnerFlushState StatesThatDontTriggerReset = ERunnerFlushState::Import | ERunnerFlushState::Spawn;
 
 	// We only need to reset the flush state if we are far enough through an evaluation
 	if (EnumHasAnyFlags(FlushState, StatesThatTriggerReset) && !EnumHasAnyFlags(FlushState, StatesThatDontTriggerReset))
 	{
 		// When resetting - we don't need to (or want to) re-import anything, we just want to re-run the current
 		// frame of updates 
-		FlushState = UE::MovieScene::FlushState::LoopEval & ~(UE::MovieScene::ERunnerFlushState::ConditionalRecompile | UE::MovieScene::ERunnerFlushState::Import);
+		FlushState = UE::MovieScene::FlushState::LoopEval & ~UE::MovieScene::ERunnerFlushState::Import;
 	}
 }
 
@@ -595,7 +617,7 @@ void FMovieSceneEntitySystemRunner::SkipFlushState(UE::MovieScene::ERunnerFlushS
 	FlushState &= ~FlushStateToSkip;
 }
 
-bool FMovieSceneEntitySystemRunner::GameThread_ConditionalRecompile(UMovieSceneEntitySystemLinker* Linker)
+UE::MovieScene::ERunnerFlushResult FMovieSceneEntitySystemRunner::GameThread_ConditionalRecompile(UMovieSceneEntitySystemLinker* Linker)
 {
 	using namespace UE::MovieScene;
 
@@ -604,23 +626,33 @@ bool FMovieSceneEntitySystemRunner::GameThread_ConditionalRecompile(UMovieSceneE
 	FInstanceRegistry* InstanceRegistry = GetInstanceRegistry();
 
 	bool bAnyRecompile = false;
-	for (int32 UpdateIndex = 0; UpdateIndex < UpdateQueue.Num(); ++UpdateIndex)
-	{
-		const FUpdateParamsAndContext& Request = UpdateQueue[UpdateIndex];
-		if (!InstanceRegistry->IsHandleValid(Request.Params.InstanceHandle))
-		{
-			continue;
-		}
 
-		FSequenceInstance& Instance = InstanceRegistry->MutateInstance(Request.Params.InstanceHandle);
-		if (!Instance.IsRootSequence())
+	// If we have currently running instances, use those for the conditional compile
+	if (CurrentInstances.Num())
+	{
+		for (const FQueuedUpdateParams& UpdatedInstance : CurrentInstances)
 		{
-			continue;
+			if (InstanceRegistry->IsHandleValid(UpdatedInstance.InstanceHandle) && !EnumHasAnyFlags(UpdatedInstance.UpdateFlags, ERunnerUpdateFlags::Finish | ERunnerUpdateFlags::Destroy))
+			{
+				FSequenceInstance& Instance = InstanceRegistry->MutateInstance(UpdatedInstance.InstanceHandle);
+				if (Instance.IsRootSequence() && Instance.ConditionalRecompile(Linker))
+				{
+					bAnyRecompile = true;
+				}
+			}
 		}
-
-		if (Instance.ConditionalRecompile(Linker))
+	}
+	// Otherwise process the update queue
+	else for (int32 UpdateIndex = 0; UpdateIndex < UpdateQueue.Num(); ++UpdateIndex)
+	{
+		const FUpdateParamsAndContext& Request = UpdateQueue[UpdateIndex];
+		if (InstanceRegistry->IsHandleValid(Request.Params.InstanceHandle))
 		{
-			bAnyRecompile = true;
+			FSequenceInstance& Instance = InstanceRegistry->MutateInstance(Request.Params.InstanceHandle);
+			if (Instance.IsRootSequence() && Instance.ConditionalRecompile(Linker))
+			{
+				bAnyRecompile = true;
+			}
 		}
 	}
 
@@ -628,18 +660,29 @@ bool FMovieSceneEntitySystemRunner::GameThread_ConditionalRecompile(UMovieSceneE
 	{
 		if (!EnumHasAnyFlags(FlushState, ERunnerFlushState::Import))
 		{
+			// If we are part-way through a phase that depends on NeedsLink/NeedsUnlink to match,
+			// we need to flush through those before we allow Reimport to run
+			if (Linker->EntityManager.ContainsComponent(FBuiltInComponentTypes::Get()->Tags.NeedsUnlink) || Linker->EntityManager.ContainsComponent(FBuiltInComponentTypes::Get()->Tags.NeedsLink))
+			{
+				TGuardValue<ERunnerFlushState> FlushStateGuard(CurrentFlushState, ERunnerFlushState::None);
+				FlushOutstanding(0.0, ERunnerFlushState::Spawn | ERunnerFlushState::Instantiation);
+			}
+
 			// If we have already imported everything, we have to run a reimport to make sure everything is imported correctly
 			FlushState |= ERunnerFlushState::ReimportAfterCompile;
 		}
 
 		// If a recompile has occurred, we need to make sure we run everything again (except for checking for another recompile!)
-		FlushState |= (FlushState::LoopEval & ~ERunnerFlushState::ConditionalRecompile);
+		FlushState |= (FlushState::LoopEval & ~(ERunnerFlushState::ConditionalRecompile | ERunnerFlushState::Import));
+
+		return ERunnerFlushResult::ContinueAllowBudget;
 	}
 
-	return true;
+	// Do not allow us to get stuck in an infinite loop checking for compilation
+	return ERunnerFlushResult::ContinueNoBudgeting;
 }
 
-bool FMovieSceneEntitySystemRunner::GameThread_UpdateSequenceInstances(UMovieSceneEntitySystemLinker* Linker)
+UE::MovieScene::ERunnerFlushResult FMovieSceneEntitySystemRunner::GameThread_UpdateSequenceInstances(UMovieSceneEntitySystemLinker* Linker)
 {
 	using namespace UE::MovieScene;
 
@@ -683,7 +726,7 @@ bool FMovieSceneEntitySystemRunner::GameThread_UpdateSequenceInstances(UMovieSce
 
 			// If the request is to destroy the instance, we must assume it is no longer valid and should
 			// not be updated
-			if (EnumHasAnyFlags(Request.Params.UpdateFlags, ERunnerUpdateFlags::Destroy))
+			if (EnumHasAnyFlags(Request.Params.UpdateFlags, ERunnerUpdateFlags::Finish | ERunnerUpdateFlags::Destroy))
 			{
 				DissectedUpdates.Add(FDissectedUpdate{ MoveTemp(Request.OnFlushed), Request.Context, Request.Params.InstanceHandle, MAX_int32, Request.Params.UpdateFlags });
 				MarkForUpdate(Request.Params.InstanceHandle, Request.Params.UpdateFlags);
@@ -806,6 +849,8 @@ bool FMovieSceneEntitySystemRunner::GameThread_UpdateSequenceInstances(UMovieSce
 	// Update all systems
 	if (DissectedUpdates.Num() != 0)
 	{
+		TGuardValue<bool> IsUpdatingSequenceGuard(bIsUpdatingSequence, true);
+
 		const int32 PredicateOrder = DissectedUpdates[0].Order;
 		int32 Index = 0;
 		for (; Index < DissectedUpdates.Num() && DissectedUpdates[Index].Order == PredicateOrder; ++Index)
@@ -821,11 +866,6 @@ bool FMovieSceneEntitySystemRunner::GameThread_UpdateSequenceInstances(UMovieSce
 				AccumulatedUpdateFlags |= ESequenceInstanceUpdateFlags::NeedsPostEvaluation;
 			}
 
-			if (EnumHasAnyFlags(Update.UpdateFlags, ERunnerUpdateFlags::Destroy))
-			{
-				continue;
-			}
-
 			if (ensure(InstanceRegistry->IsHandleValid(Update.InstanceHandle)))
 			{
 				FSequenceInstance& Instance = InstanceRegistry->MutateInstance(Update.InstanceHandle);
@@ -835,6 +875,10 @@ bool FMovieSceneEntitySystemRunner::GameThread_UpdateSequenceInstances(UMovieSce
 					// Context is irrelevant for Finishing sequences
 					Instance.Finish(Linker);
 				}
+				else if (EnumHasAnyFlags(Update.UpdateFlags, ERunnerUpdateFlags::Destroy))
+				{
+					ensure(Instance.HasFinished());
+				}
 				else
 				{
 					Instance.Update(Linker, Update.Context);
@@ -861,37 +905,64 @@ bool FMovieSceneEntitySystemRunner::GameThread_UpdateSequenceInstances(UMovieSce
 	}
 #endif
 
-	return true;
+	return ERunnerFlushResult::ContinueAllowBudget;
 }
 
-bool FMovieSceneEntitySystemRunner::GameThread_ReimportSequenceInstances(UMovieSceneEntitySystemLinker* Linker)
+UE::MovieScene::ERunnerFlushResult FMovieSceneEntitySystemRunner::GameThread_ReimportSequenceInstances(UMovieSceneEntitySystemLinker* Linker)
 {
 	using namespace UE::MovieScene;
 
 	// Only called after a sequence has been recompiled after we have already updated the current instances
 	// This allows us to re-update all the sequence instances in case anything has changed
 
+	TGuardValue<bool> IsUpdatingSequenceGuard(bIsUpdatingSequence, true);
+
 	FInstanceRegistry* InstanceRegistry = WeakLinker.Get()->GetInstanceRegistry();
 
+	// NOTE: IncrementSystemSerial must be called before any instance updates are made
+	//       to ensure that up-to-date versions are used inside FEntityManager::OnStructureChanged
+	Linker->EntityManager.IncrementSystemSerial();
+
 	// Operate on a copy of the update params since it is possible CurrentInstances can change if sequence instances are destroyed
-	TArray<FQueuedUpdateParams> CurrentInstancesCopy = CurrentInstances;
+	TArray<FQueuedUpdateParams> CurrentInstancesCopy(MoveTemp(CurrentInstances));
+	CurrentInstances.Reset(CurrentInstancesCopy.Num());
+
 	for (const FQueuedUpdateParams& UpdatedInstance : CurrentInstancesCopy)
 	{
-		if (InstanceRegistry->IsHandleValid(UpdatedInstance.InstanceHandle) && !EnumHasAnyFlags(UpdatedInstance.UpdateFlags, ERunnerUpdateFlags::Destroy | ERunnerUpdateFlags::Finish))
+		if (!InstanceRegistry->IsHandleValid(UpdatedInstance.InstanceHandle))
+		{
+			continue;
+		}
+
+		if (EnumHasAnyFlags(UpdatedInstance.UpdateFlags, ERunnerUpdateFlags::Destroy | ERunnerUpdateFlags::Finish))
+		{
+			FSequenceInstance& SequenceInstance = InstanceRegistry->MutateInstance(UpdatedInstance.InstanceHandle);
+			// Need to check whether this has actually finished yet or not
+			if (!SequenceInstance.HasFinished())
+			{
+				SequenceInstance.Finish(Linker);
+			}
+
+			// If it is being finished or destroyed, just re-add this update back.
+			CurrentInstances.Add(UpdatedInstance);
+		}
+		else
 		{
 			FSequenceInstance& SequenceInstance = InstanceRegistry->MutateInstance(UpdatedInstance.InstanceHandle);
 			// Only update root instances
 			if (SequenceInstance.IsRootSequence())
 			{
+				// Add root sequences back to the CurrentInstances array. When we call Update any active sub-sequences will be re-added via MarkForUpdate
+				CurrentInstances.Add(UpdatedInstance);
 				SequenceInstance.Update(Linker, SequenceInstance.GetContext());
 			}
 		}
 	}
 
-	return true;
+	return ERunnerFlushResult::ContinueAllowBudget;
 }
 
-bool FMovieSceneEntitySystemRunner::GameThread_SpawnPhase(UMovieSceneEntitySystemLinker* Linker)
+UE::MovieScene::ERunnerFlushResult FMovieSceneEntitySystemRunner::GameThread_SpawnPhase(UMovieSceneEntitySystemLinker* Linker)
 {
 	using namespace UE::MovieScene;
 
@@ -954,10 +1025,10 @@ bool FMovieSceneEntitySystemRunner::GameThread_SpawnPhase(UMovieSceneEntitySyste
 		SkipFlushState(ERunnerFlushState::Instantiation);
 	}
 
-	return true;
+	return ERunnerFlushResult::ContinueAllowBudget;
 }
 
-bool FMovieSceneEntitySystemRunner::GameThread_InstantiationPhase(UMovieSceneEntitySystemLinker* Linker)
+UE::MovieScene::ERunnerFlushResult FMovieSceneEntitySystemRunner::GameThread_InstantiationPhase(UMovieSceneEntitySystemLinker* Linker)
 {
 	using namespace UE::MovieScene;
 
@@ -982,7 +1053,7 @@ bool FMovieSceneEntitySystemRunner::GameThread_InstantiationPhase(UMovieSceneEnt
 	return GameThread_PostInstantiation(Linker);
 }
 
-bool FMovieSceneEntitySystemRunner::GameThread_PostInstantiation(UMovieSceneEntitySystemLinker* Linker)
+UE::MovieScene::ERunnerFlushResult FMovieSceneEntitySystemRunner::GameThread_PostInstantiation(UMovieSceneEntitySystemLinker* Linker)
 {
 	using namespace UE::MovieScene;
 
@@ -1016,10 +1087,10 @@ bool FMovieSceneEntitySystemRunner::GameThread_PostInstantiation(UMovieSceneEnti
 	{
 		Linker->SystemGraph.ReconstructTaskSchedule(&Linker->EntityManager);
 	}
-	return true;
+	return ERunnerFlushResult::ContinueAllowBudget;
 }
 
-bool FMovieSceneEntitySystemRunner::GameThread_EvaluationPhase(UMovieSceneEntitySystemLinker* Linker)
+UE::MovieScene::ERunnerFlushResult FMovieSceneEntitySystemRunner::GameThread_EvaluationPhase(UMovieSceneEntitySystemLinker* Linker)
 {
 	using namespace UE::MovieScene;
 
@@ -1060,7 +1131,7 @@ bool FMovieSceneEntitySystemRunner::GameThread_EvaluationPhase(UMovieSceneEntity
 		SkipFlushState(ERunnerFlushState::Finalization | ERunnerFlushState::EventTriggers);
 	}
 
-	return true;
+	return ERunnerFlushResult::ContinueAllowBudget;
 }
 
 void FMovieSceneEntitySystemRunner::GameThread_EvaluationFinalizationPhase(UMovieSceneEntitySystemLinker* Linker)
@@ -1227,6 +1298,24 @@ void FMovieSceneEntitySystemRunner::GameThread_PostEvaluationPhase(UMovieSceneEn
 
 void FMovieSceneEntitySystemRunner::MarkForUpdate(FInstanceHandle InInstanceHandle, UE::MovieScene::ERunnerUpdateFlags UpdateFlags)
 {
+	using namespace UE::MovieScene;
+
+	// If we are reimporting after a compile, we might be finishing/destroying a sequence for the first time and so the call to
+	//    Finish from GameThread_UpdateSequenceInstances may have been missed. In this case we have to check whether the sequence
+	//    needs finishing so we don't have entities left in the ledger when it comes to being destroyed.
+	if (bIsUpdatingSequence && EnumHasAnyFlags(UpdateFlags, ERunnerUpdateFlags::Finish | ERunnerUpdateFlags::Destroy))
+	{
+		UMovieSceneEntitySystemLinker* Linker = WeakLinker.Get();
+		FInstanceRegistry* InstanceRegistry = Linker->GetInstanceRegistry();
+		FSequenceInstance& SequenceInstance = InstanceRegistry->MutateInstance(InInstanceHandle);
+
+		// Need to check whether this has actually finished yet or not
+		if (!SequenceInstance.HasFinished())
+		{
+			SequenceInstance.Finish(Linker);
+		}
+	}
+
 	CurrentInstances.Add(FQueuedUpdateParams{ InInstanceHandle, UpdateFlags });
 }
 
diff --git a/Engine/Source/Runtime/MovieScene/Private/Evaluation/MovieSceneEvaluationTemplateInstance.cpp b/Engine/Source/Runtime/MovieScene/Private/Evaluation/MovieSceneEvaluationTemplateInstance.cpp
--- a/Engine/Source/Runtime/MovieScene/Private/Evaluation/MovieSceneEvaluationTemplateInstance.cpp
+++ b/Engine/Source/Runtime/MovieScene/Private/Evaluation/MovieSceneEvaluationTemplateInstance.cpp
@@ -32,16 +32,21 @@ FMovieSceneRootEvaluationTemplateInstance::FMovieSceneRootEvaluationTemplateInst
 
 void FMovieSceneRootEvaluationTemplateInstance::TearDown()
 {
+	using namespace UE::MovieScene;
+
 	// Avoid redundant work if the linker is being destroyed anyway
 	if (EntitySystemLinker && IsValidChecked(EntitySystemLinker) && !EntitySystemLinker->IsUnreachable() && !EntitySystemLinker->HasAnyFlags(RF_BeginDestroyed))
 	{
 		if (TSharedPtr<FMovieSceneEntitySystemRunner> Runner = WeakRunner.Pin())
 		{
-			Runner->DiscardQueuedUpdates(RootInstanceHandle);
+			Runner->AbandonAndDestroyInstance(RootInstanceHandle);
+			RootInstanceHandle = FRootInstanceHandle();
+		}
+		else
+		{
+			EntitySystemLinker->GetInstanceRegistry()->DestroyInstance(RootInstanceHandle);
+			EntitySystemLinker->ResetActiveRunners();
 		}
-
-		EntitySystemLinker->GetInstanceRegistry()->DestroyInstance(RootInstanceHandle);
-		EntitySystemLinker->ResetActiveRunners();
 	}
 
 	RootInstanceHandle = UE::MovieScene::FRootInstanceHandle();
@@ -407,8 +412,10 @@ void FMovieSceneRootEvaluationTemplateInstance::PlaybackContextChanged(IMovieSce
 
 		if (TSharedPtr<FMovieSceneEntitySystemRunner> Runner = WeakRunner.Pin())
 		{
-			Runner->QueueFinalUpdate(RootInstanceHandle);
-			Runner->Flush();
+			if (Runner->QueueFinalUpdate(RootInstanceHandle))
+			{
+				Runner->Flush();
+			}
 		}
 
 		if (bGlobalCapture)
diff --git a/Engine/Source/Runtime/MovieScene/Public/EntitySystem/MovieSceneEntitySystemRunner.h b/Engine/Source/Runtime/MovieScene/Public/EntitySystem/MovieSceneEntitySystemRunner.h
--- a/Engine/Source/Runtime/MovieScene/Public/EntitySystem/MovieSceneEntitySystemRunner.h
+++ b/Engine/Source/Runtime/MovieScene/Public/EntitySystem/MovieSceneEntitySystemRunner.h
@@ -55,6 +55,17 @@ enum class ERunnerUpdateFlags
 };
 ENUM_CLASS_FLAGS(ERunnerUpdateFlags)
 
+/** Result from a runner flush state that indicates how to proceed */
+enum class ERunnerFlushResult
+{
+	/** Continue evaluation, allow budgeting */
+	ContinueAllowBudget,
+	/** Continue evaluation without budgeting the next step. */
+	ContinueNoBudgeting,
+	/** Do not continue evaluation and break out of our loop. */
+	Break,
+};
+
 } // namespace UE::MovieScene
 
 DECLARE_MULTICAST_DELEGATE(FMovieSceneEntitySystemEventTriggers);
@@ -109,7 +120,7 @@ public:
 	 *                      and will process the outstanding work on the next call to Flush. A value of 0.0 signifies no budget - the queue
 	 *                      will be fully processed without leaving any outstanding work
 	 */
-	MOVIESCENE_API void Flush(double BudgetMs = 0.f);
+	MOVIESCENE_API void Flush(double BudgetMs = 0.f, UE::MovieScene::ERunnerFlushState TargetState = UE::MovieScene::ERunnerFlushState::None);
 
 	/**
 	 * Flushes any outstanding update tasks in the current evaluation scope with a given budget. Only performs work if this runner is part-way through an evaluation
@@ -143,6 +154,11 @@ public:
 	 */
 	MOVIESCENE_API bool IsCurrentlyEvaluating() const;
 
+	/**
+	 * Check whether this runner is currently updating a sequence
+	 */
+	MOVIESCENE_API bool IsUpdatingSequence() const;
+
 	/**
 	 * Run a single evaluation phase
 	 *
@@ -180,28 +196,28 @@ private:
 	 * @param  Linker   The linker we are arrached to
 	 * @return True if the loop is allowed to continue, or false if we should not flush any more
 	 */
-	MOVIESCENE_API bool FlushNext(UMovieSceneEntitySystemLinker* Linker);
+	MOVIESCENE_API UE::MovieScene::ERunnerFlushResult FlushNext(UMovieSceneEntitySystemLinker* Linker);
 
 	/**
 	 * Set up initial state before any evaluation runs. Only called once regardless of the number of pending updates we have to process
 	 * Primarily used for setting up external 'is evaluating' flags for re-entrancy and async checks.
 	 */
-	MOVIESCENE_API bool StartEvaluation(UMovieSceneEntitySystemLinker* Linker);
+	MOVIESCENE_API UE::MovieScene::ERunnerFlushResult StartEvaluation(UMovieSceneEntitySystemLinker* Linker);
 
 	/** Execute any pending conditional recompiles on the currently queued update requests */
-	MOVIESCENE_API bool GameThread_ConditionalRecompile(UMovieSceneEntitySystemLinker* Linker);
+	MOVIESCENE_API UE::MovieScene::ERunnerFlushResult GameThread_ConditionalRecompile(UMovieSceneEntitySystemLinker* Linker);
 	/** Update sequence instances based on currently queued update requests, or outstanding dissected updates */
-	MOVIESCENE_API bool GameThread_UpdateSequenceInstances(UMovieSceneEntitySystemLinker* Linker);
+	MOVIESCENE_API UE::MovieScene::ERunnerFlushResult GameThread_UpdateSequenceInstances(UMovieSceneEntitySystemLinker* Linker);
 	/** Re-update sequence instances after a recompile on a partially evaluated sequence */
-	MOVIESCENE_API bool GameThread_ReimportSequenceInstances(UMovieSceneEntitySystemLinker* Linker);
+	MOVIESCENE_API UE::MovieScene::ERunnerFlushResult GameThread_ReimportSequenceInstances(UMovieSceneEntitySystemLinker* Linker);
 	/** Execute the spawn phase of the entity system graph, if there is anything to do */
-	MOVIESCENE_API bool GameThread_SpawnPhase(UMovieSceneEntitySystemLinker* Linker);
+	MOVIESCENE_API UE::MovieScene::ERunnerFlushResult GameThread_SpawnPhase(UMovieSceneEntitySystemLinker* Linker);
 	/** Execute the instantiation phase of the entity system graph, if there is anything to do */
-	MOVIESCENE_API bool GameThread_InstantiationPhase(UMovieSceneEntitySystemLinker* Linker);
+	MOVIESCENE_API UE::MovieScene::ERunnerFlushResult GameThread_InstantiationPhase(UMovieSceneEntitySystemLinker* Linker);
 	/** Called immediately after instantiation to execute cleanup and bookkeeping tasks. Skipped if instantiation is skipped.*/
-	MOVIESCENE_API bool GameThread_PostInstantiation(UMovieSceneEntitySystemLinker* Linker);
+	MOVIESCENE_API UE::MovieScene::ERunnerFlushResult GameThread_PostInstantiation(UMovieSceneEntitySystemLinker* Linker);
 	/** Main entity-system evaluation phase. Blocks this thread until completion. */
-	MOVIESCENE_API bool GameThread_EvaluationPhase(UMovieSceneEntitySystemLinker* Linker);
+	MOVIESCENE_API UE::MovieScene::ERunnerFlushResult GameThread_EvaluationPhase(UMovieSceneEntitySystemLinker* Linker);
 	/** Finalization phase for triggering external events and other behavior. */
 	MOVIESCENE_API void GameThread_EvaluationFinalizationPhase(UMovieSceneEntitySystemLinker* Linker);
 	/** Post-evaluation phase for triggering events. */
@@ -284,4 +300,6 @@ private:
 
 	/** True if any update has been queued with the ERunnerUpdateFlags::Flush flag since our last flush */
 	bool bRequireFullFlush;
+
+	bool bIsUpdatingSequence;
 };
