Sequencer: Speculative fix for null movie scene sequence

diff --git a/Engine/Source/Runtime/MovieScene/Private/MovieSceneSequencePlayer.cpp b/Engine/Source/Runtime/MovieScene/Private/MovieSceneSequencePlayer.cpp
--- a/Engine/Source/Runtime/MovieScene/Private/MovieSceneSequencePlayer.cpp
+++ b/Engine/Source/Runtime/MovieScene/Private/MovieSceneSequencePlayer.cpp
@@ -1171,7 +1171,7 @@ void UMovieSceneSequencePlayer::UpdateTimeCursorPosition_Internal(FFrameTime New
 			: UpdatePlayPosition(PlayPosition, NewPosition, Method);
 
 		UMovieSceneSequence* MovieSceneSequence = RootTemplateInstance.GetSequence(MovieSceneSequenceID::Root);
-		const bool bIsSequenceBlocking = EnumHasAnyFlags(MovieSceneSequence->GetFlags(), EMovieSceneSequenceFlags::BlockingEvaluation);
+		const bool bIsSequenceBlocking = MovieSceneSequence ? EnumHasAnyFlags(MovieSceneSequence->GetFlags(), EMovieSceneSequenceFlags::BlockingEvaluation) : false;
 		
 		// Just update the time and sequence... if we are in the main level update we want, if possible,
 		// to only queue this sequence's update, so everything updates in parallel. If not possible, or if
