Add a Native Finished delegate for Movie Scene Sequence Player

diff --git a/Engine/Source/Runtime/MovieScene/Private/MovieSceneSequencePlayer.cpp b/Engine/Source/Runtime/MovieScene/Private/MovieSceneSequencePlayer.cpp
--- a/Engine/Source/Runtime/MovieScene/Private/MovieSceneSequencePlayer.cpp
+++ b/Engine/Source/Runtime/MovieScene/Private/MovieSceneSequencePlayer.cpp
@@ -505,6 +505,8 @@ void UMovieSceneSequencePlayer::FinishPlaybackInternal(FFrameTime TimeToFinishAt
 	{
 		OnFinished.Broadcast();
 	}
+
+	OnNativeFinished.ExecuteIfBound();
 }
 
 void UMovieSceneSequencePlayer::GoToEndAndStop()
diff --git a/Engine/Source/Runtime/MovieScene/Public/MovieSceneSequencePlayer.h b/Engine/Source/Runtime/MovieScene/Public/MovieSceneSequencePlayer.h
--- a/Engine/Source/Runtime/MovieScene/Public/MovieSceneSequencePlayer.h
+++ b/Engine/Source/Runtime/MovieScene/Public/MovieSceneSequencePlayer.h
@@ -30,6 +30,7 @@ namespace UE::MovieScene
 }
 
 DECLARE_DYNAMIC_MULTICAST_DELEGATE(FOnMovieSceneSequencePlayerEvent);
+DECLARE_DELEGATE(FOnMovieSceneSequencePlayerNativeEvent);
 
 /**
  * Enum used to define how to update to a particular time
@@ -397,6 +398,8 @@ public:
 	UPROPERTY(BlueprintAssignable, Category = "Sequencer|Player")
 	FOnMovieSceneSequencePlayerEvent OnFinished;
 
+	/** Native event triggered when the level sequence player finishes naturally (without explicitly calling stop) */
+	FOnMovieSceneSequencePlayerNativeEvent OnNativeFinished;
 
 public:
 
