Sequencer: Do not flush the runner inside BeginDestroy

GameThread_ConditionalRecompile no longer tries to recompile things if we are Finishing or Destroying the instance (since there's no need, and it could crash if the sequence has been GC'd).

FMovieSceneRootEvaluationTemplateInstance::BeginDestroy() now goes through a path that does not need to flush any runners; it now forcibly unlinks everything (removing NeedsLink first if there is a runner that is half-way through instantiation when this happens), and uses the CleanGarbage() method to clean up the stale references. This is the same approach we take for PostGarbageCollection.

diff --git a/Engine/Source/Runtime/MovieScene/Private/EntitySystem/MovieSceneEntityLedger.cpp b/Engine/Source/Runtime/MovieScene/Private/EntitySystem/MovieSceneEntityLedger.cpp
--- a/Engine/Source/Runtime/MovieScene/Private/EntitySystem/MovieSceneEntityLedger.cpp
+++ b/Engine/Source/Runtime/MovieScene/Private/EntitySystem/MovieSceneEntityLedger.cpp
@@ -189,14 +189,19 @@ void FEntityLedger::ImportEntity(UMovieSceneEntitySystemLinker* Linker, const FE
 	}
 }
 
-void FEntityLedger::UnlinkEverything(UMovieSceneEntitySystemLinker* Linker)
+void FEntityLedger::UnlinkEverything(UMovieSceneEntitySystemLinker* Linker, EUnlinkEverythingMode UnlinkMode)
 {
+	FComponentTypeID NeedsLink = FBuiltInComponentTypes::Get()->Tags.NeedsLink;
 	FComponentMask FinishedMask = FBuiltInComponentTypes::Get()->FinishedMask;
 
 	for (TPair<FMovieSceneEvaluationFieldEntityKey, FImportedEntityData>& Pair : ImportedEntities)
 	{
 		if (Pair.Value.EntityID)
 		{
+			if (UnlinkMode == EUnlinkEverythingMode::CleanGarbage)
+			{
+				Linker->EntityManager.RemoveComponent(Pair.Value.EntityID, NeedsLink, EEntityRecursion::Full);
+			}
 			Linker->EntityManager.AddComponents(Pair.Value.EntityID, FinishedMask, EEntityRecursion::Full);
 		}
 	}
diff --git a/Engine/Source/Runtime/MovieScene/Private/EntitySystem/MovieSceneEntitySystemLinker.cpp b/Engine/Source/Runtime/MovieScene/Private/EntitySystem/MovieSceneEntitySystemLinker.cpp
--- a/Engine/Source/Runtime/MovieScene/Private/EntitySystem/MovieSceneEntitySystemLinker.cpp
+++ b/Engine/Source/Runtime/MovieScene/Private/EntitySystem/MovieSceneEntitySystemLinker.cpp
@@ -458,6 +458,17 @@ void UMovieSceneEntitySystemLinker::ResetActiveRunners()
 	}
 }
 
+void UMovieSceneEntitySystemLinker::DestroyInstanceImmediately(UE::MovieScene::FRootInstanceHandle Instance)
+{
+	// Ensure that any changes are under a new serial
+	EntityManager.IncrementSystemSerial();
+
+	// Destroy the instance and any sub sequences. Any pre-existing NeedsLink entities will be forcibly made NeedsUnlink and cleaned as garbage
+	GetInstanceRegistry()->DestroyInstance(Instance);
+	CleanGarbage();
+	ResetActiveRunners();
+}
+
 void UMovieSceneEntitySystemLinker::OnObjectsReplaced(const TMap<UObject*, UObject*>& ReplacementMap)
 {
 #if WITH_EDITOR
diff --git a/Engine/Source/Runtime/MovieScene/Private/EntitySystem/MovieSceneEntitySystemRunner.cpp b/Engine/Source/Runtime/MovieScene/Private/EntitySystem/MovieSceneEntitySystemRunner.cpp
--- a/Engine/Source/Runtime/MovieScene/Private/EntitySystem/MovieSceneEntitySystemRunner.cpp
+++ b/Engine/Source/Runtime/MovieScene/Private/EntitySystem/MovieSceneEntitySystemRunner.cpp
@@ -646,7 +646,7 @@ UE::MovieScene::ERunnerFlushResult FMovieSceneEntitySystemRunner::GameThread_Con
 	else for (int32 UpdateIndex = 0; UpdateIndex < UpdateQueue.Num(); ++UpdateIndex)
 	{
 		const FUpdateParamsAndContext& Request = UpdateQueue[UpdateIndex];
-		if (InstanceRegistry->IsHandleValid(Request.Params.InstanceHandle))
+		if (InstanceRegistry->IsHandleValid(Request.Params.InstanceHandle) && !EnumHasAnyFlags(Request.Params.UpdateFlags, ERunnerUpdateFlags::Finish | ERunnerUpdateFlags::Destroy))
 		{
 			FSequenceInstance& Instance = InstanceRegistry->MutateInstance(Request.Params.InstanceHandle);
 			if (Instance.IsRootSequence() && Instance.ConditionalRecompile(Linker))
diff --git a/Engine/Source/Runtime/MovieScene/Private/EntitySystem/MovieSceneSequenceInstance.cpp b/Engine/Source/Runtime/MovieScene/Private/EntitySystem/MovieSceneSequenceInstance.cpp
--- a/Engine/Source/Runtime/MovieScene/Private/EntitySystem/MovieSceneSequenceInstance.cpp
+++ b/Engine/Source/Runtime/MovieScene/Private/EntitySystem/MovieSceneSequenceInstance.cpp
@@ -394,7 +394,7 @@ void FSequenceInstance::DestroyImmediately(UMovieSceneEntitySystemLinker* Linker
 	if (!Ledger.IsEmpty())
 	{
 		UE_LOG(LogMovieSceneECS, Verbose, TEXT("Instance being destroyed without first having been finished by calling Finish()"));
-		Ledger.UnlinkEverything(Linker);
+		Ledger.UnlinkEverything(Linker, EUnlinkEverythingMode::CleanGarbage);
 	}
 
 	if (SequenceUpdater)
diff --git a/Engine/Source/Runtime/MovieScene/Private/Evaluation/MovieSceneEvaluationTemplateInstance.cpp b/Engine/Source/Runtime/MovieScene/Private/Evaluation/MovieSceneEvaluationTemplateInstance.cpp
--- a/Engine/Source/Runtime/MovieScene/Private/Evaluation/MovieSceneEvaluationTemplateInstance.cpp
+++ b/Engine/Source/Runtime/MovieScene/Private/Evaluation/MovieSceneEvaluationTemplateInstance.cpp
@@ -8,6 +8,7 @@
 #include "Compilation/MovieSceneCompiledDataManager.h"
 #include "EntitySystem/MovieSceneEntitySystemLinker.h"
 #include "EntitySystem/MovieSceneEntitySystemRunner.h"
+#include "EntitySystem/BuiltInComponentTypes.h"
 #include "Evaluation/Instances/MovieSceneTrackEvaluator.h"
 
 #include "Sections/MovieSceneSubSection.h"
@@ -35,18 +36,9 @@ void FMovieSceneRootEvaluationTemplateInstance::TearDown()
 	using namespace UE::MovieScene;
 
 	// Avoid redundant work if the linker is being destroyed anyway
-	if (EntitySystemLinker && IsValidChecked(EntitySystemLinker) && !EntitySystemLinker->IsUnreachable() && !EntitySystemLinker->HasAnyFlags(RF_BeginDestroyed))
+	if (RootInstanceHandle.IsValid() && EntitySystemLinker && IsValidChecked(EntitySystemLinker) && !EntitySystemLinker->IsUnreachable() && !EntitySystemLinker->HasAnyFlags(RF_BeginDestroyed))
 	{
-		if (TSharedPtr<FMovieSceneEntitySystemRunner> Runner = WeakRunner.Pin())
-		{
-			Runner->AbandonAndDestroyInstance(RootInstanceHandle);
-			RootInstanceHandle = FRootInstanceHandle();
-		}
-		else
-		{
-			EntitySystemLinker->GetInstanceRegistry()->DestroyInstance(RootInstanceHandle);
-			EntitySystemLinker->ResetActiveRunners();
-		}
+		EntitySystemLinker->DestroyInstanceImmediately(RootInstanceHandle);
 	}
 
 	RootInstanceHandle = UE::MovieScene::FRootInstanceHandle();
diff --git a/Engine/Source/Runtime/MovieScene/Public/EntitySystem/MovieSceneEntityLedger.h b/Engine/Source/Runtime/MovieScene/Public/EntitySystem/MovieSceneEntityLedger.h
--- a/Engine/Source/Runtime/MovieScene/Public/EntitySystem/MovieSceneEntityLedger.h
+++ b/Engine/Source/Runtime/MovieScene/Public/EntitySystem/MovieSceneEntityLedger.h
@@ -24,6 +24,11 @@ namespace MovieScene
 struct FEntityImportSequenceParams;
 struct IMovieScenePerEntityMutation;
 
+enum class EUnlinkEverythingMode
+{
+	Normal,
+	CleanGarbage,
+};
 
 /**
  * An entity ledger is responsible for tracking the entites that have been imported for the currently relevant frame of a sequence instance.
@@ -92,7 +97,7 @@ public:
 	 *
 	 * @param Linker The linker that owns this ledger
 	 */
-	MOVIESCENE_API void UnlinkEverything(UMovieSceneEntitySystemLinker* Linker);
+	MOVIESCENE_API void UnlinkEverything(UMovieSceneEntitySystemLinker* Linker, EUnlinkEverythingMode Garbage = EUnlinkEverythingMode::Normal);
 
 	/**
 	 * Unlink all imported one-shot linker entities and their children and clear the list of one shots
diff --git a/Engine/Source/Runtime/MovieScene/Public/EntitySystem/MovieSceneEntitySystemLinker.h b/Engine/Source/Runtime/MovieScene/Public/EntitySystem/MovieSceneEntitySystemLinker.h
--- a/Engine/Source/Runtime/MovieScene/Public/EntitySystem/MovieSceneEntitySystemLinker.h
+++ b/Engine/Source/Runtime/MovieScene/Public/EntitySystem/MovieSceneEntitySystemLinker.h
@@ -344,6 +344,8 @@ public:
 
 	MOVIESCENE_API void ResetActiveRunners();
 
+	MOVIESCENE_API void DestroyInstanceImmediately(UE::MovieScene::FRootInstanceHandle Instance);
+
 private:
 
 	MOVIESCENE_API UMovieSceneEntitySystem* LinkSystemImpl(TSubclassOf<UMovieSceneEntitySystem> InClassType);
