Sequencer: Force a sort of the bindings and possessables after replacing

Follow up fix to 27759703 and #jira UE-111561

#rnx

diff --git a/Engine/Plugins/MovieScene/LevelSequenceEditor/Source/LevelSequenceEditor/Private/LevelSequenceEditorSubsystem.cpp b/Engine/Plugins/MovieScene/LevelSequenceEditor/Source/LevelSequenceEditor/Private/LevelSequenceEditorSubsystem.cpp
--- a/Engine/Plugins/MovieScene/LevelSequenceEditor/Source/LevelSequenceEditor/Private/LevelSequenceEditorSubsystem.cpp
+++ b/Engine/Plugins/MovieScene/LevelSequenceEditor/Source/LevelSequenceEditor/Private/LevelSequenceEditorSubsystem.cpp
@@ -1462,6 +1462,8 @@ void ULevelSequenceEditorSubsystem::FixActorReferences()
 	{
 		FSequencerUtilities::UpdateBindingIDs(Sequencer.ToSharedRef(), GuidPair.Key, GuidPair.Value);
 	}
+
+	Sequencer->NotifyMovieSceneDataChanged(EMovieSceneDataChangeType::MovieSceneStructureItemsChanged);
 }
 
 void ULevelSequenceEditorSubsystem::AddActorsToBindingInternal()
diff --git a/Engine/Source/Editor/Sequencer/Private/Sequencer.cpp b/Engine/Source/Editor/Sequencer/Private/Sequencer.cpp
--- a/Engine/Source/Editor/Sequencer/Private/Sequencer.cpp
+++ b/Engine/Source/Editor/Sequencer/Private/Sequencer.cpp
@@ -2376,7 +2376,7 @@ void FSequencer::NotifyMovieSceneDataChanged( EMovieSceneDataChangeType DataChan
 		EMovieScenePlayerStatus::Type StoredPlaybackState = GetPlaybackStatus();
 		SetPlaybackStatus( EMovieScenePlayerStatus::Stopped );
 		SelectionPreview.Empty();
-		RefreshTree();
+		RefreshUI();
 		SetPlaybackStatus( StoredPlaybackState );
 	}
 	else if (DataChangeType == EMovieSceneDataChangeType::TrackValueChangedRefreshImmediately)
diff --git a/Engine/Source/Runtime/MovieScene/Private/MovieScene.cpp b/Engine/Source/Runtime/MovieScene/Private/MovieScene.cpp
--- a/Engine/Source/Runtime/MovieScene/Private/MovieScene.cpp
+++ b/Engine/Source/Runtime/MovieScene/Private/MovieScene.cpp
@@ -420,6 +420,9 @@ bool UMovieScene::ReplacePossessable( const FGuid& OldGuid, const FMovieScenePos
 			*OldPossessable = InNewPosessable;
 		}
 
+		// Replace directly changes the guid, so force a sort here
+		Possessables.Sort();
+
 		ReplaceBinding( OldGuid, InNewPosessable.GetGuid(), InNewPosessable.GetName() );
 		bAnythingReplaced = true;
 	}
@@ -1516,6 +1519,9 @@ void UMovieScene::ReplaceBinding(const FGuid& OldGuid, const FGuid& NewGuid, con
 		Binding->SetObjectGuid(NewGuid);
 		Binding->SetName(Name);
 
+		// Replace directly changes the guid, so force a sort here
+		ObjectBindings.Sort();
+
 		// Changing a binding guid invalidates any tracks contained within the binding
 		// Make sure they are written into the transaction buffer by calling modify
 		for (UMovieSceneTrack* Track : Binding->GetTracks())
diff --git a/Engine/Source/Runtime/MovieScene/Public/MovieScenePossessable.h b/Engine/Source/Runtime/MovieScene/Public/MovieScenePossessable.h
--- a/Engine/Source/Runtime/MovieScene/Public/MovieScenePossessable.h
+++ b/Engine/Source/Runtime/MovieScene/Public/MovieScenePossessable.h
@@ -166,7 +166,7 @@ public:
 	/* For sorts so we can search quickly by Guid */
 	FORCEINLINE bool operator<(const FMovieScenePossessable& RHS) const { return Guid < RHS.Guid; }
 	FORCEINLINE friend bool operator<(const FGuid& InGuid, const FMovieScenePossessable& RHS) { return InGuid < RHS.GetGuid(); }
-	bool operator<(const FGuid& InGuid) const { return Guid < InGuid; }
+	FORCEINLINE bool operator<(const FGuid& InGuid) const { return Guid < InGuid; }
 
 private:
 
