Sequencer: don't free group IDs during GC

We only want to free group IDs during the instantiation phase. Otherwise, a downstream system could see a group ID with entities A and B, and next instantiation see it with entities C and D. It wouldn't be clear whether A, B, C, and D genuinely belong to the same group but A and B happened to go away at the exact same time C and D arrive, or whether the group was freed due to a GC pass, and the group with C and D is actually completely different (but reusing the same ID because it was freed).

This confusion was likely causing the property instantiator to assign entities C and D to the property resolved for A and B, thus making C and D animate a property that is totally wrong.

If we don't free up IDs during GC passes, we don't have this issue, as the ID won't get re-used until a further subsequent instantiation phase, so C and D get a new ID before A and B's ID is freed.

diff --git a/Engine/Source/Runtime/MovieScene/Private/EntitySystem/MovieSceneEntityGroupingSystem.cpp b/Engine/Source/Runtime/MovieScene/Private/EntitySystem/MovieSceneEntityGroupingSystem.cpp
--- a/Engine/Source/Runtime/MovieScene/Private/EntitySystem/MovieSceneEntityGroupingSystem.cpp
+++ b/Engine/Source/Runtime/MovieScene/Private/EntitySystem/MovieSceneEntityGroupingSystem.cpp
@@ -109,10 +109,12 @@ struct FUpdateGroupsTask
 	UMovieSceneEntityGroupingSystem* System;
 	FBuiltInComponentTypes* BuiltInComponents;
 	TBitArray<> ModifiedGroups;
+	bool bFreeGroupIDs = true;
 
-	FUpdateGroupsTask(UMovieSceneEntityGroupingSystem* InSystem)
+	FUpdateGroupsTask(UMovieSceneEntityGroupingSystem* InSystem, bool bInFreeGroupIDs = true)
 		: System(InSystem)
 		, BuiltInComponents(nullptr)
+		, bFreeGroupIDs(bInFreeGroupIDs)
 	{
 	}
 
@@ -167,7 +169,7 @@ struct FUpdateGroupsTask
 		// Run the handlers' post-task callback.
 		for (const UMovieSceneEntityGroupingSystem::FEntityGroupingHandlerInfo& HandlerInfo : System->GroupHandlers)
 		{
-			HandlerInfo.Handler->PostTask();
+			HandlerInfo.Handler->PostTask(bFreeGroupIDs);
 		}
 	}
 };
@@ -274,7 +276,10 @@ void UMovieSceneEntityGroupingSystem::OnCleanTaggedGarbage()
 	// group.
 	FBuiltInComponentTypes* BuiltInComponents = FBuiltInComponentTypes::Get();
 
-	FUpdateGroupsTask GroupTask(this);
+	// Don't free group IDs. We only want to free them on instantiation phases, so that downstream
+	// systems don't see any surprisingly re-used IDs from one instantiation frame to another.
+	const bool bFreeGroupIDs = false;
+	FUpdateGroupsTask GroupTask(this, bFreeGroupIDs);
 	FEntityTaskBuilder()
 	.ReadEntityIDs()
 	.Write(BuiltInComponents->Group)
diff --git a/Engine/Source/Runtime/MovieScene/Public/EntitySystem/MovieSceneEntityGroupingSystem.h b/Engine/Source/Runtime/MovieScene/Public/EntitySystem/MovieSceneEntityGroupingSystem.h
--- a/Engine/Source/Runtime/MovieScene/Public/EntitySystem/MovieSceneEntityGroupingSystem.h
+++ b/Engine/Source/Runtime/MovieScene/Public/EntitySystem/MovieSceneEntityGroupingSystem.h
@@ -54,7 +54,7 @@ struct IEntityGroupingHandler
 	virtual ~IEntityGroupingHandler() {}
 	virtual void PreTask() {}
 	virtual void ProcessAllocation(FEntityAllocationIteratorItem Item, FReadEntityIDs EntityIDs, TWrite<FEntityGroupID> GroupIDs, FEntityGroupBuilder* Builder) = 0;
-	virtual void PostTask() {}
+	virtual void PostTask(bool bFreeGroupIDs) {}
 
 #if WITH_EDITOR
 	virtual void OnObjectsReplaced(const TMap<UObject*, UObject*>& ReplacementMap) = 0;
@@ -152,7 +152,6 @@ struct TEntityGroupingHandlerImpl<GroupingPolicy, TIntegerSequence<int, Componen
 
 	virtual void PreTask() override
 	{
-		FreedGroupIndices.Reset();
 		PreTaskImpl(&Policy);
 	}
 
@@ -274,14 +273,14 @@ struct TEntityGroupingHandlerImpl<GroupingPolicy, TIntegerSequence<int, Componen
 	}
 
 	static void PostTaskImpl(void*, ...){}
-	template <typename T> static void PostTaskImpl(T* InPolicy, decltype(&T::PostTask)* = 0)
+	template <typename T> static void PostTaskImpl(T* InPolicy, bool bInFreeGroupIDs, decltype(&T::PostTask)* = 0)
 	{
-		InPolicy->PostTask();
+		InPolicy->PostTask(bInFreeGroupIDs);
 	}
 	
-	virtual void PostTask() override
+	virtual void PostTask(bool bInFreeGroupIDs) override
 	{
-		if (FreedGroupIndices.CountSetBits() > 0)
+		if (bInFreeGroupIDs && FreedGroupIndices.Find(true) != INDEX_NONE)
 		{
 			// Build a reverse lookup map to figure out which group keys we don't need anymore
 			// based on the group indices we have freed.
