Sequencer: Fix edge case where sequences with server-only or client-only subsections do not behave correctly with a single target executable for both server and client using TargetType.Game. In that case both server and client-only subsections stayed included. The fix is to detect when a sequence containing server-only or client-only subsections is being compiled for such a target, and mark that sequence as volatile, forcing it to be recompiled in those circumstances.

diff --git a/Engine/Source/Runtime/MovieScene/Private/Compilation/MovieSceneCompiledDataManager.cpp b/Engine/Source/Runtime/MovieScene/Private/Compilation/MovieSceneCompiledDataManager.cpp
--- a/Engine/Source/Runtime/MovieScene/Private/Compilation/MovieSceneCompiledDataManager.cpp
+++ b/Engine/Source/Runtime/MovieScene/Private/Compilation/MovieSceneCompiledDataManager.cpp
@@ -668,11 +668,17 @@ bool UMovieSceneCompiledDataManager::IsDirty(UMovieSceneSequence* Sequence) cons
 
 
 void UMovieSceneCompiledDataManager::Compile(FMovieSceneCompiledDataID DataID)
+{
+	Compile(DataID, NetworkMask);
+}
+
+
+void UMovieSceneCompiledDataManager::Compile(FMovieSceneCompiledDataID DataID, EMovieSceneServerClientMask InNetworkMask)
 {
 	check(DataID.IsValid() && CompiledDataEntries.IsValidIndex(DataID.Value));
 	UMovieSceneSequence* Sequence = CompiledDataEntries[DataID.Value].GetSequence();
 	check(Sequence);
-	Compile(DataID, Sequence);
+	Compile(DataID, Sequence, InNetworkMask);
 }
 
 FMovieSceneCompiledDataID UMovieSceneCompiledDataManager::Compile(UMovieSceneSequence* Sequence)
@@ -683,6 +689,11 @@ FMovieSceneCompiledDataID UMovieSceneCompiledDataManager::Compile(UMovieSceneSeq
 }
 
 void UMovieSceneCompiledDataManager::Compile(FMovieSceneCompiledDataID DataID, UMovieSceneSequence* Sequence)
+{
+	Compile(DataID, Sequence, NetworkMask);
+}
+
+void UMovieSceneCompiledDataManager::Compile(FMovieSceneCompiledDataID DataID, UMovieSceneSequence* Sequence, EMovieSceneServerClientMask InNetworkMask)
 {
 	check(DataID.IsValid() && CompiledDataEntries.IsValidIndex(DataID.Value));
 	FMovieSceneCompiledDataEntry Entry = CompiledDataEntries[DataID.Value];
@@ -697,13 +708,20 @@ void UMovieSceneCompiledDataManager::Compile(FMovieSceneCompiledDataID DataID, U
 	Entry.DeterminismFences.Empty();
 	Entry.AccumulatedFlags = Sequence->GetFlags();
 	Params.TemplateGenerator.Reset(&Entry);
-	Params.NetworkMask = NetworkMask;
+	Params.NetworkMask = InNetworkMask;
 
 	// ---------------------------------------------------------------------------------------------------
 	// Step 1 - Always ensure the hierarchy information is completely up to date first
 	FMovieSceneSequenceHierarchy NewHierarchy;
 	const bool bHasHierarchy = CompileHierarchy(Sequence, Params, &NewHierarchy);
 
+	// If the network mask of the compiled data manager is 'all', but the sequence has been created with client-only and/or server-only subsections,
+	// then we mark the sequence volatile as we may need to recompile it at runtime in order to exclude these subsections depending on the net mode at runtime.
+	if (Params.NetworkMask == EMovieSceneServerClientMask::All && NewHierarchy.GetAccumulatedNetworkMask() != EMovieSceneServerClientMask::All)
+	{
+		Entry.AccumulatedFlags |= EMovieSceneSequenceFlags::Volatile;
+	}
+
 	if (IMovieSceneDeterminismSource* DeterminismSource = Cast<IMovieSceneDeterminismSource>(Sequence))
 	{
 		DeterminismSource->PopulateDeterminismData(GatheredData.DeterminismData, TRange<FFrameNumber>::All());
@@ -1569,6 +1587,9 @@ void UMovieSceneCompiledDataManager::PopulateSubSequenceTree(UMovieSceneSubTrack
 			continue;
 		}
 
+		InOutHierarchy->AccumulateNetworkMask(SubSection->GetNetworkMask());
+		
+
 		const FMovieSceneSequenceTransform SequenceToRootTransform = Params.RootToSequenceTransform.InverseFromLoop(Params.RootToSequenceWarpCounter);
 
 		// In the case the Sequence to Root Transform contains an infinite timescale, then one of the timescales above us is zero. In this case, we cannot
diff --git a/Engine/Source/Runtime/MovieScene/Private/Compilation/MovieSceneCompiledVolatilityManager.cpp b/Engine/Source/Runtime/MovieScene/Private/Compilation/MovieSceneCompiledVolatilityManager.cpp
--- a/Engine/Source/Runtime/MovieScene/Private/Compilation/MovieSceneCompiledVolatilityManager.cpp
+++ b/Engine/Source/Runtime/MovieScene/Private/Compilation/MovieSceneCompiledVolatilityManager.cpp
@@ -88,7 +88,30 @@ bool FCompiledDataVolatilityManager::ConditionalRecompile(IMovieScenePlayer& Pla
 
 	if (CompiledDataManager->IsDirty(RootDataID))
 	{
-		CompiledDataManager->Compile(RootDataID);
+		// We override the network mask from the compiled data manager here simply because it may not be correct.
+		// In a non-editor/PIE executable, we have a single global compiled data manager, and in its current 'global static' implementation,
+		// it may not have been created at a time with enough context to determine the net mode.
+		// In certain edge cases, such as networked games where both server and client are compiled using a single target executable of TargetType.Game,
+		// we may not have the context at Sequence cook/compile time to determine which subsections may be included/excluded at runtime based on NetworkMask,
+		// and so these Sequences are marked as volatile on compile time. Therefore, it's here, upon conditional recompile, when we need to know the correct
+		// network mask to use, which we override and apply here.
+		EMovieSceneServerClientMask NetworkMask = CompiledDataManager->GetNetworkMask();
+		UObject* PlaybackContext = Player.GetPlaybackContext();
+		UWorld* World = PlaybackContext ? PlaybackContext->GetWorld() : nullptr;
+
+		if (World)
+		{
+			ENetMode NetMode = World->GetNetMode();
+			if (NetMode == ENetMode::NM_DedicatedServer)
+			{
+				NetworkMask = EMovieSceneServerClientMask::Server;
+			}
+			else if (NetMode == ENetMode::NM_Client)
+			{
+				NetworkMask = EMovieSceneServerClientMask::Client;
+			}
+		}
+		CompiledDataManager->Compile(RootDataID, NetworkMask);
 		bRecompiled = true;
 	}
 	else
diff --git a/Engine/Source/Runtime/MovieScene/Public/Compilation/MovieSceneCompiledDataManager.h b/Engine/Source/Runtime/MovieScene/Public/Compilation/MovieSceneCompiledDataManager.h
--- a/Engine/Source/Runtime/MovieScene/Public/Compilation/MovieSceneCompiledDataManager.h
+++ b/Engine/Source/Runtime/MovieScene/Public/Compilation/MovieSceneCompiledDataManager.h
@@ -228,6 +228,8 @@ public:
 
 	MOVIESCENE_API UMovieSceneCompiledData* MakeCompiledData(UMovieSceneSequence* Sequence) const;
 
+	MOVIESCENE_API EMovieSceneServerClientMask GetNetworkMask() const { return NetworkMask; }
+
 	MOVIESCENE_API void SetEmulatedNetworkMask(EMovieSceneServerClientMask NewMask);
 
 	MOVIESCENE_API void Reset(UMovieSceneSequence* Sequence);
@@ -275,10 +277,14 @@ public:
 
 	MOVIESCENE_API void Compile(FMovieSceneCompiledDataID DataID);
 
+	MOVIESCENE_API void Compile(FMovieSceneCompiledDataID DataID, EMovieSceneServerClientMask InNetworkMask);
+
 	MOVIESCENE_API FMovieSceneCompiledDataID Compile(UMovieSceneSequence* Sequence);
 
 	MOVIESCENE_API void Compile(FMovieSceneCompiledDataID DataID, UMovieSceneSequence* Sequence);
 
+	MOVIESCENE_API void Compile(FMovieSceneCompiledDataID DataID, UMovieSceneSequence* Sequence, EMovieSceneServerClientMask InNetworkMask);
+
 	static MOVIESCENE_API bool CompileHierarchy(UMovieSceneSequence* Sequence, FMovieSceneSequenceHierarchy* InOutHierarchy, EMovieSceneServerClientMask NetworkMask);
 
 	MOVIESCENE_API void CopyCompiledData(UMovieSceneSequence* Sequence);
diff --git a/Engine/Source/Runtime/MovieScene/Public/Evaluation/MovieSceneSequenceHierarchy.h b/Engine/Source/Runtime/MovieScene/Public/Evaluation/MovieSceneSequenceHierarchy.h
--- a/Engine/Source/Runtime/MovieScene/Public/Evaluation/MovieSceneSequenceHierarchy.h
+++ b/Engine/Source/Runtime/MovieScene/Public/Evaluation/MovieSceneSequenceHierarchy.h
@@ -344,6 +344,16 @@ struct FMovieSceneSequenceHierarchy
 		return Tree.Data;
 	}
 
+	EMovieSceneServerClientMask GetAccumulatedNetworkMask() const
+	{
+		return AccumulatedNetworkMask;
+	}
+
+	void AccumulateNetworkMask(EMovieSceneServerClientMask Mask)
+	{
+		AccumulatedNetworkMask &= Mask;
+	}
+
 #if !NO_LOGGING
 	void LogHierarchy() const;
 	void LogSubSequenceTree() const;
@@ -365,4 +375,13 @@ private:
 	/** Structural information describing the structure of the sequence */
 	UPROPERTY()
 	TMap<FMovieSceneSequenceID, FMovieSceneSequenceHierarchyNode> Hierarchy;
+
+	/** Holds the accumulated network mask from all included sub-sections. 
+	* If client or server-only subsections are found and included based on the gather params network mask, other bits will be excluded.
+	* If the gather param network mask excludes client or server-only sub-sections, these will be skipped, and so not accumulated.
+	* If no client or server-only subsections are found and included, the mask will be All.
+	* If both client and server-only subsections are found and included, the mask will be None as each would exclude the other.
+	*/
+	UPROPERTY()
+	EMovieSceneServerClientMask AccumulatedNetworkMask = EMovieSceneServerClientMask::All;
 };
