Sequencer: fix trimming of looping sub-sections

diff --git a/Engine/Source/Runtime/MovieScene/Private/Sections/MovieSceneSubSection.cpp b/Engine/Source/Runtime/MovieScene/Private/Sections/MovieSceneSubSection.cpp
--- a/Engine/Source/Runtime/MovieScene/Private/Sections/MovieSceneSubSection.cpp
+++ b/Engine/Source/Runtime/MovieScene/Private/Sections/MovieSceneSubSection.cpp
@@ -389,28 +389,34 @@ void UMovieSceneSubSection::TrimSection( FQualifiedFrameTime TrimTime, bool bTri
 		return;
 	}
 
-	FFrameNumber InitialStartOffset = Parameters.StartFrameOffset;
-
-	UMovieSceneSection::TrimSection( TrimTime, bTrimLeft, bDeleteKeys );
+	SetFlags(RF_Transactional);
+	if (!TryModify())
+	{
+		return;
+	}
 
 	// If trimming off the left, set the offset of the shot
-	if ( bTrimLeft && InitialRange.GetLowerBound().IsClosed() && GetSequence())
+	if (bTrimLeft && InitialRange.GetLowerBound().IsClosed() && GetSequence())
 	{
 		// Sections need their offsets calculated in their local resolution. Different sequences can have different tick resolutions 
 		// so we need to transform from the parent resolution to the local one before splitting them.
-		FFrameRate LocalTickResolution = GetSequence()->GetMovieScene()->GetTickResolution();
-		FFrameNumber LocalResolutionStartOffset = FFrameRate::TransformTime(TrimTime.Time.GetFrame() - UE::MovieScene::DiscreteInclusiveLower(InitialRange), TrimTime.Rate, LocalTickResolution).FrameNumber;
-
-
-		FFrameNumber NewStartOffset = LocalResolutionStartOffset * Parameters.TimeScale;
-		NewStartOffset += InitialStartOffset;
-
-		// Ensure start offset is not less than 0
-		if (NewStartOffset >= 0)
-		{
-			Parameters.StartFrameOffset = NewStartOffset;
-		}
+		UMovieScene* LocalMovieScene = GetSequence()->GetMovieScene();
+		const FFrameRate LocalTickResolution = LocalMovieScene->GetTickResolution();
+		const FFrameTime LocalTickResolutionTrimTime = FFrameRate::TransformTime(TrimTime.Time, TrimTime.Rate, LocalTickResolution);
+
+		// The new first loop start offset is where the trim time fell inside the sub-sequence (this time is already
+		// normalized in the case of looping sub-sequences).
+		FFrameTime LocalTrimTime = OuterToInnerTransform().TransformTime(LocalTickResolutionTrimTime);
+		FFrameNumber NewFirstLoopStartOffset = LocalTrimTime.FrameNumber - Parameters.StartFrameOffset;
+		
+		// Make sure we don't have negative offsets (this shouldn't happen, though).
+		NewFirstLoopStartOffset = FMath::Max(FFrameNumber(0), NewFirstLoopStartOffset);
+		
+		Parameters.FirstLoopStartFrameOffset = NewFirstLoopStartOffset;
 	}
+
+	// Actually trim the section range!
+	UMovieSceneSection::TrimSection(TrimTime, bTrimLeft, bDeleteKeys);
 }
 
 void UMovieSceneSubSection::GetSnapTimes(TArray<FFrameNumber>& OutSnapTimes, bool bGetSectionBorders) const
