Sequencer: GetDisplayName() if track row name is empty

#rnx

diff --git a/Engine/Source/Runtime/MovieScene/Private/MovieSceneNameableTrack.cpp b/Engine/Source/Runtime/MovieScene/Private/MovieSceneNameableTrack.cpp
--- a/Engine/Source/Runtime/MovieScene/Private/MovieSceneNameableTrack.cpp
+++ b/Engine/Source/Runtime/MovieScene/Private/MovieSceneNameableTrack.cpp
@@ -84,12 +84,7 @@ FText UMovieSceneNameableTrack::GetTrackRowDisplayName(int32 TrackRowIndex) cons
 		return TrackRowDisplayNames[TrackRowIndex];
 	}
 
-	if (DisplayName.IsEmpty())
-	{
-		return GetDefaultDisplayName();
-	}
-
-	return DisplayName;
+	return GetDisplayName();
 }
 
 FText UMovieSceneNameableTrack::GetDefaultDisplayName() const
