Sequencer: Fix major sync issues with replicated sequences while that sequence is itself animating the global time dilation. The code around calculating latency was not taking into account effective time dilation.
Also, fix a spelling error for a (non UPROPERTY/exposed) variable.

diff --git a/Engine/Source/Runtime/MovieScene/Private/MovieSceneSequencePlayer.cpp b/Engine/Source/Runtime/MovieScene/Private/MovieSceneSequencePlayer.cpp
--- a/Engine/Source/Runtime/MovieScene/Private/MovieSceneSequencePlayer.cpp
+++ b/Engine/Source/Runtime/MovieScene/Private/MovieSceneSequencePlayer.cpp
@@ -1413,9 +1413,18 @@ FFrameTime UMovieSceneSequencePlayer::UpdateServerTimeSamples()
 	const double Lifetime         = CurrentWallClock - float(GSequencerMaxSmoothedNetSyncSampleAge) / 1000.f;
 	const float PlaybackMultiplier = bReversePlayback ? -1.f : 1.f;
 
+	float TimeDilation = 1.0f;
+	if (const UWorld* World = GetPlaybackWorld())
+	{
+		if (AWorldSettings* WorldSettings = World->GetWorldSettings())
+		{
+			TimeDilation = WorldSettings->GetEffectiveTimeDilation();
+		}
+	}
+
 	// Cull any old samples that were taken more than GSequencerMaxSmoothedNetSyncSampleAge ms ago by
 	// Finding the index of the first sample younger than this time
-	const int32 FirstValidSample = Algo::LowerBoundBy(ServerTimeSamples, Lifetime, &FServerTimeSample::ReceievedTime);
+	const int32 FirstValidSample = Algo::LowerBoundBy(ServerTimeSamples, Lifetime, &FServerTimeSample::ReceivedTime);
 	if (FirstValidSample >= ServerTimeSamples.Num())
 	{
 		// Never found a sample that is recent enough, all samples are too old
@@ -1443,10 +1452,24 @@ FFrameTime UMovieSceneSequencePlayer::UpdateServerTimeSamples()
 		ServerTimeSamples.RemoveAt(MaxNumSamples, ServerTimeSamples.Num() - MaxNumSamples, true);
 	}
 
+	auto UpdateSamplesForChangedTimeDilation = [&]()
+	{
+		// Project all server time samples back based on the new time dilation so future updates will be accurate
+		if (LastEffectiveTimeDilation != TimeDilation)
+		{
+			for (FServerTimeSample& Sample : ServerTimeSamples)
+			{
+				const double ThisSample = Sample.ServerTime + (CurrentWallClock - Sample.ReceivedTime) * PlaybackMultiplier * LastEffectiveTimeDilation;
+				Sample.ReceivedTime = CurrentWallClock - (ThisSample - Sample.ServerTime) / (PlaybackMultiplier * TimeDilation);
+			}
+			LastEffectiveTimeDilation = TimeDilation;
+		}
+	};
 
 	if (ServerTimeSamples.Num() < 10)
 	{
 		// Fallback to the current time if there are not enough samples
+		UpdateSamplesForChangedTimeDilation();
 		return PlayPosition.GetCurrentPosition();
 	}
 
@@ -1454,7 +1477,7 @@ FFrameTime UMovieSceneSequencePlayer::UpdateServerTimeSamples()
 	double MeanTime = 0;
 	for (const FServerTimeSample& Sample : ServerTimeSamples)
 	{
-		const double ThisSample = Sample.ServerTime + (CurrentWallClock - Sample.ReceievedTime) * PlaybackMultiplier;
+		const double ThisSample = Sample.ServerTime + (CurrentWallClock - Sample.ReceivedTime) * PlaybackMultiplier * LastEffectiveTimeDilation;
 		MeanTime += ThisSample;
 	}
 	MeanTime = MeanTime / ServerTimeSamples.Num();
@@ -1463,7 +1486,7 @@ FFrameTime UMovieSceneSequencePlayer::UpdateServerTimeSamples()
 	double StandardDeviation = 0;
 	for (const FServerTimeSample& Sample : ServerTimeSamples)
 	{
-		const double ThisSample = Sample.ServerTime + (CurrentWallClock - Sample.ReceievedTime) * PlaybackMultiplier;
+		const double ThisSample = Sample.ServerTime + (CurrentWallClock - Sample.ReceivedTime) * PlaybackMultiplier * LastEffectiveTimeDilation;
 		StandardDeviation += FMath::Square(ThisSample - MeanTime);
 	}
 	StandardDeviation = StandardDeviation / ServerTimeSamples.Num();
@@ -1481,7 +1504,7 @@ FFrameTime UMovieSceneSequencePlayer::UpdateServerTimeSamples()
 		// Discard anything outside the standard deviation in the hopes that future samples will converge
 		for (int32 SampleIndex = ServerTimeSamples.Num()-1; SampleIndex >= 0; --SampleIndex)
 		{
-			const double ThisSample = ServerTimeSamples[SampleIndex].ServerTime + (CurrentWallClock - ServerTimeSamples[SampleIndex].ReceievedTime) * PlaybackMultiplier;
+			const double ThisSample = ServerTimeSamples[SampleIndex].ServerTime + (CurrentWallClock - ServerTimeSamples[SampleIndex].ReceivedTime) * PlaybackMultiplier * LastEffectiveTimeDilation;
 			if (FMath::Abs(ThisSample - MeanTime) > StandardDeviation)
 			{
 				ServerTimeSamples.RemoveAt(SampleIndex, 1, false);
@@ -1494,6 +1517,8 @@ FFrameTime UMovieSceneSequencePlayer::UpdateServerTimeSamples()
 		NewMeanTime = NewMeanTime / ServerTimeSamples.Num();
 	}
 
+	UpdateSamplesForChangedTimeDilation();
+
 	// If we didn't cull too many samples, we have confidence in the data set
 	if (ServerTimeSamples.Num() >= OriginalNum/2)
 	{
@@ -1694,11 +1719,21 @@ void UMovieSceneSequencePlayer::PostNetReceive()
 	const bool bHasChangedStatus  = NetSyncProps.LastKnownStatus   != Status;
 	const bool bHasChangedTime    = NetSyncProps.LastKnownPosition != PlayPosition.GetCurrentPosition();
 
+	float TimeDilation = 1.0f;
+	if (const UWorld* World = GetPlaybackWorld())
+	{
+		if (AWorldSettings* WorldSettings = World->GetWorldSettings())
+		{
+			TimeDilation = WorldSettings->GetEffectiveTimeDilation();
+		}
+	}
+
 	const float PingMs            = GetPing();
-	const FFrameTime PingLag      = (PingMs/1000.f) * PlayPosition.GetInputRate();
+	const FFrameTime PingLag      = (PingMs/1000.f) * PlayPosition.GetInputRate() * TimeDilation;
 	//const FFrameTime LagThreshold = 0.2f * PlayPosition.GetInputRate();
 	//const FFrameTime LagDisparity = FMath::Abs(PlayPosition.GetCurrentPosition() - NetSyncProps.LastKnownPosition);
-	const FFrameTime LagThreshold = (GSequencerNetSyncThresholdMS * 0.001f) * PlayPosition.GetInputRate();
+
+	const FFrameTime LagThreshold = (GSequencerNetSyncThresholdMS * 0.001f) * PlayPosition.GetInputRate() * TimeDilation;
 
 	if (!bHasChangedStatus && !bHasChangedTime)
 	{
@@ -1797,9 +1832,17 @@ void UMovieSceneSequencePlayer::UpdateNetworkSync()
 	// Only process net playback synchronization if we are still Playing.
 	if (Status == EMovieScenePlayerStatus::Playing)
 	{
-		const float PingMs            = GetPing();
-		const FFrameTime PingLag      = (PingMs/1000.f) * PlayPosition.GetInputRate();
-		const FFrameTime LagThreshold = (GSequencerNetSyncThresholdMS * 0.001f) * PlayPosition.GetInputRate();
+		const float PingMs            = GetPing();	
+		float TimeDilation = 1.0f;
+		if (const UWorld* World = GetPlaybackWorld())
+		{
+			if (AWorldSettings* WorldSettings = World->GetWorldSettings())
+			{
+				TimeDilation = WorldSettings->GetEffectiveTimeDilation();
+			}
+		}
+		const FFrameTime PingLag      = (PingMs/1000.f) * PlayPosition.GetInputRate() * TimeDilation;
+		const FFrameTime LagThreshold = (GSequencerNetSyncThresholdMS * 0.001f) * PlayPosition.GetInputRate() * TimeDilation;
 		
 		// When the server has looped back to the start but a client is near the end (and is thus about to loop), we don't want to forcibly synchronize the time unless
 		// the *real* difference in time is above the threshold. We compute the real-time difference by adding SequenceDuration*LoopCountDifference to the server position:
diff --git a/Engine/Source/Runtime/MovieScene/Public/MovieSceneSequencePlayer.h b/Engine/Source/Runtime/MovieScene/Public/MovieSceneSequencePlayer.h
--- a/Engine/Source/Runtime/MovieScene/Public/MovieSceneSequencePlayer.h
+++ b/Engine/Source/Runtime/MovieScene/Public/MovieSceneSequencePlayer.h
@@ -695,7 +695,7 @@ protected:
 		/** The actual server sequence time in seconds, with client ping at the time of the sample baked in */
 		double ServerTime;
 		/** Wall-clock time that the sample was receieved */
-		double ReceievedTime;
+		double ReceivedTime;
 	};
 	/**
 	 * Array of server sequence times in seconds, with ping compensation baked in.
@@ -703,6 +703,11 @@ protected:
 	 */
 	TArray<FServerTimeSample> ServerTimeSamples;
 
+	/*
+	* On UpdateServerTimeSamples, the last recorded time dilation. Used to update the server time samples each update to ensure we can smooth server time even on changing time dilation.
+	*/
+	float LastEffectiveTimeDilation = 1.0f;
+
 	/** Replicated playback status and current time that are replicated to clients */
 	UPROPERTY(replicated)
 	FMovieSceneSequenceReplProperties NetSyncProps;
