Sequencer: Remove ensure for re-entrancy.

This seems to be hit a lot because of a cinematic sequence device being (re)started on some overlap.

#rnx

diff --git a/Engine/Source/Runtime/MovieScene/Private/EntitySystem/MovieSceneEntitySystemRunner.cpp b/Engine/Source/Runtime/MovieScene/Private/EntitySystem/MovieSceneEntitySystemRunner.cpp
--- a/Engine/Source/Runtime/MovieScene/Private/EntitySystem/MovieSceneEntitySystemRunner.cpp
+++ b/Engine/Source/Runtime/MovieScene/Private/EntitySystem/MovieSceneEntitySystemRunner.cpp
@@ -479,15 +479,17 @@ void FMovieSceneEntitySystemRunner::FlushOutstanding(double BudgetMs, UE::MovieS
 		return;
 	}
 
-	UMovieSceneEntitySystemLinker* Linker = GetLinker();
-
-	// Check that we are attached to a linker that allows starting a new evaluation.
-	if (!ensureMsgf(Linker, TEXT("Runner isn't attached to a valid linker")))
+	// If this runner is already being flushed, early return
+	if (CurrentFlushState != ERunnerFlushState::None)
 	{
+		UE_LOG(LogMovieSceneECS, Warning, TEXT("Cannot flush this runner while it is already being flushed outside of a re-entrancy window"));
 		return;
 	}
 
-	if (!ensureMsgf(CurrentFlushState == ERunnerFlushState::None, TEXT("Cannot flush this runner while it is already being flushed outside of a re-entrancy window")))
+	UMovieSceneEntitySystemLinker* Linker = GetLinker();
+
+	// Check that we are attached to a linker that allows starting a new evaluation.
+	if (!ensureMsgf(Linker, TEXT("Runner isn't attached to a valid linker")))
 	{
 		return;
 	}
