Sequencer: Any number of hierarchical blend targets are now supported

  - FHierarchicalBlendTarget now stores each hbias in the chain as part of an inline buffer (the same as before), but can optionally switch to a heap allocation if more than 7 levels are encountered. Using the heap here is very undesireable for performance, but at least it will work. Having more than 7 nested levels of blended hbias indicates a bigger level design problem anyway.
  - The type remains 16 bytes, so multiple can still fit neatly on a cache line. There is some manual memory management necessary to support this, but that is necessary to keep the size of the type down.
  - Added some automation tests to verify behavior

diff --git a/Engine/Source/Runtime/MovieScene/Private/EntitySystem/BuiltInComponentTypes.cpp b/Engine/Source/Runtime/MovieScene/Private/EntitySystem/BuiltInComponentTypes.cpp
--- a/Engine/Source/Runtime/MovieScene/Private/EntitySystem/BuiltInComponentTypes.cpp
+++ b/Engine/Source/Runtime/MovieScene/Private/EntitySystem/BuiltInComponentTypes.cpp
@@ -18,46 +18,164 @@ namespace MovieScene
 {
 
 static constexpr int16 GInvalidBlendTarget = std::numeric_limits<int16>::lowest();
+static const int16 GInitialBlendTargets[7] = {
+	GInvalidBlendTarget, GInvalidBlendTarget, GInvalidBlendTarget, GInvalidBlendTarget,
+	GInvalidBlendTarget, GInvalidBlendTarget, GInvalidBlendTarget
+};
 
 FHierarchicalBlendTarget::FHierarchicalBlendTarget()
+	: Capacity(InlineCapacity)
+{
 	// Fill with invalid values
-	: HBiasChain{
-		GInvalidBlendTarget, GInvalidBlendTarget, GInvalidBlendTarget, GInvalidBlendTarget,
-		GInvalidBlendTarget, GInvalidBlendTarget, GInvalidBlendTarget, GInvalidBlendTarget
+	static_assert(sizeof(Data) == sizeof(GInitialBlendTargets));
+	FMemory::Memcpy(this->Data, GInitialBlendTargets, sizeof(Data));
+}
+
+FHierarchicalBlendTarget::FHierarchicalBlendTarget(const FHierarchicalBlendTarget& RHS)
+	: Capacity(InlineCapacity)
+{
+	// Make an allocation if necessary
+	if (RHS.Capacity != InlineCapacity)
+	{
+		int16* NewAllocation = new int16[RHS.Capacity];
+		Capacity = RHS.Capacity;
+		*reinterpret_cast<int16**>(Data) = NewAllocation;
 	}
+
+	// Copy the data
+	FMemory::Memcpy(this->GetMemory(), RHS.GetMemory(), sizeof(int16)*Capacity);
+}
+
+FHierarchicalBlendTarget& FHierarchicalBlendTarget::operator=(const FHierarchicalBlendTarget& RHS)
 {
+	// Release our allocation if it's the wrong size
+	if (this->Capacity != RHS.Capacity)
+	{
+		FreeAllocation();
+	}
+
+	// Make an allocation if necessary
+	if (RHS.Capacity != InlineCapacity)
+	{
+		int16* NewAllocation = new int16[RHS.Capacity];
+		Capacity = RHS.Capacity;
+		*reinterpret_cast<int16**>(Data) = NewAllocation;
+	}
+
+	// Copy the data
+	FMemory::Memcpy(this->GetMemory(), RHS.GetMemory(), sizeof(int16)*Capacity);
+	return *this;
 }
 
-void FHierarchicalBlendTarget::Add(int16 InValue)
+FHierarchicalBlendTarget::FHierarchicalBlendTarget(FHierarchicalBlendTarget&& RHS)
+	: Capacity(RHS.Capacity)
 {
-	static constexpr int32 MaxNum = UE_ARRAY_COUNT(HBiasChain);
+	FMemory::Memcpy(this->Data, RHS.Data, sizeof(Data));
 
-	const int32 Index = Algo::LowerBound(HBiasChain, InValue, TGreater<>());
-	if (Index >= MaxNum)
+	// Reset RHS
+	RHS.Capacity = InlineCapacity;
+	static_assert(sizeof(Data) == sizeof(GInitialBlendTargets));
+	FMemory::Memcpy(RHS.Data, GInitialBlendTargets, sizeof(Data));
+}
+
+FHierarchicalBlendTarget& FHierarchicalBlendTarget::operator=(FHierarchicalBlendTarget&& RHS)
+{
+	// Free our allocation if it is allocated
+	FreeAllocation();
+
+	// Steal RHS
+	Capacity = RHS.Capacity;
+	FMemory::Memcpy(this->Data, RHS.Data, sizeof(Data));
+
+	// Reset RHS
+	RHS.Capacity = InlineCapacity;
+	static_assert(sizeof(Data) == sizeof(GInitialBlendTargets));
+	FMemory::Memcpy(RHS.Data, GInitialBlendTargets, sizeof(Data));
+
+	return *this;
+}
+
+FHierarchicalBlendTarget::~FHierarchicalBlendTarget()
+{
+	FreeAllocation();
+}
+
+void FHierarchicalBlendTarget::FreeAllocation()
+{
+	// Free our allocation if it is allocated
+	if (Capacity != InlineCapacity)
 	{
-		checkf(false, TEXT("Unable to support more than %d levels of nested hierarchical blending"), MaxNum);
-		return;
+		delete[] GetMemory();
 	}
+}
 
-	if (HBiasChain[Index] == InValue)
+void FHierarchicalBlendTarget::Grow(uint16 NewCapacity)
+{
+	int16* NewAllocation = new int16[NewCapacity];
+	int16* OldAllocation = GetMemory();
+
+	// Copy old allocation
+	FMemory::Memmove(NewAllocation, OldAllocation, Capacity*sizeof(int16));
+
+	// Initialize the tail of the new allocation
+	for (uint16 Index = Capacity; Index < NewCapacity; ++Index)
+	{
+		NewAllocation[Index] = GInvalidBlendTarget;
+	}
+
+	if (Capacity != InlineCapacity)
+	{
+		delete[] OldAllocation;
+	}
+
+	Capacity = NewCapacity;
+	*reinterpret_cast<int16**>(Data) = NewAllocation;
+}
+
+void FHierarchicalBlendTarget::Add(int16 InValue)
+{
+	TArrayView<int16> HBiasChain = GetAllEntries();
+
+	const int32 Index = Algo::LowerBound(HBiasChain, InValue, TGreater<>());
+	
+	if (Index < HBiasChain.Num() && HBiasChain[Index] == InValue)
 	{
 		// Already exists
 		return;
 	}
 
-	if (HBiasChain[MaxNum-1] != GInvalidBlendTarget)
+	// If our insertion index is beyond our capacity we need to grow.
+	const bool bAtMaxCapacity = HBiasChain[Capacity-1] != GInvalidBlendTarget;
+	const bool bNeedToGrow    = bAtMaxCapacity || Index >= Capacity;
+
+	if (bNeedToGrow)
 	{
-		checkf(false, TEXT("Unable to support more than %d levels of nested hierarchical blending"), MaxNum);
-		return;
+		if (Capacity == InlineCapacity)
+		{
+			Grow(16u);
+		}
+		else
+		{
+			Grow(Capacity + 16u);
+		}
+
+		// Re-retrieve the memory
+		HBiasChain = GetAllEntries();
 	}
 
 	// Right shift by 1 if necessary
-	if (Index < MaxNum-1 && HBiasChain[Index] != GInvalidBlendTarget)
+	if (Index < Capacity-1 && HBiasChain[Index] != GInvalidBlendTarget)
 	{
-		void* Src = HBiasChain + Index;
-		void* Dst = HBiasChain + Index + 1;
-		int32 Slack = MaxNum-Index-2; // -2 to avoid buffer overrun. The last one will always get overwritten
-		FMemory::Memmove(Dst, Src, sizeof(int16)*Slack);
+		// Example: Inserting 500 into the following:
+		// [ 1000, 100, 50, 10, 5, 0, -32768 ] w/ Capacity=7
+		// Results in:
+		// [ 1000, 500, 100, 50, 10, 5, 0 ] 
+		// Index=1
+		// NumToRelocate=5
+		void* Src = &HBiasChain[Index];
+		void* Dst = &HBiasChain[Index + 1];
+		int32 NumToRelocate = Capacity-Index-1; // -1 to avoid buffer overrun. The last one will always get overwritten. See comment above.
+		FMemory::Memmove(Dst, Src, sizeof(int16)*NumToRelocate);
 	}
 
 	// Set the value
@@ -66,20 +184,46 @@ void FHierarchicalBlendTarget::Add(int16 InValue)
 
 int32 FHierarchicalBlendTarget::Num() const
 {
+	TArrayView<const int16> HBiasChain = GetAllEntries();
 	return Algo::LowerBound(HBiasChain, GInvalidBlendTarget, TGreater<>());
 }
 
 int16 FHierarchicalBlendTarget::operator[](int32 Index) const
 {
-	static constexpr int32 MaxNum = UE_ARRAY_COUNT(HBiasChain);
+	check(Index < static_cast<int32>(Capacity));
+	return GetAllEntries()[Index];
+}
 
-	check(Index < MaxNum);
-	return HBiasChain[Index];
+int16* FHierarchicalBlendTarget::GetMemory()
+{
+	return const_cast<int16*>(const_cast<const FHierarchicalBlendTarget*>(this)->GetMemory());
+}
+
+const int16* FHierarchicalBlendTarget::GetMemory() const
+{
+	if (Capacity == InlineCapacity)
+	{
+		return reinterpret_cast<const int16*>(Data);
+	}
+	else
+	{
+		return *reinterpret_cast<const int16* const*>(Data);
+	}
+}
+
+TArrayView<int16> FHierarchicalBlendTarget::GetAllEntries()
+{
+	return MakeArrayView(GetMemory(), Capacity);
+}
+
+TArrayView<const int16> FHierarchicalBlendTarget::GetAllEntries() const
+{
+	return MakeArrayView(GetMemory(), Capacity);
 }
 
 TArrayView<const int16> FHierarchicalBlendTarget::AsArray() const
 {
-	return TArrayView<const int16>(HBiasChain, Num());
+	return GetAllEntries().Left(Num());
 }
 
 bool FObjectComponent::IsStrongReference() const
diff --git a/Engine/Source/Runtime/MovieScene/Private/Tests/MovieSceneEntitySystemTypesTests.cpp b/Engine/Source/Runtime/MovieScene/Private/Tests/MovieSceneEntitySystemTypesTests.cpp
--- a/Engine/Source/Runtime/MovieScene/Private/Tests/MovieSceneEntitySystemTypesTests.cpp
+++ b/Engine/Source/Runtime/MovieScene/Private/Tests/MovieSceneEntitySystemTypesTests.cpp
@@ -5,6 +5,9 @@
 #include "EntitySystem/MovieSceneEntitySystemTypes.h"
 #include "EntitySystem/MovieSceneEntityManager.h"
 #include "EntitySystem/MovieSceneEntityFactoryTemplates.h"
+#include "EntitySystem/BuiltInComponentTypes.h"
+#include "Containers/Array.h"
+#include "Containers/Set.h"
 #include "Containers/ArrayView.h"
 #include "Misc/AutomationTest.h"
 
@@ -344,6 +347,203 @@ bool FMovieSceneEntityMigrationTest::RunTest(const FString& Parameters)
 	return true;
 }
 
+
+IMPLEMENT_SIMPLE_AUTOMATION_TEST(FMovieSceneBlendTargetComponentTest, 
+		"System.Engine.Sequencer.EntitySystem.BlendTargetComponent",
+		EAutomationTestFlags::EditorContext | EAutomationTestFlags::EngineFilter)
+bool FMovieSceneBlendTargetComponentTest::RunTest(const FString& Parameters)
+{
+	/*~~~~~~~~~~~~~~~~~~~~~~~~~~~~~~~~~~~~~~~~~~~~~~~~~~~~~~~~~~~~~~~~~~~~~~~~~~~~~~~~~
+	This test checks the behavior of FHierarchicalBlendTarget w.r.t.
+	move operations and memory management
+	~~~~~~~~~~~~~~~~~~~~~~~~~~~~~~~~~~~~~~~~~~~~~~~~~~~~~~~~~~~~~~~~~~~~~~~~~~~~~~~~~~~*/
+
+	using namespace UE::MovieScene;
+
+	auto CompareBlendTarget = [this](const FHierarchicalBlendTarget& BlendTarget, TArrayView<const int16> Expected, const TCHAR* InTestName)
+	{
+		TArrayView<const int16> Result = BlendTarget.AsArray();
+		if (Result.Num() != Expected.Num() || !CompareItems(Result.GetData(), Expected.GetData(), Result.Num()))
+		{
+			TStringBuilder<128> ErrorString;
+			ErrorString.Appendf(TEXT("Test %s: Blend Target does not match the expected result ("), InTestName);
+			ErrorString.Join(Result, TEXT(","));
+			ErrorString.Append(TEXT(" != "));
+			ErrorString.Join(Expected, TEXT(","));
+			ErrorString.Append(TEXT(")"));
+			this->AddError(ErrorString.ToString());
+		}
+	};
+
+	// Simple test - use inline allocation
+	{
+		FHierarchicalBlendTarget BlendTarget;
+		BlendTarget.Add(0);
+		BlendTarget.Add(100);
+		BlendTarget.Add(-100);
+		BlendTarget.Add(-10);
+		BlendTarget.Add(10);
+		BlendTarget.Add(500);
+		BlendTarget.Add(1000);
+
+		this->TestEqual(TEXT("Inline Basic Capacity"), BlendTarget.GetCapacity(), 7);
+		CompareBlendTarget(BlendTarget, { 1000, 500, 100, 10, 0, -10, -100 }, TEXT("Inline Basic"));
+		GetTypeHash(BlendTarget);
+	}
+
+	// Test adding duplicates
+	{
+		FHierarchicalBlendTarget BlendTarget;
+		BlendTarget.Add(0);
+		BlendTarget.Add(100);
+		BlendTarget.Add(-100);
+		BlendTarget.Add(-10);
+		BlendTarget.Add(10);
+		BlendTarget.Add(500);
+		BlendTarget.Add(1000);
+		BlendTarget.Add(0);
+		BlendTarget.Add(100);
+		BlendTarget.Add(-100);
+		BlendTarget.Add(-10);
+		BlendTarget.Add(10);
+		BlendTarget.Add(500);
+		BlendTarget.Add(1000);
+
+		this->TestEqual(TEXT("Inline Duplicate Capacity"), BlendTarget.GetCapacity(), 7);
+		CompareBlendTarget(BlendTarget, { 1000, 500, 100, 10, 0, -10, -100 }, TEXT("Inline Duplicate"));
+		GetTypeHash(BlendTarget);
+	}
+
+	// Simple test - use inline allocation
+	{
+		FHierarchicalBlendTarget BlendTarget;
+		BlendTarget.Add(0);
+		BlendTarget.Add(100);
+		BlendTarget.Add(-100);
+		BlendTarget.Add(-10);
+		BlendTarget.Add(10);
+		BlendTarget.Add(500);
+		BlendTarget.Add(1000);
+
+		this->TestEqual(TEXT("Inline Basic Capacity"), BlendTarget.GetCapacity(), 7);
+		CompareBlendTarget(BlendTarget, { 1000, 500, 100, 10, 0, -10, -100 }, TEXT("Inline Basic"));
+		GetTypeHash(BlendTarget);
+	}
+
+	// Test growing the allocation
+	{
+		FRandomStream RandomStream;
+		RandomStream.GenerateNewSeed();
+
+		FHierarchicalBlendTarget BlendTarget;
+		TSet<int16> ExpectedSet;
+		// Add 
+		for (int16 Index = 0; Index < 25; ++Index)
+		{
+			const int16 Value = static_cast<int16>(FMath::RandRange(-10000, 10000));
+			BlendTarget.Add(Value);
+			ExpectedSet.Add(Value);
+		}
+
+		TArray<int16> Expected = ExpectedSet.Array();
+		Algo::Sort(Expected, TGreater<>());
+
+		this->TestEqual(TEXT("Basic Heap Capacity"), BlendTarget.GetCapacity(), 32);
+		CompareBlendTarget(BlendTarget, Expected, TEXT("Basic Heap"));
+		GetTypeHash(BlendTarget);
+	}
+
+	// Test growing the allocation with duplicates
+	{
+		FRandomStream RandomStream;
+		RandomStream.GenerateNewSeed();
+
+		FHierarchicalBlendTarget BlendTarget;
+		TSet<int16> ExpectedSet;
+		// Add 
+		for (int16 Index = 0; Index < 25; ++Index)
+		{
+			const int16 Value = static_cast<int16>(FMath::RandRange(-10000, 10000));
+			BlendTarget.Add(Value);
+			ExpectedSet.Add(Value);
+		}
+
+		for (int16 Expected : ExpectedSet)
+		{
+			BlendTarget.Add(Expected);
+		}
+
+		TArray<int16> Expected = ExpectedSet.Array();
+		Algo::Sort(Expected, TGreater<>());
+
+		this->TestEqual(TEXT("Duplicate Heap Capacity"), BlendTarget.GetCapacity(), 32);
+		CompareBlendTarget(BlendTarget, Expected, TEXT("Duplicate Heap"));
+		GetTypeHash(BlendTarget);
+	}
+
+	// Test move and copy operators for heap allocations
+	{
+		FRandomStream RandomStream;
+		RandomStream.GenerateNewSeed();
+
+		FHierarchicalBlendTarget BlendTarget[2];
+		TSet<int16> ExpectedSet[2];
+		// Add 
+		for (int16 Index = 0; Index < 25; ++Index)
+		{
+			const int16 Value1 = static_cast<int16>(FMath::RandRange(-10000, 10000));
+			const int16 Value2 = static_cast<int16>(FMath::RandRange(-10000, 10000));
+			BlendTarget[0].Add(Value1);
+			ExpectedSet[0].Add(Value1);
+			BlendTarget[1].Add(Value2);
+			ExpectedSet[1].Add(Value2);
+		}
+
+		TArray<int16> Expected[2] = { ExpectedSet[0].Array(), ExpectedSet[1].Array() };
+		Algo::Sort(Expected[0], TGreater<>());
+		Algo::Sort(Expected[1], TGreater<>());
+
+		// Make some copies
+		{
+			// Copy Construction
+			FHierarchicalBlendTarget BlendTargetCopy(BlendTarget[0]);
+			CompareBlendTarget(BlendTargetCopy, Expected[0], TEXT("Copy 1"));
+			CompareBlendTarget(BlendTarget[0], Expected[0], TEXT("Copied From 2"));
+			GetTypeHash(BlendTargetCopy);
+			GetTypeHash(BlendTarget[0]);
+
+			// Copy Assignment
+			BlendTargetCopy = BlendTarget[1];
+			CompareBlendTarget(BlendTargetCopy, Expected[1], TEXT("Copy 2"));
+			CompareBlendTarget(BlendTarget[1], Expected[1], TEXT("Copied From 2"));
+			GetTypeHash(BlendTargetCopy);
+			GetTypeHash(BlendTarget[1]);
+		}
+
+		// Make some moves
+		{
+			FHierarchicalBlendTarget BlendTargetCopy1(BlendTarget[0]);
+			FHierarchicalBlendTarget BlendTargetCopy2(BlendTarget[1]);
+
+			// Move Construction
+			FHierarchicalBlendTarget BlendTargetMoved(MoveTemp(BlendTargetCopy1));
+			CompareBlendTarget(BlendTargetMoved, Expected[0], TEXT("Move 1"));
+			CompareBlendTarget(BlendTargetCopy1, FHierarchicalBlendTarget().AsArray(), TEXT("Moved From 1"));
+			GetTypeHash(BlendTargetMoved);
+			GetTypeHash(BlendTargetCopy1);
+
+			// Move Assignment
+			BlendTargetMoved = MoveTemp(BlendTargetCopy2);
+			CompareBlendTarget(BlendTargetMoved, Expected[1], TEXT("Move 2"));
+			CompareBlendTarget(BlendTargetCopy2, FHierarchicalBlendTarget().AsArray(), TEXT("Moved From 2"));
+			GetTypeHash(BlendTargetMoved);
+			GetTypeHash(BlendTargetCopy2);
+		}
+	}
+
+	return true;
+}
+
 #undef LOCTEXT_NAMESPACE
 
 #endif // WITH_DEV_AUTOMATION_TESTS
diff --git a/Engine/Source/Runtime/MovieScene/Public/EntitySystem/BuiltInComponentTypes.h b/Engine/Source/Runtime/MovieScene/Public/EntitySystem/BuiltInComponentTypes.h
--- a/Engine/Source/Runtime/MovieScene/Public/EntitySystem/BuiltInComponentTypes.h
+++ b/Engine/Source/Runtime/MovieScene/Public/EntitySystem/BuiltInComponentTypes.h
@@ -326,6 +326,7 @@ private:
 	FObjectKey ObjectKey;
 };
 
+
 /**
  * Specifies a unique, sorted path of hbiases that contribute to a blended output
  * Supports up to 8 unique HBiases in its path
@@ -335,6 +336,14 @@ struct FHierarchicalBlendTarget
 	/** Default Constructor */
 	MOVIESCENE_API FHierarchicalBlendTarget();
 
+	MOVIESCENE_API FHierarchicalBlendTarget(const FHierarchicalBlendTarget& RHS);
+	MOVIESCENE_API FHierarchicalBlendTarget& operator=(const FHierarchicalBlendTarget& RHS);
+
+	MOVIESCENE_API FHierarchicalBlendTarget(FHierarchicalBlendTarget&& RHS);
+	MOVIESCENE_API FHierarchicalBlendTarget& operator=(FHierarchicalBlendTarget&& RHS);
+
+	MOVIESCENE_API ~FHierarchicalBlendTarget();
+
 	/**
 	 * Add the specified HBias to this blend target.
 	 * Duplicates are not supported - if the HBias already exists in this chain, this function does nothing.
@@ -356,33 +365,94 @@ struct FHierarchicalBlendTarget
 	 */
 	MOVIESCENE_API TArrayView<const int16> AsArray() const;
 
+	/**
+	 * Retrieve the current capacity of this container - only to be used for debugging and testing purposes
+	 */
+	uint16 GetCapacity() const
+	{
+		return Capacity;
+	}
+
 public:
 
 	friend uint32 GetTypeHash(const FHierarchicalBlendTarget& In)
 	{
-		const void* Data = In.HBiasChain;
-		const uint32* Data32 = static_cast<const uint32*>(Data);
-		static_assert(sizeof(HBiasChain) == sizeof(uint32)*4);
-
 		// Use 32 bit ints for hashing speed
-		return HashCombine(Data32[0], Data32[1]) ^ HashCombine(Data32[2], Data32[3]);
+		const void* Memory = In.GetMemory();
+		const uint32* Data32 = static_cast<const uint32*>(Memory);
+
+		if (In.Capacity == InlineCapacity)
+		{
+			static_assert(sizeof(FHierarchicalBlendTarget) == sizeof(uint32)*4);
+
+			// Use 32 bit ints for hashing speed
+			// this actually ends up hashing the capacity as well since we have 7 int16s + 1 uint16
+			return HashCombine(Data32[0], Data32[1]) ^ HashCombine(Data32[2], Data32[3]);
+		}
+		else
+		{
+			static_assert(GrowAmount%4 == 0, "GrowAmount is not a multiple of 2 which is required for this loop to work");
+			check(In.Capacity%4 == 0);
+			uint32 Hash = 0;
+			for (int32 Index = 0; Index < In.Capacity/4; Index += 4)
+			{
+				Hash ^= HashCombine(Data32[Index+0], Data32[Index+1]) ^ HashCombine(Data32[Index+2], Data32[Index+3]);
+			}
+			return Hash;
+		}
 	}
 	friend bool operator<(const FHierarchicalBlendTarget& A, const FHierarchicalBlendTarget& B)
 	{
-		return FMemory::Memcmp(A.HBiasChain, B.HBiasChain, sizeof(HBiasChain)) < 0;
+		if (A.Capacity != B.Capacity)
+		{
+			return A.Capacity < B.Capacity;
+		}
+		return FMemory::Memcmp(A.GetMemory(), B.GetMemory(), sizeof(int16)*A.Capacity) < 0;
 	}
 	friend bool operator==(const FHierarchicalBlendTarget& A, const FHierarchicalBlendTarget& B)
 	{
-		return FMemory::Memcmp(A.HBiasChain, B.HBiasChain, sizeof(HBiasChain)) == 0;
+		if (A.Capacity != B.Capacity)
+		{
+			return false;
+		}
+		return FMemory::Memcmp(A.GetMemory(), B.GetMemory(), sizeof(int16)*A.Capacity) == 0;
 	}
 	friend bool operator!=(const FHierarchicalBlendTarget& A, const FHierarchicalBlendTarget& B)
 	{
-		return FMemory::Memcmp(A.HBiasChain, B.HBiasChain, sizeof(HBiasChain)) != 0;
+		if (A.Capacity != B.Capacity)
+		{
+			return true;
+		}
+		return FMemory::Memcmp(A.GetMemory(), B.GetMemory(), sizeof(int16)*A.Capacity) != 0;
 	}
 
 private:
-	/** 16 bytes of hbias chain. Supports up to a maximum of 8 sub-sequences. */
-	int16 HBiasChain[8];
+
+	void FreeAllocation();
+
+	MOVIESCENE_API int16* GetMemory();
+	MOVIESCENE_API const int16* GetMemory() const;
+
+	TArrayView<int16> GetAllEntries();
+	TArrayView<const int16> GetAllEntries() const;
+
+	void Grow(uint16 NewCapacity);
+
+	/**
+	 * 14 bytes of hbias chain.
+	 * Type is either:
+	 *		int16[7] where Capacity == InlineCapacity
+	 *		int16*   (heap allocated) where Capacity > InlineCapacity
+	 * Supports inline up to a maximum of 7 sub-sequences, or any number on the heap allocation.
+	 * Implemented using type-erasure since a union would be padded up to 16 bytes, and we want to keep sizeof this type == 16 bytes
+	 */
+	alignas(8) uint8 Data[14];
+
+	/** 2 bytes - maximum number of hbiases in this chain. */
+	uint16 Capacity;
+
+	static constexpr int32 InlineCapacity = sizeof(Data) / sizeof(int16);
+	static constexpr uint16 GrowAmount = 16u;
 };
 
 /**
diff --git a/Engine/Source/Runtime/MovieSceneTracks/Private/Systems/MovieScenePropertyInstantiator.cpp b/Engine/Source/Runtime/MovieSceneTracks/Private/Systems/MovieScenePropertyInstantiator.cpp
--- a/Engine/Source/Runtime/MovieSceneTracks/Private/Systems/MovieScenePropertyInstantiator.cpp
+++ b/Engine/Source/Runtime/MovieSceneTracks/Private/Systems/MovieScenePropertyInstantiator.cpp
@@ -759,10 +759,13 @@ void UMovieScenePropertyInstantiatorSystem::UpdatePropertyInfo(const FPropertyPa
 		}
 		else
 		{
-			ActiveBiasEntry = IgnoredHBiasEntry;
+			ActiveBiasEntry = MoveTemp(IgnoredHBiasEntry);
 		}
 	}
 
+	// -----------------------------------
+	// NOW UNSAFE TO USE IgnoredHBiasEntry
+
 	// Reset the restore state status of the property if we still have contributors
 	// We do not do this if there are no contributors to ensure that stale properties are restored correctly
 	Params.PropertyInfo->EmptyChannels = ActiveBiasEntry.bBlendHierarchicalBias
@@ -780,7 +783,7 @@ void UMovieScenePropertyInstantiatorSystem::UpdatePropertyInfo(const FPropertyPa
 
 	Params.PropertyInfo->bIsPartiallyAnimated = bIsPartial;
 	Params.PropertyInfo->bMaxHBiasHasChanged  = Params.PropertyInfo->HierarchicalMetaData.HBias != ActiveBiasEntry.HBias;
-	Params.PropertyInfo->HierarchicalMetaData = ActiveBiasEntry;
+	Params.PropertyInfo->HierarchicalMetaData = MoveTemp(ActiveBiasEntry);
 }
 
 void UMovieScenePropertyInstantiatorSystem::InitializeFastPath(const FPropertyParameters& Params)
