Sequencer: Fix issue with possessable tracks that map to spawnables. If the spawnable despawns/spawns, the possessable track was not getting correctly invalidated, and so was failing to reconnect after the spawnable respawns.
In addition, update the label color logic so that if a possessable references a spawnable, and the spawnable is despawned that we show a 'gray' color instead of red, just like we do for regular spawnable tracks.

diff --git a/Engine/Source/Editor/Sequencer/Private/MVVM/ViewModels/ObjectBindingModel.cpp b/Engine/Source/Editor/Sequencer/Private/MVVM/ViewModels/ObjectBindingModel.cpp
--- a/Engine/Source/Editor/Sequencer/Private/MVVM/ViewModels/ObjectBindingModel.cpp
+++ b/Engine/Source/Editor/Sequencer/Private/MVVM/ViewModels/ObjectBindingModel.cpp
@@ -349,18 +349,18 @@ FSlateColor FObjectBindingModel::GetLabelColor() const
 		}
 	}
 
-	// Spawnables don't have valid object bindings when their track hasn't spawned them yet,
+	// Find the last objecting binding ancestor and ask it for the invalid color to use.
+	// e.g. Spawnables don't have valid object bindings when their track hasn't spawned them yet,
 	// so we override the default behavior of red with a gray so that users don't think there is something wrong.
-	constexpr bool bIncludeThis = true;
-	for (TSharedPtr<FObjectBindingModel> Parent : GetAncestorsOfType<FObjectBindingModel>(bIncludeThis))
-	{
-		if (Parent->GetType() == EObjectBindingType::Spawnable)
+	TFunction<FSlateColor(const FObjectBindingModel&)> GetObjectBindingAncestorInvalidLabelColor = [&](const FObjectBindingModel& InObjectBindingModel) -> FSlateColor {
+		if (TSharedPtr<FObjectBindingModel> ParentBindingModel = InObjectBindingModel.FindAncestorOfType<FObjectBindingModel>())
 		{
-			return FSlateColor::UseSubduedForeground();
+			return GetObjectBindingAncestorInvalidLabelColor(*ParentBindingModel.Get());
 		}
-	}
+		return InObjectBindingModel.GetInvalidBindingLabelColor();
+	};
 
-	return FLinearColor::Red;
+	return GetObjectBindingAncestorInvalidLabelColor(*this);
 }
 
 FText FObjectBindingModel::GetTooltipForSingleObjectBinding() const
diff --git a/Engine/Source/Editor/Sequencer/Private/MVVM/ViewModels/PossessableModel.cpp b/Engine/Source/Editor/Sequencer/Private/MVVM/ViewModels/PossessableModel.cpp
--- a/Engine/Source/Editor/Sequencer/Private/MVVM/ViewModels/PossessableModel.cpp
+++ b/Engine/Source/Editor/Sequencer/Private/MVVM/ViewModels/PossessableModel.cpp
@@ -140,6 +140,17 @@ void FPossessableModel::Delete()
 	}
 }
 
+FSlateColor FPossessableModel::GetInvalidBindingLabelColor() const
+{
+	UMovieScene* MovieScene = OwnerModel ? OwnerModel->GetMovieScene() : nullptr;
+	FMovieScenePossessable* Possessable = MovieScene ? MovieScene->FindPossessable(ObjectBindingID) : nullptr;
+	if (Possessable && Possessable->GetSpawnableObjectBindingID().IsValid())
+	{
+		return FSlateColor::UseSubduedForeground();
+	}
+	return FLinearColor::Red;
+}
+
 } // namespace Sequencer
 } // namespace UE
 
diff --git a/Engine/Source/Editor/Sequencer/Public/MVVM/ViewModels/ObjectBindingModel.h b/Engine/Source/Editor/Sequencer/Public/MVVM/ViewModels/ObjectBindingModel.h
--- a/Engine/Source/Editor/Sequencer/Public/MVVM/ViewModels/ObjectBindingModel.h
+++ b/Engine/Source/Editor/Sequencer/Public/MVVM/ViewModels/ObjectBindingModel.h
@@ -114,6 +114,7 @@ public:
 	virtual FText GetTooltipForSingleObjectBinding() const;
 	virtual const UClass* FindObjectClass() const;
 	virtual bool SupportsRebinding() const;
+	virtual FSlateColor GetInvalidBindingLabelColor() const { return FLinearColor::Red; }
 
 public:
 
diff --git a/Engine/Source/Editor/Sequencer/Public/MVVM/ViewModels/PossessableModel.h b/Engine/Source/Editor/Sequencer/Public/MVVM/ViewModels/PossessableModel.h
--- a/Engine/Source/Editor/Sequencer/Public/MVVM/ViewModels/PossessableModel.h
+++ b/Engine/Source/Editor/Sequencer/Public/MVVM/ViewModels/PossessableModel.h
@@ -43,6 +43,8 @@ public:
 
 	/*~ FObjectBindingModel interface */
 	void Delete() override;
+	FSlateColor GetInvalidBindingLabelColor() const override;
+
 };
 
 } // namespace Sequencer
diff --git a/Engine/Source/Editor/Sequencer/Public/MVVM/ViewModels/SpawnableModel.h b/Engine/Source/Editor/Sequencer/Public/MVVM/ViewModels/SpawnableModel.h
--- a/Engine/Source/Editor/Sequencer/Public/MVVM/ViewModels/SpawnableModel.h
+++ b/Engine/Source/Editor/Sequencer/Public/MVVM/ViewModels/SpawnableModel.h
@@ -43,6 +43,8 @@ public:
 
 	/*~ FObjectBindingModel interface */
 	void Delete() override;
+	FSlateColor GetInvalidBindingLabelColor() const override { return FSlateColor::UseSubduedForeground(); }
+
 };
 
 } // namespace Sequencer
diff --git a/Engine/Source/Runtime/MovieScene/Private/Evaluation/MovieSceneEvaluationState.cpp b/Engine/Source/Runtime/MovieScene/Private/Evaluation/MovieSceneEvaluationState.cpp
--- a/Engine/Source/Runtime/MovieScene/Private/Evaluation/MovieSceneEvaluationState.cpp
+++ b/Engine/Source/Runtime/MovieScene/Private/Evaluation/MovieSceneEvaluationState.cpp
@@ -239,6 +239,26 @@ void FMovieSceneObjectCache::Invalidate(const FGuid& InGuid)
 	UpdateSerialNumber();
 }
 
+void FMovieSceneObjectCache::Invalidate(const FGuid& InGuid, FMovieSceneSequenceIDRef InSequenceID)
+{
+	if (InSequenceID == SequenceID)
+	{
+		Invalidate(InGuid);
+	}
+	else
+	{
+		FMovieSceneObjectBindingID BindingID(UE::MovieScene::FFixedObjectBindingID(InGuid, InSequenceID));
+		if (FGuidArray* ReferencedGuids = ReverseMappedBindings.Find(BindingID))
+		{
+			for (FGuid ReferencedGuid : *ReferencedGuids)
+			{
+				Invalidate(ReferencedGuid);
+			}
+			ReverseMappedBindings.Remove(BindingID);
+		}
+	}
+}
+
 bool FMovieSceneObjectCache::InvalidateInternal(const FGuid& InGuid)
 {
 	// Don't manipulate the actual map structure, since this can be called from inside an iterator
@@ -266,6 +286,7 @@ void FMovieSceneObjectCache::Clear(IMovieScenePlayer& Player)
 {
 	BoundObjects.Reset();
 	ChildBindings.Reset();
+	ReverseMappedBindings.Reset();
 
 	UpdateSerialNumber();
 
@@ -390,13 +411,17 @@ void FMovieSceneObjectCache::UpdateBindings(const FGuid& InGuid, IMovieScenePlay
 
 				if (Possessable->GetSpawnableObjectBindingID().IsValid())
 				{
+					// We resolve this binding to fixed here, as we conveniently have a Player pointer already, and when being invalidated,
+					// the binding ID passed down will be relative to the root.
+					FMovieSceneObjectBindingID SpawnableFixedBindingID = Possessable->GetSpawnableObjectBindingID().ResolveToFixed(SequenceID, Player);
+					ReverseMappedBindings.FindOrAdd(SpawnableFixedBindingID).AddUnique(InGuid);
 					for (TWeakObjectPtr<> BoundObject : Possessable->GetSpawnableObjectBindingID().ResolveBoundObjects(SequenceID, Player))
 					{
 						if (BoundObject.IsValid())
 						{
 							FoundObjects.Add(BoundObject.Get());
 						}
-					}				
+					}
 				}
 				else
 				{
@@ -499,10 +524,10 @@ void FMovieSceneEvaluationState::InvalidateExpiredObjects()
 
 void FMovieSceneEvaluationState::Invalidate(const FGuid& InGuid, FMovieSceneSequenceIDRef SequenceID)
 {
-	FVersionedObjectCache* Cache = ObjectCaches.Find(SequenceID);
-	if (Cache)
+	// We need to send the invalidation method to all of the caches, as there may be other bindings in other sequences referencing this one that is being invalidated
+	for (auto& Pair : ObjectCaches)
 	{
-		Cache->ObjectCache.Invalidate(InGuid);
+		Pair.Value.ObjectCache.Invalidate(InGuid, SequenceID);
 	}
 }
 
diff --git a/Engine/Source/Runtime/MovieScene/Public/Evaluation/MovieSceneEvaluationState.h b/Engine/Source/Runtime/MovieScene/Public/Evaluation/MovieSceneEvaluationState.h
--- a/Engine/Source/Runtime/MovieScene/Public/Evaluation/MovieSceneEvaluationState.h
+++ b/Engine/Source/Runtime/MovieScene/Public/Evaluation/MovieSceneEvaluationState.h
@@ -96,12 +96,21 @@ struct FMovieSceneObjectCache
 	MOVIESCENE_API void InvalidateExpiredObjects();
 
 	/**
-	 * Invalidate the object bindings for a specific object binding ID
+	 * Invalidate the object bindings for a specific object binding ID in this sequence
 	 *
 	 * @param InGuid			The object binding ID to invalidate bindings for
 	 */
 	MOVIESCENE_API void Invalidate(const FGuid& InGuid);
 
+	/**
+	 * Invalidate the object bindings for a specific object binding ID in the specified sequence ID.
+	 * If the sequence ID matches this one, then we will look for the guid in this sequence. 
+	 * If it does not, then we will see if any of our object bindings reference that one, and if so, they will be invalidated.
+	 * 
+	 * @param InGuid			The object binding ID to invalidate bindings for
+	 */
+	MOVIESCENE_API void Invalidate(const FGuid& InGuid, FMovieSceneSequenceIDRef InSequenceID);
+
 	/**
 	 * Invalidate the object bindings for a specific object binding ID if they are not already invalidated
 	 *
@@ -207,6 +216,11 @@ private:
 	/** Map of child bindings for any given object binding */
 	typedef TArray<FGuid, TInlineAllocator<4>> FGuidArray;
 	TMap<FGuid, FGuidArray, FDefaultSetAllocator, TFastGuidKeyFuncs<FGuidArray>> ChildBindings;
+	
+	/** For possessables in this scene that map to a binding in another scene (for example to a spawnable in another scene).
+	  * Stored as reverse-lookup for speed of invalidation.
+	  */
+	TMap<FMovieSceneObjectBindingID, FGuidArray, FDefaultSetAllocator> ReverseMappedBindings;
 
 	/** Serial number for this cache */
 	uint32 SerialNumber = 0;
