Sequencer: fix upstream dependencies for the grouping system

This makes the grouping system correctly run after other systems that take care of resolving bound objects, such as the one that takes care of camera animations.

diff --git a/Engine/Source/Runtime/MovieScene/Private/EntitySystem/MovieSceneEntityGroupingSystem.cpp b/Engine/Source/Runtime/MovieScene/Private/EntitySystem/MovieSceneEntityGroupingSystem.cpp
--- a/Engine/Source/Runtime/MovieScene/Private/EntitySystem/MovieSceneEntityGroupingSystem.cpp
+++ b/Engine/Source/Runtime/MovieScene/Private/EntitySystem/MovieSceneEntityGroupingSystem.cpp
@@ -199,8 +199,7 @@ UMovieSceneEntityGroupingSystem::UMovieSceneEntityGroupingSystem(const FObjectIn
 		// This isn't something we *really* need, but pretty much all use-cases we have will want
 		// to group things using bound objects and/or bound scene components, so let's run after the
 		// systems that set those up.
-		DefineImplicitPrerequisite(UMovieSceneGenericBoundObjectInstantiator::StaticClass(), GetClass());
-		DefineImplicitPrerequisite(UMovieSceneBoundSceneComponentInstantiator::StaticClass(), GetClass());
+		DefineComponentConsumer(GetClass(), BuiltInComponents->BoundObject);
 	}
 }
 
