Sequencer: Fix setting the StartFrameOffset when the subsection is not looping

diff --git a/Engine/Source/Runtime/MovieScene/Private/Sections/MovieSceneSubSection.cpp b/Engine/Source/Runtime/MovieScene/Private/Sections/MovieSceneSubSection.cpp
--- a/Engine/Source/Runtime/MovieScene/Private/Sections/MovieSceneSubSection.cpp
+++ b/Engine/Source/Runtime/MovieScene/Private/Sections/MovieSceneSubSection.cpp
@@ -407,12 +407,20 @@ void UMovieSceneSubSection::TrimSection( FQualifiedFrameTime TrimTime, bool bTri
 		// The new first loop start offset is where the trim time fell inside the sub-sequence (this time is already
 		// normalized in the case of looping sub-sequences).
 		FFrameTime LocalTrimTime = OuterToInnerTransform().TransformTime(LocalTickResolutionTrimTime);
-		FFrameNumber NewFirstLoopStartOffset = LocalTrimTime.FrameNumber - Parameters.StartFrameOffset;
+		FFrameNumber NewStartOffset = LocalTrimTime.FrameNumber - Parameters.StartFrameOffset;
 		
 		// Make sure we don't have negative offsets (this shouldn't happen, though).
-		NewFirstLoopStartOffset = FMath::Max(FFrameNumber(0), NewFirstLoopStartOffset);
+		NewStartOffset = FMath::Max(FFrameNumber(0), NewStartOffset);
 		
-		Parameters.FirstLoopStartFrameOffset = NewFirstLoopStartOffset;
+		const bool bCanLoop = Parameters.bCanLoop;
+		if (!bCanLoop)
+		{
+			Parameters.StartFrameOffset = NewStartOffset;
+		}
+		else
+		{
+			Parameters.FirstLoopStartFrameOffset = NewStartOffset;
+		}
 	}
 
 	// Actually trim the section range!
