Sequencer: speculative fix for a rare issue with replicated sequences running an incorrect update soon after their creation

diff --git a/Engine/Source/Runtime/MovieScene/Private/MovieSceneSequencePlayer.cpp b/Engine/Source/Runtime/MovieScene/Private/MovieSceneSequencePlayer.cpp
--- a/Engine/Source/Runtime/MovieScene/Private/MovieSceneSequencePlayer.cpp
+++ b/Engine/Source/Runtime/MovieScene/Private/MovieSceneSequencePlayer.cpp
@@ -146,9 +146,9 @@ UMovieSceneSequencePlayer::~UMovieSceneSequencePlayer()
 	}
 }
 
-void UMovieSceneSequencePlayer::UpdateNetworkSyncProperties()
+void UMovieSceneSequencePlayer::UpdateNetworkSyncProperties(bool bServerOnly)
 {
-	if (HasAuthority())
+	if (!bServerOnly || HasAuthority())
 	{
 		NetSyncProps.LastKnownPosition = PlayPosition.GetCurrentPosition();
 		NetSyncProps.LastKnownStatus   = Status;
@@ -942,7 +942,14 @@ void UMovieSceneSequencePlayer::Initialize(UMovieSceneSequence* InSequence)
 	PlayPosition.Reset(StartTimeWithOffset);
 	TimeController->Reset(GetCurrentTime());
 
-	UpdateNetworkSyncProperties();
+	// Update the sync properties on the server *and* the client. We do it on the client because our first PostNetReceive
+	// could be called with an incomplete set of replicated values in very rare cases... so for instance we might
+	// get the proper LastKnownStatus from the server, but, say, not the proper LastKnownPosition. If the sequence does
+	// not start at frame 0, we would see LastKnownPosition left at 0, while our own client-side position is the first
+	// frame of the sequence, as initialized above (SetFrameRange). At this point, we would incorrectly assume that the server 
+	// jumped to frame 0 and we would do the same, when really the server hasn't moved and it's just that the correct
+	// LastKnownPosition value is coming in a later net packet.
+	UpdateNetworkSyncProperties(false);
 }
 
 void UMovieSceneSequencePlayer::Update(const float DeltaSeconds)
@@ -1187,7 +1194,7 @@ void UMovieSceneSequencePlayer::UpdateTimeCursorPosition_Internal(FFrameTime New
 		Args.bIsAsync = (bIsAsyncUpdate && !bIsSequenceBlocking);
 		Args.bHasJumped = bHasJumpedOverride;
 
-		PostEvaluationCallbacks.Add(FOnEvaluationCallback::CreateUObject(this, &UMovieSceneSequencePlayer::UpdateNetworkSyncProperties));
+		PostEvaluationCallbacks.Add(FOnEvaluationCallback::CreateUObject(this, &UMovieSceneSequencePlayer::UpdateNetworkSyncProperties, true));
 
 		UpdateMovieSceneInstance(Range, StatusOverride, Args);
 
diff --git a/Engine/Source/Runtime/MovieScene/Public/MovieSceneSequencePlayer.h b/Engine/Source/Runtime/MovieScene/Public/MovieSceneSequencePlayer.h
--- a/Engine/Source/Runtime/MovieScene/Public/MovieSceneSequencePlayer.h
+++ b/Engine/Source/Runtime/MovieScene/Public/MovieSceneSequencePlayer.h
@@ -602,8 +602,10 @@ private:
 
 	/**
 	 * Update the replicated properties required for synchronizing to clients of this sequence player
+	 *
+	 * @param bServerOnly  Only update the properties if this player is on the server
 	 */
-	MOVIESCENE_API void UpdateNetworkSyncProperties();
+	MOVIESCENE_API void UpdateNetworkSyncProperties(bool bServerOnly = true);
 
 	/**
 	 * Analyse the set of samples we have estimating the server time if we have confidence over the data.
