Sequencer: Fixed registered tick interval not being reset when player was torn down

diff --git a/Engine/Source/Runtime/MovieScene/Private/MovieSceneSequencePlayer.cpp b/Engine/Source/Runtime/MovieScene/Private/MovieSceneSequencePlayer.cpp
--- a/Engine/Source/Runtime/MovieScene/Private/MovieSceneSequencePlayer.cpp
+++ b/Engine/Source/Runtime/MovieScene/Private/MovieSceneSequencePlayer.cpp
@@ -1282,6 +1282,7 @@ void UMovieSceneSequencePlayer::TearDown()
 
 	if (TickManager)
 	{
+		RegisteredTickInterval.Reset();
 		TickManager->UnregisterTickClient(this);
 		TickManager = nullptr;
 	}
