Sequencer: move evaluation information onto FSharedPlaybackState

This is another step in the broader refactor to remove IMovieScenePlayer from the core evaluation code. This step moves some of the main evaluation information from the player and its root evaluation template instance, to the shared playback state.

With this knowledge, FSharedPlaybackState, FSequenceInstance, and a few other classes have access to basic things like the linker, runner, and more, so this change removes the need to pass these pointers into many methods. We do however still need to pass the player into a few places because the creation and initialization of sequence instances needs to execute some binding-related logic in a careful order. This will be further cleaned-up in a follow-up change when we move object caches and binidng overrides out of the player class.

diff --git a/Engine/Plugins/Animation/ControlRig/Source/ControlRigEditor/Private/ControlRigSequencerEditorLibrary.cpp b/Engine/Plugins/Animation/ControlRig/Source/ControlRigEditor/Private/ControlRigSequencerEditorLibrary.cpp
--- a/Engine/Plugins/Animation/ControlRig/Source/ControlRigEditor/Private/ControlRigSequencerEditorLibrary.cpp
+++ b/Engine/Plugins/Animation/ControlRig/Source/ControlRigEditor/Private/ControlRigSequencerEditorLibrary.cpp
@@ -1385,7 +1385,7 @@ static bool LocalGetControlRigControlValues(IMovieScenePlayer* Player, UMovieSce
 			GlobalTime = GlobalTime * RootToLocalTransform.InverseNoLooping();
 			FMovieSceneContext Context = FMovieSceneContext(FMovieSceneEvaluationRange(GlobalTime, TickResolution), Player->GetPlaybackStatus()).SetHasJumped(true);
 
-			Player->GetEvaluationTemplate().EvaluateSynchronousBlocking(Context, *Player);
+			Player->GetEvaluationTemplate().EvaluateSynchronousBlocking(Context);
 			ControlRig->Evaluate_AnyThread();
 			OutValues[Index] = ControlRig->GetControlValue(ControlName);
 		}
diff --git a/Engine/Plugins/Animation/ControlRig/Source/ControlRigEditor/Private/Sequencer/ControlRigParameterTrackEditor.cpp b/Engine/Plugins/Animation/ControlRig/Source/ControlRigEditor/Private/Sequencer/ControlRigParameterTrackEditor.cpp
--- a/Engine/Plugins/Animation/ControlRig/Source/ControlRigEditor/Private/Sequencer/ControlRigParameterTrackEditor.cpp
+++ b/Engine/Plugins/Animation/ControlRig/Source/ControlRigEditor/Private/Sequencer/ControlRigParameterTrackEditor.cpp
@@ -4673,7 +4673,7 @@ bool FControlRigParameterTrackEditor::CollapseAllLayers(TSharedPtr<ISequencer>&S
 
 						FMovieSceneContext Context = FMovieSceneContext(FMovieSceneEvaluationRange(GlobalTime, TickResolution), SequencerPtr->GetPlaybackStatus()).SetHasJumped(true);
 
-						SequencerPtr->GetEvaluationTemplate().EvaluateSynchronousBlocking(Context, *SequencerPtr);
+						SequencerPtr->GetEvaluationTemplate().EvaluateSynchronousBlocking(Context);
 						ControlRig->Evaluate_AnyThread();
 						for (TPair<FName, TArray<FTransform>>& TrailControlTransform : ControlLocalTransforms)
 						{
diff --git a/Engine/Plugins/Animation/ControlRig/Source/ControlRigEditor/Private/Sequencer/ControlRigSpaceChannelEditors.cpp b/Engine/Plugins/Animation/ControlRig/Source/ControlRigEditor/Private/Sequencer/ControlRigSpaceChannelEditors.cpp
--- a/Engine/Plugins/Animation/ControlRig/Source/ControlRigEditor/Private/Sequencer/ControlRigSpaceChannelEditors.cpp
+++ b/Engine/Plugins/Animation/ControlRig/Source/ControlRigEditor/Private/Sequencer/ControlRigSpaceChannelEditors.cpp
@@ -449,7 +449,7 @@ FKeyHandle FControlRigSpaceChannelHelpers::SequencerKeyControlRigSpaceChannel(UC
 			GlobalTime = GlobalTime * RootToLocalTransform.InverseNoLooping();
 
 			FMovieSceneContext SceneContext = FMovieSceneContext(FMovieSceneEvaluationRange(GlobalTime, TickResolution), Sequencer->GetPlaybackStatus()).SetHasJumped(true);
-			Sequencer->GetEvaluationTemplate().EvaluateSynchronousBlocking(SceneContext, *Sequencer);
+			Sequencer->GetEvaluationTemplate().EvaluateSynchronousBlocking(SceneContext);
 
 			Context.LocalTime = TickResolution.AsSeconds(FFrameTime(Time - 1));
 			ControlRig->SetControlGlobalTransform(ControlKey.Name, ControlWorldTransforms[0], true, Context, false /*undo*/, false /*bPrintPython*/, true/* bFixEulerFlips*/);
@@ -479,7 +479,7 @@ FKeyHandle FControlRigSpaceChannelHelpers::SequencerKeyControlRigSpaceChannel(UC
 			GlobalTime = GlobalTime * RootToLocalTransform.InverseNoLooping();
 
 			FMovieSceneContext SceneContext = FMovieSceneContext(FMovieSceneEvaluationRange(GlobalTime, TickResolution), Sequencer->GetPlaybackStatus()).SetHasJumped(true);
-			Sequencer->GetEvaluationTemplate().EvaluateSynchronousBlocking(SceneContext, *Sequencer);
+			Sequencer->GetEvaluationTemplate().EvaluateSynchronousBlocking(SceneContext);
 
 			ControlRig->Evaluate_AnyThread();
 			Context.LocalTime = TickResolution.AsSeconds(FFrameTime(Frame));
@@ -689,7 +689,7 @@ void  FControlRigSpaceChannelHelpers::SequencerSpaceChannelKeyDeleted(UControlRi
 			GlobalTime = GlobalTime * RootToLocalTransform.InverseNoLooping();
 
 			FMovieSceneContext SceneContext = FMovieSceneContext(FMovieSceneEvaluationRange(GlobalTime, TickResolution), Sequencer->GetPlaybackStatus()).SetHasJumped(true);
-			Sequencer->GetEvaluationTemplate().EvaluateSynchronousBlocking(SceneContext, *Sequencer);
+			Sequencer->GetEvaluationTemplate().EvaluateSynchronousBlocking(SceneContext);
 			//make sure to set rig hierarchy correct since key is not deleted yet
 			switch (PreviousValue.SpaceType)
 			{
@@ -1089,7 +1089,7 @@ void FControlRigSpaceChannelHelpers::SequencerBakeControlInSpace(UControlRig* Co
 				GlobalTime = GlobalTime * RootToLocalTransform.InverseNoLooping();
 
 				FMovieSceneContext SceneContext = FMovieSceneContext(FMovieSceneEvaluationRange(GlobalTime, TickResolution), Sequencer->GetPlaybackStatus()).SetHasJumped(true);
-				Sequencer->GetEvaluationTemplate().EvaluateSynchronousBlocking(SceneContext, *Sequencer);
+				Sequencer->GetEvaluationTemplate().EvaluateSynchronousBlocking(SceneContext);
 
 				//evaluate control rig
 				ControlRig->Evaluate_AnyThread();
@@ -1129,7 +1129,7 @@ void FControlRigSpaceChannelHelpers::SequencerBakeControlInSpace(UControlRig* Co
 				GlobalTime = GlobalTime * RootToLocalTransform.InverseNoLooping();
 
 				FMovieSceneContext SceneContext = FMovieSceneContext(FMovieSceneEvaluationRange(GlobalTime, TickResolution), Sequencer->GetPlaybackStatus()).SetHasJumped(true);
-				Sequencer->GetEvaluationTemplate().EvaluateSynchronousBlocking(SceneContext, *Sequencer);
+				Sequencer->GetEvaluationTemplate().EvaluateSynchronousBlocking(SceneContext);
 		
 				//evaluate control rig
 				ControlRig->Evaluate_AnyThread();
diff --git a/Engine/Plugins/Animation/ControlRig/Source/ControlRigEditor/Private/Tools/ControlRigSnapper.cpp b/Engine/Plugins/Animation/ControlRig/Source/ControlRigEditor/Private/Tools/ControlRigSnapper.cpp
--- a/Engine/Plugins/Animation/ControlRig/Source/ControlRigEditor/Private/Tools/ControlRigSnapper.cpp
+++ b/Engine/Plugins/Animation/ControlRig/Source/ControlRigEditor/Private/Tools/ControlRigSnapper.cpp
@@ -144,7 +144,7 @@ static bool LocalGetControlRigControlTransforms(IMovieScenePlayer* Player, const
 				FMovieSceneContext Context = FMovieSceneContext(FMovieSceneEvaluationRange(GlobalTime, TickResolution), Player->GetPlaybackStatus()).SetHasJumped(true);
 
 				DeltaTime = 1.0/Context.GetFrameRate().AsDecimal();
-				Player->GetEvaluationTemplate().EvaluateSynchronousBlocking(Context, *Player);
+				Player->GetEvaluationTemplate().EvaluateSynchronousBlocking(Context);
 			}
 			if (ControlRig->IsAdditive())
 			{
diff --git a/Engine/Plugins/Editor/SequencerAnimTools/Source/SequencerAnimTools/Private/Sequencer/MovieSceneTransformTrail.cpp b/Engine/Plugins/Editor/SequencerAnimTools/Source/SequencerAnimTools/Private/Sequencer/MovieSceneTransformTrail.cpp
--- a/Engine/Plugins/Editor/SequencerAnimTools/Source/SequencerAnimTools/Private/Sequencer/MovieSceneTransformTrail.cpp
+++ b/Engine/Plugins/Editor/SequencerAnimTools/Source/SequencerAnimTools/Private/Sequencer/MovieSceneTransformTrail.cpp
@@ -755,7 +755,7 @@ bool FMovieSceneControlRigTransformTrail::ApplyDelta(const FVector& Pos, const F
 				NewTransform.SetLocation(NewTransform.GetLocation() + Pos);
 				KeyInfo->Transform = NewTransform;
 				NewTransform = NewTransform.GetRelativeTransform(KeyInfo->ParentTransform);
-				GetSequencer()->GetEvaluationTemplate().EvaluateSynchronousBlocking(MovieSceneContext, *Player);
+				GetSequencer()->GetEvaluationTemplate().EvaluateSynchronousBlocking(MovieSceneContext);
 				ControlRig->Evaluate_AnyThread();
 				ControlRig->SetControlGlobalTransform(ControlName, NewTransform, true, Context, false /*undo*/, false /*bPrintPython*/, true/* bFixEulerFlips*/);
 			}
@@ -810,7 +810,7 @@ bool FMovieSceneControlRigTransformTrail::EndTracking()
 
 		FMovieSceneContext MovieSceneContext = FMovieSceneContext(FMovieSceneEvaluationRange(StartTime, TickResolution), Player->GetPlaybackStatus()).SetHasJumped(true);
 		
-		Player->GetEvaluationTemplate().EvaluateSynchronousBlocking(MovieSceneContext, *Player);
+		Player->GetEvaluationTemplate().EvaluateSynchronousBlocking(MovieSceneContext);
 		ControlRig->Evaluate_AnyThread();
 		return true;
 	}
@@ -850,7 +850,7 @@ bool FMovieSceneControlRigTransformTrail::HandleAltClick(FEditorViewportClient*
 	Context.LocalTime = TickResolution.AsSeconds(GlobalTime);
 	Context.KeyMask = (uint32)EControlRigContextChannelToKey::Translation;
 	FMovieSceneContext MovieSceneContext = FMovieSceneContext(FMovieSceneEvaluationRange(GlobalTime, TickResolution), Player->GetPlaybackStatus()).SetHasJumped(true);
-	GetSequencer()->GetEvaluationTemplate().EvaluateSynchronousBlocking(MovieSceneContext, *Player);
+	GetSequencer()->GetEvaluationTemplate().EvaluateSynchronousBlocking(MovieSceneContext);
 	ControlRig->Evaluate_AnyThread();
 	FTransform NewTransform(ControlRig->GetControlGlobalTransform(ControlName));
 	ControlRig->SetControlGlobalTransform(ControlName, NewTransform, true, Context, false /*undo*/, false /*bPrintPython*/, true/* bFixEulerFlips*/);
@@ -860,7 +860,7 @@ bool FMovieSceneControlRigTransformTrail::HandleAltClick(FEditorViewportClient*
 	StartTime = StartTime * RootToLocalTransform.InverseNoLooping(); //player evals in root time so need to go back to it.
 
 	MovieSceneContext = FMovieSceneContext(FMovieSceneEvaluationRange(StartTime, TickResolution), Player->GetPlaybackStatus()).SetHasJumped(true);
-	Player->GetEvaluationTemplate().EvaluateSynchronousBlocking(MovieSceneContext, *Player);
+	Player->GetEvaluationTemplate().EvaluateSynchronousBlocking(MovieSceneContext);
 	ControlRig->Evaluate_AnyThread();
 
 	//create new keys
diff --git a/Engine/Plugins/Editor/SequencerAnimTools/Source/SequencerAnimTools/Private/Sequencer/SequencerTrailHierarchy.cpp b/Engine/Plugins/Editor/SequencerAnimTools/Source/SequencerAnimTools/Private/Sequencer/SequencerTrailHierarchy.cpp
--- a/Engine/Plugins/Editor/SequencerAnimTools/Source/SequencerAnimTools/Private/Sequencer/SequencerTrailHierarchy.cpp
+++ b/Engine/Plugins/Editor/SequencerAnimTools/Source/SequencerAnimTools/Private/Sequencer/SequencerTrailHierarchy.cpp
@@ -260,7 +260,7 @@ void FSequencerTrailHierarchy::UpdateControlRig(const TArray<FFrameNumber> &Fram
 
 					FMovieSceneContext Context = FMovieSceneContext(FMovieSceneEvaluationRange(GlobalTime, TickResolution), Player->GetPlaybackStatus()).SetHasJumped(true);
 
-					Player->GetEvaluationTemplate().EvaluateSynchronousBlocking(Context, *Player);
+					Player->GetEvaluationTemplate().EvaluateSynchronousBlocking(Context);
 					ControlRig->Evaluate_AnyThread();
 					for (FTrailControlTransforms& TrailControlTransform : TrailControlTransforms)
 					{
@@ -293,7 +293,7 @@ void FSequencerTrailHierarchy::UpdateControlRig(const TArray<FFrameNumber> &Fram
 			FFrameTime StartTime = Sequencer->GetLocalTime().Time;
 			StartTime = StartTime * RootToLocalTransform.InverseNoLooping(); //player evals in root time so need to go back to it.
 			FMovieSceneContext Context = FMovieSceneContext(FMovieSceneEvaluationRange(StartTime, TickResolution), Player->GetPlaybackStatus()).SetHasJumped(true);
-			Player->GetEvaluationTemplate().EvaluateSynchronousBlocking(Context, *Player);
+			Player->GetEvaluationTemplate().EvaluateSynchronousBlocking(Context);
 			ControlRig->Evaluate_AnyThread();
 		}
 	}
diff --git a/Engine/Plugins/Importers/USDImporter/Source/USDUtilities/Private/USDSkeletalDataConversion.cpp b/Engine/Plugins/Importers/USDImporter/Source/USDUtilities/Private/USDSkeletalDataConversion.cpp
--- a/Engine/Plugins/Importers/USDImporter/Source/USDUtilities/Private/USDSkeletalDataConversion.cpp
+++ b/Engine/Plugins/Importers/USDImporter/Source/USDUtilities/Private/USDSkeletalDataConversion.cpp
@@ -3587,7 +3587,7 @@ bool UnrealToUsd::ConvertControlRigSection(
 			InPlayer->GetPlaybackStatus()
 		).SetHasJumped( true );
 
-		InPlayer->GetEvaluationTemplate().EvaluateSynchronousBlocking( Context, *InPlayer );
+		InPlayer->GetEvaluationTemplate().EvaluateSynchronousBlocking( Context );
 		ControlRig->Evaluate_AnyThread();
 
 		URigHierarchy* Hierarchy = ControlRig->GetHierarchy();
diff --git a/Engine/Plugins/MovieScene/ReplayTracks/Source/ReplayTracks/Private/Systems/MovieSceneReplaySystem.cpp b/Engine/Plugins/MovieScene/ReplayTracks/Source/ReplayTracks/Private/Systems/MovieSceneReplaySystem.cpp
--- a/Engine/Plugins/MovieScene/ReplayTracks/Source/ReplayTracks/Private/Systems/MovieSceneReplaySystem.cpp
+++ b/Engine/Plugins/MovieScene/ReplayTracks/Source/ReplayTracks/Private/Systems/MovieSceneReplaySystem.cpp
@@ -442,7 +442,7 @@ void UMovieSceneReplaySystem::OnPostLoadMap(UWorld* World, IMovieScenePlayer* La
 				World->GetTimerManager().ClearTimer(ReEvaluateHandle);
 
 				FMovieSceneRootEvaluationTemplateInstance& RootEvalTemplate = LastPlayer->GetEvaluationTemplate();
-				RootEvalTemplate.EvaluateSynchronousBlocking(LastContext, *LastPlayer);
+				RootEvalTemplate.EvaluateSynchronousBlocking(LastContext);
 			}
 		},
 		0.1f,
diff --git a/Engine/Source/Editor/MovieSceneTools/Private/Constraints/ComponentConstraintChannelInterface.cpp b/Engine/Source/Editor/MovieSceneTools/Private/Constraints/ComponentConstraintChannelInterface.cpp
--- a/Engine/Source/Editor/MovieSceneTools/Private/Constraints/ComponentConstraintChannelInterface.cpp
+++ b/Engine/Source/Editor/MovieSceneTools/Private/Constraints/ComponentConstraintChannelInterface.cpp
@@ -390,7 +390,7 @@ void FComponentConstraintChannelInterface::RecomposeTransforms(
 			const FMovieSceneEvaluationRange EvaluationRange = FMovieSceneEvaluationRange(FFrameTime(FrameNumber), TickResolution);
 			const FMovieSceneContext Context = FMovieSceneContext(EvaluationRange, PlaybackStatus).SetHasJumped(true);
 
-			EvaluationTemplate.EvaluateSynchronousBlocking(Context, *InSequencer);
+			EvaluationTemplate.EvaluateSynchronousBlocking(Context);
 
 			if (EntityIDs.IsEmpty())
 			{
@@ -445,4 +445,4 @@ void FComponentConstraintChannelInterface::RecomposeTransforms(
 				FVector(CurrentTransformChannels[6], CurrentTransformChannels[7], CurrentTransformChannels[8]));
 		}
 	}
-}
\ No newline at end of file
+}
diff --git a/Engine/Source/Editor/MovieSceneTools/Private/Constraints/MovieSceneConstraintChannelHelper.cpp b/Engine/Source/Editor/MovieSceneTools/Private/Constraints/MovieSceneConstraintChannelHelper.cpp
--- a/Engine/Source/Editor/MovieSceneTools/Private/Constraints/MovieSceneConstraintChannelHelper.cpp
+++ b/Engine/Source/Editor/MovieSceneTools/Private/Constraints/MovieSceneConstraintChannelHelper.cpp
@@ -105,7 +105,7 @@ void FCompensationEvaluator::ComputeLocalTransforms(
 				BakeHelper->PreEvaluation(MovieScene, FrameNumber);
 			}
 		}
-		InSequencer->GetEvaluationTemplate().EvaluateSynchronousBlocking(Context, *InSequencer);
+		InSequencer->GetEvaluationTemplate().EvaluateSynchronousBlocking(Context);
 
 		// evaluate constraints
 		for (const TWeakObjectPtr<UTickableConstraint>& InConstraint : AllConstraints)
@@ -265,7 +265,7 @@ void FCompensationEvaluator::ComputeLocalTransformsForBaking(UWorld* InWorld, co
 				BakeHelper->PreEvaluation(MovieScene, FrameNumber);
 			}
 		}
-		InSequencer->GetEvaluationTemplate().EvaluateSynchronousBlocking(Context, *InSequencer);
+		InSequencer->GetEvaluationTemplate().EvaluateSynchronousBlocking(Context);
 
 		// evaluate constraints
 		for (const TWeakObjectPtr<UTickableConstraint>& InConstraint : AllConstraints)
@@ -383,7 +383,7 @@ void FCompensationEvaluator::ComputeLocalTransformsBeforeDeletion(
 				BakeHelper->PreEvaluation(MovieScene, FrameNumber);
 			}
 		}
-		InSequencer->GetEvaluationTemplate().EvaluateSynchronousBlocking(Context, *InSequencer);
+		InSequencer->GetEvaluationTemplate().EvaluateSynchronousBlocking(Context);
 
 		// evaluate constraints
 		for (const TWeakObjectPtr<UTickableConstraint>& InConstraint : AllConstraints)
@@ -486,7 +486,7 @@ void FCompensationEvaluator::ComputeCompensation(UWorld* InWorld, const TSharedP
 				BakeHelper->PreEvaluation(MovieScene, InFrame);
 			}
 		}
-		InSequencer->GetEvaluationTemplate().EvaluateSynchronousBlocking(Context0, *InSequencer);
+		InSequencer->GetEvaluationTemplate().EvaluateSynchronousBlocking(Context0);
 
 		for (const TWeakObjectPtr<UTickableConstraint>& InConstraint : AllConstraints)
 		{
@@ -590,7 +590,7 @@ void FCompensationEvaluator::CacheTransforms(UWorld* InWorld, const TSharedPtr<I
 				BakeHelper->PreEvaluation(MovieScene, InFrame);
 			}
 		}
-		InSequencer->GetEvaluationTemplate().EvaluateSynchronousBlocking(Context, *InSequencer);
+		InSequencer->GetEvaluationTemplate().EvaluateSynchronousBlocking(Context);
 
 		// evaluate constraints
 		for (const TWeakObjectPtr<UTickableConstraint>& InConstraint : AllConstraints)
@@ -1363,4 +1363,4 @@ void FMovieSceneConstraintChannelHelper::HandleOffsetChanged(
 		// force evaluation so that new local values are evaluated before the constraint 
 		InSequencer->ForceEvaluate();
 	}
-}
\ No newline at end of file
+}
diff --git a/Engine/Source/Editor/MovieSceneTools/Private/MovieSceneToolHelpers.cpp b/Engine/Source/Editor/MovieSceneTools/Private/MovieSceneToolHelpers.cpp
--- a/Engine/Source/Editor/MovieSceneTools/Private/MovieSceneToolHelpers.cpp
+++ b/Engine/Source/Editor/MovieSceneTools/Private/MovieSceneToolHelpers.cpp
@@ -4257,9 +4257,9 @@ static void GetSequencerActorWorldTransforms(IMovieScenePlayer* Player, UMovieSc
 					FMovieSceneContext Context = FMovieSceneContext(FMovieSceneEvaluationRange(GlobalTime, TickResolution), Player->GetPlaybackStatus()).SetHasJumped(true);
 					if (Index == 0) // similar with baking first time in we need to evaluate twice (think due to double buffering that happens with skel mesh components).
 					{
-						Player->GetEvaluationTemplate().EvaluateSynchronousBlocking(Context, *Player);
+						Player->GetEvaluationTemplate().EvaluateSynchronousBlocking(Context);
 					}
-					Player->GetEvaluationTemplate().EvaluateSynchronousBlocking(Context, *Player);
+					Player->GetEvaluationTemplate().EvaluateSynchronousBlocking(Context);
 
 
 					const FConstraintsManagerController& Controller = FConstraintsManagerController::Get(Actor->GetWorld());
diff --git a/Engine/Source/Editor/Sequencer/Private/SSequencer.cpp b/Engine/Source/Editor/Sequencer/Private/SSequencer.cpp
--- a/Engine/Source/Editor/Sequencer/Private/SSequencer.cpp
+++ b/Engine/Source/Editor/Sequencer/Private/SSequencer.cpp
@@ -2041,7 +2041,7 @@ void SSequencer::FillAdvancedMenu(FMenuBuilder& MenuBuilder)
 				UMovieSceneSequence* RootSequence = Template.GetRootSequence();
 
 				// Set the new emulation mode
-				Template.SetEmulatedNetworkMask(InMode, *SequencerPin);
+				Template.SetEmulatedNetworkMask(InMode);
 				// Since sequencer owns its own compiled data manager, it's ok to override the mask here and reset everything
 				Template.GetCompiledDataManager()->SetEmulatedNetworkMask(InMode);
 				// Reinitialize the template again
diff --git a/Engine/Source/Editor/Sequencer/Private/Sequencer.cpp b/Engine/Source/Editor/Sequencer/Private/Sequencer.cpp
--- a/Engine/Source/Editor/Sequencer/Private/Sequencer.cpp
+++ b/Engine/Source/Editor/Sequencer/Private/Sequencer.cpp
@@ -720,6 +720,8 @@ FSequencer::~FSequencer()
 	SequencerWidget.Reset();
 
 	TrackEditors.Empty();
+
+	RootTemplateInstance.TearDown();
 }
 
 
@@ -1213,7 +1215,7 @@ void FSequencer::FocusSequenceInstance(UMovieSceneSubSection& InSubSection)
 		RestorePreAnimatedState();
 
 		UMovieSceneEntitySystemLinker* Linker = RootTemplateInstance.GetEntitySystemLinker();
-		RootTemplateInstance.FindInstance(MovieSceneSequenceID::Root)->OverrideRootSequence(Linker, ActiveTemplateIDs.Top());
+		RootTemplateInstance.FindInstance(MovieSceneSequenceID::Root)->OverrideRootSequence(ActiveTemplateIDs.Top());
 	}
 
 	UpdateSubSequenceData();
@@ -1379,7 +1381,7 @@ void FSequencer::PopToSequenceInstance(FMovieSceneSequenceIDRef SequenceID)
 		if (Settings->ShouldEvaluateSubSequencesInIsolation())
 		{
 			UMovieSceneEntitySystemLinker* Linker = RootTemplateInstance.GetEntitySystemLinker();
-			RootTemplateInstance.FindInstance(MovieSceneSequenceID::Root)->OverrideRootSequence(Linker, ActiveTemplateIDs.Top());
+			RootTemplateInstance.FindInstance(MovieSceneSequenceID::Root)->OverrideRootSequence(ActiveTemplateIDs.Top());
 		}
 
 		UpdateSequencerCustomizations(PreviousFocusedSequence);
@@ -3231,7 +3233,7 @@ void FSequencer::EvaluateInternal(FMovieSceneEvaluationRange InRange, bool bHasJ
 	Context.SetHasJumped(bHasJumped);
 
 	
-	RootTemplateInstance.EvaluateSynchronousBlocking(Context, *this);
+	RootTemplateInstance.EvaluateSynchronousBlocking(Context);
 	SuppressAutoEvalSignature.Reset();
 
 	if (Settings->ShouldRerunConstructionScripts())
@@ -10843,7 +10845,7 @@ void FSequencer::BindCommands()
 
 			FMovieSceneSequenceID NewOverrideRoot = bNewValue ? ActiveTemplateIDs.Top() : MovieSceneSequenceID::Root;
 			UMovieSceneEntitySystemLinker* Linker = RootTemplateInstance.GetEntitySystemLinker();
-			RootTemplateInstance.FindInstance(MovieSceneSequenceID::Root)->OverrideRootSequence(Linker, NewOverrideRoot);
+			RootTemplateInstance.FindInstance(MovieSceneSequenceID::Root)->OverrideRootSequence(NewOverrideRoot);
 
 			ForceEvaluate();
 		} ),
diff --git a/Engine/Source/Editor/UnrealEd/Private/Fbx/FbxMainExport.cpp b/Engine/Source/Editor/UnrealEd/Private/Fbx/FbxMainExport.cpp
--- a/Engine/Source/Editor/UnrealEd/Private/Fbx/FbxMainExport.cpp
+++ b/Engine/Source/Editor/UnrealEd/Private/Fbx/FbxMainExport.cpp
@@ -1733,7 +1733,7 @@ void FLevelSequenceAnimTrackAdapter::UpdateAnimation( int32 LocalFrame )
 
 	FMovieSceneContext Context = FMovieSceneContext(FMovieSceneEvaluationRange(GlobalTime, TickResolution), MovieScenePlayer->GetPlaybackStatus()).SetHasJumped(true);
 
-	MovieScenePlayer->GetEvaluationTemplate().EvaluateSynchronousBlocking( Context, *MovieScenePlayer );
+	MovieScenePlayer->GetEvaluationTemplate().EvaluateSynchronousBlocking( Context );
 }
 
 float FLevelSequenceAnimTrackAdapter::GetFrameRate() const
diff --git a/Engine/Source/Runtime/MovieScene/Private/Compilation/MovieSceneCompiledVolatilityManager.cpp b/Engine/Source/Runtime/MovieScene/Private/Compilation/MovieSceneCompiledVolatilityManager.cpp
--- a/Engine/Source/Runtime/MovieScene/Private/Compilation/MovieSceneCompiledVolatilityManager.cpp
+++ b/Engine/Source/Runtime/MovieScene/Private/Compilation/MovieSceneCompiledVolatilityManager.cpp
@@ -2,6 +2,7 @@
 
 #include "Compilation/MovieSceneCompiledVolatilityManager.h"
 #include "Compilation/MovieSceneCompiledDataManager.h"
+#include "EntitySystem/MovieSceneSharedPlaybackState.h"
 #include "MovieSceneFwd.h"
 
 namespace UE
@@ -32,8 +33,11 @@ FORCEINLINE EMovieSceneSequenceFlags GetEditorVolatilityFlags()
 }
 
 
-TUniquePtr<FCompiledDataVolatilityManager> FCompiledDataVolatilityManager::Construct(IMovieScenePlayer& Player, FMovieSceneCompiledDataID RootDataID, UMovieSceneCompiledDataManager* CompiledDataManager)
+TUniquePtr<FCompiledDataVolatilityManager> FCompiledDataVolatilityManager::Construct(TSharedRef<const FSharedPlaybackState> SharedPlaybackState)
 {
+	FMovieSceneCompiledDataID RootDataID = SharedPlaybackState->GetRootCompiledDataID();
+	UMovieSceneCompiledDataManager* CompiledDataManager = SharedPlaybackState->GetCompiledDataManager();
+
 	const FMovieSceneCompiledDataEntry& Entry = CompiledDataManager->GetEntryRef(RootDataID);
 	EMovieSceneSequenceFlags SequenceFlags = Entry.AccumulatedFlags | GetEditorVolatilityFlags();
 	if (!EnumHasAnyFlags(SequenceFlags, EMovieSceneSequenceFlags::Volatile))
@@ -47,14 +51,23 @@ TUniquePtr<FCompiledDataVolatilityManager> FCompiledDataVolatilityManager::Const
 		}
 	}
 
-	TUniquePtr<FCompiledDataVolatilityManager> VolatilityManager = MakeUnique<FCompiledDataVolatilityManager>();
-	VolatilityManager->ConditionalRecompile(Player, RootDataID, CompiledDataManager);
+	TUniquePtr<FCompiledDataVolatilityManager> VolatilityManager = MakeUnique<FCompiledDataVolatilityManager>(SharedPlaybackState);
+	VolatilityManager->ConditionalRecompile();
 	return VolatilityManager;
 }
 
-bool FCompiledDataVolatilityManager::HasBeenRecompiled(FMovieSceneCompiledDataID RootDataID, UMovieSceneCompiledDataManager* CompiledDataManager) const
+FCompiledDataVolatilityManager::FCompiledDataVolatilityManager(TSharedRef<const FSharedPlaybackState> SharedPlaybackState)
+	: WeakSharedPlaybackState(SharedPlaybackState)
+{
+}
+
+bool FCompiledDataVolatilityManager::HasBeenRecompiled() const
 {
-	if (HasSequenceBeenRecompiled(RootDataID, MovieSceneSequenceID::Root, CompiledDataManager))
+	TSharedPtr<const FSharedPlaybackState> SharedPlaybackState = WeakSharedPlaybackState.Pin();
+	FMovieSceneCompiledDataID RootDataID = SharedPlaybackState->GetRootCompiledDataID();
+	UMovieSceneCompiledDataManager* CompiledDataManager = SharedPlaybackState->GetCompiledDataManager();
+
+	if (HasSequenceBeenRecompiled(RootDataID, MovieSceneSequenceID::Root))
 	{
 		return true;
 	}
@@ -64,7 +77,7 @@ bool FCompiledDataVolatilityManager::HasBeenRecompiled(FMovieSceneCompiledDataID
 		for (const TTuple<FMovieSceneSequenceID, FMovieSceneSubSequenceData>& Pair : Hierarchy->AllSubSequenceData())
 		{
 			FMovieSceneCompiledDataID SubDataID = CompiledDataManager->GetSubDataID(RootDataID, Pair.Key);
-			if (SubDataID.IsValid() && HasSequenceBeenRecompiled(SubDataID, Pair.Key, CompiledDataManager))
+			if (SubDataID.IsValid() && HasSequenceBeenRecompiled(SubDataID, Pair.Key))
 			{
 				return true;
 			}
@@ -74,18 +87,23 @@ bool FCompiledDataVolatilityManager::HasBeenRecompiled(FMovieSceneCompiledDataID
 	return false;
 }
 
-bool FCompiledDataVolatilityManager::HasSequenceBeenRecompiled(FMovieSceneCompiledDataID DataID, FMovieSceneSequenceID SequenceID, UMovieSceneCompiledDataManager* CompiledDataManager) const
+bool FCompiledDataVolatilityManager::HasSequenceBeenRecompiled(FMovieSceneCompiledDataID DataID, FMovieSceneSequenceID SequenceID) const
 {
 	const FGuid* CachedSignature = CachedCompilationSignatures.Find(SequenceID);
 
+	TSharedPtr<const FSharedPlaybackState> SharedPlaybackState = WeakSharedPlaybackState.Pin();
+	UMovieSceneCompiledDataManager* CompiledDataManager = SharedPlaybackState->GetCompiledDataManager();
 	const FMovieSceneCompiledDataEntry& CompiledEntry = CompiledDataManager->GetEntryRef(DataID);
 	return CachedSignature == nullptr || *CachedSignature != CompiledEntry.CompiledSignature;
 }
 
-bool FCompiledDataVolatilityManager::ConditionalRecompile(IMovieScenePlayer& Player, FMovieSceneCompiledDataID RootDataID, UMovieSceneCompiledDataManager* CompiledDataManager)
+bool FCompiledDataVolatilityManager::ConditionalRecompile()
 {
 	bool bRecompiled = false;
 
+	TSharedPtr<const FSharedPlaybackState> SharedPlaybackState = WeakSharedPlaybackState.Pin();
+	FMovieSceneCompiledDataID RootDataID = SharedPlaybackState->GetRootCompiledDataID();
+	UMovieSceneCompiledDataManager* CompiledDataManager = SharedPlaybackState->GetCompiledDataManager();
 	if (CompiledDataManager->IsDirty(RootDataID))
 	{
 		// We override the network mask from the compiled data manager here simply because it may not be correct.
@@ -96,7 +114,7 @@ bool FCompiledDataVolatilityManager::ConditionalRecompile(IMovieScenePlayer& Pla
 		// and so these Sequences are marked as volatile on compile time. Therefore, it's here, upon conditional recompile, when we need to know the correct
 		// network mask to use, which we override and apply here.
 		EMovieSceneServerClientMask NetworkMask = CompiledDataManager->GetNetworkMask();
-		UObject* PlaybackContext = Player.GetPlaybackContext();
+		UObject* PlaybackContext = SharedPlaybackState->GetPlaybackContext();
 		UWorld* World = PlaybackContext ? PlaybackContext->GetWorld() : nullptr;
 
 		if (World)
@@ -116,29 +134,34 @@ bool FCompiledDataVolatilityManager::ConditionalRecompile(IMovieScenePlayer& Pla
 	}
 	else
 	{
-		bRecompiled = HasBeenRecompiled(RootDataID, CompiledDataManager);
+		bRecompiled = HasBeenRecompiled();
 	}
 
 	if (bRecompiled)
 	{
-		UpdateCachedSignatures(Player, RootDataID, CompiledDataManager);
+		UpdateCachedSignatures();
 	}
 
 	return bRecompiled;
 }
 
-void FCompiledDataVolatilityManager::UpdateCachedSignatures(IMovieScenePlayer& Player, FMovieSceneCompiledDataID RootDataID, UMovieSceneCompiledDataManager* CompiledDataManager)
+void FCompiledDataVolatilityManager::UpdateCachedSignatures()
 {
 	CachedCompilationSignatures.Reset();
 
+	TSharedPtr<const FSharedPlaybackState> SharedPlaybackState = WeakSharedPlaybackState.Pin();
+	FMovieSceneCompiledDataID RootDataID = SharedPlaybackState->GetRootCompiledDataID();
+	UMovieSceneCompiledDataManager* CompiledDataManager = SharedPlaybackState->GetCompiledDataManager();
+	IMovieScenePlayer* Player = FPlayerIndexPlaybackCapability::GetPlayer(SharedPlaybackState.ToSharedRef());
+
 	{
 		const FMovieSceneCompiledDataEntry& RootEntry = CompiledDataManager->GetEntryRef(RootDataID);
 		CachedCompilationSignatures.Add(MovieSceneSequenceID::Root, RootEntry.CompiledSignature);
 
 		UMovieSceneSequence* RootSequence = RootEntry.GetSequence();
-		if (RootSequence)
+		if (RootSequence && Player)
 		{
-			Player.State.AssignSequence(MovieSceneSequenceID::Root, *RootSequence, Player);
+			Player->State.AssignSequence(MovieSceneSequenceID::Root, *RootSequence, *Player);
 		}
 	}
 
@@ -158,9 +181,9 @@ void FCompiledDataVolatilityManager::UpdateCachedSignatures(IMovieScenePlayer& P
 			CachedCompilationSignatures.Add(SubData.Key, SubEntry.CompiledSignature);
 
 			UMovieSceneSequence* Sequence = SubData.Value.GetSequence();
-			if (Sequence)
+			if (Sequence && Player)
 			{
-				Player.State.AssignSequence(SubData.Key, *Sequence, Player);
+				Player->State.AssignSequence(SubData.Key, *Sequence, *Player);
 			}
 		}
 	}
diff --git a/Engine/Source/Runtime/MovieScene/Private/EntitySystem/MovieSceneEntitySystemRunner.cpp b/Engine/Source/Runtime/MovieScene/Private/EntitySystem/MovieSceneEntitySystemRunner.cpp
--- a/Engine/Source/Runtime/MovieScene/Private/EntitySystem/MovieSceneEntitySystemRunner.cpp
+++ b/Engine/Source/Runtime/MovieScene/Private/EntitySystem/MovieSceneEntitySystemRunner.cpp
@@ -266,13 +266,13 @@ bool FMovieSceneEntitySystemRunner::QueueFinalUpdateImpl(FInstanceHandle InInsta
 	// 2. we're not in the middle of an update loop 
 	// 3. the instance has no current updates
 	//
-	const bool bCanFinishImmediately = Instance.CanFinishImmediately(Linker);
+	const bool bCanFinishImmediately = Instance.CanFinishImmediately();
 	const ERunnerFlushState UnsafeDestroyMask = FlushState::Everything & ~(ERunnerFlushState::PostEvaluation | ERunnerFlushState::End);
 	const bool bSafeToDestroyNow = !EnumHasAnyFlags(FlushState, UnsafeDestroyMask);
 	if (bCanFinishImmediately && bSafeToDestroyNow && !HasQueuedUpdates(InInstanceHandle))
 	{
-		Instance.Finish(Linker);
-		Instance.PostEvaluation(Linker);
+		Instance.Finish();
+		Instance.PostEvaluation();
 
 		InOnLastFlushDelegate.ExecuteIfBound();
 
@@ -637,7 +637,7 @@ UE::MovieScene::ERunnerFlushResult FMovieSceneEntitySystemRunner::GameThread_Con
 			if (InstanceRegistry->IsHandleValid(UpdatedInstance.InstanceHandle) && !EnumHasAnyFlags(UpdatedInstance.UpdateFlags, ERunnerUpdateFlags::Finish | ERunnerUpdateFlags::Destroy))
 			{
 				FSequenceInstance& Instance = InstanceRegistry->MutateInstance(UpdatedInstance.InstanceHandle);
-				if (Instance.IsRootSequence() && Instance.ConditionalRecompile(Linker))
+				if (Instance.IsRootSequence() && Instance.ConditionalRecompile())
 				{
 					bAnyRecompile = true;
 				}
@@ -651,7 +651,7 @@ UE::MovieScene::ERunnerFlushResult FMovieSceneEntitySystemRunner::GameThread_Con
 		if (InstanceRegistry->IsHandleValid(Request.Params.InstanceHandle) && !EnumHasAnyFlags(Request.Params.UpdateFlags, ERunnerUpdateFlags::Finish | ERunnerUpdateFlags::Destroy))
 		{
 			FSequenceInstance& Instance = InstanceRegistry->MutateInstance(Request.Params.InstanceHandle);
-			if (Instance.IsRootSequence() && Instance.ConditionalRecompile(Linker))
+			if (Instance.IsRootSequence() && Instance.ConditionalRecompile())
 			{
 				bAnyRecompile = true;
 			}
@@ -751,7 +751,7 @@ UE::MovieScene::ERunnerFlushResult FMovieSceneEntitySystemRunner::GameThread_Upd
 				continue;
 			}
 
-			Instance.DissectContext(Linker, Request.Context, Dissections);
+			Instance.DissectContext(Request.Context, Dissections);
 
 			if (Dissections.Num() != 0)
 			{
@@ -839,7 +839,7 @@ UE::MovieScene::ERunnerFlushResult FMovieSceneEntitySystemRunner::GameThread_Upd
 
 			if (EnumHasAnyFlags(SequenceInstance.GetUpdateFlags(), ESequenceInstanceUpdateFlags::NeedsPreEvaluation))
 			{
-				SequenceInstance.PreEvaluation(Linker);
+				SequenceInstance.PreEvaluation();
 			}
 		}
 	}
@@ -875,7 +875,7 @@ UE::MovieScene::ERunnerFlushResult FMovieSceneEntitySystemRunner::GameThread_Upd
 				if (EnumHasAnyFlags(Update.UpdateFlags, ERunnerUpdateFlags::Finish))
 				{
 					// Context is irrelevant for Finishing sequences
-					Instance.Finish(Linker);
+					Instance.Finish();
 				}
 				else if (EnumHasAnyFlags(Update.UpdateFlags, ERunnerUpdateFlags::Destroy))
 				{
@@ -883,7 +883,7 @@ UE::MovieScene::ERunnerFlushResult FMovieSceneEntitySystemRunner::GameThread_Upd
 				}
 				else
 				{
-					Instance.Update(Linker, Update.Context);
+					Instance.Update(Update.Context);
 				}
 			}
 		}
@@ -942,7 +942,7 @@ UE::MovieScene::ERunnerFlushResult FMovieSceneEntitySystemRunner::GameThread_Rei
 			// Need to check whether this has actually finished yet or not
 			if (!SequenceInstance.HasFinished())
 			{
-				SequenceInstance.Finish(Linker);
+				SequenceInstance.Finish();
 			}
 
 			// If it is being finished or destroyed, just re-add this update back.
@@ -956,7 +956,7 @@ UE::MovieScene::ERunnerFlushResult FMovieSceneEntitySystemRunner::GameThread_Rei
 			{
 				// Add root sequences back to the CurrentInstances array. When we call Update any active sub-sequences will be re-added via MarkForUpdate
 				CurrentInstances.Add(UpdatedInstance);
-				SequenceInstance.Update(Linker, SequenceInstance.GetContext());
+				SequenceInstance.Update(SequenceInstance.GetContext());
 			}
 		}
 	}
@@ -1247,7 +1247,7 @@ void FMovieSceneEntitySystemRunner::GameThread_PostEvaluationPhase(UMovieSceneEn
 
 				if (EnumHasAnyFlags(Instance.GetUpdateFlags(), ESequenceInstanceUpdateFlags::NeedsPostEvaluation))
 				{
-					Instance.PostEvaluation(Linker);
+					Instance.PostEvaluation();
 				}
 
 				if (EnumHasAnyFlags(UpdateParams.UpdateFlags, ERunnerUpdateFlags::Destroy))
@@ -1314,7 +1314,7 @@ void FMovieSceneEntitySystemRunner::MarkForUpdate(FInstanceHandle InInstanceHand
 		// Need to check whether this has actually finished yet or not
 		if (!SequenceInstance.HasFinished())
 		{
-			SequenceInstance.Finish(Linker);
+			SequenceInstance.Finish();
 		}
 	}
 
diff --git a/Engine/Source/Runtime/MovieScene/Private/EntitySystem/MovieSceneInstanceRegistry.cpp b/Engine/Source/Runtime/MovieScene/Private/EntitySystem/MovieSceneInstanceRegistry.cpp
--- a/Engine/Source/Runtime/MovieScene/Private/EntitySystem/MovieSceneInstanceRegistry.cpp
+++ b/Engine/Source/Runtime/MovieScene/Private/EntitySystem/MovieSceneInstanceRegistry.cpp
@@ -5,6 +5,7 @@
 #include "EntitySystem/MovieSceneEntityMutations.h"
 #include "EntitySystem/MovieSceneEntitySystemLinker.h"
 #include "EntitySystem/MovieSceneEntityInstantiatorSystem.h"
+#include "EntitySystem/MovieSceneEntitySystemRunner.h"
 #include "EntitySystem/MovieSceneEntitySystemTask.h"
 #include "EntitySystem/MovieSceneSharedPlaybackState.h"
 #include "EntitySystem/EntityAllocationIterator.h"
@@ -56,7 +57,7 @@ FInstanceHandle FInstanceRegistry::FindRelatedInstanceHandle(FInstanceHandle Ins
 	return RootInstance->FindSubInstance(SequenceID);
 }
 
-FRootInstanceHandle FInstanceRegistry::AllocateRootInstance(IMovieScenePlayer* Player)
+FRootInstanceHandle FInstanceRegistry::AllocateRootInstance(IMovieScenePlayer* Player, UMovieSceneSequence& RootSequence, TSharedPtr<FMovieSceneEntitySystemRunner> Runner, UMovieSceneCompiledDataManager* CompiledDataManager)
 {
 	check(Instances.Num() < 65535);
 
@@ -65,9 +66,16 @@ FRootInstanceHandle FInstanceRegistry::AllocateRootInstance(IMovieScenePlayer* P
 	FSparseArrayAllocationInfo NewAllocation = Instances.AddUninitialized();
 	FRootInstanceHandle InstanceHandle { (uint16)NewAllocation.Index, InstanceSerial };
 
-	TSharedRef<FSharedPlaybackState> NewPlaybackState = MakeShared<FSharedPlaybackState>();
+	FSharedPlaybackStateCreateParams PlaybackStateCreateParams;
+	PlaybackStateCreateParams.PlaybackContext = Player->GetPlaybackContext();
+	PlaybackStateCreateParams.RootInstanceHandle = InstanceHandle;
+	PlaybackStateCreateParams.Runner = Runner;
+	PlaybackStateCreateParams.CompiledDataManager = CompiledDataManager;
+	TSharedRef<FSharedPlaybackState> NewPlaybackState = MakeShared<FSharedPlaybackState>(RootSequence, PlaybackStateCreateParams);
 
-	new (NewAllocation) FSequenceInstance(Linker, Player, NewPlaybackState, InstanceHandle);
+	FSequenceInstance* NewInstance = new (NewAllocation) FSequenceInstance(NewPlaybackState, InstanceHandle);
+
+	NewInstance->Initialize(Player);
 
 	return InstanceHandle;
 }
@@ -82,7 +90,9 @@ FInstanceHandle FInstanceRegistry::AllocateSubInstance(IMovieScenePlayer* Player
 	
 	TSharedRef<FSharedPlaybackState> PlaybackState = GetInstance(RootInstanceHandle).GetSharedPlaybackState();
 
-	new (NewAllocation) FSequenceInstance(Linker, Player, PlaybackState, InstanceHandle, ParentInstanceHandle, RootInstanceHandle, SequenceID);
+	FSequenceInstance* NewInstance = new (NewAllocation) FSequenceInstance(PlaybackState, InstanceHandle, ParentInstanceHandle, RootInstanceHandle, SequenceID);
+
+	NewInstance->Initialize(Player);
 
 	return InstanceHandle;
 }
@@ -92,13 +102,12 @@ void FInstanceRegistry::DestroyInstance(FInstanceHandle InstanceHandle)
 	if (ensureMsgf(Instances.IsValidIndex(InstanceHandle.InstanceID) && Instances[InstanceHandle.InstanceID].GetSerialNumber() == InstanceHandle.InstanceSerial, TEXT("Attempting to destroy an instance an invalid instance handle.")))
 	{
 		FSequenceInstance& Instance = Instances[InstanceHandle.InstanceID];
-		const bool bIsRootInstance = Instance.IsRootSequence();
 		const bool bHasFinished = (GExitPurge || Instance.HasFinished());
 		if (!bHasFinished)
 		{
 			UE_LOG(LogMovieSceneECS, Verbose, TEXT("Instance being destroyed without finishing evaluation."));
 		}
-		Instance.DestroyImmediately(Linker);
+		Instance.DestroyImmediately();
 		Instances.RemoveAt(InstanceHandle.InstanceID);
 	}
 }
@@ -140,7 +149,7 @@ FScopedVolatilityManagerSuppression::~FScopedVolatilityManagerSuppression()
 {
 	FSequenceInstance& Instance = InstanceRegistry->MutateInstance(RootInstanceHandle);
 	Instance.VolatilityManager = MoveTemp(PreviousVolatilityManager);
-	Instance.ConditionalRecompile(InstanceRegistry->GetLinker());
+	Instance.ConditionalRecompile();
 }
 
 } // namespace MovieScene
diff --git a/Engine/Source/Runtime/MovieScene/Private/EntitySystem/MovieSceneSequenceInstance.cpp b/Engine/Source/Runtime/MovieScene/Private/EntitySystem/MovieSceneSequenceInstance.cpp
--- a/Engine/Source/Runtime/MovieScene/Private/EntitySystem/MovieSceneSequenceInstance.cpp
+++ b/Engine/Source/Runtime/MovieScene/Private/EntitySystem/MovieSceneSequenceInstance.cpp
@@ -66,13 +66,10 @@ void PurgeStaleTrackTemplates(UMovieSceneCompiledDataManager* CompiledDataManage
 }
 
 
-
-
-FSequenceInstance::FSequenceInstance(UMovieSceneEntitySystemLinker* Linker, IMovieScenePlayer* Player, TSharedRef<FSharedPlaybackState> PlaybackState, FRootInstanceHandle InInstanceHandle)
+FSequenceInstance::FSequenceInstance(TSharedRef<FSharedPlaybackState> PlaybackState, FRootInstanceHandle InInstanceHandle)
 	: SharedPlaybackState(PlaybackState)
 	, SequenceID(MovieSceneSequenceID::Root)
 	, RootOverrideSequenceID(MovieSceneSequenceID::Root)
-	, PlayerIndex(Player->GetUniqueIndex())
 	, InstanceHandle(InInstanceHandle)
 	, RootInstanceHandle(InInstanceHandle)
 {
@@ -85,21 +82,15 @@ FSequenceInstance::FSequenceInstance(UMovieSceneEntitySystemLinker* Linker, IMov
 	bHasEverUpdated = false;
 
 #if !UE_BUILD_SHIPPING && !UE_BUILD_TEST
-	UMovieSceneSequence* RootSequence = Player->GetEvaluationTemplate().GetRootSequence();
+	UMovieSceneSequence* RootSequence = PlaybackState->GetRootSequence();
 	RootSequenceName = RootSequence->GetPathName();
 #endif
-
-	FMovieSceneObjectCache& ObjectCache = Player->State.GetObjectCache(SequenceID);
-	OnInvalidateObjectBindingHandle = ObjectCache.OnBindingInvalidated.AddUObject(Linker, &UMovieSceneEntitySystemLinker::InvalidateObjectBinding, InstanceHandle);
-
-	InvalidateCachedData(Linker);
 }
 
-FSequenceInstance::FSequenceInstance(UMovieSceneEntitySystemLinker* Linker, IMovieScenePlayer* Player, TSharedRef<FSharedPlaybackState> PlaybackState, FInstanceHandle InInstanceHandle, FInstanceHandle InParentInstanceHandle, FRootInstanceHandle InRootInstanceHandle, FMovieSceneSequenceID InSequenceID)
+FSequenceInstance::FSequenceInstance(TSharedRef<FSharedPlaybackState> PlaybackState, FInstanceHandle InInstanceHandle, FInstanceHandle InParentInstanceHandle, FRootInstanceHandle InRootInstanceHandle, FMovieSceneSequenceID InSequenceID)
 	: SharedPlaybackState(PlaybackState)
 	, SequenceID(InSequenceID)
 	, RootOverrideSequenceID(MovieSceneSequenceID::Invalid)
-	, PlayerIndex(Player->GetUniqueIndex())
 	, InstanceHandle(InInstanceHandle)
 	, ParentInstanceHandle(InParentInstanceHandle)
 	, RootInstanceHandle(InRootInstanceHandle)
@@ -110,11 +101,29 @@ FSequenceInstance::FSequenceInstance(UMovieSceneEntitySystemLinker* Linker, IMov
 	// be created if they are active, and the Start/Update/Finish loop does not apply to sub-instances
 	bFinished = false;
 	bHasEverUpdated = false;
+}
 
-	FMovieSceneObjectCache& ObjectCache = Player->State.GetObjectCache(SequenceID);
-	OnInvalidateObjectBindingHandle = ObjectCache.OnBindingInvalidated.AddUObject(Linker, &UMovieSceneEntitySystemLinker::InvalidateObjectBinding, InstanceHandle);
+void FSequenceInstance::Initialize(IMovieScenePlayer* Player)
+{
+	if (!Player)
+	{
+		return;
+	}
 
-	InvalidateCachedData(Linker);
+	// Initialize playback capabilities if this is the root sequence.
+	if (IsRootSequence())
+	{
+		SharedPlaybackState->AddCapability<FPlayerIndexPlaybackCapability>(Player->GetUniqueIndex());
+	}
+
+	UMovieSceneEntitySystemLinker* Linker = SharedPlaybackState->GetLinker();
+	if (ensure(Linker))
+	{
+		FMovieSceneObjectCache& ObjectCache = Player->State.GetObjectCache(SequenceID);
+		OnInvalidateObjectBindingHandle = ObjectCache.OnBindingInvalidated.AddUObject(Linker, &UMovieSceneEntitySystemLinker::InvalidateObjectBinding, InstanceHandle);
+	}
+
+	InvalidateCachedData();
 }
 
 FSequenceInstance::~FSequenceInstance()
@@ -126,17 +135,23 @@ FSequenceInstance& FSequenceInstance::operator=(FSequenceInstance&&) = default;
 
 IMovieScenePlayer* FSequenceInstance::GetPlayer() const
 {
-	return IMovieScenePlayer::Get(PlayerIndex);
+	using namespace UE::MovieScene;
+
+	return FPlayerIndexPlaybackCapability::GetPlayer(SharedPlaybackState);
 }
 
-void FSequenceInstance::InitializeLegacyEvaluator(UMovieSceneEntitySystemLinker* Linker)
+uint16 FSequenceInstance::GetPlayerIndex() const
 {
-	IMovieScenePlayer* Player = GetPlayer();
-	check(Player);
+	using namespace UE::MovieScene;
 
-	FMovieSceneRootEvaluationTemplateInstance& RootTemplate        = Player->GetEvaluationTemplate();
-	UMovieSceneCompiledDataManager*            CompiledDataManager = RootTemplate.GetCompiledDataManager();
-	const FMovieSceneCompiledDataEntry&        CompiledEntry       = CompiledDataManager->GetEntryRef(RootTemplate.GetCompiledDataID());
+	return FPlayerIndexPlaybackCapability::GetPlayerIndex(SharedPlaybackState);
+}
+
+void FSequenceInstance::InitializeLegacyEvaluator()
+{
+	const FMovieSceneCompiledDataID RootCompiledDataID = SharedPlaybackState->GetRootCompiledDataID();
+	UMovieSceneCompiledDataManager* CompiledDataManager = SharedPlaybackState->GetCompiledDataManager();
+	const FMovieSceneCompiledDataEntry& CompiledEntry = CompiledDataManager->GetEntryRef(RootCompiledDataID);
 
 	if (EnumHasAnyFlags(CompiledEntry.AccumulatedMask, EMovieSceneSequenceCompilerMask::EvaluationTemplate))
 	{
@@ -144,11 +159,15 @@ void FSequenceInstance::InitializeLegacyEvaluator(UMovieSceneEntitySystemLinker*
 
 		if (!LegacyEvaluator)
 		{
-			LegacyEvaluator = MakeUnique<FMovieSceneTrackEvaluator>(CompiledEntry.GetSequence(), RootTemplate.GetCompiledDataID(), CompiledDataManager);
+			UMovieSceneSequence* RootSequence = SharedPlaybackState->GetRootSequence();
+			LegacyEvaluator = MakeUnique<FMovieSceneTrackEvaluator>(RootSequence, RootCompiledDataID, CompiledDataManager);
 		}
 	}
 	else if (LegacyEvaluator)
 	{
+		IMovieScenePlayer* Player = GetPlayer();
+		check(Player);
+
 		LegacyEvaluator->Finish(*Player);
 		LegacyEvaluator = nullptr;
 
@@ -156,61 +175,67 @@ void FSequenceInstance::InitializeLegacyEvaluator(UMovieSceneEntitySystemLinker*
 	}
 }
 
-void FSequenceInstance::InvalidateCachedData(UMovieSceneEntitySystemLinker* Linker)
+void FSequenceInstance::InvalidateCachedData()
 {
+	// WARNING: this method is called from Initialize. If this is a root sequence, the Player is still
+	// in the process of creating us, and is waiting for the call stack to return in order to set our
+	// SharedPlaybackState and other relevant references on itself. So do NOT access anything through
+	// the Player's RootEvaluationTemplateInstance here because it's most probably not initalized yet.
+
 	Ledger.Invalidate();
 
-	IMovieScenePlayer* Player = GetPlayer();
-	check(Player);
+	IMovieScenePlayer* Player = FPlayerIndexPlaybackCapability::GetPlayer(SharedPlaybackState);
 
 	UpdateFlags = ESequenceInstanceUpdateFlags::None;
 
-	FMovieSceneRootEvaluationTemplateInstance& RootTemplate        = Player->GetEvaluationTemplate();
-	UMovieSceneCompiledDataManager*            CompiledDataManager = RootTemplate.GetCompiledDataManager();
-	FMovieSceneCompiledDataID                  RootCompiledDataID  = RootTemplate.GetCompiledDataID();
+	const FMovieSceneCompiledDataID RootCompiledDataID = SharedPlaybackState->GetRootCompiledDataID();
+	UMovieSceneCompiledDataManager* CompiledDataManager = SharedPlaybackState->GetCompiledDataManager();
 
 	if (SequenceID == MovieSceneSequenceID::Root)
 	{
-		SharedPlaybackState->InvalidateCachedData(Linker);
+		SharedPlaybackState->InvalidateCachedData();
 
-		Player->State.AssignSequence(SequenceID, *RootTemplate.GetRootSequence(), *Player);
+		if (Player)
+		{
+			UMovieSceneSequence* RootSequence = SharedPlaybackState->GetRootSequence();
+			Player->State.AssignSequence(SequenceID, *RootSequence, *Player);
+		}
 
 		// Try and recreate the volatility manager if this sequence is now volatile
 		if (!VolatilityManager)
 		{
-			VolatilityManager = FCompiledDataVolatilityManager::Construct(*Player, RootCompiledDataID, CompiledDataManager);
+			VolatilityManager = FCompiledDataVolatilityManager::Construct(SharedPlaybackState);
 		}
 
 		ISequenceUpdater::FactoryInstance(SequenceUpdater, CompiledDataManager, RootCompiledDataID);
 
+		UMovieSceneEntitySystemLinker* Linker = SharedPlaybackState->GetLinker();
 		SequenceUpdater->InvalidateCachedData(Linker, RootInstanceHandle);
-		SequenceUpdater->PopulateUpdateFlags(Linker, Player, UpdateFlags);
+		SequenceUpdater->PopulateUpdateFlags(Linker, SharedPlaybackState, UpdateFlags);
 
 		if (LegacyEvaluator)
 		{
 			LegacyEvaluator->InvalidateCachedData();
 		}
 
-		InitializeLegacyEvaluator(Linker);
+		InitializeLegacyEvaluator();
 	}
-	else if (UMovieSceneSequence* SubSequence = RootTemplate.GetSequence(SequenceID))
+	else if (UMovieSceneSequence* SubSequence = SharedPlaybackState->GetSequence(SequenceID))
 	{
-		Player->State.AssignSequence(SequenceID, *SubSequence, *Player);
+		if (Player)
+		{
+			Player->State.AssignSequence(SequenceID, *SubSequence, *Player);
+		}
 	}
 }
 
-bool FSequenceInstance::ConditionalRecompile(UMovieSceneEntitySystemLinker* Linker)
+bool FSequenceInstance::ConditionalRecompile()
 {
 	if (VolatilityManager)
 	{
-		IMovieScenePlayer*                         Player              = GetPlayer();
-		FMovieSceneRootEvaluationTemplateInstance& RootTemplate        = Player->GetEvaluationTemplate();
-		UMovieSceneCompiledDataManager*            CompiledDataManager = RootTemplate.GetCompiledDataManager();
-		FMovieSceneCompiledDataID                  RootCompiledDataID  = RootTemplate.GetCompiledDataID();
-
-		if (VolatilityManager->ConditionalRecompile(*Player, RootCompiledDataID, CompiledDataManager))
+		if (VolatilityManager->ConditionalRecompile())
 		{
-			InvalidateCachedData(Linker);
+			InvalidateCachedData();
 			return true;
 		}
 	}
@@ -218,16 +243,17 @@ bool FSequenceInstance::ConditionalRecompile(UMovieSceneEntitySystemLinker* Link
 	return false;
 }
 
-void FSequenceInstance::DissectContext(UMovieSceneEntitySystemLinker* Linker, const FMovieSceneContext& InContext, TArray<TRange<FFrameTime>>& OutDissections)
+void FSequenceInstance::DissectContext(const FMovieSceneContext& InContext, TArray<TRange<FFrameTime>>& OutDissections)
 {
 	if (EnumHasAnyFlags(UpdateFlags, ESequenceInstanceUpdateFlags::NeedsDissection))
 	{
 		check(SequenceID == MovieSceneSequenceID::Root);
-		SequenceUpdater->DissectContext(Linker, GetPlayer(), InContext, OutDissections);
+		UMovieSceneEntitySystemLinker* Linker = SharedPlaybackState->GetLinker();
+		SequenceUpdater->DissectContext(Linker, SharedPlaybackState, InContext, OutDissections);
 	}
 }
 
-void FSequenceInstance::Start(UMovieSceneEntitySystemLinker* Linker, const FMovieSceneContext& InContext)
+void FSequenceInstance::Start(const FMovieSceneContext& InContext)
 {
 	check(SequenceID == MovieSceneSequenceID::Root);
 
@@ -236,54 +262,57 @@ void FSequenceInstance::Start(UMovieSceneEntitySystemLinker* Linker, const FMovi
 
 	check(RootInstanceHandle == InstanceHandle);
 
-	IMovieScenePlayer* Player = GetPlayer();
-	SequenceUpdater->Start(Linker, RootInstanceHandle, Player, InContext);
+	UMovieSceneEntitySystemLinker* Linker = SharedPlaybackState->GetLinker();
+	SequenceUpdater->Start(Linker, RootInstanceHandle, SharedPlaybackState, InContext);
 }
 
-void FSequenceInstance::Update(UMovieSceneEntitySystemLinker* Linker, const FMovieSceneContext& InContext)
+void FSequenceInstance::Update(const FMovieSceneContext& InContext)
 {
 	SCOPE_CYCLE_COUNTER(MovieSceneEval_SequenceInstanceUpdate);
 	SCOPE_CYCLE_UOBJECT(ContextScope, GetPlayer()->AsUObject());
 
 	bHasEverUpdated = true;
+	UMovieSceneEntitySystemLinker* Linker = SharedPlaybackState->GetLinker();
 
 	if (bFinished)
 	{
-		Start(Linker, InContext);
+		Start(InContext);
 	}
 
 	check(RootInstanceHandle == InstanceHandle);
 
 	Context = InContext;
-	SequenceUpdater->Update(Linker, RootInstanceHandle, GetPlayer(), InContext);
+	SequenceUpdater->Update(Linker, RootInstanceHandle, SharedPlaybackState, InContext);
 }
 
-bool FSequenceInstance::CanFinishImmediately(UMovieSceneEntitySystemLinker* Linker) const
+bool FSequenceInstance::CanFinishImmediately() const
 {
 	if (SequenceUpdater)
 	{
 		check(RootInstanceHandle == InstanceHandle);
 
+		UMovieSceneEntitySystemLinker* Linker = SharedPlaybackState->GetLinker();
 		return SequenceUpdater->CanFinishImmediately(Linker, RootInstanceHandle);
 	}
 
 	return true;
 }
 
-void FSequenceInstance::Finish(UMovieSceneEntitySystemLinker* Linker)
+void FSequenceInstance::Finish()
 {
 	if (IsRootSequence() && !bHasEverUpdated)
 	{
 		return;
 	}
 
+	UMovieSceneEntitySystemLinker* Linker = SharedPlaybackState->GetLinker();
 	Linker->EntityManager.IncrementSystemSerial();
 	bFinished = true;
 	Ledger.UnlinkEverything(Linker);
 
 	Ledger = FEntityLedger();
 
-	IMovieScenePlayer* Player = IMovieScenePlayer::Get(PlayerIndex);
+	IMovieScenePlayer* Player = GetPlayer();
 	if (!ensure(Player))
 	{
 		return;
@@ -292,7 +321,7 @@ void FSequenceInstance::Finish(UMovieSceneEntitySystemLinker* Linker)
 	if (SequenceUpdater)
 	{
 		check(RootInstanceHandle == InstanceHandle);
-		SequenceUpdater->Finish(Linker, RootInstanceHandle, Player);
+		SequenceUpdater->Finish(Linker, RootInstanceHandle, SharedPlaybackState);
 	}
 
 	if (LegacyEvaluator)
@@ -313,7 +342,7 @@ void FSequenceInstance::Finish(UMovieSceneEntitySystemLinker* Linker)
 	}
 }
 
-void FSequenceInstance::PreEvaluation(UMovieSceneEntitySystemLinker* Linker)
+void FSequenceInstance::PreEvaluation()
 {
 	if (!EnumHasAnyFlags(UpdateFlags, ESequenceInstanceUpdateFlags::NeedsPreEvaluation))
 	{
@@ -334,7 +363,7 @@ void FSequenceInstance::RunLegacyTrackTemplates()
 {
 	if (LegacyEvaluator)
 	{
-		IMovieScenePlayer* Player = IMovieScenePlayer::Get(PlayerIndex);
+		IMovieScenePlayer* Player = GetPlayer();
 		if (ensure(Player))
 		{
 			if (bFinished)
@@ -349,7 +378,7 @@ void FSequenceInstance::RunLegacyTrackTemplates()
 	}
 }
 
-void FSequenceInstance::PostEvaluation(UMovieSceneEntitySystemLinker* Linker)
+void FSequenceInstance::PostEvaluation()
 {
 	if (IsRootSequence() && EnumHasAnyFlags(UpdateFlags, ESequenceInstanceUpdateFlags::NeedsPostEvaluation))
 	{
@@ -385,8 +414,10 @@ void FSequenceInstance::PostEvaluation(UMovieSceneEntitySystemLinker* Linker)
 	}
 }
 
-void FSequenceInstance::DestroyImmediately(UMovieSceneEntitySystemLinker* Linker)
+void FSequenceInstance::DestroyImmediately()
 {
+	UMovieSceneEntitySystemLinker* Linker = SharedPlaybackState->GetLinker();
+
 	if (!Ledger.IsEmpty())
 	{
 		UE_LOG(LogMovieSceneECS, Verbose, TEXT("Instance being destroyed without first having been finished by calling Finish()"));
@@ -399,11 +430,12 @@ void FSequenceInstance::DestroyImmediately(UMovieSceneEntitySystemLinker* Linker
 	}
 }
 
-void FSequenceInstance::OverrideRootSequence(UMovieSceneEntitySystemLinker* Linker, FMovieSceneSequenceID NewRootSequenceID)
+void FSequenceInstance::OverrideRootSequence(FMovieSceneSequenceID NewRootSequenceID)
 {
 	if (SequenceUpdater)
 	{
 		check(RootInstanceHandle == InstanceHandle);
+		UMovieSceneEntitySystemLinker* Linker = SharedPlaybackState->GetLinker();
 		SequenceUpdater->OverrideRootSequence(Linker, RootInstanceHandle, NewRootSequenceID);
 	}
 
@@ -430,5 +462,60 @@ FSubSequencePath FSequenceInstance::GetSubSequencePath() const
 	return FSubSequencePath(SequenceID, *GetPlayer());
 }
 
+bool FSequenceInstance::ConditionalRecompile(UMovieSceneEntitySystemLinker* Linker)
+{
+	return ConditionalRecompile();
+}
+
+void FSequenceInstance::DissectContext(UMovieSceneEntitySystemLinker* Linker, const FMovieSceneContext& InContext, TArray<TRange<FFrameTime>>& OutDissections)
+{
+	DissectContext(InContext, OutDissections);
+}
+
+void FSequenceInstance::Start(UMovieSceneEntitySystemLinker* Linker, const FMovieSceneContext& InContext)
+{
+	Start(InContext);
+}
+
+void FSequenceInstance::PreEvaluation(UMovieSceneEntitySystemLinker* Linker)
+{
+	PreEvaluation();
+}
+
+void FSequenceInstance::Update(UMovieSceneEntitySystemLinker* Linker, const FMovieSceneContext& InContext)
+{
+	Update(InContext);
+}
+
+bool FSequenceInstance::CanFinishImmediately(UMovieSceneEntitySystemLinker* Linker) const
+{
+	return CanFinishImmediately();
+}
+
+void FSequenceInstance::Finish(UMovieSceneEntitySystemLinker* Linker)
+{
+	Finish();
+}
+
+void FSequenceInstance::PostEvaluation(UMovieSceneEntitySystemLinker* Linker)
+{
+	PostEvaluation();
+}
+
+void FSequenceInstance::InvalidateCachedData(UMovieSceneEntitySystemLinker* Linker)
+{
+	InvalidateCachedData();
+}
+
+void FSequenceInstance::DestroyImmediately(UMovieSceneEntitySystemLinker* Linker)
+{
+	DestroyImmediately();
+}
+
+void FSequenceInstance::OverrideRootSequence(UMovieSceneEntitySystemLinker* Linker, FMovieSceneSequenceID NewRootSequenceID)
+{
+	OverrideRootSequence(NewRootSequenceID);
+}
+
 } // namespace MovieScene
 } // namespace UE
diff --git a/Engine/Source/Runtime/MovieScene/Private/EntitySystem/MovieSceneSequenceUpdaters.cpp b/Engine/Source/Runtime/MovieScene/Private/EntitySystem/MovieSceneSequenceUpdaters.cpp
--- a/Engine/Source/Runtime/MovieScene/Private/EntitySystem/MovieSceneSequenceUpdaters.cpp
+++ b/Engine/Source/Runtime/MovieScene/Private/EntitySystem/MovieSceneSequenceUpdaters.cpp
@@ -35,12 +35,12 @@ struct FSequenceUpdater_Flat : ISequenceUpdater
 	explicit FSequenceUpdater_Flat(FMovieSceneCompiledDataID InCompiledDataID);
 	~FSequenceUpdater_Flat();
 
-	virtual void PopulateUpdateFlags(UMovieSceneEntitySystemLinker* InLinker, IMovieScenePlayer* InPlayer, ESequenceInstanceUpdateFlags& OutUpdateFlags) override;
-	virtual void DissectContext(UMovieSceneEntitySystemLinker* Linker, IMovieScenePlayer* InPlayer, const FMovieSceneContext& Context, TArray<TRange<FFrameTime>>& OutDissections) override;
-	virtual void Start(UMovieSceneEntitySystemLinker* Linker, FRootInstanceHandle InstanceHandle, IMovieScenePlayer* InPlayer, const FMovieSceneContext& InContext) override;
-	virtual void Update(UMovieSceneEntitySystemLinker* Linker, FRootInstanceHandle InstanceHandle, IMovieScenePlayer* InPlayer, const FMovieSceneContext& Context) override;
+	virtual void PopulateUpdateFlags(UMovieSceneEntitySystemLinker* InLinker, TSharedRef<const FSharedPlaybackState> SharedPlaybackState, ESequenceInstanceUpdateFlags& OutUpdateFlags) override;
+	virtual void DissectContext(UMovieSceneEntitySystemLinker* Linker, TSharedRef<const FSharedPlaybackState> SharedPlaybackState, const FMovieSceneContext& Context, TArray<TRange<FFrameTime>>& OutDissections) override;
+	virtual void Start(UMovieSceneEntitySystemLinker* Linker, FRootInstanceHandle InstanceHandle, TSharedRef<const FSharedPlaybackState> SharedPlaybackState, const FMovieSceneContext& InContext) override;
+	virtual void Update(UMovieSceneEntitySystemLinker* Linker, FRootInstanceHandle InstanceHandle, TSharedRef<const FSharedPlaybackState> SharedPlaybackState, const FMovieSceneContext& Context) override;
 	virtual bool CanFinishImmediately(UMovieSceneEntitySystemLinker* InLinker, FRootInstanceHandle RootInstanceHandle) const override;
-	virtual void Finish(UMovieSceneEntitySystemLinker* Linker, FRootInstanceHandle InstanceHandle, IMovieScenePlayer* InPlayer) override;
+	virtual void Finish(UMovieSceneEntitySystemLinker* Linker, FRootInstanceHandle InstanceHandle, TSharedRef<const FSharedPlaybackState> SharedPlaybackState) override;
 	virtual void InvalidateCachedData(UMovieSceneEntitySystemLinker* Linker, FRootInstanceHandle RootInstanceHandle) override;
 	virtual void Destroy(UMovieSceneEntitySystemLinker* Linker) override;
 	virtual TUniquePtr<ISequenceUpdater> MigrateToHierarchical() override;
@@ -64,12 +64,12 @@ struct FSequenceUpdater_Hierarchical : ISequenceUpdater
 
 	~FSequenceUpdater_Hierarchical();
 
-	virtual void PopulateUpdateFlags(UMovieSceneEntitySystemLinker* InLinker, IMovieScenePlayer* InPlayer, ESequenceInstanceUpdateFlags& OutUpdateFlags) override;
-	virtual void DissectContext(UMovieSceneEntitySystemLinker* Linker, IMovieScenePlayer* InPlayer, const FMovieSceneContext& Context, TArray<TRange<FFrameTime>>& OutDissections) override;
-	virtual void Start(UMovieSceneEntitySystemLinker* Linker, FRootInstanceHandle InstanceHandle, IMovieScenePlayer* InPlayer, const FMovieSceneContext& InContext) override;
-	virtual void Update(UMovieSceneEntitySystemLinker* Linker, FRootInstanceHandle InstanceHandle, IMovieScenePlayer* InPlayer, const FMovieSceneContext& Context) override;
+	virtual void PopulateUpdateFlags(UMovieSceneEntitySystemLinker* InLinker, TSharedRef<const FSharedPlaybackState> SharedPlaybackState, ESequenceInstanceUpdateFlags& OutUpdateFlags) override;
+	virtual void DissectContext(UMovieSceneEntitySystemLinker* Linker, TSharedRef<const FSharedPlaybackState> SharedPlaybackState, const FMovieSceneContext& Context, TArray<TRange<FFrameTime>>& OutDissections) override;
+	virtual void Start(UMovieSceneEntitySystemLinker* Linker, FRootInstanceHandle InstanceHandle, TSharedRef<const FSharedPlaybackState> SharedPlaybackState, const FMovieSceneContext& InContext) override;
+	virtual void Update(UMovieSceneEntitySystemLinker* Linker, FRootInstanceHandle InstanceHandle, TSharedRef<const FSharedPlaybackState> SharedPlaybackState, const FMovieSceneContext& Context) override;
 	virtual bool CanFinishImmediately(UMovieSceneEntitySystemLinker* InLinker, FRootInstanceHandle RootInstanceHandle) const override;
-	virtual void Finish(UMovieSceneEntitySystemLinker* Linker, FRootInstanceHandle InstanceHandle, IMovieScenePlayer* InPlayer) override;
+	virtual void Finish(UMovieSceneEntitySystemLinker* Linker, FRootInstanceHandle InstanceHandle, TSharedRef<const FSharedPlaybackState> SharedPlaybackState) override;
 	virtual void InvalidateCachedData(UMovieSceneEntitySystemLinker* Linker, FRootInstanceHandle RootInstanceHandle) override;
 	virtual void Destroy(UMovieSceneEntitySystemLinker* Linker) override;
 	virtual TUniquePtr<ISequenceUpdater> MigrateToHierarchical() override { return nullptr; }
@@ -80,7 +80,7 @@ private:
 
 	TRange<FFrameNumber> UpdateEntitiesForSequence(const FMovieSceneEntityComponentField* ComponentField, FFrameTime SequenceTime, FMovieSceneEvaluationFieldEntitySet& OutEntities);
 
-	FInstanceHandle GetOrCreateSequenceInstance(IMovieScenePlayer* InPlayer, UMovieSceneSequence* SubSequence, const FMovieSceneSequenceHierarchy* Hierarchy, FInstanceRegistry* InstanceRegistry, FRootInstanceHandle RootInstanceHandle, FMovieSceneSequenceID SequenceID);
+	FInstanceHandle GetOrCreateSequenceInstance(TSharedRef<const FSharedPlaybackState> SharedPlaybackState, UMovieSceneSequence* SubSequence, const FMovieSceneSequenceHierarchy* Hierarchy, FInstanceRegistry* InstanceRegistry, FRootInstanceHandle RootInstanceHandle, FMovieSceneSequenceID SequenceID);
 
 private:
 
@@ -196,11 +196,11 @@ TUniquePtr<ISequenceUpdater> FSequenceUpdater_Flat::MigrateToHierarchical()
 	return MakeUnique<FSequenceUpdater_Hierarchical>(CompiledDataID);
 }
 
-void FSequenceUpdater_Flat::PopulateUpdateFlags(UMovieSceneEntitySystemLinker* InLinker, IMovieScenePlayer* InPlayer, ESequenceInstanceUpdateFlags& OutUpdateFlags)
+void FSequenceUpdater_Flat::PopulateUpdateFlags(UMovieSceneEntitySystemLinker* InLinker, TSharedRef<const FSharedPlaybackState> SharedPlaybackState, ESequenceInstanceUpdateFlags& OutUpdateFlags)
 {
 	if (!CachedDeterminismFences.IsSet())
 	{
-		UMovieSceneCompiledDataManager* CompiledDataManager = InPlayer->GetEvaluationTemplate().GetCompiledDataManager();
+		UMovieSceneCompiledDataManager* CompiledDataManager = SharedPlaybackState->GetCompiledDataManager();
 		TArrayView<const FFrameTime>    DeterminismFences   = CompiledDataManager->GetEntryRef(CompiledDataID).DeterminismFences;
 
 		if (DeterminismFences.Num() != 0)
@@ -213,7 +213,10 @@ void FSequenceUpdater_Flat::PopulateUpdateFlags(UMovieSceneEntitySystemLinker* I
 		}
 	}
 
-	InPlayer->PopulateUpdateFlags(OutUpdateFlags);
+	if (IMovieScenePlayer* Player = FPlayerIndexPlaybackCapability::GetPlayer(SharedPlaybackState))
+	{
+		Player->PopulateUpdateFlags(OutUpdateFlags);
+	}
 
 	if (CachedDeterminismFences.IsSet() && CachedDeterminismFences->Num() > 0)
 	{
@@ -221,11 +224,11 @@ void FSequenceUpdater_Flat::PopulateUpdateFlags(UMovieSceneEntitySystemLinker* I
 	}
 }
 
-void FSequenceUpdater_Flat::DissectContext(UMovieSceneEntitySystemLinker* Linker, IMovieScenePlayer* InPlayer, const FMovieSceneContext& Context, TArray<TRange<FFrameTime>>& OutDissections)
+void FSequenceUpdater_Flat::DissectContext(UMovieSceneEntitySystemLinker* Linker, TSharedRef<const FSharedPlaybackState> SharedPlaybackState, const FMovieSceneContext& Context, TArray<TRange<FFrameTime>>& OutDissections)
 {
 	if (!CachedDeterminismFences.IsSet())
 	{
-		UMovieSceneCompiledDataManager* CompiledDataManager = InPlayer->GetEvaluationTemplate().GetCompiledDataManager();
+		UMovieSceneCompiledDataManager* CompiledDataManager = SharedPlaybackState->GetCompiledDataManager();
 		TArrayView<const FFrameTime>    DeterminismFences   = CompiledDataManager->GetEntryRef(CompiledDataID).DeterminismFences;
 
 		if (DeterminismFences.Num() != 0)
@@ -245,17 +248,18 @@ void FSequenceUpdater_Flat::DissectContext(UMovieSceneEntitySystemLinker* Linker
 	}
 }
 
-void FSequenceUpdater_Flat::Start(UMovieSceneEntitySystemLinker* Linker, FRootInstanceHandle InstanceHandle, IMovieScenePlayer* InPlayer, const FMovieSceneContext& InContext)
+void FSequenceUpdater_Flat::Start(UMovieSceneEntitySystemLinker* Linker, FRootInstanceHandle InstanceHandle, TSharedRef<const FSharedPlaybackState> SharedPlaybackState, const FMovieSceneContext& InContext)
 {
 }
 
-void FSequenceUpdater_Flat::Update(UMovieSceneEntitySystemLinker* Linker, FRootInstanceHandle InstanceHandle, IMovieScenePlayer* InPlayer, const FMovieSceneContext& Context)
+void FSequenceUpdater_Flat::Update(UMovieSceneEntitySystemLinker* Linker, FRootInstanceHandle InstanceHandle, TSharedRef<const FSharedPlaybackState> SharedPlaybackState, const FMovieSceneContext& Context)
 {
 	FSequenceInstance& SequenceInstance = Linker->GetInstanceRegistry()->MutateInstance(InstanceHandle);
 	SequenceInstance.SetContext(Context);
 
-	const FMovieSceneEntityComponentField* ComponentField = InPlayer->GetEvaluationTemplate().GetCompiledDataManager()->FindEntityComponentField(CompiledDataID);
-	UMovieSceneSequence* Sequence = InPlayer->GetEvaluationTemplate().GetSequence(MovieSceneSequenceID::Root);
+	UMovieSceneCompiledDataManager* CompiledDataManager = SharedPlaybackState->GetCompiledDataManager();
+	const FMovieSceneEntityComponentField* ComponentField = CompiledDataManager->FindEntityComponentField(CompiledDataID);
+	UMovieSceneSequence* Sequence = SharedPlaybackState->GetRootSequence();
 	if (Sequence == nullptr)
 	{
 		SequenceInstance.Ledger.UnlinkEverything(Linker);
@@ -264,7 +268,11 @@ void FSequenceUpdater_Flat::Update(UMovieSceneEntitySystemLinker* Linker, FRootI
 
 	if (!bDynamicWeighting.IsSet())
 	{
-		bDynamicWeighting = InPlayer->HasDynamicWeighting() || EnumHasAnyFlags(Sequence->GetFlags(), EMovieSceneSequenceFlags::DynamicWeighting);
+		bDynamicWeighting = EnumHasAnyFlags(Sequence->GetFlags(), EMovieSceneSequenceFlags::DynamicWeighting);
+		if (IMovieScenePlayer* Player = FPlayerIndexPlaybackCapability::GetPlayer(SharedPlaybackState))
+		{
+			bDynamicWeighting = bDynamicWeighting || Player->HasDynamicWeighting();
+		}
 	}
 
 	FMovieSceneEvaluationFieldEntitySet EntitiesScratch;
@@ -321,7 +329,7 @@ bool FSequenceUpdater_Flat::CanFinishImmediately(UMovieSceneEntitySystemLinker*
 	return SequenceInstance.Ledger.IsEmpty();
 }
 
-void FSequenceUpdater_Flat::Finish(UMovieSceneEntitySystemLinker* Linker, FRootInstanceHandle InstanceHandle, IMovieScenePlayer* InPlayer)
+void FSequenceUpdater_Flat::Finish(UMovieSceneEntitySystemLinker* Linker, FRootInstanceHandle InstanceHandle, TSharedRef<const FSharedPlaybackState> SharedPlaybackState)
 {
 	InvalidateCachedData(Linker, InstanceHandle);
 }
@@ -352,9 +360,12 @@ FSequenceUpdater_Hierarchical::~FSequenceUpdater_Hierarchical()
 {
 }
 
-void FSequenceUpdater_Hierarchical::PopulateUpdateFlags(UMovieSceneEntitySystemLinker* InLinker, IMovieScenePlayer* InPlayer, ESequenceInstanceUpdateFlags& OutUpdateFlags)
+void FSequenceUpdater_Hierarchical::PopulateUpdateFlags(UMovieSceneEntitySystemLinker* InLinker, TSharedRef<const FSharedPlaybackState> SharedPlaybackState, ESequenceInstanceUpdateFlags& OutUpdateFlags)
 {
-	InPlayer->PopulateUpdateFlags(OutUpdateFlags);
+	if (IMovieScenePlayer* Player = FPlayerIndexPlaybackCapability::GetPlayer(SharedPlaybackState))
+	{
+		Player->PopulateUpdateFlags(OutUpdateFlags);
+	}
 
 	// If we've already been told we need dissection there's nothing left to do
 	if (EnumHasAnyFlags(OutUpdateFlags, ESequenceInstanceUpdateFlags::NeedsDissection))
@@ -362,7 +373,7 @@ void FSequenceUpdater_Hierarchical::PopulateUpdateFlags(UMovieSceneEntitySystemL
 		return;
 	}
 
-	UMovieSceneCompiledDataManager* CompiledDataManager = InPlayer->GetEvaluationTemplate().GetCompiledDataManager();
+	UMovieSceneCompiledDataManager* CompiledDataManager = SharedPlaybackState->GetCompiledDataManager();
 
 	{
 		const FMovieSceneCompiledDataEntry& RootDataEntry = CompiledDataManager->GetEntryRef(CompiledDataID);
@@ -388,9 +399,9 @@ void FSequenceUpdater_Hierarchical::PopulateUpdateFlags(UMovieSceneEntitySystemL
 	}
 }
 
-void FSequenceUpdater_Hierarchical::DissectContext(UMovieSceneEntitySystemLinker* Linker, IMovieScenePlayer* InPlayer, const FMovieSceneContext& Context, TArray<TRange<FFrameTime>>& OutDissections)
+void FSequenceUpdater_Hierarchical::DissectContext(UMovieSceneEntitySystemLinker* Linker, TSharedRef<const FSharedPlaybackState> SharedPlaybackState, const FMovieSceneContext& Context, TArray<TRange<FFrameTime>>& OutDissections)
 {
-	UMovieSceneCompiledDataManager* CompiledDataManager = InPlayer->GetEvaluationTemplate().GetCompiledDataManager();
+	UMovieSceneCompiledDataManager* CompiledDataManager = SharedPlaybackState->GetCompiledDataManager();
 
 	TRange<FFrameNumber> TraversedRange = Context.GetFrameNumberRange();
 	TArray<FFrameTime>   RootDissectionTimes;
@@ -457,7 +468,7 @@ void FSequenceUpdater_Hierarchical::DissectContext(UMovieSceneEntitySystemLinker
 	}
 }
 
-FInstanceHandle FSequenceUpdater_Hierarchical::GetOrCreateSequenceInstance(IMovieScenePlayer* InPlayer, UMovieSceneSequence* SubSequence, const FMovieSceneSequenceHierarchy* Hierarchy, FInstanceRegistry* InstanceRegistry, FRootInstanceHandle RootInstanceHandle, FMovieSceneSequenceID SequenceID)
+FInstanceHandle FSequenceUpdater_Hierarchical::GetOrCreateSequenceInstance(TSharedRef<const FSharedPlaybackState> SharedPlaybackState, UMovieSceneSequence* SubSequence, const FMovieSceneSequenceHierarchy* Hierarchy, FInstanceRegistry* InstanceRegistry, FRootInstanceHandle RootInstanceHandle, FMovieSceneSequenceID SequenceID)
 {
 	check(SequenceID != MovieSceneSequenceID::Root);
 
@@ -467,6 +478,7 @@ FInstanceHandle FSequenceUpdater_Hierarchical::GetOrCreateSequenceInstance(IMovi
 	}
 
 	const FMovieSceneSequenceHierarchyNode* Node = Hierarchy->FindNode(SequenceID);
+	IMovieScenePlayer* Player = FPlayerIndexPlaybackCapability::GetPlayer(SharedPlaybackState);
 
 	checkf(Node, TEXT("Attempting to construct a new sub sequence instance with a sub sequence ID that does not exist in the hierarchy"));
 	checkf(Node->ParentID != MovieSceneSequenceID::Invalid, TEXT("Parent should never be invalid for a non-root SequenceID"));
@@ -478,10 +490,10 @@ FInstanceHandle FSequenceUpdater_Hierarchical::GetOrCreateSequenceInstance(IMovi
 	}
 	else if (UMovieSceneSequence* ParentSequence = Hierarchy->FindSubData(Node->ParentID)->GetSequence())
 	{
-		ParentInstance = GetOrCreateSequenceInstance(InPlayer, ParentSequence, Hierarchy, InstanceRegistry, RootInstanceHandle, Node->ParentID);
+		ParentInstance = GetOrCreateSequenceInstance(SharedPlaybackState, ParentSequence, Hierarchy, InstanceRegistry, RootInstanceHandle, Node->ParentID);
 	}
 
-	FInstanceHandle InstanceHandle = InstanceRegistry->AllocateSubInstance(InPlayer, SequenceID, RootInstanceHandle, ParentInstance);
+	FInstanceHandle InstanceHandle = InstanceRegistry->AllocateSubInstance(Player, SequenceID, RootInstanceHandle, ParentInstance);
 	SequenceInstances.Add(SequenceID, FSubInstanceData { SubSequence->GetSignature(), InstanceHandle });
 
 	return InstanceHandle;
@@ -506,14 +518,15 @@ void FSequenceUpdater_Hierarchical::OverrideRootSequence(UMovieSceneEntitySystem
 	}
 }
 
-void FSequenceUpdater_Hierarchical::Start(UMovieSceneEntitySystemLinker* Linker, FRootInstanceHandle InstanceHandle, IMovieScenePlayer* InPlayer, const FMovieSceneContext& InContext)
+void FSequenceUpdater_Hierarchical::Start(UMovieSceneEntitySystemLinker* Linker, FRootInstanceHandle InstanceHandle, TSharedRef<const FSharedPlaybackState> SharedPlaybackState, const FMovieSceneContext& InContext)
 {
 }
 
-void FSequenceUpdater_Hierarchical::Update(UMovieSceneEntitySystemLinker* Linker, FRootInstanceHandle InstanceHandle, IMovieScenePlayer* InPlayer, const FMovieSceneContext& Context)
+void FSequenceUpdater_Hierarchical::Update(UMovieSceneEntitySystemLinker* Linker, FRootInstanceHandle InstanceHandle, TSharedRef<const FSharedPlaybackState> SharedPlaybackState, const FMovieSceneContext& Context)
 {
 	FInstanceRegistry* InstanceRegistry = Linker->GetInstanceRegistry();
-	UMovieSceneCompiledDataManager* CompiledDataManager = InPlayer->GetEvaluationTemplate().GetCompiledDataManager();
+	UMovieSceneCompiledDataManager* CompiledDataManager = SharedPlaybackState->GetCompiledDataManager();
+	IMovieScenePlayer* Player = FPlayerIndexPlaybackCapability::GetPlayer(SharedPlaybackState);
 
 	FMovieSceneEvaluationFieldEntitySet EntitiesScratch;
 
@@ -532,7 +545,7 @@ void FSequenceUpdater_Hierarchical::Update(UMovieSceneEntitySystemLinker* Linker
 		UMovieSceneSequence*                RootSequence    = SubData ? SubData->GetSequence() : nullptr;
 		if (ensure(RootSequence))
 		{
-			RootInstanceHandle = GetOrCreateSequenceInstance(InPlayer, RootSequence, RootHierarchy, InstanceRegistry, InstanceHandle, RootOverrideSequenceID);
+			RootInstanceHandle = GetOrCreateSequenceInstance(SharedPlaybackState, RootSequence, RootHierarchy, InstanceRegistry, InstanceHandle, RootOverrideSequenceID);
 			RootCompiledDataID = CompiledDataManager->GetDataID(RootSequence);
 			RootContext        = Context.Transform(SubData->RootToSequenceTransform, SubData->TickResolution);
 
@@ -563,7 +576,11 @@ void FSequenceUpdater_Hierarchical::Update(UMovieSceneEntitySystemLinker* Linker
 		{
 			if (!bDynamicWeighting.IsSet())
 			{
-				bDynamicWeighting = InPlayer->HasDynamicWeighting() || EnumHasAnyFlags(CompiledDataManager->GetEntryRef(RootCompiledDataID).AccumulatedFlags, EMovieSceneSequenceFlags::DynamicWeighting);
+				bDynamicWeighting = EnumHasAnyFlags(CompiledDataManager->GetEntryRef(RootCompiledDataID).AccumulatedFlags, EMovieSceneSequenceFlags::DynamicWeighting);
+				if (Player)
+				{
+					bDynamicWeighting = bDynamicWeighting || Player->HasDynamicWeighting();
+				}
 			}
 
 			// Update entities if necessary
@@ -641,7 +658,7 @@ void FSequenceUpdater_Hierarchical::Update(UMovieSceneEntitySystemLinker* Linker
 				FMovieSceneCompiledDataID SubDataID = CompiledDataManager->GetDataID(SubSequence);
 
 				// Set the context for the root sequence instance
-				FInstanceHandle    SubSequenceHandle = GetOrCreateSequenceInstance(InPlayer, SubSequence, RootHierarchy, InstanceRegistry, InstanceHandle, SequenceIDFromRoot);
+				FInstanceHandle    SubSequenceHandle = GetOrCreateSequenceInstance(SharedPlaybackState, SubSequence, RootHierarchy, InstanceRegistry, InstanceHandle, SequenceIDFromRoot);
 				FSequenceInstance& SubSequenceInstance = InstanceRegistry->MutateInstance(SubSequenceHandle);
 
 				// Update the sub sequence's context
@@ -766,14 +783,14 @@ bool FSequenceUpdater_Hierarchical::CanFinishImmediately(UMovieSceneEntitySystem
 	return true;
 }
 
-void FSequenceUpdater_Hierarchical::Finish(UMovieSceneEntitySystemLinker* Linker, FRootInstanceHandle InstanceHandle, IMovieScenePlayer* InPlayer)
+void FSequenceUpdater_Hierarchical::Finish(UMovieSceneEntitySystemLinker* Linker, FRootInstanceHandle InstanceHandle, TSharedRef<const FSharedPlaybackState> SharedPlaybackState)
 {
 	FInstanceRegistry* InstanceRegistry = Linker->GetInstanceRegistry();
 
 	// Finish all sub sequences as well
 	for (TPair<FMovieSceneSequenceID, FSubInstanceData> Pair : SequenceInstances)
 	{
-		InstanceRegistry->MutateInstance(Pair.Value.Handle).Finish(Linker);
+		InstanceRegistry->MutateInstance(Pair.Value.Handle).Finish();
 	}
 
 	InvalidateCachedData(Linker, InstanceHandle);
diff --git a/Engine/Source/Runtime/MovieScene/Private/EntitySystem/MovieSceneSharedPlaybackState.cpp b/Engine/Source/Runtime/MovieScene/Private/EntitySystem/MovieSceneSharedPlaybackState.cpp
new file mode 100644
--- /dev/null
+++ b/Engine/Source/Runtime/MovieScene/Private/EntitySystem/MovieSceneSharedPlaybackState.cpp
@@ -0,0 +1,62 @@
+// Copyright Epic Games, Inc. All Rights Reserved.
+
+#include "EntitySystem/MovieSceneSharedPlaybackState.h"
+
+#include "Compilation/MovieSceneCompiledDataManager.h"
+#include "EntitySystem/MovieSceneEntitySystemRunner.h"
+
+namespace UE::MovieScene
+{
+
+FSharedPlaybackState::FSharedPlaybackState()
+{
+}
+
+FSharedPlaybackState::FSharedPlaybackState(
+		UMovieSceneSequence& InRootSequence,
+		const FSharedPlaybackStateCreateParams& CreateParams)
+	: WeakRootSequence(&InRootSequence)
+	, WeakPlaybackContext(CreateParams.PlaybackContext)
+	, WeakRunner(CreateParams.Runner)
+	, CompiledDataManager(CreateParams.CompiledDataManager)
+	, RootInstanceHandle(CreateParams.RootInstanceHandle)
+{
+	if (CompiledDataManager)
+	{
+		RootCompiledDataID = CompiledDataManager->GetDataID(&InRootSequence);
+	}
+}
+
+UMovieSceneEntitySystemLinker* FSharedPlaybackState::GetLinker() const
+{
+	if (TSharedPtr<FMovieSceneEntitySystemRunner> Runner = WeakRunner.Pin())
+	{
+		return Runner->GetLinker();
+	}
+	return nullptr;
+}
+
+const FMovieSceneSequenceHierarchy* FSharedPlaybackState::GetHierarchy() const
+{
+	if (CompiledDataManager && RootCompiledDataID.IsValid())
+	{
+		return CompiledDataManager->FindHierarchy(RootCompiledDataID);
+	}
+	return nullptr;
+}
+
+UMovieSceneSequence* FSharedPlaybackState::GetSequence(FMovieSceneSequenceIDRef InSequenceID) const
+{
+	const FMovieSceneSequenceHierarchy* Hierarchy = GetHierarchy();
+	const FMovieSceneSubSequenceData*   SubData   = Hierarchy ? Hierarchy->FindSubData(InSequenceID) : nullptr;
+	return SubData ? SubData->GetSequence() : nullptr;
+}
+
+void FSharedPlaybackState::InvalidateCachedData()
+{
+	UMovieSceneEntitySystemLinker* Linker = GetLinker();
+	Capabilities.InvalidateCachedData(Linker);
+}
+
+} // namespace UE::MovieScene
+
diff --git a/Engine/Source/Runtime/MovieScene/Private/Evaluation/MovieSceneEvaluationTemplateInstance.cpp b/Engine/Source/Runtime/MovieScene/Private/Evaluation/MovieSceneEvaluationTemplateInstance.cpp
--- a/Engine/Source/Runtime/MovieScene/Private/Evaluation/MovieSceneEvaluationTemplateInstance.cpp
+++ b/Engine/Source/Runtime/MovieScene/Private/Evaluation/MovieSceneEvaluationTemplateInstance.cpp
@@ -2,6 +2,7 @@
 
 #include "Evaluation/MovieSceneEvaluationTemplateInstance.h"
 #include "Containers/SortedMap.h"
+#include "EntitySystem/MovieSceneSequenceInstanceHandle.h"
 #include "IMovieScenePlayer.h"
 #include "MovieSceneSequence.h"
 #include "MovieSceneSequence.h"
@@ -22,8 +23,7 @@ DECLARE_CYCLE_STAT(TEXT("Entire Evaluation Cost"), MovieSceneEval_EntireEvaluati
 
 
 FMovieSceneRootEvaluationTemplateInstance::FMovieSceneRootEvaluationTemplateInstance()
-	: CompiledDataManager(nullptr)
-	, EntitySystemLinker(nullptr)
+	: EntitySystemLinker(nullptr)
 	, RootID(MovieSceneSequenceID::Root)
 {
 #if WITH_EDITOR
@@ -36,13 +36,17 @@ void FMovieSceneRootEvaluationTemplateInstance::TearDown()
 	using namespace UE::MovieScene;
 
 	// Avoid redundant work if the linker is being destroyed anyway
-	if (RootInstanceHandle.IsValid() && EntitySystemLinker && IsValidChecked(EntitySystemLinker) && !EntitySystemLinker->IsUnreachable() && !EntitySystemLinker->HasAnyFlags(RF_BeginDestroyed))
+	if (SharedPlaybackState.IsValid() && 
+			SharedPlaybackState->GetRootInstanceHandle().IsValid() &&
+			EntitySystemLinker && 
+			IsValidChecked(EntitySystemLinker) && 
+			!EntitySystemLinker->IsUnreachable() && 
+			!EntitySystemLinker->HasAnyFlags(RF_BeginDestroyed))
 	{
-		EntitySystemLinker->DestroyInstanceImmediately(RootInstanceHandle);
+		EntitySystemLinker->DestroyInstanceImmediately(SharedPlaybackState->GetRootInstanceHandle());
 	}
 
-	RootInstanceHandle = UE::MovieScene::FRootInstanceHandle();
-	CompiledDataManager = nullptr;
+	SharedPlaybackState.Reset();
 	EntitySystemLinker = nullptr;
 }
 
@@ -63,20 +67,21 @@ UMovieSceneEntitySystemLinker* FMovieSceneRootEvaluationTemplateInstance::Constr
 	return UMovieSceneEntitySystemLinker::FindOrCreateLinker(PlaybackContext, UE::MovieScene::EEntitySystemLinkerRole::Standalone, TEXT("DefaultEntitySystemLinker"));
 }
 
-void FMovieSceneRootEvaluationTemplateInstance::Initialize(UMovieSceneSequence& InRootSequence, IMovieScenePlayer& Player, UMovieSceneCompiledDataManager* InCompiledDataManager, TWeakPtr<FMovieSceneEntitySystemRunner> InWeakRunner)
+void FMovieSceneRootEvaluationTemplateInstance::Initialize(UMovieSceneSequence& InRootSequence, IMovieScenePlayer& Player, UMovieSceneCompiledDataManager* InCompiledDataManager, TSharedPtr<FMovieSceneEntitySystemRunner> InRunner)
 {
+	using namespace UE::MovieScene;
+
 	bool bReinitialize = (
 			// Initialize if we weren't initialized before and this is our first sequence.
-			WeakRootSequence.Get() == nullptr ||
+			!SharedPlaybackState.IsValid() ||
+			SharedPlaybackState->GetRootSequence() == nullptr ||
 			// Initialize if we lost our linker.
 			EntitySystemLinker == nullptr);
 
-	const UMovieSceneCompiledDataManager* PreviousCompiledDataManager = CompiledDataManager;
-	if (InCompiledDataManager)
-	{
-		CompiledDataManager = InCompiledDataManager;
-	}
-	else
+	// Reinitialize if the compiled data manager has changed.
+	const UMovieSceneCompiledDataManager* PreviousCompiledDataManager = SharedPlaybackState ? 
+		SharedPlaybackState->GetCompiledDataManager() : nullptr;
+	if (!InCompiledDataManager)
 	{
 #if WITH_EDITOR
 		EMovieSceneServerClientMask Mask = EmulatedNetworkMask;
@@ -98,39 +103,37 @@ void FMovieSceneRootEvaluationTemplateInstance::Initialize(UMovieSceneSequence&
 				}
 			}
 		}
-		CompiledDataManager = UMovieSceneCompiledDataManager::GetPrecompiledData(Mask);
+		InCompiledDataManager = UMovieSceneCompiledDataManager::GetPrecompiledData(Mask);
 #else
-		CompiledDataManager = UMovieSceneCompiledDataManager::GetPrecompiledData();
+		InCompiledDataManager = UMovieSceneCompiledDataManager::GetPrecompiledData();
 #endif
 	}
-	bReinitialize |= (PreviousCompiledDataManager != CompiledDataManager);
+	bReinitialize |= (PreviousCompiledDataManager != InCompiledDataManager);
 
-	TSharedPtr<FMovieSceneEntitySystemRunner> PreviousRunner = WeakRunner.Pin();
+	// Reinitialize if the runner has changed.
+	TSharedPtr<FMovieSceneEntitySystemRunner> PreviousRunner = SharedPlaybackState ?
+		SharedPlaybackState->GetRunner() : nullptr;
+	bReinitialize |= (PreviousRunner != InRunner);
 
-	if (UMovieSceneSequence* ExistingSequence = WeakRootSequence.Get())
-	{
-		if (ExistingSequence != &InRootSequence)
-		{
-			bReinitialize = true;
-		}
-	}
+	// Reinitialize if the root sequence has changed.
+	UMovieSceneSequence* PreviousRootSequence = SharedPlaybackState ? 
+		SharedPlaybackState->GetRootSequence() : nullptr;
+	bReinitialize |= (PreviousRootSequence != &InRootSequence);
 
-	CompiledDataID = CompiledDataManager->GetDataID(&InRootSequence);
-	WeakRootSequence = &InRootSequence;
 	RootID = MovieSceneSequenceID::Root;
-	WeakRunner = InWeakRunner;
 
 	if (bReinitialize)
 	{
-		if (RootInstanceHandle.IsValid())
+		// Tear down previous sequence first.
+		if (SharedPlaybackState.IsValid() && SharedPlaybackState->GetRootInstanceHandle().IsValid())
 		{
 			if (PreviousRunner)
 			{
-				PreviousRunner->AbandonAndDestroyInstance(RootInstanceHandle);
+				PreviousRunner->AbandonAndDestroyInstance(SharedPlaybackState->GetRootInstanceHandle());
 			}
 			else if (EntitySystemLinker)
 			{
-				EntitySystemLinker->GetInstanceRegistry()->DestroyInstance(RootInstanceHandle);
+				EntitySystemLinker->GetInstanceRegistry()->DestroyInstance(SharedPlaybackState->GetRootInstanceHandle());
 			}
 			else
 			{
@@ -141,27 +144,40 @@ void FMovieSceneRootEvaluationTemplateInstance::Initialize(UMovieSceneSequence&
 		Player.State.PersistentEntityData.Reset();
 		Player.State.PersistentSharedData.Reset();
 		
+		// Build a new linker and immediately attach the runner to it, so that we can find both
+		// when initializing the new root instance and shared playback state.
 		EntitySystemLinker = ConstructEntityLinker(Player);
-		if (EntitySystemLinker != nullptr && EntitySystemLinker->GetInstanceRegistry())
-		{
-			RootInstanceHandle = EntitySystemLinker->GetInstanceRegistry()->AllocateRootInstance(&Player);
-		}
 
-		TSharedPtr<FMovieSceneEntitySystemRunner> Runner = WeakRunner.Pin();
-		if (ensure(Runner) && EntitySystemLinker)
+		if (ensure(InRunner) && EntitySystemLinker)
 		{
-			if (Runner->IsAttachedToLinker() && Runner->GetLinker() != EntitySystemLinker)
+			if (InRunner->IsAttachedToLinker() && InRunner->GetLinker() != EntitySystemLinker)
 			{
-				Runner->DetachFromLinker();
+				InRunner->DetachFromLinker();
 			}
 
-			if (!Runner->IsAttachedToLinker())
+			if (!InRunner->IsAttachedToLinker())
 			{
-				Runner->AttachToLinker(EntitySystemLinker);
+				InRunner->AttachToLinker(EntitySystemLinker);
 			}
 		}
 
-		Player.PreAnimatedState.Initialize(EntitySystemLinker, RootInstanceHandle);
+		// Create the new root instance and save its new shared playback state.
+		FRootInstanceHandle RootInstanceHandle;
+		if (EntitySystemLinker != nullptr && EntitySystemLinker->GetInstanceRegistry())
+		{
+			FInstanceRegistry* InstanceRegistry = EntitySystemLinker->GetInstanceRegistry();
+			RootInstanceHandle = InstanceRegistry->AllocateRootInstance(&Player, InRootSequence, InRunner, InCompiledDataManager);
+			SharedPlaybackState = InstanceRegistry->GetInstance(RootInstanceHandle).GetSharedPlaybackState();
+		}
+		else
+		{
+			SharedPlaybackState.Reset();
+		}
+
+		if (RootInstanceHandle.IsValid())
+		{
+			Player.PreAnimatedState.Initialize(EntitySystemLinker, RootInstanceHandle);
+		}
 
 		ResetDirectorInstances();
 	}
@@ -169,14 +185,19 @@ void FMovieSceneRootEvaluationTemplateInstance::Initialize(UMovieSceneSequence&
 
 void FMovieSceneRootEvaluationTemplateInstance::Evaluate(FMovieSceneContext Context, IMovieScenePlayer& Player)
 {
-	EvaluateSynchronousBlocking(Context, Player);
+	EvaluateSynchronousBlocking(Context);
 }
 
 void FMovieSceneRootEvaluationTemplateInstance::EvaluateSynchronousBlocking(FMovieSceneContext Context, IMovieScenePlayer& Player)
 {
-	if (TSharedPtr<FMovieSceneEntitySystemRunner> Runner = WeakRunner.Pin())
+	EvaluateSynchronousBlocking(Context);
+}
+
+void FMovieSceneRootEvaluationTemplateInstance::EvaluateSynchronousBlocking(FMovieSceneContext Context)
+{
+	if (TSharedPtr<FMovieSceneEntitySystemRunner> Runner = GetRunner())
 	{
-		Runner->QueueUpdate(Context, RootInstanceHandle);
+		Runner->QueueUpdate(Context, GetRootInstanceHandle());
 		Runner->Flush();
 	}
 }
@@ -191,11 +212,48 @@ void FMovieSceneRootEvaluationTemplateInstance::ResetDirectorInstances()
 	DirectorInstances.Reset();
 }
 
+bool FMovieSceneRootEvaluationTemplateInstance::IsValid() const
+{
+	return EntitySystemLinker && SharedPlaybackState;
+}
+
+UE::MovieScene::FRootInstanceHandle FMovieSceneRootEvaluationTemplateInstance::GetRootInstanceHandle() const
+{
+	return SharedPlaybackState ? SharedPlaybackState->GetRootInstanceHandle() : UE::MovieScene::FRootInstanceHandle();
+}
+
+UMovieSceneSequence* FMovieSceneRootEvaluationTemplateInstance::GetRootSequence() const
+{
+	using namespace UE::MovieScene;
+
+	if (SharedPlaybackState)
+	{
+		return SharedPlaybackState->GetRootSequence();
+	}
+	return nullptr;
+}
+
+FMovieSceneCompiledDataID FMovieSceneRootEvaluationTemplateInstance::GetCompiledDataID() const
+{
+	return SharedPlaybackState ? SharedPlaybackState->GetRootCompiledDataID() : FMovieSceneCompiledDataID();
+}
+
+UMovieSceneCompiledDataManager* FMovieSceneRootEvaluationTemplateInstance::GetCompiledDataManager() const
+{
+	return SharedPlaybackState ? SharedPlaybackState->GetCompiledDataManager() : nullptr;
+}
+
+TSharedPtr<FMovieSceneEntitySystemRunner> FMovieSceneRootEvaluationTemplateInstance::GetRunner() const
+{
+	return SharedPlaybackState ? SharedPlaybackState->GetRunner() : TSharedPtr<FMovieSceneEntitySystemRunner>();
+}
+
 void FMovieSceneRootEvaluationTemplateInstance::EnableGlobalPreAnimatedStateCapture()
 {
 	using namespace UE::MovieScene;
 
-	if (ensure(EntitySystemLinker))
+	const FRootInstanceHandle RootInstanceHandle = GetRootInstanceHandle();
+	if (ensure(EntitySystemLinker && RootInstanceHandle.IsValid()))
 	{
 		const FSequenceInstance& Instance = EntitySystemLinker->GetInstanceRegistry()->GetInstance(RootInstanceHandle);
 		Instance.GetPlayer()->PreAnimatedState.EnableGlobalPreAnimatedStateCapture();
@@ -206,12 +264,11 @@ UMovieSceneSequence* FMovieSceneRootEvaluationTemplateInstance::GetSequence(FMov
 {
 	if (SequenceID == MovieSceneSequenceID::Root)
 	{
-		return WeakRootSequence.Get();
+		return SharedPlaybackState->GetRootSequence();
 	}
-	else if (CompiledDataID.IsValid())
+	else if (const FMovieSceneSequenceHierarchy* Hierarchy = SharedPlaybackState->GetHierarchy())
 	{
-		const FMovieSceneSequenceHierarchy* Hierarchy       = CompiledDataManager->FindHierarchy(CompiledDataID);
-		const FMovieSceneSubSequenceData*   SubSequenceData = Hierarchy ? Hierarchy->FindSubData(SequenceID) : nullptr;
+		const FMovieSceneSubSequenceData* SubSequenceData = Hierarchy->FindSubData(SequenceID);
 		if (SubSequenceData)
 		{
 			return SubSequenceData->GetSequence();
@@ -227,10 +284,11 @@ UMovieSceneEntitySystemLinker* FMovieSceneRootEvaluationTemplateInstance::GetEnt
 
 bool FMovieSceneRootEvaluationTemplateInstance::HasEverUpdated() const
 {
-	if (EntitySystemLinker)
+	const UE::MovieScene::FRootInstanceHandle RootInstanceHandle = GetRootInstanceHandle();
+	if (EntitySystemLinker && RootInstanceHandle.IsValid())
 	{
-		const UE::MovieScene::FSequenceInstance* SequenceInstance = &EntitySystemLinker->GetInstanceRegistry()->GetInstance(RootInstanceHandle);
-		return SequenceInstance && SequenceInstance->HasEverUpdated();
+		const UE::MovieScene::FSequenceInstance& SequenceInstance = EntitySystemLinker->GetInstanceRegistry()->GetInstance(RootInstanceHandle);
+		return SequenceInstance.HasEverUpdated();
 	}
 
 	return false;
@@ -238,7 +296,7 @@ bool FMovieSceneRootEvaluationTemplateInstance::HasEverUpdated() const
 
 const FMovieSceneSequenceHierarchy* FMovieSceneRootEvaluationTemplateInstance::GetHierarchy() const
 {
-	return CompiledDataManager->FindHierarchy(GetCompiledDataID());
+	return SharedPlaybackState ? SharedPlaybackState->GetHierarchy() : nullptr;
 }
 
 void FMovieSceneRootEvaluationTemplateInstance::GetSequenceParentage(const UE::MovieScene::FInstanceHandle InstanceHandle, TArray<UE::MovieScene::FInstanceHandle>& OutParentHandles) const
@@ -256,6 +314,7 @@ void FMovieSceneRootEvaluationTemplateInstance::GetSequenceParentage(const UE::M
 	check(InstanceHandle.IsValid());
 	const FSequenceInstance& Instance = InstanceRegistry->GetInstance(InstanceHandle);
 
+	const UE::MovieScene::FRootInstanceHandle RootInstanceHandle = GetRootInstanceHandle();
 	checkf(Instance.GetRootInstanceHandle() == RootInstanceHandle, TEXT("The provided instance handle relates to a different root sequence."));
 	const FSequenceInstance& RootInstance = InstanceRegistry->GetInstance(RootInstanceHandle);
 
@@ -282,11 +341,23 @@ void FMovieSceneRootEvaluationTemplateInstance::GetSequenceParentage(const UE::M
 	}
 }
 
+const UE::MovieScene::FSequenceInstance* FMovieSceneRootEvaluationTemplateInstance::GetRootInstance() const
+{
+	const UE::MovieScene::FRootInstanceHandle RootInstanceHandle = GetRootInstanceHandle();
+	if (RootInstanceHandle.IsValid())
+	{
+		const UE::MovieScene::FInstanceRegistry* InstanceRegistry = EntitySystemLinker->GetInstanceRegistry();
+		return &InstanceRegistry->GetInstance(RootInstanceHandle);
+	}
+	return nullptr;
+}
+
 UE::MovieScene::FSequenceInstance* FMovieSceneRootEvaluationTemplateInstance::FindInstance(FMovieSceneSequenceID SequenceID)
 {
 	using namespace UE::MovieScene;
 
-	if (ensure(EntitySystemLinker))
+	const UE::MovieScene::FRootInstanceHandle RootInstanceHandle = GetRootInstanceHandle();
+	if (ensure(EntitySystemLinker && RootInstanceHandle.IsValid()))
 	{
 		FSequenceInstance* SequenceInstance = &EntitySystemLinker->GetInstanceRegistry()->MutateInstance(RootInstanceHandle);
 
@@ -309,7 +380,8 @@ const UE::MovieScene::FSequenceInstance* FMovieSceneRootEvaluationTemplateInstan
 {
 	using namespace UE::MovieScene;
 
-	if (ensure(EntitySystemLinker))
+	const UE::MovieScene::FRootInstanceHandle RootInstanceHandle = GetRootInstanceHandle();
+	if (ensure(EntitySystemLinker && RootInstanceHandle.IsValid()))
 	{
 		const FSequenceInstance* SequenceInstance = &EntitySystemLinker->GetInstanceRegistry()->GetInstance(RootInstanceHandle);
 
@@ -372,12 +444,12 @@ UObject* FMovieSceneRootEvaluationTemplateInstance::GetOrCreateDirectorInstance(
 	UObject* NewDirectorInstance = nullptr;
 	if (SequenceID == MovieSceneSequenceID::Root)
 	{
-		if (UMovieSceneSequence* Sequence = WeakRootSequence.Get())
+		if (UMovieSceneSequence* Sequence = GetRootSequence())
 		{
 			NewDirectorInstance = Sequence->CreateDirectorInstance(Player, SequenceID);
 		}
 	}
-	else if (const FMovieSceneSequenceHierarchy* Hierarchy = CompiledDataManager->FindHierarchy(CompiledDataID))
+	else if (const FMovieSceneSequenceHierarchy* Hierarchy = GetHierarchy())
 	{
 		const FMovieSceneSubSequenceData* SubData = Hierarchy->FindSubData(SequenceID);
 		check(SubData);
@@ -396,15 +468,25 @@ void FMovieSceneRootEvaluationTemplateInstance::PlaybackContextChanged(IMovieSce
 {
 	using namespace UE::MovieScene;
 
+	// Only the playback context changed, so we keep the same sequence, runner, and compiled data manager.
+	UMovieSceneSequence* RootSequence = SharedPlaybackState->GetRootSequence();
+	TSharedPtr<FMovieSceneEntitySystemRunner> Runner = SharedPlaybackState->GetRunner();
+	UMovieSceneCompiledDataManager* CompiledDataManager = SharedPlaybackState->GetCompiledDataManager();
+
 	const bool bGlobalCapture = Player.PreAnimatedState.IsCapturingGlobalPreAnimatedState();
+	const FRootInstanceHandle PreviousRootInstanceHandle = SharedPlaybackState->GetRootInstanceHandle();
 
-	if (RootInstanceHandle.IsValid() && EntitySystemLinker && IsValidChecked(EntitySystemLinker) && !EntitySystemLinker->IsUnreachable() && !EntitySystemLinker->HasAnyFlags(RF_BeginDestroyed))
+	if (PreviousRootInstanceHandle.IsValid() &&
+			EntitySystemLinker && 
+			IsValidChecked(EntitySystemLinker) &&
+			!EntitySystemLinker->IsUnreachable() && 
+			!EntitySystemLinker->HasAnyFlags(RF_BeginDestroyed))
 	{
 		EntitySystemLinker->CleanupInvalidBoundObjects();
 
-		if (TSharedPtr<FMovieSceneEntitySystemRunner> Runner = WeakRunner.Pin())
+		if (Runner)
 		{
-			if (Runner->QueueFinalUpdate(RootInstanceHandle))
+			if (Runner->QueueFinalUpdate(PreviousRootInstanceHandle))
 			{
 				Runner->Flush();
 			}
@@ -414,11 +496,11 @@ void FMovieSceneRootEvaluationTemplateInstance::PlaybackContextChanged(IMovieSce
 		{
 			Player.RestorePreAnimatedState();
 		}
-		EntitySystemLinker->GetInstanceRegistry()->DestroyInstance(RootInstanceHandle);
+		EntitySystemLinker->GetInstanceRegistry()->DestroyInstance(PreviousRootInstanceHandle);
 	}
 
 	EntitySystemLinker = ConstructEntityLinker(Player);
-	if (TSharedPtr<FMovieSceneEntitySystemRunner> Runner = WeakRunner.Pin())
+	if (Runner)
 	{
 		if (Runner->IsAttachedToLinker())
 		{
@@ -428,7 +510,10 @@ void FMovieSceneRootEvaluationTemplateInstance::PlaybackContextChanged(IMovieSce
 		Runner->AttachToLinker(EntitySystemLinker);
 	}
 
-	RootInstanceHandle = EntitySystemLinker->GetInstanceRegistry()->AllocateRootInstance(&Player);
+	FInstanceRegistry* InstanceRegistry = EntitySystemLinker->GetInstanceRegistry();
+	const FRootInstanceHandle RootInstanceHandle = InstanceRegistry->AllocateRootInstance(&Player, *RootSequence, Runner, CompiledDataManager);
+	SharedPlaybackState = InstanceRegistry->GetInstance(RootInstanceHandle).GetSharedPlaybackState();
+
 	DirectorInstances.Reset();
 
 	Player.PreAnimatedState.Initialize(EntitySystemLinker, RootInstanceHandle);
@@ -438,17 +523,20 @@ void FMovieSceneRootEvaluationTemplateInstance::PlaybackContextChanged(IMovieSce
 	}
 }
 
-
 const FMovieSceneSubSequenceData* FMovieSceneRootEvaluationTemplateInstance::FindSubData(FMovieSceneSequenceIDRef SequenceID) const
 {
-	const FMovieSceneSequenceHierarchy* Hierarchy = CompiledDataManager->FindHierarchy(CompiledDataID);
-	return Hierarchy ? Hierarchy->FindSubData(SequenceID) : nullptr;
+	if (const FMovieSceneSequenceHierarchy* Hierarchy = SharedPlaybackState->GetHierarchy())
+	{
+		return Hierarchy->FindSubData(SequenceID);
+	}
+	return nullptr;
 }
 
 void FMovieSceneRootEvaluationTemplateInstance::CopyActuators(FMovieSceneBlendingAccumulator& Accumulator) const
 {
 	using namespace UE::MovieScene;
 
+	FRootInstanceHandle RootInstanceHandle = SharedPlaybackState->GetRootInstanceHandle();
 	const FSequenceInstance& SequenceInstance = EntitySystemLinker->GetInstanceRegistry()->GetInstance(RootInstanceHandle);
 	const FMovieSceneTrackEvaluator* LegacyEvaluator = SequenceInstance.GetLegacyEvaluator();
 
@@ -460,6 +548,10 @@ void FMovieSceneRootEvaluationTemplateInstance::CopyActuators(FMovieSceneBlendin
 
 #if WITH_EDITOR
 void FMovieSceneRootEvaluationTemplateInstance::SetEmulatedNetworkMask(EMovieSceneServerClientMask InNewMask, IMovieScenePlayer& Player)
+{
+	SetEmulatedNetworkMask(InNewMask);
+}
+void FMovieSceneRootEvaluationTemplateInstance::SetEmulatedNetworkMask(EMovieSceneServerClientMask InNewMask)
 {
 	check(InNewMask != EMovieSceneServerClientMask::None);
 	EmulatedNetworkMask = InNewMask;
diff --git a/Engine/Source/Runtime/MovieScene/Private/Evaluation/PreAnimatedState/MovieScenePreAnimatedStateExtension.cpp b/Engine/Source/Runtime/MovieScene/Private/Evaluation/PreAnimatedState/MovieScenePreAnimatedStateExtension.cpp
--- a/Engine/Source/Runtime/MovieScene/Private/Evaluation/PreAnimatedState/MovieScenePreAnimatedStateExtension.cpp
+++ b/Engine/Source/Runtime/MovieScene/Private/Evaluation/PreAnimatedState/MovieScenePreAnimatedStateExtension.cpp
@@ -390,7 +390,7 @@ void FPreAnimatedStateExtension::RestoreGlobalState(const FRestoreStateParams& P
 	// Invalidate cached data for any sequence instance that belongs to the terminal instance
 	if (Params.TerminalInstanceHandle.IsValid())
 	{
-		Linker->GetInstanceRegistry()->MutateInstance(Params.TerminalInstanceHandle).InvalidateCachedData(Linker);
+		Linker->GetInstanceRegistry()->MutateInstance(Params.TerminalInstanceHandle).InvalidateCachedData();
 	}
 
 	bEntriesInvalidated = true;
@@ -507,7 +507,7 @@ void FPreAnimatedStateExtension::DiscardGlobalState(const FRestoreStateParams& P
 	// Invalidate cached data for any sequence instance that belongs to the terminal instance
 	if (Params.TerminalInstanceHandle.IsValid())
 	{
-		Linker->GetInstanceRegistry()->MutateInstance(Params.TerminalInstanceHandle).InvalidateCachedData(Linker);
+		Linker->GetInstanceRegistry()->MutateInstance(Params.TerminalInstanceHandle).InvalidateCachedData();
 	}
 
 	bEntriesInvalidated = true;
diff --git a/Engine/Source/Runtime/MovieScene/Private/MovieScenePlayer.cpp b/Engine/Source/Runtime/MovieScene/Private/MovieScenePlayer.cpp
--- a/Engine/Source/Runtime/MovieScene/Private/MovieScenePlayer.cpp
+++ b/Engine/Source/Runtime/MovieScene/Private/MovieScenePlayer.cpp
@@ -18,6 +18,26 @@ static FRWLock                          GGlobalPlayerRegistryLock;
 static TSparseArray<IMovieScenePlayer*> GGlobalPlayerRegistry;
 static TBitArray<> GGlobalPlayerUpdateFlags;
 
+TPlaybackCapabilityID<FPlayerIndexPlaybackCapability> FPlayerIndexPlaybackCapability::ID = TPlaybackCapabilityID<FPlayerIndexPlaybackCapability>::Register();
+
+IMovieScenePlayer* FPlayerIndexPlaybackCapability::GetPlayer(TSharedRef<const FSharedPlaybackState> Owner)
+{
+	if (FPlayerIndexPlaybackCapability* Cap = Owner->FindCapability<FPlayerIndexPlaybackCapability>())
+	{
+		return IMovieScenePlayer::Get(Cap->PlayerIndex);
+	}
+	return nullptr;
+}
+
+uint16 FPlayerIndexPlaybackCapability::GetPlayerIndex(TSharedRef<const FSharedPlaybackState> Owner)
+{
+	if (FPlayerIndexPlaybackCapability* Cap = Owner->FindCapability<FPlayerIndexPlaybackCapability>())
+	{
+		return Cap->PlayerIndex;
+	}
+	return (uint16)-1;
+}
+
 } // namespace MovieScene
 } // namespace UE
 
@@ -92,6 +112,6 @@ void IMovieScenePlayer::InvalidateCachedData()
 	UE::MovieScene::FSequenceInstance* RootInstance = Template.FindInstance(MovieSceneSequenceID::Root);
 	if (RootInstance)
 	{
-		RootInstance->InvalidateCachedData(Template.GetEntitySystemLinker());
+		RootInstance->InvalidateCachedData();
 	}
 }
diff --git a/Engine/Source/Runtime/MovieScene/Private/Tests/MovieScenePreAnimatedStateTests.cpp b/Engine/Source/Runtime/MovieScene/Private/Tests/MovieScenePreAnimatedStateTests.cpp
--- a/Engine/Source/Runtime/MovieScene/Private/Tests/MovieScenePreAnimatedStateTests.cpp
+++ b/Engine/Source/Runtime/MovieScene/Private/Tests/MovieScenePreAnimatedStateTests.cpp
@@ -394,13 +394,13 @@ bool FMovieScenePreAnimatedStateTrackTypesTest::RunTest(const FString& Parameter
 		for (int32 i = 0; i < NumEvaluations; ++i)
 		{
 			FMovieSceneEvaluationRange EvaluatedRange(TRange<FFrameTime>(i*(SectionLength/NumEvaluations), (i+1)*(SectionLength/NumEvaluations)), TickResolution, EPlayDirection::Forwards);
-			TestPlayer.Template.EvaluateSynchronousBlocking(EvaluatedRange, TestPlayer);
+			TestPlayer.Template.EvaluateSynchronousBlocking(EvaluatedRange);
 
 			Assert(this, TestValue1, StartValue + i, TEXT("Keep-State EvaluationHook did not Begin or Update correctly."));
 		}
 
 		FMovieSceneEvaluationRange EvaluatedRange(TRange<FFrameTime>(SectionLength, SectionLength+100), TickResolution, EPlayDirection::Forwards);
-		TestPlayer.Template.EvaluateSynchronousBlocking(EvaluatedRange, TestPlayer);
+		TestPlayer.Template.EvaluateSynchronousBlocking(EvaluatedRange);
 		Assert(this, TestValue1, EndValue, TEXT("Keep-State EvaluationHook did not End correctly."));
 
 		TestPlayer.RestorePreAnimatedState();
@@ -413,13 +413,13 @@ bool FMovieScenePreAnimatedStateTrackTypesTest::RunTest(const FString& Parameter
 		for (int32 i = 0; i < NumEvaluations; ++i)
 		{
 			FMovieSceneEvaluationRange EvaluatedRange(TRange<FFrameTime>(2000 + i*(SectionLength/NumEvaluations), 2000 + (i+1)*(SectionLength/NumEvaluations)), TickResolution, EPlayDirection::Forwards);
-			TestPlayer.Template.EvaluateSynchronousBlocking(EvaluatedRange, TestPlayer);
+			TestPlayer.Template.EvaluateSynchronousBlocking(EvaluatedRange);
 
 			Assert(this, TestValue1, StartValue + i, TEXT("Restore-State EvaluationHook did not Begin or Update correctly."));
 		}
 
 		FMovieSceneEvaluationRange FinalRange(TRange<FFrameTime>(2000 + SectionLength, 2000 + SectionLength + 100), TickResolution, EPlayDirection::Forwards);
-		TestPlayer.Template.EvaluateSynchronousBlocking(FinalRange, TestPlayer);
+		TestPlayer.Template.EvaluateSynchronousBlocking(FinalRange);
 		Assert(this, TestValue1, TestMagicNumber, TEXT("Restore-State EvaluationHook did not End correctly."));
 
 		TestPlayer.RestorePreAnimatedState();
@@ -477,7 +477,7 @@ bool FMovieScenePreAnimatedStateContextChangedTest::RunTest(const FString& Param
 		for (int32 i = 0; i < NumEvaluations; ++i)
 		{
 			FMovieSceneEvaluationRange EvaluatedRange(TRange<FFrameTime>(i*(SectionLength/NumEvaluations), (i+1)*(SectionLength/NumEvaluations)), TickResolution, EPlayDirection::Forwards);
-			TestPlayer.Template.EvaluateSynchronousBlocking(EvaluatedRange, TestPlayer);
+			TestPlayer.Template.EvaluateSynchronousBlocking(EvaluatedRange);
 
 			Assert(this, TestValue1, StartValue + i, TEXT("In-Editor 1: EvaluationHook did not Begin or Update correctly."));
 		}
@@ -494,7 +494,7 @@ bool FMovieScenePreAnimatedStateContextChangedTest::RunTest(const FString& Param
 		for (int32 i = 0; i < NumEvaluations; ++i)
 		{
 			FMovieSceneEvaluationRange EvaluatedRange(TRange<FFrameTime>(i*(SectionLength/NumEvaluations), (i+1)*(SectionLength/NumEvaluations)), TickResolution, EPlayDirection::Forwards);
-			TestPlayer.Template.EvaluateSynchronousBlocking(EvaluatedRange, TestPlayer);
+			TestPlayer.Template.EvaluateSynchronousBlocking(EvaluatedRange);
 
 			Assert(this, TestValue1, StartValue + i, TEXT("PIE: EvaluationHook did not Begin or Update correctly."));
 		}
diff --git a/Engine/Source/Runtime/MovieScene/Public/Compilation/MovieSceneCompiledVolatilityManager.h b/Engine/Source/Runtime/MovieScene/Public/Compilation/MovieSceneCompiledVolatilityManager.h
--- a/Engine/Source/Runtime/MovieScene/Public/Compilation/MovieSceneCompiledVolatilityManager.h
+++ b/Engine/Source/Runtime/MovieScene/Public/Compilation/MovieSceneCompiledVolatilityManager.h
@@ -7,7 +7,6 @@
 #include "Evaluation/MovieSceneEvaluationTemplateInstance.h"
 #include "Templates/UniquePtr.h"
 
-class IMovieScenePlayer;
 class UMovieSceneCompiledDataManager;
 struct FGuid;
 struct FMovieSceneCompiledDataID;
@@ -18,24 +17,29 @@ namespace UE
 namespace MovieScene
 {
 
+struct FSharedPlaybackState;
+
 struct FCompiledDataVolatilityManager
 {
-	static TUniquePtr<FCompiledDataVolatilityManager> Construct(IMovieScenePlayer& Player, FMovieSceneCompiledDataID RootDataID, UMovieSceneCompiledDataManager* CompiledDataManager);
+	static TUniquePtr<FCompiledDataVolatilityManager> Construct(TSharedRef<const FSharedPlaybackState> SharedPlaybackState);
+
+	FCompiledDataVolatilityManager(TSharedRef<const FSharedPlaybackState> SharedPlaybackState);
 
-	bool ConditionalRecompile(IMovieScenePlayer& Player, FMovieSceneCompiledDataID RootDataID, UMovieSceneCompiledDataManager* CompiledDataManager);
+	bool ConditionalRecompile();
 
 private:
 
-	bool HasBeenRecompiled(FMovieSceneCompiledDataID RootDataID, UMovieSceneCompiledDataManager* CompiledDataManager) const;
+	bool HasBeenRecompiled() const;
 
-	bool HasSequenceBeenRecompiled(FMovieSceneCompiledDataID DataID, FMovieSceneSequenceID SequenceID, UMovieSceneCompiledDataManager* CompiledDataManager) const;
+	bool HasSequenceBeenRecompiled(FMovieSceneCompiledDataID DataID, FMovieSceneSequenceID SequenceID) const;
 
-	void UpdateCachedSignatures(IMovieScenePlayer& Player, FMovieSceneCompiledDataID RootDataID, UMovieSceneCompiledDataManager* CompiledDataManager);
+	void UpdateCachedSignatures();
 
 private:
 
+	TWeakPtr<const FSharedPlaybackState> WeakSharedPlaybackState;
 	TSortedMap<FMovieSceneSequenceID, FGuid, TInlineAllocator<16>> CachedCompilationSignatures;
 };
 
 } // namespace MovieScene
-} // namespace UE
\ No newline at end of file
+} // namespace UE
diff --git a/Engine/Source/Runtime/MovieScene/Public/EntitySystem/MovieSceneInstanceRegistry.h b/Engine/Source/Runtime/MovieScene/Public/EntitySystem/MovieSceneInstanceRegistry.h
--- a/Engine/Source/Runtime/MovieScene/Public/EntitySystem/MovieSceneInstanceRegistry.h
+++ b/Engine/Source/Runtime/MovieScene/Public/EntitySystem/MovieSceneInstanceRegistry.h
@@ -12,7 +12,9 @@
 #include "MovieSceneSequenceID.h"
 
 
+class FMovieSceneEntitySystemRunner;
 class UMovieSceneEntitySystemLinker;
+class UMovieSceneCompiledDataManager;
 
 namespace UE
 {
@@ -81,7 +83,7 @@ struct FInstanceRegistry
 		return GetInstance(InstanceHandle).GetContext();
 	}
 
-	MOVIESCENE_API FRootInstanceHandle AllocateRootInstance(IMovieScenePlayer* Player);
+	MOVIESCENE_API FRootInstanceHandle AllocateRootInstance(IMovieScenePlayer* Player, UMovieSceneSequence& RootSequence, TSharedPtr<FMovieSceneEntitySystemRunner> Runner, UMovieSceneCompiledDataManager* CompiledDataManager);
 
 	MOVIESCENE_API FInstanceHandle AllocateSubInstance(IMovieScenePlayer* Player, FMovieSceneSequenceID SequenceID, FRootInstanceHandle RootInstance, FInstanceHandle ParentInstanceHandle);
 
diff --git a/Engine/Source/Runtime/MovieScene/Public/EntitySystem/MovieSceneSequenceInstance.h b/Engine/Source/Runtime/MovieScene/Public/EntitySystem/MovieSceneSequenceInstance.h
--- a/Engine/Source/Runtime/MovieScene/Public/EntitySystem/MovieSceneSequenceInstance.h
+++ b/Engine/Source/Runtime/MovieScene/Public/EntitySystem/MovieSceneSequenceInstance.h
@@ -62,7 +62,7 @@ public:
 	 * @param Linker     The linker that owns this sequence instance
 	 * @return true if a recompile has ocurred, false otherwise
 	 */
-	MOVIESCENE_API bool ConditionalRecompile(UMovieSceneEntitySystemLinker* Linker);
+	MOVIESCENE_API bool ConditionalRecompile();
 
 	/**
 	 * Called only for top-level sequence instances before any updates to it with the specified context.
@@ -72,8 +72,7 @@ public:
 	 * @param Context    The overall context that this sequence instance is being evaluated with
 	 * @param OutDissections   An array to populate with dissected time ranges that should be evaluated separately, in order.
 	 */
-	MOVIESCENE_API void DissectContext(UMovieSceneEntitySystemLinker* Linker, const FMovieSceneContext& Context, TArray<TRange<FFrameTime>>& OutDissections);
-
+	MOVIESCENE_API void DissectContext(const FMovieSceneContext& Context, TArray<TRange<FFrameTime>>& OutDissections);
 
 	/**
 	 * Called for root level instances that have either never evaluated, or have previously finished evaluating
@@ -81,16 +80,14 @@ public:
 	 * @param Linker     The linker that owns this sequence instance
 	 * @param InContext  The context that this sequence instance is to be evaluated with
 	 */
-	MOVIESCENE_API void Start(UMovieSceneEntitySystemLinker* Linker, const FMovieSceneContext& InContext);
-
+	MOVIESCENE_API void Start(const FMovieSceneContext& InContext);
 
 	/**
 	 * Called when this instance has been queued for evaluation in order for it to do any pre-work setup.
 	 *
 	 * @param Linker     The linker that owns this sequence instance
 	 */
-	MOVIESCENE_API void PreEvaluation(UMovieSceneEntitySystemLinker* Linker);
-
+	MOVIESCENE_API void PreEvaluation();
 
 	/**
 	 * Called after dissection for root level instances in order for this sequence instacne to update any entities it needs for evaluation.
@@ -98,8 +95,7 @@ public:
 	 * @param Linker     The linker that owns this sequence instance
 	 * @param InContext  The (potentially dissected) context that this sequence instance is to be evaluated with
 	 */
-	MOVIESCENE_API void Update(UMovieSceneEntitySystemLinker* Linker, const FMovieSceneContext& InContext);
-
+	MOVIESCENE_API void Update(const FMovieSceneContext& InContext);
 
 	/**
 	 * Returns whether this instance can be finished immediately without any last update.
@@ -107,22 +103,21 @@ public:
 	 * @param Linker     The linker that owns this sequence instance
 	 * @return           Whether the instance can be finished immediately
 	 */
-	MOVIESCENE_API bool CanFinishImmediately(UMovieSceneEntitySystemLinker* Linker) const;
-
+	MOVIESCENE_API bool CanFinishImmediately() const;
 
 	/**
 	 * Mark this instance as finished, causing all its entities to be unlinked and the instance to become inactive at the end of the next update.
 	 *
 	 * @param Linker     The linker that owns this sequence instance
 	 */
-	MOVIESCENE_API void Finish(UMovieSceneEntitySystemLinker* Linker);
+	MOVIESCENE_API void Finish();
 
 	/**
 	 * Called when this sequence instance has been evaluated in order for it to do any clean-up or other post-update work
 	 *
 	 * @param Linker     The linker that owns this sequence instance
 	 */
-	MOVIESCENE_API void PostEvaluation(UMovieSceneEntitySystemLinker* Linker);
+	MOVIESCENE_API void PostEvaluation();
 
 	/**
 	 * Called to run legacy track templates
@@ -149,10 +144,7 @@ public:
 	/**
 	 * Retrieve the IMovieScenePlayer's unique index
 	 */
-	uint16 GetPlayerIndex() const
-	{
-		return PlayerIndex;
-	}
+	MOVIESCENE_API uint16 GetPlayerIndex() const;
 
 	/**
 	 * Retrieve the SequenceID for this instance
@@ -305,25 +297,28 @@ public:
 	/**
 	 * Invalidate any cached data that may be being used for evaluation due to a change in the source asset data
 	 */
-	MOVIESCENE_API void InvalidateCachedData(UMovieSceneEntitySystemLinker* Linker);
+	MOVIESCENE_API void InvalidateCachedData();
 
 	/**
 	 * Destroy this sequence instance immediately - Finish must previously have been called
 	 */
-	MOVIESCENE_API void DestroyImmediately(UMovieSceneEntitySystemLinker* Linker);
+	MOVIESCENE_API void DestroyImmediately();
 
 	/**
 	 * Called to override the simulated root sequence ID for this instance. Only valid for IsRootSequence() instances.
 	 */
-	MOVIESCENE_API void OverrideRootSequence(UMovieSceneEntitySystemLinker* Linker, FMovieSceneSequenceID NewRootSequenceID);
+	MOVIESCENE_API void OverrideRootSequence(FMovieSceneSequenceID NewRootSequenceID);
 
 public:
 
 	/** Constructor for top level sequences */
-	MOVIESCENE_API explicit FSequenceInstance(UMovieSceneEntitySystemLinker* Linker, IMovieScenePlayer* Player, TSharedRef<FSharedPlaybackState> PlaybackState, FRootInstanceHandle ThisInstanceHandle);
+	MOVIESCENE_API explicit FSequenceInstance(TSharedRef<FSharedPlaybackState> PlaybackState, FRootInstanceHandle ThisInstanceHandle);
 
 	/** Constructor for sub sequences */
-	MOVIESCENE_API explicit FSequenceInstance(UMovieSceneEntitySystemLinker* Linker, IMovieScenePlayer* Player, TSharedRef<FSharedPlaybackState> PlaybackState, FInstanceHandle ThisInstanceHandle, FInstanceHandle InParentInstanceHandle, FRootInstanceHandle RootInstanceHandle, FMovieSceneSequenceID InSequenceID);
+	MOVIESCENE_API explicit FSequenceInstance(TSharedRef<FSharedPlaybackState> PlaybackState, FInstanceHandle ThisInstanceHandle, FInstanceHandle InParentInstanceHandle, FRootInstanceHandle RootInstanceHandle, FMovieSceneSequenceID InSequenceID);
+
+	/** Initialization of the sequence */
+	MOVIESCENE_API void Initialize(IMovieScenePlayer* Player);
 
 	/** Destructor */
 	MOVIESCENE_API ~FSequenceInstance();
@@ -336,9 +331,34 @@ public:
 	MOVIESCENE_API FSequenceInstance(FSequenceInstance&&);
 	MOVIESCENE_API FSequenceInstance& operator=(FSequenceInstance&&);
 
+public:
+
+	UE_DEPRECATED(5.4, "Please use the version of this method that doesn't take a Linker")
+	MOVIESCENE_API bool ConditionalRecompile(UMovieSceneEntitySystemLinker* Linker);
+	UE_DEPRECATED(5.4, "Please use the version of this method that doesn't take a Linker")
+	MOVIESCENE_API void DissectContext(UMovieSceneEntitySystemLinker* Linker, const FMovieSceneContext& Context, TArray<TRange<FFrameTime>>& OutDissections);
+	UE_DEPRECATED(5.4, "Please use the version of this method that doesn't take a Linker")
+	MOVIESCENE_API void Start(UMovieSceneEntitySystemLinker* Linker, const FMovieSceneContext& InContext);
+	UE_DEPRECATED(5.4, "Please use the version of this method that doesn't take a Linker")
+	MOVIESCENE_API void PreEvaluation(UMovieSceneEntitySystemLinker* Linker);
+	UE_DEPRECATED(5.4, "Please use the version of this method that doesn't take a Linker")
+	MOVIESCENE_API void Update(UMovieSceneEntitySystemLinker* Linker, const FMovieSceneContext& InContext);
+	UE_DEPRECATED(5.4, "Please use the version of this method that doesn't take a Linker")
+	MOVIESCENE_API bool CanFinishImmediately(UMovieSceneEntitySystemLinker* Linker) const;
+	UE_DEPRECATED(5.4, "Please use the version of this method that doesn't take a Linker")
+	MOVIESCENE_API void Finish(UMovieSceneEntitySystemLinker* Linker);
+	UE_DEPRECATED(5.4, "Please use the version of this method that doesn't take a Linker")
+	MOVIESCENE_API void PostEvaluation(UMovieSceneEntitySystemLinker* Linker);
+	UE_DEPRECATED(5.4, "Please use the version of this method that doesn't take a Linker")
+	MOVIESCENE_API void InvalidateCachedData(UMovieSceneEntitySystemLinker* Linker);
+	UE_DEPRECATED(5.4, "Please use the version of this method that doesn't take a Linker")
+	MOVIESCENE_API void DestroyImmediately(UMovieSceneEntitySystemLinker* Linker);
+	UE_DEPRECATED(5.4, "Please use the version of this method that doesn't take a Linker")
+	MOVIESCENE_API void OverrideRootSequence(UMovieSceneEntitySystemLinker* Linker, FMovieSceneSequenceID NewRootSequenceID);
+
 private:
 
-	MOVIESCENE_API void InitializeLegacyEvaluator(UMovieSceneEntitySystemLinker* Linker);
+	MOVIESCENE_API void InitializeLegacyEvaluator();
 
 private:
 
@@ -360,14 +380,12 @@ private:
 	TSharedRef<FSharedPlaybackState> SharedPlaybackState;
 
 
-	/** Delegate Binding for when an object binding is invalidated in this instance . */
+	/** Delegate Binding for when an object binding is invalidated in this instance. */
 	FDelegateHandle OnInvalidateObjectBindingHandle;
 	/** This sequence instances sequence ID, or MovieSceneSequenceID::Root for top-level sequences. */
 	FMovieSceneSequenceID SequenceID;
 	/** When SequenceID != MovieSceneSequenceID::Root, specifies an ID to override as a simulated root. */
 	FMovieSceneSequenceID RootOverrideSequenceID;
-	/** The index of this instance's IMovieScenePlayer retrievable through IMovieScenePlayer::Get(). */
-	uint16 PlayerIndex;
 	/** Cached update flags denoting what kinds of updates are required by this instance */
 	ESequenceInstanceUpdateFlags UpdateFlags;
 	/** This instance's handle. */
diff --git a/Engine/Source/Runtime/MovieScene/Public/EntitySystem/MovieSceneSequenceUpdaters.h b/Engine/Source/Runtime/MovieScene/Public/EntitySystem/MovieSceneSequenceUpdaters.h
--- a/Engine/Source/Runtime/MovieScene/Public/EntitySystem/MovieSceneSequenceUpdaters.h
+++ b/Engine/Source/Runtime/MovieScene/Public/EntitySystem/MovieSceneSequenceUpdaters.h
@@ -6,9 +6,9 @@
 #include "CoreTypes.h"
 #include "Math/Range.h"
 #include "Misc/FrameTime.h"
+#include "Templates/SharedPointer.h"
 #include "Templates/UniquePtr.h"
 
-class IMovieScenePlayer;
 class UMovieSceneCompiledDataManager;
 class UMovieSceneEntitySystemLinker;
 struct FFrameTime;
@@ -24,6 +24,7 @@ namespace MovieScene
 
 struct FInstanceHandle;
 struct FRootInstanceHandle;
+struct FSharedPlaybackState;
 
 enum class ESequenceInstanceUpdateFlags : uint8;
 
@@ -57,7 +58,7 @@ public:
 	/**
 	 * Called to initialize the flag structure that denotes what functions need to be called on this updater
 	 */
-	virtual void PopulateUpdateFlags(UMovieSceneEntitySystemLinker* InLinker, IMovieScenePlayer* InPlayer, ESequenceInstanceUpdateFlags& OutUpdateFlags) = 0;
+	virtual void PopulateUpdateFlags(UMovieSceneEntitySystemLinker* InLinker, TSharedRef<const FSharedPlaybackState> SharedPlaybackState, ESequenceInstanceUpdateFlags& OutUpdateFlags) = 0;
 
 	/**
 	 * Called before any updates to the sequence to allow this updater to dissect the context into smaller ranges that should be evaluated independently for the purpose of ensuring determinism.
@@ -68,7 +69,7 @@ public:
 	 * @param InContext        The total root-level context for the next evaluation to be dissected
 	 * @param OutDissections   (Out) Array to populate with dissected ranges
 	 */
-	virtual void DissectContext(UMovieSceneEntitySystemLinker* InLinker, IMovieScenePlayer* InPlayer, const FMovieSceneContext& InContext, TArray<TRange<FFrameTime>>& OutDissections) = 0;
+	virtual void DissectContext(UMovieSceneEntitySystemLinker* InLinker, TSharedRef<const FSharedPlaybackState> SharedPlaybackState, const FMovieSceneContext& InContext, TArray<TRange<FFrameTime>>& OutDissections) = 0;
 
 
 	/**
@@ -79,7 +80,7 @@ public:
 	 * @param InPlayer         The movie scene player instance playing this sequence
 	 * @param InContext        The root-level context for the current evaluation.
 	 */
-	virtual void Start(UMovieSceneEntitySystemLinker* InLinker, FRootInstanceHandle InInstanceHandle, IMovieScenePlayer* InPlayer, const FMovieSceneContext& InContext) = 0;
+	virtual void Start(UMovieSceneEntitySystemLinker* InLinker, FRootInstanceHandle InInstanceHandle, TSharedRef<const FSharedPlaybackState> SharedPlaybackState, const FMovieSceneContext& InContext) = 0;
 
 
 	/**
@@ -90,7 +91,7 @@ public:
 	 * @param InPlayer         The movie scene player instance playing this sequence
 	 * @param InContext        The root-level context for the current evaluation.
 	 */
-	virtual void Update(UMovieSceneEntitySystemLinker* InLinker, FRootInstanceHandle InInstanceHandle, IMovieScenePlayer* InPlayer, const FMovieSceneContext& InContext) = 0;
+	virtual void Update(UMovieSceneEntitySystemLinker* InLinker, FRootInstanceHandle InInstanceHandle, TSharedRef<const FSharedPlaybackState> SharedPlaybackState, const FMovieSceneContext& InContext) = 0;
 
 
 	/**
@@ -110,7 +111,7 @@ public:
 	 * @param InInstanceHandle The instance handle for the top level sequence instance that this updater belongs to
 	 * @param InPlayer         The movie scene player instance playing this sequence
 	 */
-	virtual void Finish(UMovieSceneEntitySystemLinker* InLinker, FRootInstanceHandle InInstanceHandle, IMovieScenePlayer* InPlayer) = 0;
+	virtual void Finish(UMovieSceneEntitySystemLinker* InLinker, FRootInstanceHandle InInstanceHandle, TSharedRef<const FSharedPlaybackState> SharedPlaybackState) = 0;
 
 
 	/**
diff --git a/Engine/Source/Runtime/MovieScene/Public/EntitySystem/MovieSceneSharedPlaybackState.h b/Engine/Source/Runtime/MovieScene/Public/EntitySystem/MovieSceneSharedPlaybackState.h
--- a/Engine/Source/Runtime/MovieScene/Public/EntitySystem/MovieSceneSharedPlaybackState.h
+++ b/Engine/Source/Runtime/MovieScene/Public/EntitySystem/MovieSceneSharedPlaybackState.h
@@ -2,43 +2,143 @@
 
 #pragma once
 
+#include "Compilation/MovieSceneCompiledDataID.h"
 #include "CoreTypes.h"
+#include "EntitySystem/MovieSceneSequenceInstanceHandle.h"
+#include "Evaluation/MovieSceneEvaluationOperand.h"
 #include "Evaluation/MovieScenePlaybackCapabilities.h"
+#include "MovieSceneSequenceID.h"
+
+class FMovieSceneEntitySystemRunner;
+class UMovieSceneCompiledDataManager;
+struct FMovieSceneObjectCache;
+struct FMovieSceneSequenceHierarchy;
 
 namespace UE::MovieScene
 {
 
-struct FSharedPlaybackState
+/**
+ * Parameter structure for initializing a new shared playback state.
+ */
+struct FSharedPlaybackStateCreateParams
 {
-	/** Playback capabilities for the root sequence */
-	FPlaybackCapabilities Capabilities;
+	/**
+	 * The playback context in which the root sequence will be evaluated.
+	 */
+	UObject* PlaybackContext = nullptr;
+
+	/**
+	 * The handle of the root sequence instance, if the created playback state
+	 * is meant to relate to an instance that has also been created inside
+	 * a runner/linker's instance registry.
+	 */
+	FRootInstanceHandle RootInstanceHandle;
+
+	/**
+	 * The runner that will be evaluating the sequence that the created playback
+	 * state relates to.
+	 */
+	TSharedPtr<FMovieSceneEntitySystemRunner> Runner;
 
+	/**
+	 * The compiled data manager with which the root sequence was compiled, or
+	 * will be compiled. If unset, the default global manager will be used.
+	 */
+	UMovieSceneCompiledDataManager* CompiledDataManager = nullptr;
+};
+
+/**
+ * A structure that stores playback state for an entire sequence hierarchy.
+ */
+struct FSharedPlaybackState : TSharedFromThis<FSharedPlaybackState>
+{
 public:
 
+	FSharedPlaybackState();
+	FSharedPlaybackState(
+			UMovieSceneSequence& InRootSequence,
+			const FSharedPlaybackStateCreateParams& CreateParams);
+
+public:
+	
+	/** Gets the playback context */
+	UObject* GetPlaybackContext() const { return WeakPlaybackContext.Get(); }
+
+	/** Gets the root sequence */
+	UMovieSceneSequence* GetRootSequence() const { return WeakRootSequence.Get(); }
+
+	/** Gets the runner evaluating this root sequence */
+	TSharedPtr<FMovieSceneEntitySystemRunner> GetRunner() const { return WeakRunner.Pin(); }
+
+	/** Gets the compiled data manager that contains the data for the root sequence */
+	TObjectPtr<UMovieSceneCompiledDataManager> GetCompiledDataManager() const { return CompiledDataManager; }
+
+	/** Gets the handle of the root sequence */
+	const FRootInstanceHandle& GetRootInstanceHandle() const { return RootInstanceHandle; }
+
+	/** Gets the compiled data ID for the root sequence */
+	const FMovieSceneCompiledDataID&  GetRootCompiledDataID() const { return RootCompiledDataID; }
+
+public:
+
+	// General utility methods
+
+	MOVIESCENE_API UMovieSceneEntitySystemLinker* GetLinker() const;
+
+	MOVIESCENE_API const FMovieSceneSequenceHierarchy* GetHierarchy() const;
+	MOVIESCENE_API UMovieSceneSequence* GetSequence(FMovieSceneSequenceIDRef SequenceID) const;
+
+public:
+
+	/**
+	 * Gets the capabilities container.
+	 */
+	FPlaybackCapabilities& GetCapabilities() { return Capabilities; }
+
+	/**
+	 * Gets the capabilities container.
+	 */
+	const FPlaybackCapabilities& GetCapabilities() const { return Capabilities; }
+
+	/**
+	 * Returns whether the root sequence has the specified capability.
+	 */
 	template<typename T>
 	bool HasCapability() const
 	{
 		return Capabilities.HasCapability(T::ID);
 	}
 
+	/**
+	 * Finds the specified capability on the root sequence.
+	 */
 	template<typename T>
 	T* FindCapability() const
 	{
 		return Capabilities.FindCapability<T>(T::ID);
 	}
 
+	/**
+	 * Builds the specified capability for the root sequence.
+	 */
 	template<typename T, typename ...ArgTypes>
 	T& AddCapability(ArgTypes&&... InArgs)
 	{
 		return Capabilities.AddCapability<T>(T::ID, Forward<ArgTypes>(InArgs)...);
 	}
 
+	/**
+	 * Adds the specified capability on the root sequence as a raw pointer.
+	 */
 	template<typename T, typename ...ArgTypes>
 	T& AddCapabilityRaw(T* InPointer)
 	{
 		return Capabilities.AddCapabilityRaw<T>(T::ID, InPointer);
 	}
 
+	/**
+	 * Adds the specified capability on the root sequence as a shared pointer.
+	 */
 	template<typename T, typename ...ArgTypes>
 	T& AddCapabilityShared(TSharedRef<T> InSharedRef)
 	{
@@ -47,10 +147,30 @@ public:
 
 public:
 
-	void InvalidateCachedData(UMovieSceneEntitySystemLinker* Linker)
-	{
-		Capabilities.InvalidateCachedData(Linker);
-	}
+	MOVIESCENE_API void InvalidateCachedData();
+
+private:
+
+	/** The root sequence */
+	TWeakObjectPtr<UMovieSceneSequence> WeakRootSequence;
+
+	/** The playback context */
+	TWeakObjectPtr<UObject> WeakPlaybackContext;
+
+	/** The runner evaluating this root sequence */
+	TWeakPtr<FMovieSceneEntitySystemRunner> WeakRunner;
+
+	/** The compiled data manager that contains the data for the root sequence */
+	TObjectPtr<UMovieSceneCompiledDataManager> CompiledDataManager;
+
+	/** The handle of the root sequence */
+	FRootInstanceHandle RootInstanceHandle;
+
+	/** The compiled data ID for the root sequence */
+	FMovieSceneCompiledDataID  RootCompiledDataID;
+
+	/** Playback capabilities for the root sequence */
+	FPlaybackCapabilities Capabilities;
 };
 
 } // namespace UE::MovieScene
diff --git a/Engine/Source/Runtime/MovieScene/Public/Evaluation/MovieSceneEvaluationTemplateInstance.h b/Engine/Source/Runtime/MovieScene/Public/Evaluation/MovieSceneEvaluationTemplateInstance.h
--- a/Engine/Source/Runtime/MovieScene/Public/Evaluation/MovieSceneEvaluationTemplateInstance.h
+++ b/Engine/Source/Runtime/MovieScene/Public/Evaluation/MovieSceneEvaluationTemplateInstance.h
@@ -5,6 +5,8 @@
 #include "CoreTypes.h"
 #include "Containers/Map.h"
 #include "EntitySystem/MovieSceneEntityIDs.h"
+#include "EntitySystem/MovieSceneInstanceRegistry.h"
+#include "EntitySystem/MovieSceneSharedPlaybackState.h"
 #include "Templates/UniquePtr.h"
 #include "UObject/WeakObjectPtr.h"
 #include "UObject/StrongObjectPtr.h"
@@ -48,14 +50,16 @@ public:
 	 * @param Player					The player responsible for playback
 	 * @param TemplateStore				Template store responsible for supplying templates for a given sequence
 	 */
-	MOVIESCENE_API void Initialize(UMovieSceneSequence& RootSequence, IMovieScenePlayer& Player, UMovieSceneCompiledDataManager* InCompiledDataManager, TWeakPtr<FMovieSceneEntitySystemRunner> InWeakRunner);
+	MOVIESCENE_API void Initialize(UMovieSceneSequence& RootSequence, IMovieScenePlayer& Player, UMovieSceneCompiledDataManager* InCompiledDataManager, TSharedPtr<FMovieSceneEntitySystemRunner> InRunner);
 
 	/**
 	 * Evaluate this sequence in a synchronous way.
 	 *
 	 * @param Context				Evaluation context containing the time (or range) to evaluate
-	 * @param Player				The player responsible for playback
 	 */
+	MOVIESCENE_API void EvaluateSynchronousBlocking(FMovieSceneContext Context);
+
+	UE_DEPRECATED(5.4, "Use the version without the player parameter")
 	MOVIESCENE_API void EvaluateSynchronousBlocking(FMovieSceneContext Context, IMovieScenePlayer& Player);
 
 	UE_DEPRECATED(5.1, "Use EvaluateSynchronousBlocking instead.")
@@ -68,35 +72,17 @@ public:
 
 	MOVIESCENE_API void ResetDirectorInstances();
 
-	bool IsValid() const
-	{
-		return CompiledDataManager && WeakRootSequence.Get();
-	}
+	MOVIESCENE_API bool IsValid() const;
 
-	UE::MovieScene::FRootInstanceHandle GetRootInstanceHandle() const
-	{
-		return RootInstanceHandle;
-	}
+	MOVIESCENE_API UE::MovieScene::FRootInstanceHandle GetRootInstanceHandle() const;
 
-	UMovieSceneSequence* GetRootSequence() const
-	{
-		return WeakRootSequence.Get();
-	}
+	MOVIESCENE_API UMovieSceneSequence* GetRootSequence() const;
 
-	FMovieSceneCompiledDataID GetCompiledDataID() const
-	{
-		return CompiledDataID;
-	}
+	MOVIESCENE_API FMovieSceneCompiledDataID GetCompiledDataID() const;
 
-	UMovieSceneCompiledDataManager* GetCompiledDataManager() const
-	{
-		return CompiledDataManager;
-	}
+	MOVIESCENE_API UMovieSceneCompiledDataManager* GetCompiledDataManager() const;
 
-	TSharedPtr<FMovieSceneEntitySystemRunner> GetRunner() const
-	{
-		return WeakRunner.Pin();
-	}
+	MOVIESCENE_API TSharedPtr<FMovieSceneEntitySystemRunner> GetRunner() const;
 
 	MOVIESCENE_API bool HasEverUpdated() const;
 
@@ -106,6 +92,8 @@ public:
 
 	MOVIESCENE_API void GetSequenceParentage(const UE::MovieScene::FInstanceHandle InstanceHandle, TArray<UE::MovieScene::FInstanceHandle>& OutParentHandles) const;
 
+	MOVIESCENE_API const UE::MovieScene::FSequenceInstance* GetRootInstance() const;
+
 	MOVIESCENE_API UE::MovieScene::FSequenceInstance* FindInstance(FMovieSceneSequenceID SequenceID);
 
 	MOVIESCENE_API const UE::MovieScene::FSequenceInstance* FindInstance(FMovieSceneSequenceID SequenceID) const;
@@ -134,8 +122,11 @@ public:
 	MOVIESCENE_API void EnableGlobalPreAnimatedStateCapture();
 
 #if WITH_EDITOR
-	MOVIESCENE_API void SetEmulatedNetworkMask(EMovieSceneServerClientMask InNewMask, IMovieScenePlayer& Player);
+	MOVIESCENE_API void SetEmulatedNetworkMask(EMovieSceneServerClientMask InNewMask);
 	MOVIESCENE_API EMovieSceneServerClientMask GetEmulatedNetworkMask() const;
+
+	UE_DEPRECATED(5.4, "Please use the version without the player parameter")
+	MOVIESCENE_API void SetEmulatedNetworkMask(EMovieSceneServerClientMask InNewMask, IMovieScenePlayer& Player);
 #endif
 
 private:
@@ -144,27 +135,18 @@ private:
 
 private:
 
-	UPROPERTY()
-	TWeakObjectPtr<UMovieSceneSequence> WeakRootSequence;
-
-	UPROPERTY()
-	TObjectPtr<UMovieSceneCompiledDataManager> CompiledDataManager;
-
-	TWeakPtr<FMovieSceneEntitySystemRunner> WeakRunner;
-
-	UE::MovieScene::FRootInstanceHandle RootInstanceHandle;
-
 	UPROPERTY()
 	TObjectPtr<UMovieSceneEntitySystemLinker> EntitySystemLinker;
 
+	/** The playback state for the hierarchy of sequence instances */
+	TSharedPtr<UE::MovieScene::FSharedPlaybackState> SharedPlaybackState;
+	
 	/** Map of director instances by sequence ID. Kept alive by this map assuming this struct is reference collected */
 	UPROPERTY()
 	TMap<FMovieSceneSequenceID, TObjectPtr<UObject>> DirectorInstances;
 
 	FMovieSceneSequenceID RootID;
 
-	FMovieSceneCompiledDataID CompiledDataID;
-
 #if WITH_EDITOR
 	EMovieSceneServerClientMask EmulatedNetworkMask;
 #endif
diff --git a/Engine/Source/Runtime/MovieScene/Public/IMovieScenePlayer.h b/Engine/Source/Runtime/MovieScene/Public/IMovieScenePlayer.h
--- a/Engine/Source/Runtime/MovieScene/Public/IMovieScenePlayer.h
+++ b/Engine/Source/Runtime/MovieScene/Public/IMovieScenePlayer.h
@@ -7,11 +7,13 @@
 #include "Containers/ArrayView.h"
 #include "Misc/InlineValue.h"
 
-#include "MovieSceneSpawnRegister.h"
+#include "EntitySystem/MovieSceneSharedPlaybackState.h"
+#include "Evaluation/IMovieScenePlaybackCapability.h"
 #include "Evaluation/MovieSceneAnimTypeID.h"
-#include "Evaluation/MovieScenePreAnimatedState.h"
-#include "Evaluation/MovieSceneEvaluationState.h"
 #include "Evaluation/MovieSceneEvaluationOperand.h"
+#include "Evaluation/MovieSceneEvaluationState.h"
+#include "Evaluation/MovieScenePreAnimatedState.h"
+#include "MovieSceneSpawnRegister.h"
 
 enum class EMovieSceneBuiltInEasing : uint8;
 
@@ -28,6 +30,20 @@ class IMovieSceneSequencePlayerObserver;
 namespace UE::MovieScene
 {
 	enum class ESequenceInstanceUpdateFlags : uint8;
+	struct FSharedPlaybackState;
+
+	/**
+	 * Playback capability for storing an IMovieScenePlayer unique index.
+	 */
+	struct FPlayerIndexPlaybackCapability
+	{
+		static TPlaybackCapabilityID<FPlayerIndexPlaybackCapability> ID;
+
+		static IMovieScenePlayer* GetPlayer(TSharedRef<const FSharedPlaybackState> Owner);
+		static uint16 GetPlayerIndex(TSharedRef<const FSharedPlaybackState> Owner);
+
+		uint16 PlayerIndex = (uint16)-1;
+	};
 }
 
 struct EMovieSceneViewportParams
diff --git a/Engine/Source/Runtime/UMG/Private/Animation/UMGSequencePlayer.cpp b/Engine/Source/Runtime/UMG/Private/Animation/UMGSequencePlayer.cpp
--- a/Engine/Source/Runtime/UMG/Private/Animation/UMGSequencePlayer.cpp
+++ b/Engine/Source/Runtime/UMG/Private/Animation/UMGSequencePlayer.cpp
@@ -568,7 +568,7 @@ void UUMGSequencePlayer::RemoveEvaluationData()
 		}
 
 		SequenceInstance->Ledger.UnlinkEverything(Linker);
-		SequenceInstance->InvalidateCachedData(Linker);
+		SequenceInstance->InvalidateCachedData();
 	}
 }
 
