CIS: fix compile issue on Mac ARM arch

diff --git a/Engine/Source/Runtime/MovieScene/Public/IMovieScenePlayer.h b/Engine/Source/Runtime/MovieScene/Public/IMovieScenePlayer.h
--- a/Engine/Source/Runtime/MovieScene/Public/IMovieScenePlayer.h
+++ b/Engine/Source/Runtime/MovieScene/Public/IMovieScenePlayer.h
@@ -39,8 +39,12 @@ namespace UE::MovieScene
 	{
 		static TPlaybackCapabilityID<FPlayerIndexPlaybackCapability> ID;
 
-		static IMovieScenePlayer* GetPlayer(TSharedRef<const FSharedPlaybackState> Owner);
-		static uint16 GetPlayerIndex(TSharedRef<const FSharedPlaybackState> Owner);
+		static MOVIESCENE_API IMovieScenePlayer* GetPlayer(TSharedRef<const FSharedPlaybackState> Owner);
+		static MOVIESCENE_API uint16 GetPlayerIndex(TSharedRef<const FSharedPlaybackState> Owner);
+
+		FPlayerIndexPlaybackCapability(uint16 InPlayerIndex)
+			: PlayerIndex(InPlayerIndex)
+		{}
 
 		uint16 PlayerIndex = (uint16)-1;
 	};
