Sequencer: Restablish the MovieSceneBinding pointer after sorting.

This fixes an issue where the event triggers would be sent for the wrong binding and the UI wouldn't update after replacing an actor.

diff --git a/Engine/Source/Runtime/MovieScene/Private/MovieScene.cpp b/Engine/Source/Runtime/MovieScene/Private/MovieScene.cpp
--- a/Engine/Source/Runtime/MovieScene/Private/MovieScene.cpp
+++ b/Engine/Source/Runtime/MovieScene/Private/MovieScene.cpp
@@ -1522,6 +1522,9 @@ void UMovieScene::ReplaceBinding(const FGuid& OldGuid, const FGuid& NewGuid, con
 		// Replace directly changes the guid, so force a sort here
 		ObjectBindings.Sort();
 
+		// Reget the binding after sorting
+		Binding = FindBinding(NewGuid);
+
 		// Changing a binding guid invalidates any tracks contained within the binding
 		// Make sure they are written into the transaction buffer by calling modify
 		for (UMovieSceneTrack* Track : Binding->GetTracks())
