Sequencer: don't call CleanTaggedGarbage twice on systems that happen to belong to multiple phases.

diff --git a/Engine/Source/Runtime/MovieScene/Private/EntitySystem/MovieSceneEntitySystemLinker.cpp b/Engine/Source/Runtime/MovieScene/Private/EntitySystem/MovieSceneEntitySystemLinker.cpp
--- a/Engine/Source/Runtime/MovieScene/Private/EntitySystem/MovieSceneEntitySystemLinker.cpp
+++ b/Engine/Source/Runtime/MovieScene/Private/EntitySystem/MovieSceneEntitySystemLinker.cpp
@@ -404,9 +404,16 @@ void UMovieSceneEntitySystemLinker::CleanGarbage()
 	// the next time a runner gets flushed
 	LastInstantiationVersion = 0;
 
-	auto RouteCleanTaggedGarbage = [](UMovieSceneEntitySystem* System){ System->CleanTaggedGarbage(); };
-	SystemGraph.IteratePhase(ESystemPhase::Spawn, RouteCleanTaggedGarbage);
-	SystemGraph.IteratePhase(ESystemPhase::Instantiation, RouteCleanTaggedGarbage);
+	// Since some systems might belong to both Spawn and Instantiation phase, we need to gather them in a Set
+	// to prevent calling CleanTaggedGarbage twice on those.
+	TSet<UMovieSceneEntitySystem*> SystemsToClean;
+	auto GatherSystemsToClean = [&SystemsToClean](UMovieSceneEntitySystem* System){ SystemsToClean.Add(System); };
+	SystemGraph.IteratePhase(ESystemPhase::Spawn, GatherSystemsToClean);
+	SystemGraph.IteratePhase(ESystemPhase::Instantiation, GatherSystemsToClean);
+	for (UMovieSceneEntitySystem* System : SystemsToClean)
+	{
+		System->CleanTaggedGarbage();
+	}
 
 	// Allow any other system to cleanup garbage
 	// NOTE: Order is important here - we need to broadcast this after systems have CleanTaggedGarbage called
