Sequencer: speculative fix for crash during level sequence actor teardown.

#tests PIE, frontend, EngineTest automated tests

#localization none

diff --git a/Engine/Source/Runtime/MovieScene/Private/Evaluation/PreAnimatedState/MovieScenePreAnimatedStateExtension.cpp b/Engine/Source/Runtime/MovieScene/Private/Evaluation/PreAnimatedState/MovieScenePreAnimatedStateExtension.cpp
--- a/Engine/Source/Runtime/MovieScene/Private/Evaluation/PreAnimatedState/MovieScenePreAnimatedStateExtension.cpp
+++ b/Engine/Source/Runtime/MovieScene/Private/Evaluation/PreAnimatedState/MovieScenePreAnimatedStateExtension.cpp
@@ -388,7 +388,10 @@ void FPreAnimatedStateExtension::RestoreGlobalState(const FRestoreStateParams& P
 	GroupMetaData.Shrink();
 
 	// Invalidate cached data for any sequence instance that belongs to the terminal instance
-	if (Params.TerminalInstanceHandle.IsValid())
+	if (Params.TerminalInstanceHandle.IsValid() && 
+			ensureMsgf(
+				Linker->GetInstanceRegistry()->IsHandleValid(Params.TerminalInstanceHandle),
+				TEXT("Terminal instance handle is not valid anymore, was the sequence destroyed?")))
 	{
 		Linker->GetInstanceRegistry()->MutateInstance(Params.TerminalInstanceHandle).InvalidateCachedData();
 	}
