Sequencer: make grouping system gracefully handle repeated calls to CleanTaggedGarbage

#localization none
#tests PIE, frontend, EngineTests automated tests

diff --git a/Engine/Source/Runtime/MovieScene/Private/EntitySystem/MovieSceneEntityGroupingSystem.cpp b/Engine/Source/Runtime/MovieScene/Private/EntitySystem/MovieSceneEntityGroupingSystem.cpp
--- a/Engine/Source/Runtime/MovieScene/Private/EntitySystem/MovieSceneEntityGroupingSystem.cpp
+++ b/Engine/Source/Runtime/MovieScene/Private/EntitySystem/MovieSceneEntityGroupingSystem.cpp
@@ -21,7 +21,7 @@ FEntityGroupBuilder::FEntityGroupBuilder(UMovieSceneEntityGroupingSystem* InOwne
 
 void FEntityGroupBuilder::AddEntityToGroup(const FMovieSceneEntityID& InEntity, const FEntityGroupID& InNewGroupID)
 {
-	if (ensure(InNewGroupID.IsValid()))
+	if (ensure(InNewGroupID.HasGroup()))
 	{
 		// Add the entity to the group.
 		UMovieSceneEntityGroupingSystem::FEntityGroupInfo& GroupInfo = Owner->Groups.FindOrAdd(InNewGroupID);
@@ -31,7 +31,7 @@ void FEntityGroupBuilder::AddEntityToGroup(const FMovieSceneEntityID& InEntity,
 
 bool FEntityGroupBuilder::RemoveEntityFromGroup(const FMovieSceneEntityID& InEntity, const FEntityGroupID& InPreviousGroupID)
 {
-	if (ensure(InPreviousGroupID.IsValid()))
+	if (ensure(InPreviousGroupID.HasGroup()))
 	{
 		// Remove the entity from the group. We should find that group, and find that entity inside it.
 		UMovieSceneEntityGroupingSystem::FEntityGroupInfo* PreviousGroup = Owner->Groups.Find(InPreviousGroupID);
diff --git a/Engine/Source/Runtime/MovieScene/Private/EntitySystem/MovieSceneEntitySystemLinker.cpp b/Engine/Source/Runtime/MovieScene/Private/EntitySystem/MovieSceneEntitySystemLinker.cpp
--- a/Engine/Source/Runtime/MovieScene/Private/EntitySystem/MovieSceneEntitySystemLinker.cpp
+++ b/Engine/Source/Runtime/MovieScene/Private/EntitySystem/MovieSceneEntitySystemLinker.cpp
@@ -379,9 +379,14 @@ void UMovieSceneEntitySystemLinker::HandlePostGarbageCollection()
 	// Allow any other system to tag garbage
 	Events.TagGarbage.Broadcast(this);
 
-	auto RouteTagGarbage = [](UMovieSceneEntitySystem* System){ System->TagGarbage(); };
-	SystemGraph.IteratePhase(ESystemPhase::Spawn, RouteTagGarbage);
-	SystemGraph.IteratePhase(ESystemPhase::Instantiation, RouteTagGarbage);
+	TSet<UMovieSceneEntitySystem*> SystemsToTag;
+	auto GatherSystemsToTag = [&SystemsToTag](UMovieSceneEntitySystem* System){ SystemsToTag.Add(System); };
+	SystemGraph.IteratePhase(ESystemPhase::Spawn, GatherSystemsToTag);
+	SystemGraph.IteratePhase(ESystemPhase::Instantiation, GatherSystemsToTag);
+	for (UMovieSceneEntitySystem* System : SystemsToTag)
+	{
+		System->TagGarbage();
+	}
 
 	CleanGarbage();
 }
diff --git a/Engine/Source/Runtime/MovieScene/Public/EntitySystem/BuiltInComponentTypes.h b/Engine/Source/Runtime/MovieScene/Public/EntitySystem/BuiltInComponentTypes.h
--- a/Engine/Source/Runtime/MovieScene/Public/EntitySystem/BuiltInComponentTypes.h
+++ b/Engine/Source/Runtime/MovieScene/Public/EntitySystem/BuiltInComponentTypes.h
@@ -483,6 +483,17 @@ struct FEntityGroupingPolicyKey
 	int32 Index = INDEX_NONE;
 };
 
+/**
+ * Flags for FEntityGroupID
+ */
+enum class EEntityGroupFlags : uint8
+{
+	None = 0,
+	RemovedFromGroup = 1 << 0
+};
+
+ENUM_CLASS_FLAGS(EEntityGroupFlags);
+
 /**
  * The component data for describing what group an entity belongs to (see UMovieSceneEntityGroupingSystem)
  */
@@ -492,18 +503,28 @@ struct FEntityGroupID
 
 	FEntityGroupingPolicyKey PolicyKey;
 	int32 GroupIndex = INDEX_NONE;
+	EEntityGroupFlags Flags = EEntityGroupFlags::None;
 
 	FEntityGroupID() {}
 	FEntityGroupID(const FEntityGroupingPolicyKey InPolicyKey, int32 InGroupIndex) : PolicyKey(InPolicyKey), GroupIndex(InGroupIndex) {}
 
+	/**
+	 * Returns whether this component points to a valid grouping policy, and has a valid group.
+	 */
 	bool IsValid() const
 	{
 		return PolicyKey.IsValid() && GroupIndex != INDEX_NONE;
 	}
 
+	/**
+	 * Returns whether this component is valid (see IsValid) and if it still belongs to a group.
+	 * If the entity was removed from its group, the GroupIndex would still be valid, but we would
+	 * have set Flags to RemovedFromGroup. This lets downstream systems still know what group the
+	 * entity *used* to be part of, for their own book-keeping.
+	 */
 	bool HasGroup() const
 	{
-		return GroupIndex != INDEX_NONE;
+		return PolicyKey.IsValid() && GroupIndex != INDEX_NONE && !EnumHasAnyFlags(Flags, EEntityGroupFlags::RemovedFromGroup);
 	}
 
 	friend uint32 GetTypeHash(const FEntityGroupID& GroupID)
diff --git a/Engine/Source/Runtime/MovieScene/Public/EntitySystem/MovieSceneEntityGroupingSystem.h b/Engine/Source/Runtime/MovieScene/Public/EntitySystem/MovieSceneEntityGroupingSystem.h
--- a/Engine/Source/Runtime/MovieScene/Public/EntitySystem/MovieSceneEntityGroupingSystem.h
+++ b/Engine/Source/Runtime/MovieScene/Public/EntitySystem/MovieSceneEntityGroupingSystem.h
@@ -239,7 +239,7 @@ struct TEntityGroupingHandlerImpl<GroupingPolicy, TIntegerSequence<int, Componen
 			const FMovieSceneEntityID EntityID(EntityIDs[Index]);
 			FEntityGroupID& GroupID(GroupIDs[Index]);
 
-			if (GroupID.IsValid())
+			if (GroupID.HasGroup())
 			{
 				const bool bIsGroupEmpty = Builder->RemoveEntityFromGroup(EntityID, GroupID);
 				if (bIsGroupEmpty)
@@ -260,7 +260,9 @@ struct TEntityGroupingHandlerImpl<GroupingPolicy, TIntegerSequence<int, Componen
 					FreedGroupIndices[GroupID.GroupIndex] = true;
 				}
 				// Leave the GroupID on the entity so that downstream systems can use it to track
-				// that this entity is leaving its group.
+				// that this entity is leaving its group, but flag it so we don't re-free it.
+				ensure(!EnumHasAllFlags(GroupID.Flags , EEntityGroupFlags::RemovedFromGroup));
+				GroupID.Flags |= EEntityGroupFlags::RemovedFromGroup;
 			}
 		}
 	}
