Prevent day sequence modifier component from updating time of day when the main day sequence actor has already ended play.

This adds a new IsValid method on sequence players.

This change also resets the playback status of a player when it's torn down. Before, it could stay in "Playing" mode.

diff --git a/Engine/Source/Runtime/MovieScene/Private/MovieSceneSequencePlayer.cpp b/Engine/Source/Runtime/MovieScene/Private/MovieSceneSequencePlayer.cpp
--- a/Engine/Source/Runtime/MovieScene/Private/MovieSceneSequencePlayer.cpp
+++ b/Engine/Source/Runtime/MovieScene/Private/MovieSceneSequencePlayer.cpp
@@ -1286,6 +1286,13 @@ void UMovieSceneSequencePlayer::TearDown()
 		TickManager->UnregisterTickClient(this);
 		TickManager = nullptr;
 	}
+
+	Status = EMovieScenePlayerStatus::Stopped;
+}
+
+bool UMovieSceneSequencePlayer::IsValid() const
+{
+	return RootTemplateInstance.IsValid();
 }
 
 bool UMovieSceneSequencePlayer::IsDisablingEventTriggers(FFrameTime& DisabledUntilTime) const
diff --git a/Engine/Source/Runtime/MovieScene/Public/MovieSceneSequencePlayer.h b/Engine/Source/Runtime/MovieScene/Public/MovieSceneSequencePlayer.h
--- a/Engine/Source/Runtime/MovieScene/Public/MovieSceneSequencePlayer.h
+++ b/Engine/Source/Runtime/MovieScene/Public/MovieSceneSequencePlayer.h
@@ -434,6 +434,9 @@ public:
 	/** Perform any tear-down work when this player is no longer (and will never) be needed */
 	MOVIESCENE_API void TearDown();
 
+	/** Returns whether this player is valid, i.e. it has been initialized and not torn down yet */
+	MOVIESCENE_API bool IsValid() const;
+
 public:
 
 	/**
