Sequencer: fix more edge cases with camera cut track camera locking

When unlocking the camera cut track, we were only discarding the transient pre-animated state. This caused two problems:

1. This was mostly promoting saved state to "persistent" state, which means that we could still mess up the viewports when closing the Sequencer editor.
2. We have 2 things to do in pre-animated state: clear the viewport cinematic lock and optionally restore the viewport position. Since we always need to do the former, we need to always save pre-animated state, but be able to check for the latter while restoring.

So now we always save pre-animated state, check for the restoring of viewport position inside the restoring function, and completely discard all pre-animated state when unlocking cameras. This last bit required writing a new function on the pre-animated state linker extension because we need to discard all tracking and bookkeeping of contributors... (we don't want anything to happen once we scrub out of the camera cuts)

This change also fixes several methods of FPreAnimatedStateExtension not calling into the track instance and track instance input sources.

diff --git a/Engine/Source/Editor/Sequencer/Private/Sequencer.cpp b/Engine/Source/Editor/Sequencer/Private/Sequencer.cpp
--- a/Engine/Source/Editor/Sequencer/Private/Sequencer.cpp
+++ b/Engine/Source/Editor/Sequencer/Private/Sequencer.cpp
@@ -11556,7 +11556,7 @@ UE::Sequencer::FSequencerSelection& FSequencer::GetSelection()
 	return *ViewModel->GetSelection();
 }
 
-bool FSequencer::ShouldCacheEditorPreAnimatedState()
+bool FSequencer::ShouldRestoreEditorViewports()
 {
 	return Settings->GetRestoreOriginalViewportOnCameraCutUnlock();
 }
diff --git a/Engine/Source/Editor/Sequencer/Private/Sequencer.h b/Engine/Source/Editor/Sequencer/Private/Sequencer.h
--- a/Engine/Source/Editor/Sequencer/Private/Sequencer.h
+++ b/Engine/Source/Editor/Sequencer/Private/Sequencer.h
@@ -858,7 +858,7 @@ protected:
 	// FCameraCutPlaybackCapability interface
 
 	virtual bool ShouldUpdateCameraCut() override { return IsPerspectiveViewportCameraCutEnabled(); }
-	virtual bool ShouldCacheEditorPreAnimatedState() override;
+	virtual bool ShouldRestoreEditorViewports() override;
 	virtual float GetCameraBlendPlayRate() override;
 	virtual void OnCameraCutUpdated(const UE::MovieScene::FOnCameraCutUpdatedParams& Params) override;
 
diff --git a/Engine/Source/Runtime/MovieScene/Private/Evaluation/PreAnimatedState/MovieScenePreAnimatedCaptureSources.inl b/Engine/Source/Runtime/MovieScene/Private/Evaluation/PreAnimatedState/MovieScenePreAnimatedCaptureSources.inl
--- a/Engine/Source/Runtime/MovieScene/Private/Evaluation/PreAnimatedState/MovieScenePreAnimatedCaptureSources.inl
+++ b/Engine/Source/Runtime/MovieScene/Private/Evaluation/PreAnimatedState/MovieScenePreAnimatedCaptureSources.inl
@@ -115,6 +115,30 @@ void TPreAnimatedCaptureSources<KeyType>::GatherAndRemoveMetaDataForGroup(FPreAn
 	}
 }
 
+template<typename KeyType>
+void TPreAnimatedCaptureSources<KeyType>::GatherAndRemoveMetaDataForStorage(FPreAnimatedStorageID StorageID, FPreAnimatedStorageIndex StorageIndex, TArray<FPreAnimatedStateMetaData>& OutExpiredMetaData)
+{
+	for (auto It = KeyToMetaData.CreateIterator(); It; ++It)
+	{
+		FPreAnimatedStateMetaDataArray& Array = It.Value();
+		for (int32 Index = Array.Num()-1; Index >= 0; --Index)
+		{
+			const FPreAnimatedStateMetaData& MetaData = Array[Index];
+			if (MetaData.Entry.ValueHandle.TypeID == StorageID &&
+					(!StorageIndex.IsValid() || MetaData.Entry.ValueHandle.StorageIndex == StorageIndex))
+			{
+				OutExpiredMetaData.Add(MetaData);
+				Array.RemoveAt(Index, 1, false);
+			}
+		}
+
+		if (Array.Num() == 0)
+		{
+			It.RemoveCurrent();
+		}
+	}
+}
+
 template<typename KeyType>
 bool TPreAnimatedCaptureSources<KeyType>::ContainsInstanceHandle(FRootInstanceHandle RootInstanceHandle) const
 {
diff --git a/Engine/Source/Runtime/MovieScene/Private/Evaluation/PreAnimatedState/MovieScenePreAnimatedEntityCaptureSource.cpp b/Engine/Source/Runtime/MovieScene/Private/Evaluation/PreAnimatedState/MovieScenePreAnimatedEntityCaptureSource.cpp
--- a/Engine/Source/Runtime/MovieScene/Private/Evaluation/PreAnimatedState/MovieScenePreAnimatedEntityCaptureSource.cpp
+++ b/Engine/Source/Runtime/MovieScene/Private/Evaluation/PreAnimatedState/MovieScenePreAnimatedEntityCaptureSource.cpp
@@ -130,6 +130,29 @@ void FPreAnimatedEntityCaptureSource::GatherAndRemoveMetaDataForGroup(FPreAnimat
 	}
 }
 
+void FPreAnimatedEntityCaptureSource::GatherAndRemoveMetaDataForStorage(FPreAnimatedStorageID StorageID, FPreAnimatedStorageIndex StorageIndex, TArray<FPreAnimatedStateMetaData>& OutExpiredMetaData)
+{
+	for (auto It = KeyToMetaData.CreateIterator(); It; ++It)
+	{
+		FPreAnimatedStateMetaDataArray& Array = It.Value();
+		for (int32 Index = Array.Num()-1; Index >= 0; --Index)
+		{
+			const FPreAnimatedStateMetaData& MetaData = Array[Index];
+			if (MetaData.Entry.ValueHandle.TypeID == StorageID &&
+					(!StorageIndex.IsValid() || MetaData.Entry.ValueHandle.StorageIndex == StorageIndex))
+			{
+				OutExpiredMetaData.Add(MetaData);
+				Array.RemoveAt(Index, 1, false);
+			}
+		}
+
+		if (Array.Num() == 0)
+		{
+			It.RemoveCurrent();
+		}
+	}
+}
+
 bool FPreAnimatedEntityCaptureSource::ContainsInstanceHandle(FRootInstanceHandle RootInstanceHandle) const
 {
 	for (const TPair<FMovieSceneEntityID, FPreAnimatedStateMetaDataArray>& Pair : KeyToMetaData)
diff --git a/Engine/Source/Runtime/MovieScene/Private/Evaluation/PreAnimatedState/MovieScenePreAnimatedStateExtension.cpp b/Engine/Source/Runtime/MovieScene/Private/Evaluation/PreAnimatedState/MovieScenePreAnimatedStateExtension.cpp
--- a/Engine/Source/Runtime/MovieScene/Private/Evaluation/PreAnimatedState/MovieScenePreAnimatedStateExtension.cpp
+++ b/Engine/Source/Runtime/MovieScene/Private/Evaluation/PreAnimatedState/MovieScenePreAnimatedStateExtension.cpp
@@ -309,6 +309,16 @@ void FPreAnimatedStateExtension::RestoreGlobalState(const FRestoreStateParams& P
 		EntityMetaData->GatherAndRemoveExpiredMetaData(Params, ExpiredMetaData);
 	}
 
+	if (FPreAnimatedTrackInstanceCaptureSources* TrackInstanceMetaData = GetTrackInstanceMetaData())
+	{
+		TrackInstanceMetaData->GatherAndRemoveExpiredMetaData(Params, ExpiredMetaData);
+	}
+
+	if (FPreAnimatedTrackInstanceInputCaptureSources* TrackInstanceInputMetaData = GetTrackInstanceInputMetaData())
+	{
+		TrackInstanceInputMetaData->GatherAndRemoveExpiredMetaData(Params, ExpiredMetaData);
+	}
+
 	for (int32 Index = WeakExternalCaptureSources.Num()-1; Index >= 0; --Index)
 	{
 		TSharedPtr<IPreAnimatedCaptureSource> MetaData = WeakExternalCaptureSources[Index].Pin();
@@ -429,6 +439,16 @@ void FPreAnimatedStateExtension::DiscardGlobalState(const FRestoreStateParams& P
 		EntityMetaData->GatherAndRemoveExpiredMetaData(Params, ExpiredMetaData);
 	}
 
+	if (FPreAnimatedTrackInstanceCaptureSources* TrackInstanceMetaData = GetTrackInstanceMetaData())
+	{
+		TrackInstanceMetaData->GatherAndRemoveExpiredMetaData(Params, ExpiredMetaData);
+	}
+
+	if (FPreAnimatedTrackInstanceInputCaptureSources* TrackInstanceInputMetaData = GetTrackInstanceInputMetaData())
+	{
+		TrackInstanceInputMetaData->GatherAndRemoveExpiredMetaData(Params, ExpiredMetaData);
+	}
+
 	for (int32 Index = WeakExternalCaptureSources.Num() - 1; Index >= 0; --Index)
 	{
 		TSharedPtr<IPreAnimatedCaptureSource> MetaData = WeakExternalCaptureSources[Index].Pin();
@@ -523,6 +543,16 @@ void FPreAnimatedStateExtension::DiscardTransientState()
 		EntityMetaData->Reset();
 	}
 
+	if (FPreAnimatedTrackInstanceCaptureSources* TrackInstanceMetaData = GetTrackInstanceMetaData())
+	{
+		TrackInstanceMetaData->Reset();
+	}
+
+	if (FPreAnimatedTrackInstanceInputCaptureSources* TrackInstanceInputMetaData = GetTrackInstanceInputMetaData())
+	{
+		TrackInstanceInputMetaData->Reset();
+	}
+
 	for (int32 Index = WeakExternalCaptureSources.Num()-1; Index >= 0; --Index)
 	{
 		if (TSharedPtr<IPreAnimatedCaptureSource> MetaData = WeakExternalCaptureSources[Index].Pin())
@@ -563,6 +593,16 @@ void FPreAnimatedStateExtension::DiscardStateForGroup(FPreAnimatedStorageGroupHa
 		EntityMetaData->GatherAndRemoveMetaDataForGroup(GroupHandle, MetaDataToRemove);
 	}
 
+	if (FPreAnimatedTrackInstanceCaptureSources* TrackInstanceMetaData = GetTrackInstanceMetaData())
+	{
+		TrackInstanceMetaData->GatherAndRemoveMetaDataForGroup(GroupHandle, MetaDataToRemove);
+	}
+
+	if (FPreAnimatedTrackInstanceInputCaptureSources* TrackInstanceInputMetaData = GetTrackInstanceInputMetaData())
+	{
+		TrackInstanceInputMetaData->GatherAndRemoveMetaDataForGroup(GroupHandle, MetaDataToRemove);
+	}
+
 	for (int32 Index = WeakExternalCaptureSources.Num()-1; Index >= 0; --Index)
 	{
 		if (TSharedPtr<IPreAnimatedCaptureSource> MetaData = WeakExternalCaptureSources[Index].Pin())
@@ -586,6 +626,105 @@ void FPreAnimatedStateExtension::DiscardStateForGroup(FPreAnimatedStorageGroupHa
 	bEntriesInvalidated = true;
 }
 
+void FPreAnimatedStateExtension::DiscardStateForStorage(FPreAnimatedStorageID StorageID, FPreAnimatedStorageIndex StorageIndex)
+{
+	TArray<FPreAnimatedStateMetaData> MetaDataToRemove;
+
+	if (FPreAnimatedEntityCaptureSource* EntityMetaData = GetEntityMetaData())
+	{
+		EntityMetaData->GatherAndRemoveMetaDataForStorage(StorageID, StorageIndex, MetaDataToRemove);
+	}
+
+	if (FPreAnimatedTrackInstanceCaptureSources* TrackInstanceMetaData = GetTrackInstanceMetaData())
+	{
+		TrackInstanceMetaData->GatherAndRemoveMetaDataForStorage(StorageID, StorageIndex, MetaDataToRemove);
+	}
+
+	if (FPreAnimatedTrackInstanceInputCaptureSources* TrackInstanceInputMetaData = GetTrackInstanceInputMetaData())
+	{
+		TrackInstanceInputMetaData->GatherAndRemoveMetaDataForStorage(StorageID, StorageIndex, MetaDataToRemove);
+	}
+
+	for (int32 Index = WeakExternalCaptureSources.Num()-1; Index >= 0; --Index)
+	{
+		if (TSharedPtr<IPreAnimatedCaptureSource> MetaData = WeakExternalCaptureSources[Index].Pin())
+		{
+			MetaData->GatherAndRemoveMetaDataForStorage(StorageID, StorageIndex, MetaDataToRemove);
+		}
+	}
+
+	// Remove all contributions
+	for (const FPreAnimatedStateMetaData& MetaData : MetaDataToRemove)
+	{
+		FAggregatePreAnimatedStateMetaData* Aggregate = FindMetaData(MetaData.Entry);
+		if (ensure(Aggregate))
+		{
+			const int32 TotalNum = --Aggregate->NumContributors;
+			if (MetaData.bWantsRestoreState)
+			{
+				--Aggregate->NumRestoreContributors;
+			}
+
+			if (TotalNum == 0)
+			{
+				Aggregate->bWantedRestore = false;
+				Aggregate->TerminalInstanceHandle = MetaData.RootInstanceHandle;
+			}
+		}
+	}
+
+	TSharedPtr<IPreAnimatedStorage> Storage = GetStorageChecked(StorageID);
+
+	// Discard grouped entries
+	for (int32 Index = 0; Index < GroupMetaData.GetMaxIndex(); ++Index)
+	{
+		if (!GroupMetaData.IsAllocated(Index))
+		{
+			continue;
+		}
+
+		FPreAnimatedGroupMetaData& Group = GroupMetaData[Index];
+
+		for (int32 AggregateIndex = Group.AggregateMetaData.Num() - 1; AggregateIndex >= 0; --AggregateIndex)
+		{
+			FAggregatePreAnimatedStateMetaData& Aggregate = Group.AggregateMetaData[AggregateIndex];
+			if (Aggregate.ValueHandle.TypeID == StorageID && 
+					Aggregate.NumContributors == 0 &&
+					(!StorageIndex.IsValid() || Aggregate.ValueHandle.StorageIndex == StorageIndex))
+			{
+				Storage->DiscardPreAnimatedStateStorage(Aggregate.ValueHandle.StorageIndex, EPreAnimatedStorageRequirement::Persistent);
+
+				Group.AggregateMetaData.RemoveAt(AggregateIndex, 1, false);
+			}
+
+			if (Group.AggregateMetaData.Num() == 0)
+			{
+				// Remove at will not re-allocate the array or shuffle items within the sparse array, so this is safe
+				Group.GroupManagerPtr->OnGroupDestroyed(Index);
+				GroupMetaData.RemoveAt(Index);
+			}
+		}
+	}
+
+	// Discard ungrouped entries
+	for (auto UngroupedIt = UngroupedMetaData.CreateIterator(); UngroupedIt; ++UngroupedIt)
+	{
+		FAggregatePreAnimatedStateMetaData& Aggregate = UngroupedIt.Value();
+		if (Aggregate.ValueHandle.TypeID == StorageID && 
+				Aggregate.NumContributors == 0 &&
+				(!StorageIndex.IsValid() || Aggregate.ValueHandle.StorageIndex == StorageIndex))
+		{
+			Storage->DiscardPreAnimatedStateStorage(Aggregate.ValueHandle.StorageIndex, EPreAnimatedStorageRequirement::Persistent);
+
+			UngroupedIt.RemoveCurrent();
+		}
+	}
+
+	GroupMetaData.Shrink();
+
+	bEntriesInvalidated = true;
+}
+
 bool FPreAnimatedStateExtension::ContainsAnyStateForInstanceHandle(FRootInstanceHandle RootInstanceHandle) const
 {
 	if (FPreAnimatedEntityCaptureSource* EntityMetaData = GetEntityMetaData())
diff --git a/Engine/Source/Runtime/MovieScene/Public/Evaluation/PreAnimatedState/IMovieScenePreAnimatedCaptureSource.h b/Engine/Source/Runtime/MovieScene/Public/Evaluation/PreAnimatedState/IMovieScenePreAnimatedCaptureSource.h
--- a/Engine/Source/Runtime/MovieScene/Public/Evaluation/PreAnimatedState/IMovieScenePreAnimatedCaptureSource.h
+++ b/Engine/Source/Runtime/MovieScene/Public/Evaluation/PreAnimatedState/IMovieScenePreAnimatedCaptureSource.h
@@ -15,6 +15,8 @@ struct FRootInstanceHandle;
 struct FRestoreStateParams;
 struct FPreAnimatedStateMetaData;
 struct FPreAnimatedStorageGroupHandle;
+struct FPreAnimatedStorageID;
+struct FPreAnimatedStorageIndex;
 
 struct IPreAnimatedCaptureSource
 {
@@ -24,6 +26,7 @@ struct IPreAnimatedCaptureSource
 	virtual bool ContainsInstanceHandle(FRootInstanceHandle RootInstanceHandle) const = 0;
 	virtual void GatherAndRemoveExpiredMetaData(const FRestoreStateParams& Params, TArray<FPreAnimatedStateMetaData>& OutExpiredMetaData) = 0;
 	virtual void GatherAndRemoveMetaDataForGroup(FPreAnimatedStorageGroupHandle Group, TArray<FPreAnimatedStateMetaData>& OutExpiredMetaData) = 0;
+	virtual void GatherAndRemoveMetaDataForStorage(FPreAnimatedStorageID StorageID, FPreAnimatedStorageIndex StorageIndex, TArray<FPreAnimatedStateMetaData>& OutExpiredMetaData) = 0;
 };
 
 
diff --git a/Engine/Source/Runtime/MovieScene/Public/Evaluation/PreAnimatedState/MovieScenePreAnimatedCaptureSources.h b/Engine/Source/Runtime/MovieScene/Public/Evaluation/PreAnimatedState/MovieScenePreAnimatedCaptureSources.h
--- a/Engine/Source/Runtime/MovieScene/Public/Evaluation/PreAnimatedState/MovieScenePreAnimatedCaptureSources.h
+++ b/Engine/Source/Runtime/MovieScene/Public/Evaluation/PreAnimatedState/MovieScenePreAnimatedCaptureSources.h
@@ -47,7 +47,7 @@ struct FPreAnimatedEvalHookKeyType
 };
 
 /**
- * Structure responsible for tracking contributions to pre-eanimated state entries that originate from track templates (ie, from an IMovieSceneExecutionToken::Execute)
+ * Structure responsible for tracking contributions to pre-animated state entries that originate from track templates (ie, from an IMovieSceneExecutionToken::Execute)
  */
 template<typename KeyType>
 struct TPreAnimatedCaptureSources : IPreAnimatedCaptureSource
@@ -69,6 +69,7 @@ struct TPreAnimatedCaptureSources : IPreAnimatedCaptureSource
 	bool ContainsInstanceHandle(FRootInstanceHandle RootInstanceHandle) const override;
 	void GatherAndRemoveExpiredMetaData(const FRestoreStateParams& Params, TArray<FPreAnimatedStateMetaData>& OutExpiredMetaData) override;
 	void GatherAndRemoveMetaDataForGroup(FPreAnimatedStorageGroupHandle Group, TArray<FPreAnimatedStateMetaData>& OutExpiredMetaData) override;
+	void GatherAndRemoveMetaDataForStorage(FPreAnimatedStorageID StorageID, FPreAnimatedStorageIndex StorageIndex, TArray<FPreAnimatedStateMetaData>& OutExpiredMetaData) override;
 
 private:
 
diff --git a/Engine/Source/Runtime/MovieScene/Public/Evaluation/PreAnimatedState/MovieScenePreAnimatedEntityCaptureSource.h b/Engine/Source/Runtime/MovieScene/Public/Evaluation/PreAnimatedState/MovieScenePreAnimatedEntityCaptureSource.h
--- a/Engine/Source/Runtime/MovieScene/Public/Evaluation/PreAnimatedState/MovieScenePreAnimatedEntityCaptureSource.h
+++ b/Engine/Source/Runtime/MovieScene/Public/Evaluation/PreAnimatedState/MovieScenePreAnimatedEntityCaptureSource.h
@@ -36,6 +36,7 @@ struct FPreAnimatedEntityCaptureSource
 	bool ContainsInstanceHandle(FRootInstanceHandle RootInstanceHandle) const;
 	void GatherAndRemoveExpiredMetaData(const FRestoreStateParams& Params, TArray<FPreAnimatedStateMetaData>& OutExpiredMetaData);
 	void GatherAndRemoveMetaDataForGroup(FPreAnimatedStorageGroupHandle Group, TArray<FPreAnimatedStateMetaData>& OutExpiredMetaData);
+	void GatherAndRemoveMetaDataForStorage(FPreAnimatedStorageID StorageID, FPreAnimatedStorageIndex StorageIndex, TArray<FPreAnimatedStateMetaData>& OutExpiredMetaData);
 
 private:
 
diff --git a/Engine/Source/Runtime/MovieScene/Public/Evaluation/PreAnimatedState/MovieScenePreAnimatedStateExtension.h b/Engine/Source/Runtime/MovieScene/Public/Evaluation/PreAnimatedState/MovieScenePreAnimatedStateExtension.h
--- a/Engine/Source/Runtime/MovieScene/Public/Evaluation/PreAnimatedState/MovieScenePreAnimatedStateExtension.h
+++ b/Engine/Source/Runtime/MovieScene/Public/Evaluation/PreAnimatedState/MovieScenePreAnimatedStateExtension.h
@@ -268,6 +268,11 @@ public:
 	 */
 	MOVIESCENE_API void DiscardStateForGroup(FPreAnimatedStorageGroupHandle GroupHandle);
 
+	/**
+	 * Discard the specified cached value and any and all capture source tracking related to it.
+	 */
+	MOVIESCENE_API void DiscardStateForStorage(FPreAnimatedStorageID StorageID, FPreAnimatedStorageIndex StorageIndex);
+
 	/**
 	 * Search for any captured state that originated from the specified root instance handle
 	 * WARNING: This is a linear search across all state, and so is potentially very slow
diff --git a/Engine/Source/Runtime/MovieSceneTracks/Private/TrackInstances/MovieSceneCameraCutEditorHandler.cpp b/Engine/Source/Runtime/MovieSceneTracks/Private/TrackInstances/MovieSceneCameraCutEditorHandler.cpp
--- a/Engine/Source/Runtime/MovieSceneTracks/Private/TrackInstances/MovieSceneCameraCutEditorHandler.cpp
+++ b/Engine/Source/Runtime/MovieSceneTracks/Private/TrackInstances/MovieSceneCameraCutEditorHandler.cpp
@@ -62,9 +62,16 @@ void FPreAnimatedCameraCutEditorTraits::RestorePreAnimatedValue(
 		return;
 	}
 	
-	InKey->SetViewLocation(CachedValue.ViewportLocation);
-	InKey->SetViewRotation(CachedValue.ViewportRotation);
-	InKey->ViewFOV = CachedValue.ViewportFOV;
+	// Restore pre-animated viewport position if needed. Then disable cinematic lock either way.
+	FInstanceRegistry* InstanceRegistry = Params.Linker->GetInstanceRegistry();
+	const FSequenceInstance& TerminalInstance = InstanceRegistry->GetInstance(Params.TerminalInstanceHandle);
+	FCameraCutPlaybackCapabilityCompatibilityWrapper Wrapper(TerminalInstance);
+	if (Wrapper.ShouldRestoreEditorViewports())
+	{
+		InKey->SetViewLocation(CachedValue.ViewportLocation);
+		InKey->SetViewRotation(CachedValue.ViewportRotation);
+		InKey->ViewFOV = CachedValue.ViewportFOV;
+	}
 
 	InKey->SetCinematicActorLock(nullptr);
 	InKey->bLockedCameraView = false;
@@ -83,12 +90,6 @@ void FCameraCutEditorHandler::CachePreAnimatedValue(
 		return;
 	}
 
-	FCameraCutPlaybackCapabilityCompatibilityWrapper Wrapper(SequenceInstance);
-	if (!Wrapper.ShouldCacheEditorPreAnimatedState())
-	{
-		return;
-	}
-
 	IMovieScenePlayer* Player = SequenceInstance.GetPlayer();
 	UObject* PlaybackContext = Player->GetPlaybackContext();
 	UWorld* ContextWorld = PlaybackContext ? PlaybackContext->GetWorld() : nullptr;
@@ -126,6 +127,8 @@ void FCameraCutEditorHandler::ForcePreAnimatedValueOperation(
 		return;
 	}
 
+	TGuardValue<FEntityManager*> EntityManagerForDebugging(GEntityManagerForDebuggingVisualizers, &Linker->EntityManager);
+
 	IMovieScenePlayer* Player = SequenceInstance.GetPlayer();
 	UObject* PlaybackContext = Player->GetPlaybackContext();
 	UWorld* ContextWorld = PlaybackContext ? PlaybackContext->GetWorld() : nullptr;
@@ -164,13 +167,13 @@ void FCameraCutEditorHandler::ForcePreAnimatedValueOperation(
 					{
 						FRestoreStateParams Params;
 						Params.Linker = Linker;
-						// No need to set the terminal instance handle, we don't use it (see above).
+						Params.TerminalInstanceHandle = SequenceInstance.GetRootInstanceHandle();
 						PreAnimatedStorage->RestorePreAnimatedStateStorage(StorageIndex, EPreAnimatedStorageRequirement::Transient, EPreAnimatedStorageRequirement::Persistent, Params);
 					}
 					break;
 				case EForcedCameraCutPreAnimatedStorageOperation::Discard:
 					{
-						PreAnimatedStorage->DiscardPreAnimatedStateStorage(StorageIndex, EPreAnimatedStorageRequirement::Transient);
+						Linker->PreAnimatedState.DiscardStateForStorage(FPreAnimatedCameraCutEditorStorage::StorageID, StorageIndex);
 					}
 					break;
 			}
diff --git a/Engine/Source/Runtime/MovieSceneTracks/Private/TrackInstances/MovieSceneCameraCutTrackInstance.cpp b/Engine/Source/Runtime/MovieSceneTracks/Private/TrackInstances/MovieSceneCameraCutTrackInstance.cpp
--- a/Engine/Source/Runtime/MovieSceneTracks/Private/TrackInstances/MovieSceneCameraCutTrackInstance.cpp
+++ b/Engine/Source/Runtime/MovieSceneTracks/Private/TrackInstances/MovieSceneCameraCutTrackInstance.cpp
@@ -58,11 +58,11 @@ bool FCameraCutPlaybackCapabilityCompatibilityWrapper::ShouldUpdateCameraCut()
 }
 
 #if WITH_EDITOR
-bool FCameraCutPlaybackCapabilityCompatibilityWrapper::ShouldCacheEditorPreAnimatedState()
+bool FCameraCutPlaybackCapabilityCompatibilityWrapper::ShouldRestoreEditorViewports()
 {
 	if (CameraCutCapability)
 	{
-		return CameraCutCapability->ShouldCacheEditorPreAnimatedState();
+		return CameraCutCapability->ShouldRestoreEditorViewports();
 	}
 	return true;
 }
diff --git a/Engine/Source/Runtime/MovieSceneTracks/Public/Evaluation/CameraCutPlaybackCapability.h b/Engine/Source/Runtime/MovieSceneTracks/Public/Evaluation/CameraCutPlaybackCapability.h
--- a/Engine/Source/Runtime/MovieSceneTracks/Public/Evaluation/CameraCutPlaybackCapability.h
+++ b/Engine/Source/Runtime/MovieSceneTracks/Public/Evaluation/CameraCutPlaybackCapability.h
@@ -42,7 +42,7 @@ namespace UE::MovieScene
 
 #if WITH_EDITOR
 		/** Whether the editor should cache the pre-animated viewport position */
-		virtual bool ShouldCacheEditorPreAnimatedState() { return true; }
+		virtual bool ShouldRestoreEditorViewports() { return true; }
 #endif
 	};
 
diff --git a/Engine/Source/Runtime/MovieSceneTracks/Public/TrackInstances/MovieSceneCameraCutTrackInstance.h b/Engine/Source/Runtime/MovieSceneTracks/Public/TrackInstances/MovieSceneCameraCutTrackInstance.h
--- a/Engine/Source/Runtime/MovieSceneTracks/Public/TrackInstances/MovieSceneCameraCutTrackInstance.h
+++ b/Engine/Source/Runtime/MovieSceneTracks/Public/TrackInstances/MovieSceneCameraCutTrackInstance.h
@@ -28,7 +28,7 @@ namespace UE::MovieScene
 		bool ShouldUpdateCameraCut();
 		void OnCameraCutUpdated(const FOnCameraCutUpdatedParams& Params);
 #if WITH_EDITOR
-		bool ShouldCacheEditorPreAnimatedState();
+		bool ShouldRestoreEditorViewports();
 #endif
 
 		FCameraCutPlaybackCapability* CameraCutCapability;
