Sequencer- Interface for objects to receive events when they are bound/unbound by Sequencer.

Added a new interface that allows any UObject to implement an interface (either in C++ or BP) to receive events when Sequencer binds/unbounds that object. To get the events requires the sequence add a Binding Lifetime track to the binding. If the BindingLifetime track remains infinite in length, then the events will fire at the beginning and end of sequence playback respectively. If the binding lifetime track has a finite range, then the events will fire when entering/exiting the binding lifetime section.

diff --git a/Engine/Source/Runtime/MovieScene/Private/EntitySystem/MovieSceneBindingLifetimeSystem.cpp b/Engine/Source/Runtime/MovieScene/Private/EntitySystem/MovieSceneBindingLifetimeSystem.cpp
--- a/Engine/Source/Runtime/MovieScene/Private/EntitySystem/MovieSceneBindingLifetimeSystem.cpp
+++ b/Engine/Source/Runtime/MovieScene/Private/EntitySystem/MovieSceneBindingLifetimeSystem.cpp
@@ -20,6 +20,7 @@
 #include "MovieSceneExecutionToken.h"
 #include "IMovieScenePlaybackClient.h"
 #include "EntitySystem/MovieSceneSpawnablesSystem.h"
+#include "MovieSceneBindingEventReceiverInterface.h"
 
 #include UE_INLINE_GENERATED_CPP_BY_NAME(MovieSceneBindingLifetimeSystem)
 #define LOCTEXT_NAMESPACE "MovieSceneBindingLifetimeSystem"
@@ -74,10 +75,35 @@ void UMovieSceneBindingLifetimeSystem::OnRun(FSystemTaskPrerequisites& InPrerequ
 		if (Player)
 		{
 			// For now we use the linking/unlinking of the inactive ranges to set the binding activations
-			if (BindingLifetime.BindingLifetimeState == EMovieSceneBindingLifetimeState::InActive)
+			if (BindingLifetime.BindingLifetimeState == EMovieSceneBindingLifetimeState::Inactive)
 			{
 				Player->State.SetBindingActivation(BindingLifetime.BindingGuid, SequenceID, !bLink);
 			}
+			else
+			{
+				for (TWeakObjectPtr<> WeakBoundObject : Player->FindBoundObjects(BindingLifetime.BindingGuid, SequenceID))
+				{
+					if (UObject* BoundObject = WeakBoundObject.Get())
+					{
+						if (BoundObject->Implements<UMovieSceneBindingEventReceiverInterface>())
+						{
+							TScriptInterface<IMovieSceneBindingEventReceiverInterface> BindingEventReceiver = BoundObject;
+							if (BindingEventReceiver.GetObject())
+							{
+								FMovieSceneObjectBindingID BindingID = UE::MovieScene::FRelativeObjectBindingID(MovieSceneSequenceID::Root, SequenceID, BindingLifetime.BindingGuid, *Player);
+								if (bLink)
+								{
+									IMovieSceneBindingEventReceiverInterface::Execute_OnObjectBoundBySequencer(BindingEventReceiver.GetObject(), Cast<UMovieSceneSequencePlayer>(Player->AsUObject()), BindingID);
+								}
+								else
+								{
+									IMovieSceneBindingEventReceiverInterface::Execute_OnObjectUnboundBySequencer(BindingEventReceiver.GetObject(), Cast<UMovieSceneSequencePlayer>(Player->AsUObject()), BindingID);
+								}
+							}
+						}
+					}
+				}
+			}
 		}
 	};
 
diff --git a/Engine/Source/Runtime/MovieScene/Private/Tracks/MovieSceneBindingLifetimeTrack.cpp b/Engine/Source/Runtime/MovieScene/Private/Tracks/MovieSceneBindingLifetimeTrack.cpp
--- a/Engine/Source/Runtime/MovieScene/Private/Tracks/MovieSceneBindingLifetimeTrack.cpp
+++ b/Engine/Source/Runtime/MovieScene/Private/Tracks/MovieSceneBindingLifetimeTrack.cpp
@@ -93,7 +93,7 @@ void UMovieSceneBindingLifetimeTrack::ImportEntityImpl(UMovieSceneEntitySystemLi
 
 		OutImportedEntity->AddBuilder(
 			FEntityBuilder()
-			.AddConditional(BuiltInComponentTypes->BindingLifetime, FMovieSceneBindingLifetimeComponentData{ Params.GetObjectBindingID(), EMovieSceneBindingLifetimeState::InActive }, Params.GetObjectBindingID().IsValid())
+			.AddConditional(BuiltInComponentTypes->BindingLifetime, FMovieSceneBindingLifetimeComponentData{ Params.GetObjectBindingID(), EMovieSceneBindingLifetimeState::Inactive }, Params.GetObjectBindingID().IsValid())
 		);
 }
 
diff --git a/Engine/Source/Runtime/MovieScene/Public/EntitySystem/BuiltInComponentTypes.h b/Engine/Source/Runtime/MovieScene/Public/EntitySystem/BuiltInComponentTypes.h
--- a/Engine/Source/Runtime/MovieScene/Public/EntitySystem/BuiltInComponentTypes.h
+++ b/Engine/Source/Runtime/MovieScene/Public/EntitySystem/BuiltInComponentTypes.h
@@ -332,7 +332,7 @@ enum class EMovieSceneBindingLifetimeState : uint8
 	/** Object Binding State is Active. Tracks on this binding will be active. */
 	Active,
 	/** Object Binding is Inactive. Tracks on this binding will be inactive and states will be restored/kept depending on settings. */
-	InActive
+	Inactive
 };
 
 struct FMovieSceneBindingLifetimeComponentData
diff --git a/Engine/Source/Runtime/MovieScene/Public/MovieSceneBindingEventReceiverInterface.h b/Engine/Source/Runtime/MovieScene/Public/MovieSceneBindingEventReceiverInterface.h
new file mode 100644
--- /dev/null
+++ b/Engine/Source/Runtime/MovieScene/Public/MovieSceneBindingEventReceiverInterface.h
@@ -0,0 +1,39 @@
+// Copyright Epic Games, Inc. All Rights Reserved.
+
+#pragma once
+
+#include "CoreMinimal.h"
+#include "UObject/ObjectMacros.h"
+#include "UObject/Interface.h"
+#include "MovieSceneObjectBindingID.h"
+#include "MovieSceneSequencePlayer.h"
+#include "MovieSceneBindingEventReceiverInterface.generated.h"
+
+
+class UMovieSceneSequence;
+class IPropertyHandle;
+class FStructOnScope;
+
+
+/** Interface for objects to implement used by the binding lifetime track to provide events when sequnecer activates or deactivates a binding.  */
+UINTERFACE(MinimalAPI)
+class UMovieSceneBindingEventReceiverInterface
+	: public UInterface
+{
+public:
+	GENERATED_BODY()
+};
+
+class IMovieSceneBindingEventReceiverInterface
+{
+public:
+	GENERATED_BODY()
+
+	UFUNCTION(BlueprintNativeEvent, CallInEditor)
+	void OnObjectBoundBySequencer(UMovieSceneSequencePlayer* Player, FMovieSceneObjectBindingID BindingID);
+	virtual void OnObjectBoundBySequencer_Implementation(UMovieSceneSequencePlayer* Player, FMovieSceneObjectBindingID BindingID) { };
+
+	UFUNCTION(BlueprintNativeEvent, CallInEditor)
+	void OnObjectUnboundBySequencer(UMovieSceneSequencePlayer* Player, FMovieSceneObjectBindingID BindingID);
+	virtual void OnObjectUnboundBySequencer_Implementation(UMovieSceneSequencePlayer* Player, FMovieSceneObjectBindingID BindingID) { };
+};
