Sequencer: Fixed implicit numeric cast

diff --git a/Engine/Source/Runtime/MovieScene/Public/MovieSceneSection.h b/Engine/Source/Runtime/MovieScene/Public/MovieSceneSection.h
--- a/Engine/Source/Runtime/MovieScene/Public/MovieSceneSection.h
+++ b/Engine/Source/Runtime/MovieScene/Public/MovieSceneSection.h
@@ -797,7 +797,7 @@ template<typename SectionParams>
 inline FFrameNumber GetFirstLoopStartOffsetAtTrimTime(FQualifiedFrameTime TrimTime, const SectionParams& Params, FFrameNumber StartFrame, FFrameRate FrameRate)
 {
 	const float AnimPlayRate = FMath::IsNearlyZero(Params.PlayRate) ? 1.0f : Params.PlayRate;
-	const float AnimPosition = (TrimTime.Time - StartFrame) / TrimTime.Rate * AnimPlayRate;
+	const float AnimPosition = static_cast<float>((TrimTime.Time - StartFrame) / TrimTime.Rate * AnimPlayRate);
 	const float SeqLength = Params.GetSequenceLength() - FrameRate.AsSeconds(Params.StartFrameOffset + Params.EndFrameOffset) / AnimPlayRate;
 
 	FFrameNumber NewOffset = FrameRate.AsFrameNumber(FMath::Fmod(AnimPosition, SeqLength));
