Sequencer: only set last known position on client-side initialization

Overwriting the whole NetSyncProps in Initialize could come after the first net packet already synchronized its values with the server. Since the server doesn't send values unless they have changed, some of these values could be stuck with the wrong value after we had overwritten them.

diff --git a/Engine/Source/Runtime/MovieScene/Private/MovieSceneSequencePlayer.cpp b/Engine/Source/Runtime/MovieScene/Private/MovieSceneSequencePlayer.cpp
--- a/Engine/Source/Runtime/MovieScene/Private/MovieSceneSequencePlayer.cpp
+++ b/Engine/Source/Runtime/MovieScene/Private/MovieSceneSequencePlayer.cpp
@@ -146,9 +146,9 @@ UMovieSceneSequencePlayer::~UMovieSceneSequencePlayer()
 	}
 }
 
-void UMovieSceneSequencePlayer::UpdateNetworkSyncProperties(bool bServerOnly)
+void UMovieSceneSequencePlayer::UpdateNetworkSyncProperties()
 {
-	if (!bServerOnly || HasAuthority())
+	if (HasAuthority())
 	{
 		NetSyncProps.LastKnownPosition = PlayPosition.GetCurrentPosition();
 		NetSyncProps.LastKnownStatus   = Status;
@@ -942,14 +942,16 @@ void UMovieSceneSequencePlayer::Initialize(UMovieSceneSequence* InSequence)
 	PlayPosition.Reset(StartTimeWithOffset);
 	TimeController->Reset(GetCurrentTime());
 
-	// Update the sync properties on the server *and* the client. We do it on the client because our first PostNetReceive
+	// Update the sync properties on the server.
+	UpdateNetworkSyncProperties();
+	// On the client, we also update LastKnownPosition. This is because our first PostNetReceive
 	// could be called with an incomplete set of replicated values in very rare cases... so for instance we might
 	// get the proper LastKnownStatus from the server, but, say, not the proper LastKnownPosition. If the sequence does
 	// not start at frame 0, we would see LastKnownPosition left at 0, while our own client-side position is the first
 	// frame of the sequence, as initialized above (SetFrameRange). At this point, we would incorrectly assume that the server 
 	// jumped to frame 0 and we would do the same, when really the server hasn't moved and it's just that the correct
 	// LastKnownPosition value is coming in a later net packet.
-	UpdateNetworkSyncProperties(false);
+	NetSyncProps.LastKnownPosition = PlayPosition.GetCurrentPosition();
 }
 
 void UMovieSceneSequencePlayer::Update(const float DeltaSeconds)
@@ -1194,7 +1196,7 @@ void UMovieSceneSequencePlayer::UpdateTimeCursorPosition_Internal(FFrameTime New
 		Args.bIsAsync = (bIsAsyncUpdate && !bIsSequenceBlocking);
 		Args.bHasJumped = bHasJumpedOverride;
 
-		PostEvaluationCallbacks.Add(FOnEvaluationCallback::CreateUObject(this, &UMovieSceneSequencePlayer::UpdateNetworkSyncProperties, true));
+		PostEvaluationCallbacks.Add(FOnEvaluationCallback::CreateUObject(this, &UMovieSceneSequencePlayer::UpdateNetworkSyncProperties));
 
 		UpdateMovieSceneInstance(Range, StatusOverride, Args);
 
diff --git a/Engine/Source/Runtime/MovieScene/Public/MovieSceneSequencePlayer.h b/Engine/Source/Runtime/MovieScene/Public/MovieSceneSequencePlayer.h
--- a/Engine/Source/Runtime/MovieScene/Public/MovieSceneSequencePlayer.h
+++ b/Engine/Source/Runtime/MovieScene/Public/MovieSceneSequencePlayer.h
@@ -605,10 +605,8 @@ private:
 
 	/**
 	 * Update the replicated properties required for synchronizing to clients of this sequence player
-	 *
-	 * @param bServerOnly  Only update the properties if this player is on the server
 	 */
-	MOVIESCENE_API void UpdateNetworkSyncProperties(bool bServerOnly = true);
+	MOVIESCENE_API void UpdateNetworkSyncProperties();
 
 	/**
 	 * Analyse the set of samples we have estimating the server time if we have confidence over the data.
