Sequencer: Ensure that compiled sequence data is removed from all managers on BeginDestroy

diff --git a/Engine/Source/Runtime/MovieScene/Private/MovieSceneSequence.cpp b/Engine/Source/Runtime/MovieScene/Private/MovieSceneSequence.cpp
--- a/Engine/Source/Runtime/MovieScene/Private/MovieSceneSequence.cpp
+++ b/Engine/Source/Runtime/MovieScene/Private/MovieSceneSequence.cpp
@@ -83,6 +83,11 @@ void UMovieSceneSequence::BeginDestroy()
 	if (!GExitPurge && !HasAnyFlags(RF_ClassDefaultObject))
 	{
 		UMovieSceneCompiledDataManager::GetPrecompiledData()->Reset(this);
+
+#if WITH_EDITOR
+		UMovieSceneCompiledDataManager::GetPrecompiledData(EMovieSceneServerClientMask::Client)->Reset(this);
+		UMovieSceneCompiledDataManager::GetPrecompiledData(EMovieSceneServerClientMask::Server)->Reset(this);
+#endif
 	}
 }
 
