Sequencer: move more IMovieScenePlayer features into playback capabilities

This is the third step in the overall refactor to remove the dependency of Sequencer evaluation code on IMovieScenePlayer. It puts most of the spawnables/possessables logic into playback capabilities, so that systems concerned with them can find them there instead of depending on the player interface. In order to do this, many APIs previously taking a player pointer now take a shared playback state reference, which is where they find the playback capabilities.

To keep the changelist to a manageable size, we don't deprecate the previous APIs right away except where we don't have much choice, such as the virtual methods of the spawn register. Everywhere else, passing a player is still a valid call.

Some other APIs have also changed, because the shared playback state structure often gives all the information someone may ever need, and this renders many parameters obsolete. This is for instance the case with sequence updaters.

One aspect of this change that needs to be carefully handled is the initialization and re-initialization of a sequence hierarchy prior to playback. Now that evaluation code relies on playback capabilities, these capabilities need to be added to the sequence hierarchy's shared playback state in a timely manner, and they need to react to changes in the hierarchy. The following changes have therefore been made:

- Playback capabilities can implement callbacks for initialization and for when a new sub-sequence instance is added.
- Sequence instances have a new Initialize method that needs to be called shortly after creation. This is especially needed for root instances, as it gives time for the caller to setup their playback capabilities before finishing the initialization.

Another aspect of this change is that it unearthed a problem with the capability storage, where we could stomp memory and crash while relocating capabilities to make room for a new one. This happened if we did not reallocate a new buffer (i.e. relocating in place). Moving headers could stomp the space of the first couple capabilities before we had a chance to move them. The code now moves things from back to front, first the capabilities, and second the headers.

Finally, this change adds the ability to overwrite a capability with a new value. This is needed for the way the level sequence player initializes itself (see the comment in the code).

The following things have been made into playback capabilities, or have had a new capability interface created around them:

- FMovieSceneEvaluationState, for access into object caches.
- IMovieScenePlaybackClient, for dynamic binding overrides.
- The player index, for retrieving the player interface, if any (mostly for legacy evaluators).
- Notification callbacks for when object bindings change.
- Static binding overrides.
- Storage and management of sequence director blueprints.

Since the player interface is now an optional playback capability, some progress has been made to gracefully handle the case of no such player existing.

diff --git a/Engine/Plugins/Experimental/ChaosCaching/Source/ChaosCaching/Private/Chaos/Sequencer/ChaosCacheObjectSpawner.cpp b/Engine/Plugins/Experimental/ChaosCaching/Source/ChaosCaching/Private/Chaos/Sequencer/ChaosCacheObjectSpawner.cpp
--- a/Engine/Plugins/Experimental/ChaosCaching/Source/ChaosCaching/Private/Chaos/Sequencer/ChaosCacheObjectSpawner.cpp
+++ b/Engine/Plugins/Experimental/ChaosCaching/Source/ChaosCaching/Private/Chaos/Sequencer/ChaosCacheObjectSpawner.cpp
@@ -18,9 +18,9 @@ UClass* FChaosCacheObjectSpawner::GetSupportedTemplateType() const
 	return AChaosCacheManager::StaticClass();
 }
 
-UObject* FChaosCacheObjectSpawner::SpawnObject(FMovieSceneSpawnable& Spawnable, FMovieSceneSequenceIDRef TemplateID, IMovieScenePlayer& Player)
+UObject* FChaosCacheObjectSpawner::SpawnObject(FMovieSceneSpawnable& Spawnable, FMovieSceneSequenceIDRef TemplateID, TSharedRef<const UE::MovieScene::FSharedPlaybackState> SharedPlaybackState)
 {
-	UObject* SpawnedObject = FLevelSequenceActorSpawner::SpawnObject(Spawnable, TemplateID, Player);
+	UObject* SpawnedObject = FLevelSequenceActorSpawner::SpawnObject(Spawnable, TemplateID, SharedPlaybackState);
 	
 	if (AChaosCacheManager* ChaosCache = Cast<AChaosCacheManager>(SpawnedObject))
 	{
diff --git a/Engine/Plugins/Experimental/ChaosCaching/Source/ChaosCaching/Public/Chaos/Sequencer/ChaosCacheObjectSpawner.h b/Engine/Plugins/Experimental/ChaosCaching/Source/ChaosCaching/Public/Chaos/Sequencer/ChaosCacheObjectSpawner.h
--- a/Engine/Plugins/Experimental/ChaosCaching/Source/ChaosCaching/Public/Chaos/Sequencer/ChaosCacheObjectSpawner.h
+++ b/Engine/Plugins/Experimental/ChaosCaching/Source/ChaosCaching/Public/Chaos/Sequencer/ChaosCacheObjectSpawner.h
@@ -15,7 +15,7 @@ public:
 
 	// IMovieSceneObjectSpawner interface
 	CHAOSCACHING_API virtual UClass* GetSupportedTemplateType() const override;
-	CHAOSCACHING_API virtual UObject* SpawnObject(FMovieSceneSpawnable& Spawnable, FMovieSceneSequenceIDRef TemplateID, IMovieScenePlayer& Player) override;
+	CHAOSCACHING_API virtual UObject* SpawnObject(FMovieSceneSpawnable& Spawnable, FMovieSceneSequenceIDRef TemplateID, TSharedRef<const UE::MovieScene::FSharedPlaybackState> SharedPlaybackState) override;
 	CHAOSCACHING_API virtual void DestroySpawnedObject(UObject& Object) override;
 	virtual bool IsEditor() const override { return true; }
 	virtual int32 GetSpawnerPriority() const override { return 1; }
diff --git a/Engine/Plugins/Importers/USDImporter/Source/USDExporter/Private/LevelSequenceExporterUSD.cpp b/Engine/Plugins/Importers/USDImporter/Source/USDExporter/Private/LevelSequenceExporterUSD.cpp
--- a/Engine/Plugins/Importers/USDImporter/Source/USDExporter/Private/LevelSequenceExporterUSD.cpp
+++ b/Engine/Plugins/Importers/USDImporter/Source/USDExporter/Private/LevelSequenceExporterUSD.cpp
@@ -105,7 +105,7 @@ namespace UE
 			public:
 				bool bDestroyingJustHides = true;
 
-				virtual UObject* SpawnObject( FMovieSceneSpawnable& Spawnable, FMovieSceneSequenceIDRef TemplateID, IMovieScenePlayer& Player ) override
+				virtual UObject* SpawnObject( FMovieSceneSpawnable& Spawnable, FMovieSceneSequenceIDRef TemplateID, TSharedRef<const FSharedPlaybackState> SharedPlaybackState ) override
 				{
 					// Never spawn ASphereReflectionCapture actors. These are useless in USD anyway, and we run into
 					// trouble after we're done exporting them because on the tick where they're destroyed the editor
@@ -116,8 +116,7 @@ namespace UE
 						return nullptr;
 					}
 
-					ISequencer* Sequencer = static_cast<ISequencer*>(&Player);
-					UMovieSceneSequence* RootSequence = Sequencer->GetRootMovieSceneSequence();
+					UMovieSceneSequence* RootSequence = SharedPlaybackState->GetRootSequence();
 					if ( !RootSequence )
 					{
 						return nullptr;
@@ -195,7 +194,7 @@ namespace UE
 					// Don't even have anything we can reuse: We need to spawn a brand new instance of this spawnable
 					if ( !Object )
 					{
-						Object = FLevelSequenceEditorSpawnRegister::SpawnObject( Spawnable, TemplateID, Player );
+						Object = FLevelSequenceEditorSpawnRegister::SpawnObject( Spawnable, TemplateID, SharedPlaybackState );
 						UE_LOG( LogUsd, VeryVerbose, TEXT( "Spawning '%s' (%0x) for RootSequence '%s', TemplateID '%u', Guid '%s'" ),
 							*Object->GetPathName(),
 							Object,
@@ -336,10 +335,10 @@ namespace UE
 					return SpawnableInstances.Contains(BindingGuid);
 				}
 
-				void DeleteSpawns( IMovieScenePlayer& Player )
+				void DeleteSpawns( TSharedRef<const FSharedPlaybackState> SharedPlaybackState )
 				{
 					bDestroyingJustHides = false;
-					CleanUp( Player );
+					CleanUp( SharedPlaybackState );
 
 					// If we still have existing spawns it may be because our base class' Register member didn't contain an
 					// entry for a spawnable before we called CleanUp (check its implementation: It just iterates over that Register).
@@ -460,7 +459,7 @@ namespace UE
 			public:
 				~FLevelSequenceExportContext()
 				{
-					SpawnRegister->DeleteSpawns( *Sequencer );
+					SpawnRegister->DeleteSpawns( Sequencer->GetSharedPlaybackState() );
 				}
 			};
 
@@ -1681,4 +1680,4 @@ bool ULevelSequenceExporterUsd::ExportBinary( UObject* Object, const TCHAR* Type
 #endif // #if USE_USD_SDK
 }
 
-#undef LOCTEXT_NAMESPACE
\ No newline at end of file
+#undef LOCTEXT_NAMESPACE
diff --git a/Engine/Plugins/MovieScene/LevelSequenceEditor/Source/LevelSequenceEditor/Private/Misc/LevelSequenceEditorSpawnRegister.cpp b/Engine/Plugins/MovieScene/LevelSequenceEditor/Source/LevelSequenceEditor/Private/Misc/LevelSequenceEditorSpawnRegister.cpp
--- a/Engine/Plugins/MovieScene/LevelSequenceEditor/Source/LevelSequenceEditor/Private/Misc/LevelSequenceEditorSpawnRegister.cpp
+++ b/Engine/Plugins/MovieScene/LevelSequenceEditor/Source/LevelSequenceEditor/Private/Misc/LevelSequenceEditorSpawnRegister.cpp
@@ -59,11 +59,11 @@ FLevelSequenceEditorSpawnRegister::~FLevelSequenceEditorSpawnRegister()
 /* FLevelSequenceSpawnRegister interface
  *****************************************************************************/
 
-UObject* FLevelSequenceEditorSpawnRegister::SpawnObject(FMovieSceneSpawnable& Spawnable, FMovieSceneSequenceIDRef TemplateID, IMovieScenePlayer& Player)
+UObject* FLevelSequenceEditorSpawnRegister::SpawnObject(FMovieSceneSpawnable& Spawnable, FMovieSceneSequenceIDRef TemplateID, TSharedRef<const UE::MovieScene::FSharedPlaybackState> SharedPlaybackState)
 {
 	TGuardValue<bool> Guard(bShouldClearSelectionCache, false);
 
-	UObject* NewObject = FLevelSequenceSpawnRegister::SpawnObject(Spawnable, TemplateID, Player);
+	UObject* NewObject = FLevelSequenceSpawnRegister::SpawnObject(Spawnable, TemplateID, SharedPlaybackState);
 	
 	if (AActor* NewActor = Cast<AActor>(NewObject))
 	{
@@ -105,7 +105,7 @@ void FLevelSequenceEditorSpawnRegister::PreDestroyObject(UObject& Object, const
 		if (TrackedState && TrackedState->bHasBeenModified)
 		{
 			// SaveDefaultSpawnableState will reset bHasBeenModified to false
-			SaveDefaultSpawnableStateImpl(*Spawnable, Sequence, SpawnedObject, *Sequencer);
+			SaveDefaultSpawnableStateImpl(*Spawnable, Sequence, SpawnedObject, Sequencer->GetSharedPlaybackState());
 
 			Sequence->MarkPackageDirty();
 		}
@@ -125,14 +125,14 @@ void FLevelSequenceEditorSpawnRegister::PreDestroyObject(UObject& Object, const
 	FLevelSequenceSpawnRegister::PreDestroyObject(Object, BindingId, TemplateID);
 }
 
-void FLevelSequenceEditorSpawnRegister::SaveDefaultSpawnableState(FMovieSceneSpawnable& Spawnable, FMovieSceneSequenceIDRef TemplateID, IMovieScenePlayer& Player)
+void FLevelSequenceEditorSpawnRegister::SaveDefaultSpawnableState(FMovieSceneSpawnable& Spawnable, FMovieSceneSequenceIDRef TemplateID, TSharedRef<const FSharedPlaybackState> SharedPlaybackState)
 {
-	UMovieSceneSequence* Sequence = Player.GetEvaluationTemplate().GetSequence(TemplateID);
+	UMovieSceneSequence* Sequence = SharedPlaybackState->GetSequence(TemplateID);
 
 	UObject* Object = FindSpawnedObject(Spawnable.GetGuid(), TemplateID).Get();
 	if (Object && Sequence)
 	{
-		SaveDefaultSpawnableStateImpl(Spawnable, Sequence, Object, Player);
+		SaveDefaultSpawnableStateImpl(Spawnable, Sequence, Object, SharedPlaybackState);
 		Sequence->MarkPackageDirty();
 	}
 }
@@ -164,14 +164,16 @@ void FLevelSequenceEditorSpawnRegister::SaveDefaultSpawnableState(const FGuid& B
 		UObject* Object = FindSpawnedObject(Spawnable->GetGuid(), TemplateID).Get();
 		if (Object)
 		{
-			SaveDefaultSpawnableStateImpl(*Spawnable, Sequence, Object, *Sequencer);
+			SaveDefaultSpawnableStateImpl(*Spawnable, Sequence, Object, Sequencer->GetSharedPlaybackState());
 			Sequence->MarkPackageDirty();
 		}
 	}
 }
 
-void FLevelSequenceEditorSpawnRegister::SaveDefaultSpawnableStateImpl(FMovieSceneSpawnable& Spawnable, UMovieSceneSequence* Sequence, UObject* SpawnedObject, IMovieScenePlayer& Player)
+void FLevelSequenceEditorSpawnRegister::SaveDefaultSpawnableStateImpl(FMovieSceneSpawnable& Spawnable, UMovieSceneSequence* Sequence, UObject* SpawnedObject, TSharedRef<const FSharedPlaybackState> SharedPlaybackState)
 {
+	IMovieScenePlayer* Player = UE::MovieScene::FPlayerIndexPlaybackCapability::GetPlayer(SharedPlaybackState);
+
 	FMovieSceneAnimTypeID SpawnablesTypeID = UMovieSceneSpawnablesSystem::GetAnimTypeID();
 	auto RestorePredicate = [SpawnablesTypeID](FMovieSceneAnimTypeID TypeID){ return TypeID != SpawnablesTypeID; };
 
@@ -182,13 +184,13 @@ void FLevelSequenceEditorSpawnRegister::SaveDefaultSpawnableStateImpl(FMovieScen
 		{
 			if (Component)
 			{
-				Player.PreAnimatedState.RestorePreAnimatedState(*Component, RestorePredicate);
+				Player->PreAnimatedState.RestorePreAnimatedState(*Component, RestorePredicate);
 			}
 		}
 	}
 
 	// Restore state on the object itself
-	Player.PreAnimatedState.RestorePreAnimatedState(*SpawnedObject, RestorePredicate);
+	Player->PreAnimatedState.RestorePreAnimatedState(*SpawnedObject, RestorePredicate);
 
 	// Copy the template
 	Spawnable.CopyObjectTemplate(*SpawnedObject, *Sequence);
@@ -310,7 +312,7 @@ void FLevelSequenceEditorSpawnRegister::OnPreObjectSaved(UObject* Object, FObjec
 
 				if (SpawnedObject && Spawnable && ThisSequence == SequenceBeingSaved)
 				{
-					SaveDefaultSpawnableStateImpl(*Spawnable, ThisSequence, SpawnedObject, *Sequencer);
+					SaveDefaultSpawnableStateImpl(*Spawnable, ThisSequence, SpawnedObject, Sequencer->GetSharedPlaybackState());
 				}
 			}
 		}
@@ -345,7 +347,7 @@ void FLevelSequenceEditorSpawnRegister::SetupDefaultsForSpawnable(UObject* Spawn
 	}
 }
 
-void FLevelSequenceEditorSpawnRegister::HandleConvertPossessableToSpawnable(UObject* OldObject, IMovieScenePlayer& Player, TOptional<FTransformData>& OutTransformData)
+void FLevelSequenceEditorSpawnRegister::HandleConvertPossessableToSpawnable(UObject* OldObject, TSharedRef<const FSharedPlaybackState> SharedPlaybackState, TOptional<FTransformData>& OutTransformData)
 {
 	// @TODO: this could probably be handed off to a spawner if we need anything else to be convertible between spawnable/posessable
 
@@ -361,7 +363,8 @@ void FLevelSequenceEditorSpawnRegister::HandleConvertPossessableToSpawnable(UObj
 		}
 
 		GEditor->SelectActor(OldActor, false, true);
-		UWorld* World = Cast<UWorld>(Player.GetPlaybackContext());
+		UObject* PlaybackContext = SharedPlaybackState->GetPlaybackContext()->GetWorld();
+		UWorld* World = PlaybackContext ? PlaybackContext->GetWorld() : nullptr;
 		if (World)
 		{
 			World->EditorDestroyActor(OldActor, true);
diff --git a/Engine/Plugins/MovieScene/LevelSequenceEditor/Source/LevelSequenceEditor/Public/Misc/LevelSequenceEditorSpawnRegister.h b/Engine/Plugins/MovieScene/LevelSequenceEditor/Source/LevelSequenceEditor/Public/Misc/LevelSequenceEditorSpawnRegister.h
--- a/Engine/Plugins/MovieScene/LevelSequenceEditor/Source/LevelSequenceEditor/Public/Misc/LevelSequenceEditorSpawnRegister.h
+++ b/Engine/Plugins/MovieScene/LevelSequenceEditor/Source/LevelSequenceEditor/Public/Misc/LevelSequenceEditorSpawnRegister.h
@@ -35,13 +35,13 @@ public:
 
 	// FLevelSequenceSpawnRegister interface
 
-	virtual UObject* SpawnObject(FMovieSceneSpawnable& Spawnable, FMovieSceneSequenceIDRef TemplateID, IMovieScenePlayer& Player) override;
+	virtual UObject* SpawnObject(FMovieSceneSpawnable& Spawnable, FMovieSceneSequenceIDRef TemplateID, TSharedRef<const UE::MovieScene::FSharedPlaybackState> SharedPlaybackState) override;
 	virtual void PreDestroyObject(UObject& Object, const FGuid& BindingId, FMovieSceneSequenceIDRef TemplateID) override;
-	virtual void SaveDefaultSpawnableState(FMovieSceneSpawnable& Spawnable, FMovieSceneSequenceIDRef TemplateID, IMovieScenePlayer& Player) override;
+	virtual void SaveDefaultSpawnableState(FMovieSceneSpawnable& Spawnable, FMovieSceneSequenceIDRef TemplateID, TSharedRef<const FSharedPlaybackState> SharedPlaybackState) override;
 #if WITH_EDITOR
 	virtual TValueOrError<FNewSpawnable, FText> CreateNewSpawnableType(UObject& SourceObject, UMovieScene& OwnerMovieScene, UActorFactory* ActorFactory = nullptr) override;
 	virtual void SetupDefaultsForSpawnable(UObject* SpawnedObject, const FGuid& Guid, const TOptional<FTransformData>& TransformData, TSharedRef<ISequencer> Sequencer, USequencerSettings* Settings) override;
-	virtual void HandleConvertPossessableToSpawnable(UObject* OldObject, IMovieScenePlayer& Player, TOptional<FTransformData>& OutTransformData) override;
+	virtual void HandleConvertPossessableToSpawnable(UObject* OldObject, TSharedRef<const FSharedPlaybackState> SharedPlaybackState, TOptional<FTransformData>& OutTransformData) override;
 	virtual bool CanConvertSpawnableToPossessable(FMovieSceneSpawnable& Spawnable) const override;
 #endif
 
@@ -52,7 +52,7 @@ private:
 
 	/** Saves the default state for the specified spawnable, if an instance for it currently exists */
 	void SaveDefaultSpawnableState(const FGuid& BindingId, FMovieSceneSequenceIDRef TemplateID);
-	void SaveDefaultSpawnableStateImpl(FMovieSceneSpawnable& Spawnable, UMovieSceneSequence* Sequence, UObject* SpawnedObject, IMovieScenePlayer& Player);
+	void SaveDefaultSpawnableStateImpl(FMovieSceneSpawnable& Spawnable, UMovieSceneSequence* Sequence, UObject* SpawnedObject, TSharedRef<const FSharedPlaybackState> SharedPlaybackState);
 
 	/** Called from the editor when a blueprint object replacement has occurred */
 	void OnObjectsReplaced(const TMap<UObject*, UObject*>& OldToNewInstanceMap);
diff --git a/Engine/Plugins/MovieScene/TemplateSequence/Source/TemplateSequence/Private/Systems/TemplateSequenceSystem.cpp b/Engine/Plugins/MovieScene/TemplateSequence/Source/TemplateSequence/Private/Systems/TemplateSequenceSystem.cpp
--- a/Engine/Plugins/MovieScene/TemplateSequence/Source/TemplateSequence/Private/Systems/TemplateSequenceSystem.cpp
+++ b/Engine/Plugins/MovieScene/TemplateSequence/Source/TemplateSequence/Private/Systems/TemplateSequenceSystem.cpp
@@ -75,18 +75,18 @@ void UTemplateSequenceSystem::OnRun(FSystemTaskPrerequisites& InPrerequisites, F
 			const FGuid& ObjectBindingID = ObjectBindingIDs[Index];
 			const FTemplateSequenceComponentData& TemplateSequenceData = TemplateSequenceDatas[Index];
 
-			IMovieScenePlayer* Player = SequenceInstance.GetPlayer();
-			if (ensure(Player))
+			IStaticBindingOverridesPlaybackCapability* StaticOverrides = SequenceInstance.GetSharedPlaybackState()->FindCapability<IStaticBindingOverridesPlaybackCapability>();
+			if (ensure(StaticOverrides))
 			{
 				if (bHasNeedsLink)
 				{
 					const FMovieSceneSequenceID SequenceID = SequenceInstance.GetSequenceID();
 					const FMovieSceneEvaluationOperand OuterOperand(SequenceID, ObjectBindingID);
-					Player->BindingOverrides.Add(TemplateSequenceData.InnerOperand, OuterOperand);
+					StaticOverrides->AddBindingOverride(TemplateSequenceData.InnerOperand, OuterOperand);
 				}
 				else if (bHasNeedsUnlink)
 				{
-					Player->BindingOverrides.Remove(TemplateSequenceData.InnerOperand);
+					StaticOverrides->RemoveBindingOverride(TemplateSequenceData.InnerOperand);
 				}
 			}
 		}
diff --git a/Engine/Plugins/MovieScene/TemplateSequence/Source/TemplateSequence/Public/TemplateSequencePlayer.h b/Engine/Plugins/MovieScene/TemplateSequence/Source/TemplateSequence/Public/TemplateSequencePlayer.h
--- a/Engine/Plugins/MovieScene/TemplateSequence/Source/TemplateSequence/Public/TemplateSequencePlayer.h
+++ b/Engine/Plugins/MovieScene/TemplateSequence/Source/TemplateSequence/Public/TemplateSequencePlayer.h
@@ -45,7 +45,7 @@ class FSequenceCameraShakeSpawnRegister : public FMovieSceneSpawnRegister
 public:
 	void SetSpawnedObject(UObject* InObject) { SpawnedObject = InObject; }
 
-	virtual UObject* SpawnObject(FMovieSceneSpawnable&, FMovieSceneSequenceIDRef, IMovieScenePlayer&) override { return SpawnedObject.Get(); }
+	virtual UObject* SpawnObject(FMovieSceneSpawnable&, FMovieSceneSequenceIDRef, TSharedRef<const UE::MovieScene::FSharedPlaybackState>) override { return SpawnedObject.Get(); }
 	virtual void DestroySpawnedObject(UObject&) override {}
 
 #if WITH_EDITOR
diff --git a/Engine/Source/Editor/Sequencer/Private/Sequencer.cpp b/Engine/Source/Editor/Sequencer/Private/Sequencer.cpp
--- a/Engine/Source/Editor/Sequencer/Private/Sequencer.cpp
+++ b/Engine/Source/Editor/Sequencer/Private/Sequencer.cpp
@@ -661,13 +661,13 @@ void FSequencer::InitRootSequenceInstance()
 	using namespace UE::MovieScene;
 
 	// Add the camera cut playback capability.
-	FInstanceHandle RootInstanceHandle = RootTemplateInstance.GetRootInstanceHandle();
-	UMovieSceneEntitySystemLinker* Linker = GetEvaluationTemplate().GetEntitySystemLinker();
-	FInstanceRegistry* InstanceRegistry = Linker->GetInstanceRegistry();
-	TSharedRef<FSharedPlaybackState> SharedPlaybackState = InstanceRegistry->GetInstance(RootInstanceHandle).GetSharedPlaybackState();
-	if (!SharedPlaybackState->HasCapability<FCameraCutPlaybackCapability>())
+	TSharedPtr<FSharedPlaybackState> SharedPlaybackState = RootTemplateInstance.GetSharedPlaybackState();
+	if (ensure(SharedPlaybackState.IsValid()))
 	{
-		SharedPlaybackState->AddCapabilityRaw<FCameraCutPlaybackCapability>((FCameraCutPlaybackCapability*)this);
+		if (!SharedPlaybackState->HasCapability<FCameraCutPlaybackCapability>())
+		{
+			SharedPlaybackState->AddCapabilityRaw<FCameraCutPlaybackCapability>((FCameraCutPlaybackCapability*)this);
+		}
 	}
 }
 
diff --git a/Engine/Source/Editor/Sequencer/Private/SequencerUtilities.cpp b/Engine/Source/Editor/Sequencer/Private/SequencerUtilities.cpp
--- a/Engine/Source/Editor/Sequencer/Private/SequencerUtilities.cpp
+++ b/Engine/Source/Editor/Sequencer/Private/SequencerUtilities.cpp
@@ -595,7 +595,7 @@ TArray<FGuid> ExpandMultiplePossessableBindings(TSharedRef<ISequencer> Sequencer
 				NewPossessable->SetParent(ParentGuid, MovieScene);
 			}
 
-			if (!NewPossessable->BindSpawnableObject(Sequencer->GetFocusedTemplateID(), FoundObject, &Sequencer.Get()))
+			if (!NewPossessable->BindSpawnableObject(Sequencer->GetFocusedTemplateID(), FoundObject, Sequencer->GetSharedPlaybackState()))
 			{
 				Sequence->BindPossessableObject(NewPossessableGuid, *FoundObject, BindingContext);
 			}
@@ -2534,7 +2534,7 @@ FGuid FSequencerUtilities::CreateBinding(TSharedRef<ISequencer> Sequencer, UObje
 		}
 	}
 
-	if (!OwnerMovieScene->FindPossessable(PossessableGuid)->BindSpawnableObject(Sequencer->GetFocusedTemplateID(), &InObject, &Sequencer.Get()))
+	if (!OwnerMovieScene->FindPossessable(PossessableGuid)->BindSpawnableObject(Sequencer->GetFocusedTemplateID(), &InObject, Sequencer->GetSharedPlaybackState()))
 	{
 		OwnerSequence->BindPossessableObject(PossessableGuid, InObject, BindingContext);
 	}
@@ -2644,7 +2644,7 @@ FGuid FSequencerUtilities::AssignActor(TSharedRef<ISequencer> Sequencer, AActor*
 		// Add this object
 		NewPossessableActor = FMovieScenePossessable(NewActorLabel, Actor->GetClass());
 		NewGuid = NewPossessableActor.GetGuid();
-		if (!NewPossessableActor.BindSpawnableObject(Sequencer->GetFocusedTemplateID(), Actor, &Sequencer.Get()))
+		if (!NewPossessableActor.BindSpawnableObject(Sequencer->GetFocusedTemplateID(), Actor, Sequencer->GetSharedPlaybackState()))
 		{
 			OwnerSequence->BindPossessableObject(NewPossessableActor.GetGuid(), *Actor, Sequencer->GetPlaybackContext());
 		}
@@ -2831,7 +2831,7 @@ void FSequencerUtilities::AddActorsToBinding(TSharedRef<ISequencer> Sequencer, c
 				}
 
 				ActorToAdd->Modify();
-				if (!MovieScene->FindPossessable(Guid)->BindSpawnableObject(Sequencer->GetFocusedTemplateID(), ActorToAdd, &Sequencer.Get()))
+				if (!MovieScene->FindPossessable(Guid)->BindSpawnableObject(Sequencer->GetFocusedTemplateID(), ActorToAdd, Sequencer->GetSharedPlaybackState()))
 				{
 					Sequence->BindPossessableObject(Guid, *ActorToAdd, Sequencer->GetPlaybackContext());
 				}
diff --git a/Engine/Source/Runtime/LevelSequence/Private/LevelSequenceActor.cpp b/Engine/Source/Runtime/LevelSequence/Private/LevelSequenceActor.cpp
--- a/Engine/Source/Runtime/LevelSequence/Private/LevelSequenceActor.cpp
+++ b/Engine/Source/Runtime/LevelSequence/Private/LevelSequenceActor.cpp
@@ -376,9 +376,9 @@ void ALevelSequenceActor::SetBindingByTag(FName BindingTag, const TArray<AActor*
 	const FMovieSceneObjectBindingIDs* Bindings = Sequence ? Sequence->GetMovieScene()->AllTaggedBindings().Find(BindingTag) : nullptr;
 	if (Bindings)
 	{
-		for (FMovieSceneObjectBindingID ID : Bindings->IDs)
+		for (FMovieSceneObjectBindingID BindingID : Bindings->IDs)
 		{
-			SetBinding(ID, Actors, bAllowBindingsFromAsset);
+			SetBinding(BindingID, Actors, bAllowBindingsFromAsset);
 		}
 	}
 	else if (GLevelSequenceActor_InvalidBindingTagWarnings)
@@ -414,9 +414,9 @@ void ALevelSequenceActor::AddBindingByTag(FName BindingTag, AActor* Actor, bool
 	const FMovieSceneObjectBindingIDs* Bindings = Sequence ? Sequence->GetMovieScene()->AllTaggedBindings().Find(BindingTag) : nullptr;
 	if (Bindings)
 	{
-		for (FMovieSceneObjectBindingID ID : Bindings->IDs)
+		for (FMovieSceneObjectBindingID BindingID : Bindings->IDs)
 		{
-			AddBinding(ID, Actor, bAllowBindingsFromAsset);
+			AddBinding(BindingID, Actor, bAllowBindingsFromAsset);
 		}
 	}
 	else if (GLevelSequenceActor_InvalidBindingTagWarnings)
@@ -452,9 +452,9 @@ void ALevelSequenceActor::RemoveBindingByTag(FName BindingTag, AActor* Actor)
 	const FMovieSceneObjectBindingIDs* Bindings = Sequence ? Sequence->GetMovieScene()->AllTaggedBindings().Find(BindingTag) : nullptr;
 	if (Bindings)
 	{
-		for (FMovieSceneObjectBindingID ID : Bindings->IDs)
+		for (FMovieSceneObjectBindingID BindingID : Bindings->IDs)
 		{
-			RemoveBinding(ID, Actor);
+			RemoveBinding(BindingID, Actor);
 		}
 	}
 	else if (GLevelSequenceActor_InvalidBindingTagWarnings)
diff --git a/Engine/Source/Runtime/LevelSequence/Private/LevelSequenceActorSpawner.cpp b/Engine/Source/Runtime/LevelSequence/Private/LevelSequenceActorSpawner.cpp
--- a/Engine/Source/Runtime/LevelSequence/Private/LevelSequenceActorSpawner.cpp
+++ b/Engine/Source/Runtime/LevelSequence/Private/LevelSequenceActorSpawner.cpp
@@ -3,7 +3,6 @@
 #include "LevelSequenceActorSpawner.h"
 #include "Misc/PackageName.h"
 #include "MovieSceneSpawnable.h"
-#include "IMovieScenePlayer.h"
 #include "GameFramework/Actor.h"
 #include "Engine/World.h"
 #include "Particles/ParticleSystemComponent.h"
@@ -69,7 +68,7 @@ ULevelStreaming* GetLevelStreaming(const FName& DesiredLevelName, const UWorld*
 	return nullptr;
 }
 
-UObject* FLevelSequenceActorSpawner::SpawnObject(FMovieSceneSpawnable& Spawnable, FMovieSceneSequenceIDRef TemplateID, IMovieScenePlayer& Player)
+UObject* FLevelSequenceActorSpawner::SpawnObject(FMovieSceneSpawnable& Spawnable, FMovieSceneSequenceIDRef TemplateID, TSharedRef<const UE::MovieScene::FSharedPlaybackState> SharedPlaybackState)
 {
 	AActor* ObjectTemplate = Cast<AActor>(Spawnable.GetObjectTemplate());
 	if (!ObjectTemplate)
@@ -88,7 +87,8 @@ UObject* FLevelSequenceActorSpawner::SpawnObject(FMovieSceneSpawnable& Spawnable
 	
 	// @todo sequencer actors: We need to make sure puppet objects aren't copied into PIE/SIE sessions!  They should be omitted from that duplication!
 
-	UWorld* WorldContext = Cast<UWorld>(Player.GetPlaybackContext());
+	UObject* PlaybackContext = SharedPlaybackState->GetPlaybackContext();
+	UWorld* WorldContext = PlaybackContext ? PlaybackContext->GetWorld() : nullptr;
 
 	FName DesiredLevelName = Spawnable.GetLevelName();
 	if (DesiredLevelName != NAME_None)
@@ -121,7 +121,7 @@ UObject* FLevelSequenceActorSpawner::SpawnObject(FMovieSceneSpawnable& Spawnable
 	// We use the net addressable name for spawnables on any non-editor, non-standalone world (ie, all clients, servers and PIE worlds)
 	const bool bUseNetAddressableName = Spawnable.bNetAddressableName && (WorldContext->WorldType != EWorldType::Editor) && (WorldContext->GetNetMode() != ENetMode::NM_Standalone);
 
-	FName SpawnName = bUseNetAddressableName ? Spawnable.GetNetAddressableName(Player, TemplateID) :
+	FName SpawnName = bUseNetAddressableName ? Spawnable.GetNetAddressableName(SharedPlaybackState, TemplateID) :
 #if WITH_EDITOR
 		// Construct the object with the same name that we will set later on the actor to avoid renaming it inside SetActorLabel
 		MakeUniqueObjectName(WorldContext->PersistentLevel, ObjectTemplate->GetClass(), *Spawnable.GetName());
@@ -225,7 +225,7 @@ UObject* FLevelSequenceActorSpawner::SpawnObject(FMovieSceneSpawnable& Spawnable
 	}	
 #endif
 
-	if (UMovieSceneEntitySystemLinker* Linker = Player.GetEvaluationTemplate().GetEntitySystemLinker())
+	if (UMovieSceneEntitySystemLinker* Linker = SharedPlaybackState->GetLinker())
 	{
 		if (UMovieSceneDeferredComponentMovementSystem* DeferredMovementSystem = Linker->FindSystem<UMovieSceneDeferredComponentMovementSystem>())
 		{
diff --git a/Engine/Source/Runtime/LevelSequence/Private/LevelSequencePlayer.cpp b/Engine/Source/Runtime/LevelSequence/Private/LevelSequencePlayer.cpp
--- a/Engine/Source/Runtime/LevelSequence/Private/LevelSequencePlayer.cpp
+++ b/Engine/Source/Runtime/LevelSequence/Private/LevelSequencePlayer.cpp
@@ -135,14 +135,19 @@ void ULevelSequencePlayer::Initialize(ULevelSequence* InLevelSequence, ULevel* I
 
 	UMovieSceneSequencePlayer::Initialize(InLevelSequence);
 
-	// Add the camera cut playback capability.
-	FInstanceHandle RootInstanceHandle = RootTemplateInstance.GetRootInstanceHandle();
-	UMovieSceneEntitySystemLinker* Linker = GetEvaluationTemplate().GetEntitySystemLinker();
-	FInstanceRegistry* InstanceRegistry = Linker->GetInstanceRegistry();
-	TSharedRef<FSharedPlaybackState> SharedPlaybackState = InstanceRegistry->GetInstance(RootInstanceHandle).GetSharedPlaybackState();
-	if (!SharedPlaybackState->HasCapability<FCameraCutPlaybackCapability>())
+	TSharedPtr<FSharedPlaybackState> SharedPlaybackState = RootTemplateInstance.GetSharedPlaybackState();
+	if (SharedPlaybackState)
 	{
-		SharedPlaybackState->AddCapabilityRaw<FCameraCutPlaybackCapability>((FCameraCutPlaybackCapability*)this);
+		// The parent player class' root evaluation template may or may not have re-initialized itself.
+		// For instance, if we are given the same sequence asset we already had before, and nothing else
+		// (such as playback context) has changed, no actual re-initialization occurs and we keep the
+		// same shared playback state as before.
+		// That state would already have the spawn register and camera cut capabilies... however, our 
+		// spawn register was just re-created (see a few lines above) so we need to overwrite the 
+		// capability pointer to the new object.
+		// The camera cut capability will stay as a pointer to ourselves, and that's fine.
+		SharedPlaybackState->SetOrAddCapabilityRaw<FMovieSceneSpawnRegister>(SpawnRegister.Get());
+		SharedPlaybackState->SetOrAddCapabilityRaw<FCameraCutPlaybackCapability>((FCameraCutPlaybackCapability*)this);
 	}
 }
 
diff --git a/Engine/Source/Runtime/LevelSequence/Private/LevelSequenceSpawnRegister.cpp b/Engine/Source/Runtime/LevelSequence/Private/LevelSequenceSpawnRegister.cpp
--- a/Engine/Source/Runtime/LevelSequence/Private/LevelSequenceSpawnRegister.cpp
+++ b/Engine/Source/Runtime/LevelSequence/Private/LevelSequenceSpawnRegister.cpp
@@ -14,14 +14,13 @@ FLevelSequenceSpawnRegister::FLevelSequenceSpawnRegister()
 	LevelSequenceModule.GenerateObjectSpawners(MovieSceneObjectSpawners);
 }
 
-UObject* FLevelSequenceSpawnRegister::SpawnObject(FMovieSceneSpawnable& Spawnable, FMovieSceneSequenceIDRef TemplateID, IMovieScenePlayer& Player)
+UObject* FLevelSequenceSpawnRegister::SpawnObject(FMovieSceneSpawnable& Spawnable, FMovieSceneSequenceIDRef TemplateID, TSharedRef<const FSharedPlaybackState> SharedPlaybackState)
 {
 	for (TSharedRef<IMovieSceneObjectSpawner> MovieSceneObjectSpawner : MovieSceneObjectSpawners)
 	{
 		if (Spawnable.GetObjectTemplate() != nullptr && Spawnable.GetObjectTemplate()->IsA(MovieSceneObjectSpawner->GetSupportedTemplateType()))
 		{
-			
-			UObject* SpawnedObject = MovieSceneObjectSpawner->SpawnObject(Spawnable, TemplateID, Player);
+			UObject* SpawnedObject = MovieSceneObjectSpawner->SpawnObject(Spawnable, TemplateID, SharedPlaybackState);
 			if (SpawnedObject)
 			{
 				return SpawnedObject;
@@ -60,4 +59,4 @@ bool FLevelSequenceSpawnRegister::CanSpawnObject(UClass* InClass) const
 	return false;
 }
 
-#endif
\ No newline at end of file
+#endif
diff --git a/Engine/Source/Runtime/LevelSequence/Public/LevelSequenceActorSpawner.h b/Engine/Source/Runtime/LevelSequence/Public/LevelSequenceActorSpawner.h
--- a/Engine/Source/Runtime/LevelSequence/Public/LevelSequenceActorSpawner.h
+++ b/Engine/Source/Runtime/LevelSequence/Public/LevelSequenceActorSpawner.h
@@ -12,6 +12,6 @@ public:
 
 	// IMovieSceneObjectSpawner interface
 	LEVELSEQUENCE_API virtual UClass* GetSupportedTemplateType() const override;
-	LEVELSEQUENCE_API virtual UObject* SpawnObject(FMovieSceneSpawnable& Spawnable, FMovieSceneSequenceIDRef TemplateID, IMovieScenePlayer& Player) override;
+	LEVELSEQUENCE_API virtual UObject* SpawnObject(FMovieSceneSpawnable& Spawnable, FMovieSceneSequenceIDRef TemplateID, TSharedRef<const UE::MovieScene::FSharedPlaybackState> SharedPlaybackState) override;
 	LEVELSEQUENCE_API virtual void DestroySpawnedObject(UObject& Object) override;
 };
diff --git a/Engine/Source/Runtime/LevelSequence/Public/LevelSequenceSpawnRegister.h b/Engine/Source/Runtime/LevelSequence/Public/LevelSequenceSpawnRegister.h
--- a/Engine/Source/Runtime/LevelSequence/Public/LevelSequenceSpawnRegister.h
+++ b/Engine/Source/Runtime/LevelSequence/Public/LevelSequenceSpawnRegister.h
@@ -23,7 +23,7 @@ public:
 
 protected:
 	/** ~ FMovieSceneSpawnRegister interface */
-	LEVELSEQUENCE_API virtual UObject* SpawnObject(FMovieSceneSpawnable& Spawnable, FMovieSceneSequenceIDRef TemplateID, IMovieScenePlayer& Player) override;
+	LEVELSEQUENCE_API virtual UObject* SpawnObject(FMovieSceneSpawnable& Spawnable, FMovieSceneSequenceIDRef TemplateID, TSharedRef<const FSharedPlaybackState> SharedPlaybackState) override;
 	LEVELSEQUENCE_API virtual void DestroySpawnedObject(UObject& Object) override;
 
 #if WITH_EDITOR
diff --git a/Engine/Source/Runtime/MovieScene/Private/Compilation/MovieSceneCompiledVolatilityManager.cpp b/Engine/Source/Runtime/MovieScene/Private/Compilation/MovieSceneCompiledVolatilityManager.cpp
--- a/Engine/Source/Runtime/MovieScene/Private/Compilation/MovieSceneCompiledVolatilityManager.cpp
+++ b/Engine/Source/Runtime/MovieScene/Private/Compilation/MovieSceneCompiledVolatilityManager.cpp
@@ -161,7 +161,7 @@ void FCompiledDataVolatilityManager::UpdateCachedSignatures()
 		UMovieSceneSequence* RootSequence = RootEntry.GetSequence();
 		if (RootSequence && Player)
 		{
-			Player->State.AssignSequence(MovieSceneSequenceID::Root, *RootSequence, *Player);
+			Player->State.AssignSequence(MovieSceneSequenceID::Root, *RootSequence, Player->GetSharedPlaybackState());
 		}
 	}
 
@@ -183,7 +183,7 @@ void FCompiledDataVolatilityManager::UpdateCachedSignatures()
 			UMovieSceneSequence* Sequence = SubData.Value.GetSequence();
 			if (Sequence && Player)
 			{
-				Player->State.AssignSequence(SubData.Key, *Sequence, *Player);
+				Player->State.AssignSequence(SubData.Key, *Sequence, Player->GetSharedPlaybackState());
 			}
 		}
 	}
diff --git a/Engine/Source/Runtime/MovieScene/Private/EntitySystem/MovieSceneEntityInstantiatorSystem.cpp b/Engine/Source/Runtime/MovieScene/Private/EntitySystem/MovieSceneEntityInstantiatorSystem.cpp
--- a/Engine/Source/Runtime/MovieScene/Private/EntitySystem/MovieSceneEntityInstantiatorSystem.cpp
+++ b/Engine/Source/Runtime/MovieScene/Private/EntitySystem/MovieSceneEntityInstantiatorSystem.cpp
@@ -84,7 +84,7 @@ void UMovieSceneEntityInstantiatorSystem::UnlinkStaleObjectBindings(UE::MovieSce
 		const FSequenceInstance* TargetInstance = &InstanceRegistry->GetInstance(InstanceHandle);
 
 		FMovieSceneSequenceID ThisSequenceID     = TargetInstance->GetSequenceID();
-		FMovieSceneSequenceID RemappedSequenceID = InObjectBindingID.ResolveSequenceID(ThisSequenceID, *TargetInstance->GetPlayer());
+		FMovieSceneSequenceID RemappedSequenceID = InObjectBindingID.ResolveSequenceID(ThisSequenceID, TargetInstance->GetSharedPlaybackState());
 
 		if (RemappedSequenceID == ThisSequenceID)
 		{
diff --git a/Engine/Source/Runtime/MovieScene/Private/EntitySystem/MovieSceneEntitySystemRunner.cpp b/Engine/Source/Runtime/MovieScene/Private/EntitySystem/MovieSceneEntitySystemRunner.cpp
--- a/Engine/Source/Runtime/MovieScene/Private/EntitySystem/MovieSceneEntitySystemRunner.cpp
+++ b/Engine/Source/Runtime/MovieScene/Private/EntitySystem/MovieSceneEntitySystemRunner.cpp
@@ -835,7 +835,11 @@ UE::MovieScene::ERunnerFlushResult FMovieSceneEntitySystemRunner::GameThread_Upd
 			FSequenceInstance& SequenceInstance = InstanceRegistry->MutateInstance(UpdatedInstance.InstanceHandle);
 
 			AccumulatedUpdateFlags |= SequenceInstance.GetUpdateFlags();
-			IMovieScenePlayer::SetIsEvaluatingFlag(SequenceInstance.GetPlayerIndex(), true);
+			uint16 PlayerIndex = FPlayerIndexPlaybackCapability::GetPlayerIndex(SequenceInstance.GetSharedPlaybackState());
+			if (PlayerIndex != (uint16)-1)
+			{
+				IMovieScenePlayer::SetIsEvaluatingFlag(PlayerIndex, true);
+			}
 
 			if (EnumHasAnyFlags(SequenceInstance.GetUpdateFlags(), ESequenceInstanceUpdateFlags::NeedsPreEvaluation))
 			{
@@ -1243,7 +1247,11 @@ void FMovieSceneEntitySystemRunner::GameThread_PostEvaluationPhase(UMovieSceneEn
 				FSequenceInstance& Instance = InstanceRegistry->MutateInstance(UpdateParams.InstanceHandle);
 
 				Instance.Ledger.UnlinkOneShots(Linker);
-				IMovieScenePlayer::SetIsEvaluatingFlag(Instance.GetPlayerIndex(), false);
+				uint16 PlayerIndex = FPlayerIndexPlaybackCapability::GetPlayerIndex(Instance.GetSharedPlaybackState());
+				if (PlayerIndex != (uint16)-1)
+				{
+					IMovieScenePlayer::SetIsEvaluatingFlag(PlayerIndex, false);
+				}
 
 				if (EnumHasAnyFlags(Instance.GetUpdateFlags(), ESequenceInstanceUpdateFlags::NeedsPostEvaluation))
 				{
diff --git a/Engine/Source/Runtime/MovieScene/Private/EntitySystem/MovieSceneInstanceRegistry.cpp b/Engine/Source/Runtime/MovieScene/Private/EntitySystem/MovieSceneInstanceRegistry.cpp
--- a/Engine/Source/Runtime/MovieScene/Private/EntitySystem/MovieSceneInstanceRegistry.cpp
+++ b/Engine/Source/Runtime/MovieScene/Private/EntitySystem/MovieSceneInstanceRegistry.cpp
@@ -57,7 +57,11 @@ FInstanceHandle FInstanceRegistry::FindRelatedInstanceHandle(FInstanceHandle Ins
 	return RootInstance->FindSubInstance(SequenceID);
 }
 
-FRootInstanceHandle FInstanceRegistry::AllocateRootInstance(IMovieScenePlayer* Player, UMovieSceneSequence& RootSequence, TSharedPtr<FMovieSceneEntitySystemRunner> Runner, UMovieSceneCompiledDataManager* CompiledDataManager)
+FRootInstanceHandle FInstanceRegistry::AllocateRootInstance(
+		UMovieSceneSequence& InRootSequence,
+		UObject* InPlaybackContext,
+		TSharedPtr<FMovieSceneEntitySystemRunner> InRunner,
+		UMovieSceneCompiledDataManager* InCompiledDataManager)
 {
 	check(Instances.Num() < 65535);
 
@@ -66,21 +70,30 @@ FRootInstanceHandle FInstanceRegistry::AllocateRootInstance(IMovieScenePlayer* P
 	FSparseArrayAllocationInfo NewAllocation = Instances.AddUninitialized();
 	FRootInstanceHandle InstanceHandle { (uint16)NewAllocation.Index, InstanceSerial };
 
+	if (!InRunner)
+	{
+		FMovieSceneEntitySystemRunner* ActiveRunner = Linker->GetActiveRunner();
+		InRunner = ActiveRunner ? ActiveRunner->AsShared() : TSharedPtr<FMovieSceneEntitySystemRunner>();
+	}
+	if (!InCompiledDataManager)
+	{
+		InCompiledDataManager = UMovieSceneCompiledDataManager::GetPrecompiledData();
+	}
+
 	FSharedPlaybackStateCreateParams PlaybackStateCreateParams;
-	PlaybackStateCreateParams.PlaybackContext = Player->GetPlaybackContext();
+	PlaybackStateCreateParams.PlaybackContext = InPlaybackContext;
 	PlaybackStateCreateParams.RootInstanceHandle = InstanceHandle;
-	PlaybackStateCreateParams.Runner = Runner;
-	PlaybackStateCreateParams.CompiledDataManager = CompiledDataManager;
-	TSharedRef<FSharedPlaybackState> NewPlaybackState = MakeShared<FSharedPlaybackState>(RootSequence, PlaybackStateCreateParams);
+	PlaybackStateCreateParams.Runner = InRunner;
+	PlaybackStateCreateParams.CompiledDataManager = InCompiledDataManager;
 
-	FSequenceInstance* NewInstance = new (NewAllocation) FSequenceInstance(NewPlaybackState, InstanceHandle);
+	TSharedRef<FSharedPlaybackState> NewPlaybackState = MakeShared<FSharedPlaybackState>(InRootSequence, PlaybackStateCreateParams);
 
-	NewInstance->Initialize(Player);
+	new (NewAllocation) FSequenceInstance(NewPlaybackState);
 
 	return InstanceHandle;
 }
 
-FInstanceHandle FInstanceRegistry::AllocateSubInstance(IMovieScenePlayer* Player, FMovieSceneSequenceID SequenceID, FRootInstanceHandle RootInstanceHandle, FInstanceHandle ParentInstanceHandle)
+FInstanceHandle FInstanceRegistry::AllocateSubInstance(FMovieSceneSequenceID SequenceID, FRootInstanceHandle RootInstanceHandle, FInstanceHandle ParentInstanceHandle)
 {
 	check(Instances.Num() < 65535 && SequenceID != MovieSceneSequenceID::Root && ParentInstanceHandle.IsValid());
 
@@ -90,9 +103,9 @@ FInstanceHandle FInstanceRegistry::AllocateSubInstance(IMovieScenePlayer* Player
 	
 	TSharedRef<FSharedPlaybackState> PlaybackState = GetInstance(RootInstanceHandle).GetSharedPlaybackState();
 
-	FSequenceInstance* NewInstance = new (NewAllocation) FSequenceInstance(PlaybackState, InstanceHandle, ParentInstanceHandle, RootInstanceHandle, SequenceID);
+	new (NewAllocation) FSequenceInstance(PlaybackState, InstanceHandle, ParentInstanceHandle, SequenceID);
 
-	NewInstance->Initialize(Player);
+	PlaybackState->GetCapabilities().OnSubInstanceCreated(PlaybackState, InstanceHandle);
 
 	return InstanceHandle;
 }
diff --git a/Engine/Source/Runtime/MovieScene/Private/EntitySystem/MovieSceneSequenceInstance.cpp b/Engine/Source/Runtime/MovieScene/Private/EntitySystem/MovieSceneSequenceInstance.cpp
--- a/Engine/Source/Runtime/MovieScene/Private/EntitySystem/MovieSceneSequenceInstance.cpp
+++ b/Engine/Source/Runtime/MovieScene/Private/EntitySystem/MovieSceneSequenceInstance.cpp
@@ -66,12 +66,15 @@ void PurgeStaleTrackTemplates(UMovieSceneCompiledDataManager* CompiledDataManage
 }
 
 
-FSequenceInstance::FSequenceInstance(TSharedRef<FSharedPlaybackState> PlaybackState, FRootInstanceHandle InInstanceHandle)
+
+
+FSequenceInstance::FSequenceInstance(TSharedRef<FSharedPlaybackState> PlaybackState)
 	: SharedPlaybackState(PlaybackState)
 	, SequenceID(MovieSceneSequenceID::Root)
 	, RootOverrideSequenceID(MovieSceneSequenceID::Root)
-	, InstanceHandle(InInstanceHandle)
-	, RootInstanceHandle(InInstanceHandle)
+	, InstanceHandle(PlaybackState->GetRootInstanceHandle())
+	, RootInstanceHandle(PlaybackState->GetRootInstanceHandle())
+	, bInitialized(false)
 {
 	UpdateFlags = ESequenceInstanceUpdateFlags::None;
 
@@ -87,13 +90,14 @@ FSequenceInstance::FSequenceInstance(TSharedRef<FSharedPlaybackState> PlaybackSt
 #endif
 }
 
-FSequenceInstance::FSequenceInstance(TSharedRef<FSharedPlaybackState> PlaybackState, FInstanceHandle InInstanceHandle, FInstanceHandle InParentInstanceHandle, FRootInstanceHandle InRootInstanceHandle, FMovieSceneSequenceID InSequenceID)
+FSequenceInstance::FSequenceInstance(TSharedRef<FSharedPlaybackState> PlaybackState, FInstanceHandle InInstanceHandle, FInstanceHandle InParentInstanceHandle, FMovieSceneSequenceID InSequenceID)
 	: SharedPlaybackState(PlaybackState)
 	, SequenceID(InSequenceID)
 	, RootOverrideSequenceID(MovieSceneSequenceID::Invalid)
 	, InstanceHandle(InInstanceHandle)
 	, ParentInstanceHandle(InParentInstanceHandle)
-	, RootInstanceHandle(InRootInstanceHandle)
+	, RootInstanceHandle(PlaybackState->GetRootInstanceHandle())
+	, bInitialized(false)
 {
 	UpdateFlags = ESequenceInstanceUpdateFlags::None;
 
@@ -103,31 +107,17 @@ FSequenceInstance::FSequenceInstance(TSharedRef<FSharedPlaybackState> PlaybackSt
 	bHasEverUpdated = false;
 }
 
-void FSequenceInstance::Initialize(IMovieScenePlayer* Player)
+void FSequenceInstance::Initialize()
 {
-	if (!Player)
-	{
-		return;
-	}
-
-	// Initialize playback capabilities if this is the root sequence.
-	if (IsRootSequence())
-	{
-		SharedPlaybackState->AddCapability<FPlayerIndexPlaybackCapability>(Player->GetUniqueIndex());
-	}
-
-	UMovieSceneEntitySystemLinker* Linker = SharedPlaybackState->GetLinker();
-	if (ensure(Linker))
-	{
-		FMovieSceneObjectCache& ObjectCache = Player->State.GetObjectCache(SequenceID);
-		OnInvalidateObjectBindingHandle = ObjectCache.OnBindingInvalidated.AddUObject(Linker, &UMovieSceneEntitySystemLinker::InvalidateObjectBinding, InstanceHandle);
-	}
+	ensureMsgf(!bInitialized, TEXT("This instance was already initialized!"));
+	bInitialized = true;
 
 	InvalidateCachedData();
 }
 
 FSequenceInstance::~FSequenceInstance()
-{}
+{
+}
 
 FSequenceInstance::FSequenceInstance(FSequenceInstance&&) = default;
 
@@ -135,23 +125,22 @@ FSequenceInstance& FSequenceInstance::operator=(FSequenceInstance&&) = default;
 
 IMovieScenePlayer* FSequenceInstance::GetPlayer() const
 {
-	using namespace UE::MovieScene;
-
 	return FPlayerIndexPlaybackCapability::GetPlayer(SharedPlaybackState);
 }
 
 uint16 FSequenceInstance::GetPlayerIndex() const
 {
-	using namespace UE::MovieScene;
-
 	return FPlayerIndexPlaybackCapability::GetPlayerIndex(SharedPlaybackState);
 }
 
 void FSequenceInstance::InitializeLegacyEvaluator()
 {
-	const FMovieSceneCompiledDataID RootCompiledDataID = SharedPlaybackState->GetRootCompiledDataID();
+	IMovieScenePlayer* Player = GetPlayer();
+	check(Player);
+
 	UMovieSceneCompiledDataManager* CompiledDataManager = SharedPlaybackState->GetCompiledDataManager();
-	const FMovieSceneCompiledDataEntry& CompiledEntry = CompiledDataManager->GetEntryRef(RootCompiledDataID);
+	const FMovieSceneCompiledDataID RootCompiledDataID = SharedPlaybackState->GetRootCompiledDataID();
+	const FMovieSceneCompiledDataEntry& CompiledEntry   = CompiledDataManager->GetEntryRef(RootCompiledDataID);
 
 	if (EnumHasAnyFlags(CompiledEntry.AccumulatedMask, EMovieSceneSequenceCompilerMask::EvaluationTemplate))
 	{
@@ -159,15 +148,11 @@ void FSequenceInstance::InitializeLegacyEvaluator()
 
 		if (!LegacyEvaluator)
 		{
-			UMovieSceneSequence* RootSequence = SharedPlaybackState->GetRootSequence();
-			LegacyEvaluator = MakeUnique<FMovieSceneTrackEvaluator>(RootSequence, RootCompiledDataID, CompiledDataManager);
+			LegacyEvaluator = MakeUnique<FMovieSceneTrackEvaluator>(CompiledEntry.GetSequence(), RootCompiledDataID, CompiledDataManager);
 		}
 	}
 	else if (LegacyEvaluator)
 	{
-		IMovieScenePlayer* Player = GetPlayer();
-		check(Player);
-
 		LegacyEvaluator->Finish(*Player);
 		LegacyEvaluator = nullptr;
 
@@ -177,28 +162,24 @@ void FSequenceInstance::InitializeLegacyEvaluator()
 
 void FSequenceInstance::InvalidateCachedData()
 {
-	// WARNING: this method is called from Initialize. If this is a root sequence, the Player is still
-	// in the process of creating us, and is waiting for the call stack to return in order to set our
-	// SharedPlaybackState and other relevant references on itself. So do NOT access anything through
-	// the Player's RootEvaluationTemplateInstance here because it's most probably not initalized yet.
+	ensureMsgf(bInitialized, TEXT("This instance hasn't been initialized yet!"));
 
 	Ledger.Invalidate();
 
-	IMovieScenePlayer* Player = FPlayerIndexPlaybackCapability::GetPlayer(SharedPlaybackState);
-
 	UpdateFlags = ESequenceInstanceUpdateFlags::None;
 
-	const FMovieSceneCompiledDataID RootCompiledDataID = SharedPlaybackState->GetRootCompiledDataID();
 	UMovieSceneCompiledDataManager* CompiledDataManager = SharedPlaybackState->GetCompiledDataManager();
+	const FMovieSceneCompiledDataID RootCompiledDataID = SharedPlaybackState->GetRootCompiledDataID();
+
+	FMovieSceneEvaluationState* State = SharedPlaybackState->FindCapability<FMovieSceneEvaluationState>();
 
 	if (SequenceID == MovieSceneSequenceID::Root)
 	{
 		SharedPlaybackState->InvalidateCachedData();
 
-		if (Player)
+		if (State)
 		{
-			UMovieSceneSequence* RootSequence = SharedPlaybackState->GetRootSequence();
-			Player->State.AssignSequence(SequenceID, *RootSequence, *Player);
+			State->AssignSequence(SequenceID, *SharedPlaybackState->GetRootSequence(), SharedPlaybackState);
 		}
 
 		// Try and recreate the volatility manager if this sequence is now volatile
@@ -209,9 +190,8 @@ void FSequenceInstance::InvalidateCachedData()
 
 		ISequenceUpdater::FactoryInstance(SequenceUpdater, CompiledDataManager, RootCompiledDataID);
 
-		UMovieSceneEntitySystemLinker* Linker = SharedPlaybackState->GetLinker();
-		SequenceUpdater->InvalidateCachedData(Linker, RootInstanceHandle);
-		SequenceUpdater->PopulateUpdateFlags(Linker, SharedPlaybackState, UpdateFlags);
+		SequenceUpdater->InvalidateCachedData(SharedPlaybackState);
+		SequenceUpdater->PopulateUpdateFlags(SharedPlaybackState, UpdateFlags);
 
 		if (LegacyEvaluator)
 		{
@@ -222,15 +202,17 @@ void FSequenceInstance::InvalidateCachedData()
 	}
 	else if (UMovieSceneSequence* SubSequence = SharedPlaybackState->GetSequence(SequenceID))
 	{
-		if (Player)
+		if (State)
 		{
-			Player->State.AssignSequence(SequenceID, *SubSequence, *Player);
+			State->AssignSequence(SequenceID, *SubSequence, SharedPlaybackState);
 		}
 	}
 }
 
 bool FSequenceInstance::ConditionalRecompile()
 {
+	ensureMsgf(bInitialized, TEXT("This instance hasn't been initialized yet!"));
+
 	if (VolatilityManager)
 	{
 		if (VolatilityManager->ConditionalRecompile())
@@ -245,25 +227,26 @@ bool FSequenceInstance::ConditionalRecompile()
 
 void FSequenceInstance::DissectContext(const FMovieSceneContext& InContext, TArray<TRange<FFrameTime>>& OutDissections)
 {
+	ensureMsgf(bInitialized, TEXT("This instance hasn't been initialized yet!"));
+
 	if (EnumHasAnyFlags(UpdateFlags, ESequenceInstanceUpdateFlags::NeedsDissection))
 	{
 		check(SequenceID == MovieSceneSequenceID::Root);
-		UMovieSceneEntitySystemLinker* Linker = SharedPlaybackState->GetLinker();
-		SequenceUpdater->DissectContext(Linker, SharedPlaybackState, InContext, OutDissections);
+		SequenceUpdater->DissectContext(SharedPlaybackState, InContext, OutDissections);
 	}
 }
 
 void FSequenceInstance::Start(const FMovieSceneContext& InContext)
 {
-	check(SequenceID == MovieSceneSequenceID::Root);
+	ensureMsgf(bInitialized, TEXT("This instance hasn't been initialized yet!"));
+	ensureMsgf(SequenceID == MovieSceneSequenceID::Root, TEXT("Only root sequences should be started"));
 
 	bFinished = false;
 	bHasEverUpdated = true;
 
 	check(RootInstanceHandle == InstanceHandle);
 
-	UMovieSceneEntitySystemLinker* Linker = SharedPlaybackState->GetLinker();
-	SequenceUpdater->Start(Linker, RootInstanceHandle, SharedPlaybackState, InContext);
+	SequenceUpdater->Start(SharedPlaybackState, InContext);
 }
 
 void FSequenceInstance::Update(const FMovieSceneContext& InContext)
@@ -271,8 +254,9 @@ void FSequenceInstance::Update(const FMovieSceneContext& InContext)
 	SCOPE_CYCLE_COUNTER(MovieSceneEval_SequenceInstanceUpdate);
 	SCOPE_CYCLE_UOBJECT(ContextScope, GetPlayer()->AsUObject());
 
+	ensureMsgf(bInitialized, TEXT("This instance hasn't been initialized yet!"));
+
 	bHasEverUpdated = true;
-	UMovieSceneEntitySystemLinker* Linker = SharedPlaybackState->GetLinker();
 
 	if (bFinished)
 	{
@@ -282,7 +266,7 @@ void FSequenceInstance::Update(const FMovieSceneContext& InContext)
 	check(RootInstanceHandle == InstanceHandle);
 
 	Context = InContext;
-	SequenceUpdater->Update(Linker, RootInstanceHandle, SharedPlaybackState, InContext);
+	SequenceUpdater->Update(SharedPlaybackState, InContext);
 }
 
 bool FSequenceInstance::CanFinishImmediately() const
@@ -291,8 +275,7 @@ bool FSequenceInstance::CanFinishImmediately() const
 	{
 		check(RootInstanceHandle == InstanceHandle);
 
-		UMovieSceneEntitySystemLinker* Linker = SharedPlaybackState->GetLinker();
-		return SequenceUpdater->CanFinishImmediately(Linker, RootInstanceHandle);
+		return SequenceUpdater->CanFinishImmediately(SharedPlaybackState);
 	}
 
 	return true;
@@ -312,30 +295,27 @@ void FSequenceInstance::Finish()
 
 	Ledger = FEntityLedger();
 
-	IMovieScenePlayer* Player = GetPlayer();
-	if (!ensure(Player))
-	{
-		return;
-	}
-
 	if (SequenceUpdater)
 	{
 		check(RootInstanceHandle == InstanceHandle);
-		SequenceUpdater->Finish(Linker, RootInstanceHandle, SharedPlaybackState);
+		SequenceUpdater->Finish(SharedPlaybackState);
 	}
 
-	if (LegacyEvaluator)
+	IMovieScenePlayer* Player = FPlayerIndexPlaybackCapability::GetPlayer(SharedPlaybackState);
+	if (LegacyEvaluator && ensure(Player))
 	{
 		LegacyEvaluator->Finish(*Player);
 	}
 
 	if (IsRootSequence())
 	{
-		FMovieSceneSpawnRegister& SpawnRegister = Player->GetSpawnRegister();
-		SpawnRegister.ForgetExternallyOwnedSpawnedObjects(Player->State, *Player);
-		SpawnRegister.CleanUp(*Player);
+		if (FMovieSceneSpawnRegister* SpawnRegister = SharedPlaybackState->FindCapability<FMovieSceneSpawnRegister>())
+		{
+			SpawnRegister->ForgetExternallyOwnedSpawnedObjects(SharedPlaybackState);
+			SpawnRegister->CleanUp(SharedPlaybackState);
+		}
 
-		if (Player->PreAnimatedState.IsCapturingGlobalPreAnimatedState())
+		if (Player && Player->PreAnimatedState.IsCapturingGlobalPreAnimatedState())
 		{
 			Linker->PreAnimatedState.RestoreGlobalState(FRestoreStateParams{ Linker, RootInstanceHandle });
 		}
@@ -363,7 +343,7 @@ void FSequenceInstance::RunLegacyTrackTemplates()
 {
 	if (LegacyEvaluator)
 	{
-		IMovieScenePlayer* Player = GetPlayer();
+		IMovieScenePlayer* Player = FPlayerIndexPlaybackCapability::GetPlayer(SharedPlaybackState);
 		if (ensure(Player))
 		{
 			if (bFinished)
@@ -416,17 +396,16 @@ void FSequenceInstance::PostEvaluation()
 
 void FSequenceInstance::DestroyImmediately()
 {
-	UMovieSceneEntitySystemLinker* Linker = SharedPlaybackState->GetLinker();
-
 	if (!Ledger.IsEmpty())
 	{
 		UE_LOG(LogMovieSceneECS, Verbose, TEXT("Instance being destroyed without first having been finished by calling Finish()"));
+		UMovieSceneEntitySystemLinker* Linker = SharedPlaybackState->GetLinker();
 		Ledger.UnlinkEverything(Linker, EUnlinkEverythingMode::CleanGarbage);
 	}
 
 	if (SequenceUpdater)
 	{
-		SequenceUpdater->Destroy(Linker);
+		SequenceUpdater->Destroy(SharedPlaybackState);
 	}
 }
 
@@ -435,8 +414,7 @@ void FSequenceInstance::OverrideRootSequence(FMovieSceneSequenceID NewRootSequen
 	if (SequenceUpdater)
 	{
 		check(RootInstanceHandle == InstanceHandle);
-		UMovieSceneEntitySystemLinker* Linker = SharedPlaybackState->GetLinker();
-		SequenceUpdater->OverrideRootSequence(Linker, RootInstanceHandle, NewRootSequenceID);
+		SequenceUpdater->OverrideRootSequence(SharedPlaybackState, NewRootSequenceID);
 	}
 
 	RootOverrideSequenceID = NewRootSequenceID;
diff --git a/Engine/Source/Runtime/MovieScene/Private/EntitySystem/MovieSceneSequenceUpdaters.cpp b/Engine/Source/Runtime/MovieScene/Private/EntitySystem/MovieSceneSequenceUpdaters.cpp
--- a/Engine/Source/Runtime/MovieScene/Private/EntitySystem/MovieSceneSequenceUpdaters.cpp
+++ b/Engine/Source/Runtime/MovieScene/Private/EntitySystem/MovieSceneSequenceUpdaters.cpp
@@ -35,17 +35,17 @@ struct FSequenceUpdater_Flat : ISequenceUpdater
 	explicit FSequenceUpdater_Flat(FMovieSceneCompiledDataID InCompiledDataID);
 	~FSequenceUpdater_Flat();
 
-	virtual void PopulateUpdateFlags(UMovieSceneEntitySystemLinker* InLinker, TSharedRef<const FSharedPlaybackState> SharedPlaybackState, ESequenceInstanceUpdateFlags& OutUpdateFlags) override;
-	virtual void DissectContext(UMovieSceneEntitySystemLinker* Linker, TSharedRef<const FSharedPlaybackState> SharedPlaybackState, const FMovieSceneContext& Context, TArray<TRange<FFrameTime>>& OutDissections) override;
-	virtual void Start(UMovieSceneEntitySystemLinker* Linker, FRootInstanceHandle InstanceHandle, TSharedRef<const FSharedPlaybackState> SharedPlaybackState, const FMovieSceneContext& InContext) override;
-	virtual void Update(UMovieSceneEntitySystemLinker* Linker, FRootInstanceHandle InstanceHandle, TSharedRef<const FSharedPlaybackState> SharedPlaybackState, const FMovieSceneContext& Context) override;
-	virtual bool CanFinishImmediately(UMovieSceneEntitySystemLinker* InLinker, FRootInstanceHandle RootInstanceHandle) const override;
-	virtual void Finish(UMovieSceneEntitySystemLinker* Linker, FRootInstanceHandle InstanceHandle, TSharedRef<const FSharedPlaybackState> SharedPlaybackState) override;
-	virtual void InvalidateCachedData(UMovieSceneEntitySystemLinker* Linker, FRootInstanceHandle RootInstanceHandle) override;
-	virtual void Destroy(UMovieSceneEntitySystemLinker* Linker) override;
+	virtual void PopulateUpdateFlags(TSharedRef<const FSharedPlaybackState> SharedPlaybackState, ESequenceInstanceUpdateFlags& OutUpdateFlags) override;
+	virtual void DissectContext(TSharedRef<const FSharedPlaybackState> SharedPlaybackState, const FMovieSceneContext& Context, TArray<TRange<FFrameTime>>& OutDissections) override;
+	virtual void Start(TSharedRef<const FSharedPlaybackState> SharedPlaybackState, const FMovieSceneContext& InContext) override;
+	virtual void Update(TSharedRef<const FSharedPlaybackState> SharedPlaybackState, const FMovieSceneContext& Context) override;
+	virtual bool CanFinishImmediately(TSharedRef<const FSharedPlaybackState> SharedPlaybackState) const override;
+	virtual void Finish(TSharedRef<const FSharedPlaybackState> SharedPlaybackState) override;
+	virtual void InvalidateCachedData(TSharedRef<const FSharedPlaybackState> SharedPlaybackState) override;
+	virtual void Destroy(TSharedRef<const FSharedPlaybackState> SharedPlaybackState) override;
 	virtual TUniquePtr<ISequenceUpdater> MigrateToHierarchical() override;
 	virtual FInstanceHandle FindSubInstance(FMovieSceneSequenceID SubSequenceID) const override { return FInstanceHandle(); }
-	virtual void OverrideRootSequence(UMovieSceneEntitySystemLinker* InLinker, FRootInstanceHandle InstanceHandle, FMovieSceneSequenceID NewRootOverrideSequenceID) override {}
+	virtual void OverrideRootSequence(TSharedRef<const FSharedPlaybackState> SharedPlaybackState, FMovieSceneSequenceID NewRootOverrideSequenceID) override {}
 
 private:
 
@@ -64,23 +64,23 @@ struct FSequenceUpdater_Hierarchical : ISequenceUpdater
 
 	~FSequenceUpdater_Hierarchical();
 
-	virtual void PopulateUpdateFlags(UMovieSceneEntitySystemLinker* InLinker, TSharedRef<const FSharedPlaybackState> SharedPlaybackState, ESequenceInstanceUpdateFlags& OutUpdateFlags) override;
-	virtual void DissectContext(UMovieSceneEntitySystemLinker* Linker, TSharedRef<const FSharedPlaybackState> SharedPlaybackState, const FMovieSceneContext& Context, TArray<TRange<FFrameTime>>& OutDissections) override;
-	virtual void Start(UMovieSceneEntitySystemLinker* Linker, FRootInstanceHandle InstanceHandle, TSharedRef<const FSharedPlaybackState> SharedPlaybackState, const FMovieSceneContext& InContext) override;
-	virtual void Update(UMovieSceneEntitySystemLinker* Linker, FRootInstanceHandle InstanceHandle, TSharedRef<const FSharedPlaybackState> SharedPlaybackState, const FMovieSceneContext& Context) override;
-	virtual bool CanFinishImmediately(UMovieSceneEntitySystemLinker* InLinker, FRootInstanceHandle RootInstanceHandle) const override;
-	virtual void Finish(UMovieSceneEntitySystemLinker* Linker, FRootInstanceHandle InstanceHandle, TSharedRef<const FSharedPlaybackState> SharedPlaybackState) override;
-	virtual void InvalidateCachedData(UMovieSceneEntitySystemLinker* Linker, FRootInstanceHandle RootInstanceHandle) override;
-	virtual void Destroy(UMovieSceneEntitySystemLinker* Linker) override;
+	virtual void PopulateUpdateFlags(TSharedRef<const FSharedPlaybackState> SharedPlaybackState, ESequenceInstanceUpdateFlags& OutUpdateFlags) override;
+	virtual void DissectContext(TSharedRef<const FSharedPlaybackState> SharedPlaybackState, const FMovieSceneContext& Context, TArray<TRange<FFrameTime>>& OutDissections) override;
+	virtual void Start(TSharedRef<const FSharedPlaybackState> SharedPlaybackState, const FMovieSceneContext& InContext) override;
+	virtual void Update(TSharedRef<const FSharedPlaybackState> SharedPlaybackState, const FMovieSceneContext& Context) override;
+	virtual bool CanFinishImmediately(TSharedRef<const FSharedPlaybackState> SharedPlaybackState) const override;
+	virtual void Finish(TSharedRef<const FSharedPlaybackState> SharedPlaybackState) override;
+	virtual void InvalidateCachedData(TSharedRef<const FSharedPlaybackState> SharedPlaybackState) override;
+	virtual void Destroy(TSharedRef<const FSharedPlaybackState> SharedPlaybackState) override;
 	virtual TUniquePtr<ISequenceUpdater> MigrateToHierarchical() override { return nullptr; }
 	virtual FInstanceHandle FindSubInstance(FMovieSceneSequenceID SubSequenceID) const override { return SequenceInstances.FindRef(SubSequenceID).Handle; }
-	virtual void OverrideRootSequence(UMovieSceneEntitySystemLinker* InLinker, FRootInstanceHandle InstanceHandle, FMovieSceneSequenceID NewRootOverrideSequenceID) override;
+	virtual void OverrideRootSequence(TSharedRef<const FSharedPlaybackState> SharedPlaybackState, FMovieSceneSequenceID NewRootOverrideSequenceID) override;
 
 private:
 
 	TRange<FFrameNumber> UpdateEntitiesForSequence(const FMovieSceneEntityComponentField* ComponentField, FFrameTime SequenceTime, FMovieSceneEvaluationFieldEntitySet& OutEntities);
 
-	FInstanceHandle GetOrCreateSequenceInstance(TSharedRef<const FSharedPlaybackState> SharedPlaybackState, UMovieSceneSequence* SubSequence, const FMovieSceneSequenceHierarchy* Hierarchy, FInstanceRegistry* InstanceRegistry, FRootInstanceHandle RootInstanceHandle, FMovieSceneSequenceID SequenceID);
+	FInstanceHandle GetOrCreateSequenceInstance(TSharedRef<const FSharedPlaybackState> SharedPlaybackState, UMovieSceneSequence* SubSequence, const FMovieSceneSequenceHierarchy* Hierarchy, FInstanceRegistry* InstanceRegistry, FMovieSceneSequenceID SequenceID);
 
 private:
 
@@ -196,7 +196,7 @@ TUniquePtr<ISequenceUpdater> FSequenceUpdater_Flat::MigrateToHierarchical()
 	return MakeUnique<FSequenceUpdater_Hierarchical>(CompiledDataID);
 }
 
-void FSequenceUpdater_Flat::PopulateUpdateFlags(UMovieSceneEntitySystemLinker* InLinker, TSharedRef<const FSharedPlaybackState> SharedPlaybackState, ESequenceInstanceUpdateFlags& OutUpdateFlags)
+void FSequenceUpdater_Flat::PopulateUpdateFlags(TSharedRef<const FSharedPlaybackState> SharedPlaybackState, ESequenceInstanceUpdateFlags& OutUpdateFlags)
 {
 	if (!CachedDeterminismFences.IsSet())
 	{
@@ -224,7 +224,7 @@ void FSequenceUpdater_Flat::PopulateUpdateFlags(UMovieSceneEntitySystemLinker* I
 	}
 }
 
-void FSequenceUpdater_Flat::DissectContext(UMovieSceneEntitySystemLinker* Linker, TSharedRef<const FSharedPlaybackState> SharedPlaybackState, const FMovieSceneContext& Context, TArray<TRange<FFrameTime>>& OutDissections)
+void FSequenceUpdater_Flat::DissectContext(TSharedRef<const FSharedPlaybackState> SharedPlaybackState, const FMovieSceneContext& Context, TArray<TRange<FFrameTime>>& OutDissections)
 {
 	if (!CachedDeterminismFences.IsSet())
 	{
@@ -248,12 +248,14 @@ void FSequenceUpdater_Flat::DissectContext(UMovieSceneEntitySystemLinker* Linker
 	}
 }
 
-void FSequenceUpdater_Flat::Start(UMovieSceneEntitySystemLinker* Linker, FRootInstanceHandle InstanceHandle, TSharedRef<const FSharedPlaybackState> SharedPlaybackState, const FMovieSceneContext& InContext)
+void FSequenceUpdater_Flat::Start(TSharedRef<const FSharedPlaybackState> SharedPlaybackState, const FMovieSceneContext& InContext)
 {
 }
 
-void FSequenceUpdater_Flat::Update(UMovieSceneEntitySystemLinker* Linker, FRootInstanceHandle InstanceHandle, TSharedRef<const FSharedPlaybackState> SharedPlaybackState, const FMovieSceneContext& Context)
+void FSequenceUpdater_Flat::Update(TSharedRef<const FSharedPlaybackState> SharedPlaybackState, const FMovieSceneContext& Context)
 {
+	UMovieSceneEntitySystemLinker* Linker = SharedPlaybackState->GetLinker();
+	const FRootInstanceHandle& InstanceHandle = SharedPlaybackState->GetRootInstanceHandle();
 	FSequenceInstance& SequenceInstance = Linker->GetInstanceRegistry()->MutateInstance(InstanceHandle);
 	SequenceInstance.SetContext(Context);
 
@@ -323,22 +325,24 @@ void FSequenceUpdater_Flat::Update(UMovieSceneEntitySystemLinker* Linker, FRootI
 	}
 }
 
-bool FSequenceUpdater_Flat::CanFinishImmediately(UMovieSceneEntitySystemLinker* InLinker, FRootInstanceHandle RootInstanceHandle) const
+bool FSequenceUpdater_Flat::CanFinishImmediately(TSharedRef<const FSharedPlaybackState> SharedPlaybackState) const
 {
-	const FSequenceInstance& SequenceInstance = InLinker->GetInstanceRegistry()->GetInstance(RootInstanceHandle);
+	UMovieSceneEntitySystemLinker* Linker = SharedPlaybackState->GetLinker();
+	const FRootInstanceHandle RootInstanceHandle = SharedPlaybackState->GetRootInstanceHandle();
+	const FSequenceInstance& SequenceInstance = Linker->GetInstanceRegistry()->GetInstance(RootInstanceHandle);
 	return SequenceInstance.Ledger.IsEmpty();
 }
 
-void FSequenceUpdater_Flat::Finish(UMovieSceneEntitySystemLinker* Linker, FRootInstanceHandle InstanceHandle, TSharedRef<const FSharedPlaybackState> SharedPlaybackState)
+void FSequenceUpdater_Flat::Finish(TSharedRef<const FSharedPlaybackState> SharedPlaybackState)
 {
-	InvalidateCachedData(Linker, InstanceHandle);
+	InvalidateCachedData(SharedPlaybackState);
 }
 
-void FSequenceUpdater_Flat::Destroy(UMovieSceneEntitySystemLinker* Linker)
+void FSequenceUpdater_Flat::Destroy(TSharedRef<const FSharedPlaybackState> SharedPlaybackState)
 {
 }
 
-void FSequenceUpdater_Flat::InvalidateCachedData(UMovieSceneEntitySystemLinker* Linker, FRootInstanceHandle RootInstanceHandle)
+void FSequenceUpdater_Flat::InvalidateCachedData(TSharedRef<const FSharedPlaybackState> SharedPlaybackState)
 {
 	CachedEntityRange = TRange<FFrameNumber>::Empty();
 	CachedDeterminismFences.Reset();
@@ -360,7 +364,7 @@ FSequenceUpdater_Hierarchical::~FSequenceUpdater_Hierarchical()
 {
 }
 
-void FSequenceUpdater_Hierarchical::PopulateUpdateFlags(UMovieSceneEntitySystemLinker* InLinker, TSharedRef<const FSharedPlaybackState> SharedPlaybackState, ESequenceInstanceUpdateFlags& OutUpdateFlags)
+void FSequenceUpdater_Hierarchical::PopulateUpdateFlags(TSharedRef<const FSharedPlaybackState> SharedPlaybackState, ESequenceInstanceUpdateFlags& OutUpdateFlags)
 {
 	if (IMovieScenePlayer* Player = FPlayerIndexPlaybackCapability::GetPlayer(SharedPlaybackState))
 	{
@@ -399,7 +403,7 @@ void FSequenceUpdater_Hierarchical::PopulateUpdateFlags(UMovieSceneEntitySystemL
 	}
 }
 
-void FSequenceUpdater_Hierarchical::DissectContext(UMovieSceneEntitySystemLinker* Linker, TSharedRef<const FSharedPlaybackState> SharedPlaybackState, const FMovieSceneContext& Context, TArray<TRange<FFrameTime>>& OutDissections)
+void FSequenceUpdater_Hierarchical::DissectContext(TSharedRef<const FSharedPlaybackState> SharedPlaybackState, const FMovieSceneContext& Context, TArray<TRange<FFrameTime>>& OutDissections)
 {
 	UMovieSceneCompiledDataManager* CompiledDataManager = SharedPlaybackState->GetCompiledDataManager();
 
@@ -468,7 +472,7 @@ void FSequenceUpdater_Hierarchical::DissectContext(UMovieSceneEntitySystemLinker
 	}
 }
 
-FInstanceHandle FSequenceUpdater_Hierarchical::GetOrCreateSequenceInstance(TSharedRef<const FSharedPlaybackState> SharedPlaybackState, UMovieSceneSequence* SubSequence, const FMovieSceneSequenceHierarchy* Hierarchy, FInstanceRegistry* InstanceRegistry, FRootInstanceHandle RootInstanceHandle, FMovieSceneSequenceID SequenceID)
+FInstanceHandle FSequenceUpdater_Hierarchical::GetOrCreateSequenceInstance(TSharedRef<const FSharedPlaybackState> SharedPlaybackState, UMovieSceneSequence* SubSequence, const FMovieSceneSequenceHierarchy* Hierarchy, FInstanceRegistry* InstanceRegistry, FMovieSceneSequenceID SequenceID)
 {
 	check(SequenceID != MovieSceneSequenceID::Root);
 
@@ -478,11 +482,11 @@ FInstanceHandle FSequenceUpdater_Hierarchical::GetOrCreateSequenceInstance(TShar
 	}
 
 	const FMovieSceneSequenceHierarchyNode* Node = Hierarchy->FindNode(SequenceID);
-	IMovieScenePlayer* Player = FPlayerIndexPlaybackCapability::GetPlayer(SharedPlaybackState);
-
 	checkf(Node, TEXT("Attempting to construct a new sub sequence instance with a sub sequence ID that does not exist in the hierarchy"));
 	checkf(Node->ParentID != MovieSceneSequenceID::Invalid, TEXT("Parent should never be invalid for a non-root SequenceID"));
 
+	const FRootInstanceHandle& RootInstanceHandle = SharedPlaybackState->GetRootInstanceHandle();
+
 	FInstanceHandle ParentInstance;
 	if (Node->ParentID == MovieSceneSequenceID::Root)
 	{
@@ -490,16 +494,17 @@ FInstanceHandle FSequenceUpdater_Hierarchical::GetOrCreateSequenceInstance(TShar
 	}
 	else if (UMovieSceneSequence* ParentSequence = Hierarchy->FindSubData(Node->ParentID)->GetSequence())
 	{
-		ParentInstance = GetOrCreateSequenceInstance(SharedPlaybackState, ParentSequence, Hierarchy, InstanceRegistry, RootInstanceHandle, Node->ParentID);
+		ParentInstance = GetOrCreateSequenceInstance(SharedPlaybackState, ParentSequence, Hierarchy, InstanceRegistry, Node->ParentID);
 	}
 
-	FInstanceHandle InstanceHandle = InstanceRegistry->AllocateSubInstance(Player, SequenceID, RootInstanceHandle, ParentInstance);
+	FInstanceHandle InstanceHandle = InstanceRegistry->AllocateSubInstance(SequenceID, RootInstanceHandle, ParentInstance);
 	SequenceInstances.Add(SequenceID, FSubInstanceData { SubSequence->GetSignature(), InstanceHandle });
+	InstanceRegistry->MutateInstance(InstanceHandle).Initialize();
 
 	return InstanceHandle;
 }
 
-void FSequenceUpdater_Hierarchical::OverrideRootSequence(UMovieSceneEntitySystemLinker* InLinker, FRootInstanceHandle RootInstanceHandle, FMovieSceneSequenceID NewRootOverrideSequenceID)
+void FSequenceUpdater_Hierarchical::OverrideRootSequence(TSharedRef<const FSharedPlaybackState> SharedPlaybackState, FMovieSceneSequenceID NewRootOverrideSequenceID)
 {
 	if (RootOverrideSequenceID != NewRootOverrideSequenceID)
 	{
@@ -509,28 +514,31 @@ void FSequenceUpdater_Hierarchical::OverrideRootSequence(UMovieSceneEntitySystem
 			// We unlink everything from the root sequence since we know they won't be necessary.
 			// This is because the root sequence instance is handled separately in FSequenceUpdater_Hierarchical::Update, and it wouldn't
 			// get automatically unlinked like other sub sequences would (by way of not being present in the ActiveSequences map)
-			FInstanceRegistry* InstanceRegistry = InLinker->GetInstanceRegistry();
-			InstanceRegistry->MutateInstance(RootInstanceHandle).Ledger.UnlinkEverything(InLinker);
+			UMovieSceneEntitySystemLinker* Linker = SharedPlaybackState->GetLinker();
+			FInstanceRegistry* InstanceRegistry = Linker->GetInstanceRegistry();
+			const FRootInstanceHandle& RootInstanceHandle = SharedPlaybackState->GetRootInstanceHandle();
+			InstanceRegistry->MutateInstance(RootInstanceHandle).Ledger.UnlinkEverything(Linker);
 		}
 
-		InvalidateCachedData(InLinker, RootInstanceHandle);
+		InvalidateCachedData(SharedPlaybackState);
 		RootOverrideSequenceID = NewRootOverrideSequenceID;
 	}
 }
 
-void FSequenceUpdater_Hierarchical::Start(UMovieSceneEntitySystemLinker* Linker, FRootInstanceHandle InstanceHandle, TSharedRef<const FSharedPlaybackState> SharedPlaybackState, const FMovieSceneContext& InContext)
+void FSequenceUpdater_Hierarchical::Start(TSharedRef<const FSharedPlaybackState> SharedPlaybackState, const FMovieSceneContext& InContext)
 {
 }
 
-void FSequenceUpdater_Hierarchical::Update(UMovieSceneEntitySystemLinker* Linker, FRootInstanceHandle InstanceHandle, TSharedRef<const FSharedPlaybackState> SharedPlaybackState, const FMovieSceneContext& Context)
+void FSequenceUpdater_Hierarchical::Update(TSharedRef<const FSharedPlaybackState> SharedPlaybackState, const FMovieSceneContext& Context)
 {
+	UMovieSceneEntitySystemLinker* Linker = SharedPlaybackState->GetLinker();
 	FInstanceRegistry* InstanceRegistry = Linker->GetInstanceRegistry();
 	UMovieSceneCompiledDataManager* CompiledDataManager = SharedPlaybackState->GetCompiledDataManager();
 	IMovieScenePlayer* Player = FPlayerIndexPlaybackCapability::GetPlayer(SharedPlaybackState);
 
 	FMovieSceneEvaluationFieldEntitySet EntitiesScratch;
 
-	FInstanceHandle             RootInstanceHandle = InstanceHandle;
+	FRootInstanceHandle         RootInstanceHandle = SharedPlaybackState->GetRootInstanceHandle();
 	FMovieSceneCompiledDataID   RootCompiledDataID = CompiledDataID;
 	FSubSequencePath            RootOverridePath;
 	FMovieSceneContext          RootContext = Context;
@@ -545,7 +553,8 @@ void FSequenceUpdater_Hierarchical::Update(UMovieSceneEntitySystemLinker* Linker
 		UMovieSceneSequence*                RootSequence    = SubData ? SubData->GetSequence() : nullptr;
 		if (ensure(RootSequence))
 		{
-			RootInstanceHandle = GetOrCreateSequenceInstance(SharedPlaybackState, RootSequence, RootHierarchy, InstanceRegistry, InstanceHandle, RootOverrideSequenceID);
+			FInstanceHandle NewRootInstanceHandle = GetOrCreateSequenceInstance(SharedPlaybackState, RootSequence, RootHierarchy, InstanceRegistry, RootOverrideSequenceID);
+			RootInstanceHandle = FRootInstanceHandle(NewRootInstanceHandle.InstanceID, NewRootInstanceHandle.InstanceSerial);
 			RootCompiledDataID = CompiledDataManager->GetDataID(RootSequence);
 			RootContext        = Context.Transform(SubData->RootToSequenceTransform, SubData->TickResolution);
 
@@ -591,7 +600,7 @@ void FSequenceUpdater_Hierarchical::Update(UMovieSceneEntitySystemLinker* Linker
 				FEntityImportSequenceParams Params;
 				Params.SequenceID = MovieSceneSequenceID::Root;
 				Params.InstanceHandle = RootInstanceHandle;
-				Params.RootInstanceHandle = InstanceHandle;
+				Params.RootInstanceHandle = RootInstanceHandle;
 				Params.DefaultCompletionMode = RootSequence->DefaultCompletionMode;
 				Params.HierarchicalBias = 0;
 				Params.bDynamicWeighting = bDynamicWeighting.Get(false);
@@ -610,7 +619,7 @@ void FSequenceUpdater_Hierarchical::Update(UMovieSceneEntitySystemLinker* Linker
 					FEntityImportSequenceParams Params;
 					Params.SequenceID = MovieSceneSequenceID::Root;
 					Params.InstanceHandle = RootInstanceHandle;
-					Params.RootInstanceHandle = InstanceHandle;
+					Params.RootInstanceHandle = RootInstanceHandle;
 					Params.DefaultCompletionMode = RootSequence->DefaultCompletionMode;
 					Params.HierarchicalBias = 0;
 					Params.bDynamicWeighting = bDynamicWeighting.Get(false);
@@ -658,7 +667,7 @@ void FSequenceUpdater_Hierarchical::Update(UMovieSceneEntitySystemLinker* Linker
 				FMovieSceneCompiledDataID SubDataID = CompiledDataManager->GetDataID(SubSequence);
 
 				// Set the context for the root sequence instance
-				FInstanceHandle    SubSequenceHandle = GetOrCreateSequenceInstance(SharedPlaybackState, SubSequence, RootHierarchy, InstanceRegistry, InstanceHandle, SequenceIDFromRoot);
+				FInstanceHandle    SubSequenceHandle = GetOrCreateSequenceInstance(SharedPlaybackState, SubSequence, RootHierarchy, InstanceRegistry, SequenceIDFromRoot);
 				FSequenceInstance& SubSequenceInstance = InstanceRegistry->MutateInstance(SubSequenceHandle);
 
 				// Update the sub sequence's context
@@ -695,7 +704,7 @@ void FSequenceUpdater_Hierarchical::Update(UMovieSceneEntitySystemLinker* Linker
 				FEntityImportSequenceParams Params;
 				Params.SequenceID = SequenceIDFromRoot;
 				Params.InstanceHandle = SubSequenceHandle;
-				Params.RootInstanceHandle = InstanceHandle;
+				Params.RootInstanceHandle = RootInstanceHandle;
 				Params.DefaultCompletionMode = SubSequence->DefaultCompletionMode;
 				Params.HierarchicalBias = SubData->HierarchicalBias;
 				Params.SubSectionFlags = SubData->AccumulatedFlags;
@@ -761,9 +770,12 @@ void FSequenceUpdater_Hierarchical::Update(UMovieSceneEntitySystemLinker* Linker
 	}
 }
 
-bool FSequenceUpdater_Hierarchical::CanFinishImmediately(UMovieSceneEntitySystemLinker* InLinker, FRootInstanceHandle RootInstanceHandle) const
+bool FSequenceUpdater_Hierarchical::CanFinishImmediately(TSharedRef<const FSharedPlaybackState> SharedPlaybackState) const
 {
-	FInstanceRegistry* InstanceRegistry = InLinker->GetInstanceRegistry();
+	UMovieSceneEntitySystemLinker* Linker = SharedPlaybackState->GetLinker();
+	const FRootInstanceHandle RootInstanceHandle = SharedPlaybackState->GetRootInstanceHandle();
+
+	FInstanceRegistry* InstanceRegistry = Linker->GetInstanceRegistry();
 
 	const FSequenceInstance& RootInstance = InstanceRegistry->GetInstance(RootInstanceHandle);
 	if (!RootInstance.Ledger.IsEmpty())
@@ -783,8 +795,9 @@ bool FSequenceUpdater_Hierarchical::CanFinishImmediately(UMovieSceneEntitySystem
 	return true;
 }
 
-void FSequenceUpdater_Hierarchical::Finish(UMovieSceneEntitySystemLinker* Linker, FRootInstanceHandle InstanceHandle, TSharedRef<const FSharedPlaybackState> SharedPlaybackState)
+void FSequenceUpdater_Hierarchical::Finish(TSharedRef<const FSharedPlaybackState> SharedPlaybackState)
 {
+	UMovieSceneEntitySystemLinker* Linker = SharedPlaybackState->GetLinker();
 	FInstanceRegistry* InstanceRegistry = Linker->GetInstanceRegistry();
 
 	// Finish all sub sequences as well
@@ -793,11 +806,12 @@ void FSequenceUpdater_Hierarchical::Finish(UMovieSceneEntitySystemLinker* Linker
 		InstanceRegistry->MutateInstance(Pair.Value.Handle).Finish();
 	}
 
-	InvalidateCachedData(Linker, InstanceHandle);
+	InvalidateCachedData(SharedPlaybackState);
 }
 
-void FSequenceUpdater_Hierarchical::Destroy(UMovieSceneEntitySystemLinker* Linker)
+void FSequenceUpdater_Hierarchical::Destroy(TSharedRef<const FSharedPlaybackState> SharedPlaybackState)
 {
+	UMovieSceneEntitySystemLinker* Linker = SharedPlaybackState->GetLinker();
 	FInstanceRegistry* InstanceRegistry = Linker->GetInstanceRegistry();
 
 	for (TPair<FMovieSceneSequenceID, FSubInstanceData> Pair : SequenceInstances)
@@ -806,18 +820,17 @@ void FSequenceUpdater_Hierarchical::Destroy(UMovieSceneEntitySystemLinker* Linke
 	}
 }
 
-void FSequenceUpdater_Hierarchical::InvalidateCachedData(UMovieSceneEntitySystemLinker* Linker, FRootInstanceHandle RootInstanceHandle)
+void FSequenceUpdater_Hierarchical::InvalidateCachedData(TSharedRef<const FSharedPlaybackState> SharedPlaybackState)
 {
 	bDynamicWeighting.Reset();
 	CachedEntityRange = TRange<FFrameNumber>::Empty();
 
+	UMovieSceneEntitySystemLinker* Linker = SharedPlaybackState->GetLinker();
 	FInstanceRegistry* InstanceRegistry = Linker->GetInstanceRegistry();
 
-	FMovieSceneRootEvaluationTemplateInstance& RootTemplate = InstanceRegistry->GetInstance(RootInstanceHandle).GetPlayer()->GetEvaluationTemplate();
-
 	for (TPair<FMovieSceneSequenceID, FSubInstanceData>& Pair : SequenceInstances)
 	{
-		UMovieSceneSequence* Sequence = RootTemplate.GetSequence(Pair.Key);
+		UMovieSceneSequence* Sequence = SharedPlaybackState->GetSequence(Pair.Key);
 		if (!Sequence)
 		{
 			Pair.Value.SequenceSignature = FGuid();
@@ -831,8 +844,6 @@ void FSequenceUpdater_Hierarchical::InvalidateCachedData(UMovieSceneEntitySystem
 	}
 }
 
-
-
 TRange<FFrameNumber> FSequenceUpdater_Hierarchical::UpdateEntitiesForSequence(const FMovieSceneEntityComponentField* ComponentField, FFrameTime SequenceTime, FMovieSceneEvaluationFieldEntitySet& OutEntities)
 {
 	TRange<FFrameNumber> CachedRange = TRange<FFrameNumber>::All();
@@ -846,6 +857,5 @@ TRange<FFrameNumber> FSequenceUpdater_Hierarchical::UpdateEntitiesForSequence(co
 	return CachedRange;
 }
 
-
 } // namespace MovieScene
 } // namespace UE
diff --git a/Engine/Source/Runtime/MovieScene/Private/EntitySystem/MovieSceneSharedPlaybackState.cpp b/Engine/Source/Runtime/MovieScene/Private/EntitySystem/MovieSceneSharedPlaybackState.cpp
--- a/Engine/Source/Runtime/MovieScene/Private/EntitySystem/MovieSceneSharedPlaybackState.cpp
+++ b/Engine/Source/Runtime/MovieScene/Private/EntitySystem/MovieSceneSharedPlaybackState.cpp
@@ -4,6 +4,7 @@
 
 #include "Compilation/MovieSceneCompiledDataManager.h"
 #include "EntitySystem/MovieSceneEntitySystemRunner.h"
+#include "MovieSceneSequenceID.h"
 
 namespace UE::MovieScene
 {
@@ -47,9 +48,16 @@ const FMovieSceneSequenceHierarchy* FSharedPlaybackState::GetHierarchy() const
 
 UMovieSceneSequence* FSharedPlaybackState::GetSequence(FMovieSceneSequenceIDRef InSequenceID) const
 {
-	const FMovieSceneSequenceHierarchy* Hierarchy = GetHierarchy();
-	const FMovieSceneSubSequenceData*   SubData   = Hierarchy ? Hierarchy->FindSubData(InSequenceID) : nullptr;
-	return SubData ? SubData->GetSequence() : nullptr;
+	if (InSequenceID == MovieSceneSequenceID::Root)
+	{
+		return WeakRootSequence.Get();
+	}
+	else
+	{
+		const FMovieSceneSequenceHierarchy* Hierarchy = GetHierarchy();
+		const FMovieSceneSubSequenceData*   SubData   = Hierarchy ? Hierarchy->FindSubData(InSequenceID) : nullptr;
+		return SubData ? SubData->GetSequence() : nullptr;
+	}
 }
 
 void FSharedPlaybackState::InvalidateCachedData()
diff --git a/Engine/Source/Runtime/MovieScene/Private/EntitySystem/MovieSceneSpawnablesSystem.cpp b/Engine/Source/Runtime/MovieScene/Private/EntitySystem/MovieSceneSpawnablesSystem.cpp
--- a/Engine/Source/Runtime/MovieScene/Private/EntitySystem/MovieSceneSpawnablesSystem.cpp
+++ b/Engine/Source/Runtime/MovieScene/Private/EntitySystem/MovieSceneSpawnablesSystem.cpp
@@ -6,12 +6,15 @@
 #include "EntitySystem/MovieSceneEntityRange.h"
 #include "EntitySystem/MovieSceneEntitySystemTask.h"
 #include "EntitySystem/MovieSceneEntityInstantiatorSystem.h"
+#include "EntitySystem/MovieSceneSharedPlaybackState.h"
 #include "EntitySystem/BuiltInComponentTypes.h"
 #include "Evaluation/MovieSceneEvaluationOperand.h"
 
 #include "MovieScene.h"
 #include "MovieSceneSequence.h"
 #include "MovieSceneExecutionToken.h"
+#include "MovieSceneSpawnRegister.h"
+
 #include "IMovieScenePlayer.h"
 #include "IMovieScenePlaybackClient.h"
 
@@ -40,18 +43,19 @@ struct FSpawnTrackPreAnimatedTokenProducer : IMovieScenePreAnimatedTokenProducer
 
 			virtual void RestoreState(UObject& Object, const UE::MovieScene::FRestoreStateParams& Params) override
 			{
-				IMovieScenePlayer* Player = Params.GetTerminalPlayer();
-				if (!ensure(Player))
+				TSharedPtr<const FSharedPlaybackState> PlaybackState = Params.GetTerminalPlaybackState();
+				if (!ensure(PlaybackState))
 				{
 					return;
 				}
 
-				if (!Player->GetSpawnRegister().DestroySpawnedObject(OperandToDestroy.ObjectBindingID, OperandToDestroy.SequenceID, *Player))
+				FMovieSceneSpawnRegister* SpawnRegister = PlaybackState->FindCapability<FMovieSceneSpawnRegister>();
+				if (SpawnRegister && !SpawnRegister->DestroySpawnedObject(OperandToDestroy.ObjectBindingID, OperandToDestroy.SequenceID, PlaybackState.ToSharedRef()))
 				{
 					// This branch should only be taken for Externally owned spawnables that have been 'forgotten',
 					// but still had RestoreState tokens generated for them (ie, in FSequencer, or if bRestoreState is enabled)
 					// on a UMovieSceneSequencePlayer
-					Player->GetSpawnRegister().DestroyObjectDirectly(Object);
+					SpawnRegister->DestroyObjectDirectly(Object);
 				}
 			}
 		};
@@ -127,13 +131,13 @@ void UMovieSceneSpawnablesSystem::OnRun(FSystemTaskPrerequisites& InPrerequisite
 		if (ensure(InstanceRegistry->IsHandleValid(InstanceHandle)))
 		{
 			const FSequenceInstance& Instance = InstanceRegistry->GetInstance(InstanceHandle);
-			IMovieScenePlayer* Player = Instance.GetPlayer();
+			TSharedRef<FSharedPlaybackState> SharedPlaybackState = Instance.GetSharedPlaybackState();
 
 			// If the sequence instance has finished and it is a sub sequence, we do not destroy the spawnable
 			// if it is owned by the root sequence or externally. These will get destroyed or forgotten by the player when it ends
 			if (Instance.HasFinished() && Instance.IsSubSequence())
 			{
-				const UMovieSceneSequence* Sequence = Player->State.FindSequence(Instance.GetSequenceID());
+				const UMovieSceneSequence* Sequence = SharedPlaybackState->GetSequence(Instance.GetSequenceID());
 				FMovieSceneSpawnable* Spawnable = Sequence ? Sequence->GetMovieScene()->FindSpawnable(SpawnableObjectID) : nullptr;
 				if (!Spawnable || Spawnable->GetSpawnOwnership() != ESpawnOwnership::InnerSequence)
 				{
@@ -153,20 +157,27 @@ void UMovieSceneSpawnablesSystem::OnRun(FSystemTaskPrerequisites& InPrerequisite
 
 		const FSequenceInstance& SequenceInstance = InstanceRegistry->GetInstance(InstanceHandle);
 
-		FMovieSceneSequenceID SequenceID  = SequenceInstance.GetSequenceID();
-		IMovieScenePlayer*    Player      = SequenceInstance.GetPlayer();
+		TSharedRef<const FSharedPlaybackState> SharedPlaybackState = SequenceInstance.GetSharedPlaybackState();
+		FMovieSceneSpawnRegister* SpawnRegister = SharedPlaybackState->FindCapability<FMovieSceneSpawnRegister>();
+		if (!SpawnRegister)
+		{
+			return;
+		}
 
+		FMovieSceneSequenceID SequenceID = SequenceInstance.GetSequenceID();
 		const FMovieSceneEvaluationOperand SpawnableOperand(SequenceID, SpawnableBindingID);
-		if (const FMovieSceneEvaluationOperand* OperandOverride = Player->BindingOverrides.Find(SpawnableOperand))
+		IStaticBindingOverridesPlaybackCapability* StaticOverrides = SharedPlaybackState->FindCapability<IStaticBindingOverridesPlaybackCapability>();
+		if (StaticOverrides && StaticOverrides->GetBindingOverride(SpawnableOperand))
 		{
 			// Don't do anything if this operand was overriden... someone else will take care of it (either another spawn track, or
 			// some possessable).
 			return;
 		}
 
-		UObject* ExistingSpawnedObject = Player->GetSpawnRegister().FindSpawnedObject(SpawnableBindingID, SequenceID).Get();
+		UObject* ExistingSpawnedObject = SpawnRegister->FindSpawnedObject(SpawnableBindingID, SequenceID).Get();
 
-		if (!Player->State.GetBindingActivation(SpawnableBindingID, SequenceID))
+		FMovieSceneEvaluationState* State = SharedPlaybackState->FindCapability<FMovieSceneEvaluationState>();
+		if (State && !State->GetBindingActivation(SpawnableBindingID, SequenceID))
 		{
 			// If the binding is currently inactive, don't spawn the object.
 
@@ -179,10 +190,10 @@ void UMovieSceneSpawnablesSystem::OnRun(FSystemTaskPrerequisites& InPrerequisite
 		}
 
 		// Check whether the binding is overridden - if it is we cannot spawn a new object
-		if (const IMovieScenePlaybackClient* PlaybackClient = Player->GetPlaybackClient())
+		if (IMovieScenePlaybackClient* DynamicOverrides = SharedPlaybackState->FindCapability<IMovieScenePlaybackClient>())
 		{
 			TArray<UObject*, TInlineAllocator<1>> FoundObjects;
-			bool bUseDefaultBinding = PlaybackClient->RetrieveBindingOverrides(SpawnableBindingID, SequenceID, FoundObjects);
+			bool bUseDefaultBinding = DynamicOverrides->RetrieveBindingOverrides(SpawnableBindingID, SequenceID, FoundObjects);
 			if (!bUseDefaultBinding)
 			{
 				// If the binding has been overridden but we have an existing spawned object, then the binding is new and we need to destroy the spawned object.
@@ -203,14 +214,15 @@ void UMovieSceneSpawnablesSystem::OnRun(FSystemTaskPrerequisites& InPrerequisite
 		}
 
 		// At this point we've decided that we need to spawn a whole new object
-		const UMovieSceneSequence* Sequence = Player->State.FindSequence(SequenceID);
+		const UMovieSceneSequence* Sequence = SharedPlaybackState->GetSequence(SequenceID);
 		if (!Sequence)
 		{
 			return;
 		}
 
-		UObject* SpawnedObject = Player->GetSpawnRegister().SpawnObject(SpawnableBindingID, *Sequence->GetMovieScene(), SequenceID, *Player);
-		if (SpawnedObject)
+		UObject* SpawnedObject = SpawnRegister->SpawnObject(SpawnableBindingID, *Sequence->GetMovieScene(), SequenceID, SharedPlaybackState);
+		IMovieScenePlayer* Player = FPlayerIndexPlaybackCapability::GetPlayer(SharedPlaybackState);
+		if (SpawnedObject && Player)
 		{
 			FMovieSceneEvaluationOperand Operand(SequenceID, SpawnableBindingID);
 			Player->OnObjectSpawned(SpawnedObject, Operand);
@@ -242,8 +254,11 @@ void UMovieSceneSpawnablesSystem::OnRun(FSystemTaskPrerequisites& InPrerequisite
 		// Have to check whether the player is still valid because there is a possibility it got cleaned up
 		if (InstanceRegistry->IsHandleValid(Tuple.Get<2>()))
 		{
-			IMovieScenePlayer* Player = InstanceRegistry->GetInstance(Tuple.Get<2>()).GetPlayer();
-			Player->GetSpawnRegister().DestroySpawnedObject(Tuple.Get<0>(), Tuple.Get<1>(), *Player);
+			TSharedRef<const FSharedPlaybackState> SharedPlaybackState = InstanceRegistry->GetInstance(Tuple.Get<2>()).GetSharedPlaybackState();
+			if (FMovieSceneSpawnRegister* SpawnRegister = SharedPlaybackState->FindCapability<FMovieSceneSpawnRegister>())
+			{
+				SpawnRegister->DestroySpawnedObject(Tuple.Get<0>(), Tuple.Get<1>(), SharedPlaybackState);
+			}
 		}
 	}
 }
diff --git a/Engine/Source/Runtime/MovieScene/Private/Evaluation/MovieSceneEvaluationState.cpp b/Engine/Source/Runtime/MovieScene/Private/Evaluation/MovieSceneEvaluationState.cpp
--- a/Engine/Source/Runtime/MovieScene/Private/Evaluation/MovieSceneEvaluationState.cpp
+++ b/Engine/Source/Runtime/MovieScene/Private/Evaluation/MovieSceneEvaluationState.cpp
@@ -5,7 +5,9 @@
 #include "Algo/Sort.h"
 #include "Algo/Unique.h"
 #include "Engine/World.h"
-#include "Evaluation/MovieSceneEvaluationTemplateInstance.h"
+#include "EntitySystem/MovieSceneEntitySystemLinker.h"
+#include "EntitySystem/MovieSceneInstanceRegistry.h"
+#include "EntitySystem/MovieSceneSequenceInstance.h"
 #include "IMovieScenePlaybackClient.h"
 #include "IMovieScenePlayer.h"
 #include "MovieScene.h"
@@ -16,6 +18,14 @@
 DECLARE_CYCLE_STAT(TEXT("Find Bound Objects"), MovieSceneEval_FindBoundObjects, STATGROUP_MovieSceneEval);
 DECLARE_CYCLE_STAT(TEXT("Iterate Bound Objects"), MovieSceneEval_IterateBoundObjects, STATGROUP_MovieSceneEval);
 
+namespace UE::MovieScene
+{
+
+TPlaybackCapabilityID<IObjectBindingNotifyPlaybackCapability> IObjectBindingNotifyPlaybackCapability::ID = TPlaybackCapabilityID<IObjectBindingNotifyPlaybackCapability>::Register();
+TPlaybackCapabilityID<IStaticBindingOverridesPlaybackCapability> IStaticBindingOverridesPlaybackCapability::ID = TPlaybackCapabilityID<IStaticBindingOverridesPlaybackCapability>::Register();
+
+}  // namespace UE::MovieScene
+
 FMovieSceneSharedDataId FMovieSceneSharedDataId::Allocate()
 {
 	static uint32 Counter = 0;
@@ -26,7 +36,7 @@ FMovieSceneSharedDataId FMovieSceneSharedDataId::Allocate()
 	return Value;
 }
 
-TArrayView<TWeakObjectPtr<>> FMovieSceneObjectCache::FindBoundObjects(const FGuid& InBindingID, IMovieScenePlayer& Player)
+TArrayView<TWeakObjectPtr<>> FMovieSceneObjectCache::FindBoundObjects(const FGuid& InBindingID, TSharedRef<const FSharedPlaybackState> SharedPlaybackState)
 {
 	MOVIESCENE_DETAILED_SCOPE_CYCLE_COUNTER(MovieSceneEval_FindBoundObjects)
 	
@@ -41,7 +51,7 @@ TArrayView<TWeakObjectPtr<>> FMovieSceneObjectCache::FindBoundObjects(const FGui
 	}
 
 	// Attempt to update the bindings
-	UpdateBindings(InBindingID, Player);
+	UpdateBindings(InBindingID, SharedPlaybackState);
 
 	Bindings = BoundObjects.Find(InBindingID);
 	if (Bindings)
@@ -70,7 +80,7 @@ TArrayView<const TWeakObjectPtr<>> FMovieSceneObjectCache::IterateBoundObjects(c
 	return TArrayView<TWeakObjectPtr<>>();
 }
 
-FGuid FMovieSceneObjectCache::FindObjectId(UObject& InObject, IMovieScenePlayer& Player)
+FGuid FMovieSceneObjectCache::FindObjectId(UObject& InObject, TSharedRef<const FSharedPlaybackState> SharedPlaybackState)
 {
 	UMovieSceneSequence* Sequence = WeakSequence.Get();
 	UMovieScene* MovieScene = Sequence ? Sequence->GetMovieScene() : nullptr;
@@ -83,13 +93,13 @@ FGuid FMovieSceneObjectCache::FindObjectId(UObject& InObject, IMovieScenePlayer&
 	{
 		// @todo: Currently we delete the entire object cache when attempting to find an object's ID to ensure that we do a 
 		// complete lookup from scratch. This is required for UMG as it interchanges content slots without notifying sequencer.
-		Clear(Player);
+		Clear(SharedPlaybackState);
 	}
 
-	return FindCachedObjectId(InObject, Player);
+	return FindCachedObjectId(InObject, SharedPlaybackState);
 }
 
-FGuid FMovieSceneObjectCache::FindCachedObjectId(UObject& InObject, IMovieScenePlayer& Player)
+FGuid FMovieSceneObjectCache::FindCachedObjectId(UObject& InObject, TSharedRef<const FSharedPlaybackState> SharedPlaybackState)
 {
 	UMovieSceneSequence* Sequence = WeakSequence.Get();
 	UMovieScene* MovieScene = Sequence ? Sequence->GetMovieScene() : nullptr;
@@ -104,7 +114,7 @@ FGuid FMovieSceneObjectCache::FindCachedObjectId(UObject& InObject, IMovieSceneP
 	for (int32 Index = 0; Index < MovieScene->GetPossessableCount(); ++Index)
 	{
 		FGuid ThisGuid = MovieScene->GetPossessable(Index).GetGuid();
-		if (FindBoundObjects(ThisGuid, Player).Contains(ObjectToFind))
+		if (FindBoundObjects(ThisGuid, SharedPlaybackState).Contains(ObjectToFind))
 		{
 			return ThisGuid;
 		}
@@ -114,7 +124,7 @@ FGuid FMovieSceneObjectCache::FindCachedObjectId(UObject& InObject, IMovieSceneP
 	for (int32 Index = 0; Index < MovieScene->GetSpawnableCount(); ++Index)
 	{
 		FGuid ThisGuid = MovieScene->GetSpawnable(Index).GetGuid();
-		if (FindBoundObjects(ThisGuid, Player).Contains(ObjectToFind))
+		if (FindBoundObjects(ThisGuid, SharedPlaybackState).Contains(ObjectToFind))
 		{
 			return ThisGuid;
 		}
@@ -123,7 +133,7 @@ FGuid FMovieSceneObjectCache::FindCachedObjectId(UObject& InObject, IMovieSceneP
 	return FGuid();
 }
 
-void FMovieSceneObjectCache::FilterObjectBindings(UObject* PredicateObject, IMovieScenePlayer& Player, TArray<FMovieSceneObjectBindingID>* OutBindings)
+void FMovieSceneObjectCache::FilterObjectBindings(UObject* PredicateObject, TSharedRef<const FSharedPlaybackState> SharedPlaybackState, TArray<FMovieSceneObjectBindingID>* OutBindings)
 {
 	check(OutBindings);
 
@@ -150,7 +160,7 @@ void FMovieSceneObjectCache::FilterObjectBindings(UObject* PredicateObject, IMov
 
 	for (const FGuid& DirtyBinding : OutOfDateBindings)
 	{
-		UpdateBindings(DirtyBinding, Player);
+		UpdateBindings(DirtyBinding, SharedPlaybackState);
 
 		const FBoundObjects& Bindings = BoundObjects.FindChecked(DirtyBinding);
 		for (TWeakObjectPtr<> WeakObject : Bindings.Objects)
@@ -300,32 +310,39 @@ bool FMovieSceneObjectCache::InvalidateInternal(const FGuid& InGuid)
 	return true;
 }
 
-void FMovieSceneObjectCache::Clear(IMovieScenePlayer& Player)
+void FMovieSceneObjectCache::Clear(TSharedRef<const FSharedPlaybackState> SharedPlaybackState)
 {
+	using namespace UE::MovieScene;
+
 	BoundObjects.Reset();
 	ChildBindings.Reset();
 	ReverseMappedBindings.Reset();
 
 	UpdateSerialNumber();
 
-	Player.NotifyBindingsChanged();
+	if (IObjectBindingNotifyPlaybackCapability* Notify = SharedPlaybackState->FindCapability<IObjectBindingNotifyPlaybackCapability>())
+	{
+		Notify->NotifyBindingsChanged();
+	}
 	OnBindingInvalidated.Broadcast(FGuid());
 }
 
 
-void FMovieSceneObjectCache::SetSequence(UMovieSceneSequence& InSequence, FMovieSceneSequenceIDRef InSequenceID, IMovieScenePlayer& Player)
+void FMovieSceneObjectCache::SetSequence(UMovieSceneSequence& InSequence, FMovieSceneSequenceIDRef InSequenceID, TSharedRef<const FSharedPlaybackState> SharedPlaybackState)
 {
 	if (WeakSequence != &InSequence)
 	{
-		Clear(Player);
+		Clear(SharedPlaybackState);
 	}
 
 	WeakSequence = &InSequence;
 	SequenceID = InSequenceID;
 }
 
-void FMovieSceneObjectCache::UpdateBindings(const FGuid& InGuid, IMovieScenePlayer& Player)
+void FMovieSceneObjectCache::UpdateBindings(const FGuid& InGuid, TSharedRef<const FSharedPlaybackState> SharedPlaybackState)
 {
+	using namespace UE::MovieScene;
+
 	TGuardValue<bool> ReentrancyGuard(bReentrantUpdate, true);
 
 	// Invalidate existing bindings, we're going to rebuild them.
@@ -360,16 +377,19 @@ void FMovieSceneObjectCache::UpdateBindings(const FGuid& InGuid, IMovieScenePlay
 	// for a different sequence).
 	// TODO-lchabant: we could technically end up in a circular override that creates an infinite loop...
 	const FMovieSceneEvaluationOperand Operand(SequenceID, InGuid);
-	if (const FMovieSceneEvaluationOperand* OverrideOperand = Player.BindingOverrides.Find(Operand))
+	FMovieSceneEvaluationState* State = SharedPlaybackState->FindCapability<FMovieSceneEvaluationState>();
+	IStaticBindingOverridesPlaybackCapability* StaticOverrides = SharedPlaybackState->FindCapability<IStaticBindingOverridesPlaybackCapability>();
+	if (const FMovieSceneEvaluationOperand* OverrideOperand = StaticOverrides ? StaticOverrides->GetBindingOverride(Operand) : nullptr)
 	{
-		const TArrayView<TWeakObjectPtr<>> OverrideBoundObjects = Player.FindBoundObjects(*OverrideOperand);
+		const TArrayView<TWeakObjectPtr<>> OverrideBoundObjects = State->FindBoundObjects(*OverrideOperand, SharedPlaybackState);
 		Bindings->Objects.Append(OverrideBoundObjects.GetData(), OverrideBoundObjects.Num());
 	}
 	else
 	{
 		const bool bUseParentsAsContext = Sequence->AreParentContextsSignificant();
 
-		UObject* Context = Player.GetPlaybackContext();
+		UObject* Context = SharedPlaybackState->GetPlaybackContext();
+		IMovieScenePlayer* Player = FPlayerIndexPlaybackCapability::GetPlayer(SharedPlaybackState);
 
 		const FMovieScenePossessable* Possessable = Sequence->GetMovieScene()->FindPossessable(InGuid);
 		if (Possessable)
@@ -379,7 +399,7 @@ void FMovieSceneObjectCache::UpdateBindings(const FGuid& InGuid, IMovieScenePlay
 			// Because these are ordered parent-first, the parent must have already been bound, if it exists
 			if (Possessable->GetParent().IsValid())
 			{
-				TArrayView<TWeakObjectPtr<>> ParentBoundObjects = FindBoundObjects(Possessable->GetParent(), Player);
+				TArrayView<TWeakObjectPtr<>> ParentBoundObjects = FindBoundObjects(Possessable->GetParent(), SharedPlaybackState);
 
 				ChildBindings.FindOrAdd(Possessable->GetParent()).AddUnique(InGuid);
 
@@ -400,7 +420,7 @@ void FMovieSceneObjectCache::UpdateBindings(const FGuid& InGuid, IMovieScenePlay
 					
 					if (Possessable->GetSpawnableObjectBindingID().IsValid())
 					{
-						for (TWeakObjectPtr<> BoundObject : Possessable->GetSpawnableObjectBindingID().ResolveBoundObjects(SequenceID, Player))
+						for (TWeakObjectPtr<> BoundObject : Possessable->GetSpawnableObjectBindingID().ResolveBoundObjects(SequenceID, SharedPlaybackState))
 						{
 							if (BoundObject.IsValid())
 							{
@@ -410,15 +430,19 @@ void FMovieSceneObjectCache::UpdateBindings(const FGuid& InGuid, IMovieScenePlay
 					}
 					else
 					{
-						FMovieSceneDynamicBindingResolveResult ResolveResult = FMovieSceneDynamicBindingInvoker::ResolveDynamicBinding(Player, Sequence, SequenceID, InGuid, Possessable->DynamicBinding);
+						FMovieSceneDynamicBindingResolveResult ResolveResult = FMovieSceneDynamicBindingInvoker::ResolveDynamicBinding(*Player, Sequence, SequenceID, InGuid, Possessable->DynamicBinding);
 						if (ResolveResult.Object)
 						{
 							ensureMsgf(ResolveResult.bIsPossessedObject, TEXT("Possessables don't support spawnable-type ownership"));
 							FoundObjects.Add(ResolveResult.Object);
 						}
+						else if (Player)
+						{
+							Player->ResolveBoundObjects(InGuid, SequenceID, *Sequence, ResolutionContext, FoundObjects);
+						}
 						else
 						{
-							Player.ResolveBoundObjects(InGuid, SequenceID, *Sequence, ResolutionContext, FoundObjects);
+							Sequence->LocateBoundObjects(InGuid, ResolutionContext, FoundObjects);
 						}
 					}
 					
@@ -437,9 +461,9 @@ void FMovieSceneObjectCache::UpdateBindings(const FGuid& InGuid, IMovieScenePlay
 				{
 					// We resolve this binding to fixed here, as we conveniently have a Player pointer already, and when being invalidated,
 					// the binding ID passed down will be relative to the root.
-					FMovieSceneObjectBindingID SpawnableFixedBindingID = Possessable->GetSpawnableObjectBindingID().ResolveToFixed(SequenceID, Player);
+					FMovieSceneObjectBindingID SpawnableFixedBindingID = Possessable->GetSpawnableObjectBindingID().ResolveToFixed(SequenceID, SharedPlaybackState);
 					ReverseMappedBindings.FindOrAdd(SpawnableFixedBindingID).AddUnique(InGuid);
-					for (TWeakObjectPtr<> BoundObject : Possessable->GetSpawnableObjectBindingID().ResolveBoundObjects(SequenceID, Player))
+					for (TWeakObjectPtr<> BoundObject : Possessable->GetSpawnableObjectBindingID().ResolveBoundObjects(SequenceID, SharedPlaybackState))
 					{
 						if (BoundObject.IsValid())
 						{
@@ -449,15 +473,19 @@ void FMovieSceneObjectCache::UpdateBindings(const FGuid& InGuid, IMovieScenePlay
 				}
 				else
 				{
-					FMovieSceneDynamicBindingResolveResult ResolveResult = FMovieSceneDynamicBindingInvoker::ResolveDynamicBinding(Player, Sequence, SequenceID, InGuid, Possessable->DynamicBinding);
+					FMovieSceneDynamicBindingResolveResult ResolveResult = FMovieSceneDynamicBindingInvoker::ResolveDynamicBinding(*Player, Sequence, SequenceID, InGuid, Possessable->DynamicBinding);
 					if (ResolveResult.Object)
 					{
 						ensureMsgf(ResolveResult.bIsPossessedObject, TEXT("Possessables don't support spawnable-type ownership"));
 						FoundObjects.Add(ResolveResult.Object);
 					}
+					else if (Player)
+					{
+						Player->ResolveBoundObjects(InGuid, SequenceID, *Sequence, ResolutionContext, FoundObjects);
+					}
 					else
 					{
-						Player.ResolveBoundObjects(InGuid, SequenceID, *Sequence, ResolutionContext, FoundObjects);
+						Sequence->LocateBoundObjects(InGuid, ResolutionContext, FoundObjects);
 					}
 				}
 				
@@ -474,11 +502,11 @@ void FMovieSceneObjectCache::UpdateBindings(const FGuid& InGuid, IMovieScenePlay
 			bool bUseDefault = true;
 
 			// Allow external overrides for spawnables
-			const IMovieScenePlaybackClient* PlaybackClient = Player.GetPlaybackClient();
-			if (PlaybackClient)
+			const IMovieScenePlaybackClient* DynamicOverrides = SharedPlaybackState->FindCapability<IMovieScenePlaybackClient>();
+			if (DynamicOverrides)
 			{
 				TArray<UObject*, TInlineAllocator<1>> FoundObjects;
-				bUseDefault = PlaybackClient->RetrieveBindingOverrides(InGuid, SequenceID, FoundObjects);
+				bUseDefault = DynamicOverrides->RetrieveBindingOverrides(InGuid, SequenceID, FoundObjects);
 				for (UObject* Object : FoundObjects)
 				{
 					Bindings->Objects.Add(Object);
@@ -488,7 +516,8 @@ void FMovieSceneObjectCache::UpdateBindings(const FGuid& InGuid, IMovieScenePlay
 			// If we have no overrides, or they want to allow the default spawnable, do that now
 			if (bUseDefault)
 			{
-				UObject* SpawnedObject = Player.GetSpawnRegister().FindSpawnedObject(InGuid, SequenceID).Get();
+				const FMovieSceneSpawnRegister* SpawnRegister = SharedPlaybackState->FindCapability<FMovieSceneSpawnRegister>();
+				UObject* SpawnedObject = SpawnRegister ? SpawnRegister->FindSpawnedObject(InGuid, SequenceID).Get() : nullptr;
 				if (SpawnedObject)
 				{
 					Bindings->Objects.Add(SpawnedObject);
@@ -519,7 +548,11 @@ void FMovieSceneObjectCache::UpdateBindings(const FGuid& InGuid, IMovieScenePlay
 	if (NumBoundObjects > 0)
 	{
 		Bindings->bUpToDate = true;
-		Player.NotifyBindingUpdate(InGuid, SequenceID, Bindings->Objects);
+
+		if (IObjectBindingNotifyPlaybackCapability* Notify = SharedPlaybackState->FindCapability<IObjectBindingNotifyPlaybackCapability>())
+		{
+			Notify->NotifyBindingUpdate(InGuid, SequenceID, Bindings->Objects);
+		}
 
 		if (auto* Children = ChildBindings.Find(InGuid))
 		{
@@ -538,6 +571,38 @@ void FMovieSceneObjectCache::UpdateSerialNumber()
 	++SerialNumber;
 }
 
+TArrayView<TWeakObjectPtr<>> FMovieSceneObjectCache::FindBoundObjects(const FGuid& InBindingID, IMovieScenePlayer& Player)
+{
+	return FindBoundObjects(InBindingID, Player.GetSharedPlaybackState());
+}
+
+void FMovieSceneObjectCache::SetSequence(UMovieSceneSequence& InSequence, FMovieSceneSequenceIDRef InSequenceID, IMovieScenePlayer& Player)
+{
+	SetSequence(InSequence, InSequenceID, Player.GetSharedPlaybackState());
+}
+
+FGuid FMovieSceneObjectCache::FindObjectId(UObject& InObject, IMovieScenePlayer& Player)
+{
+	return FindObjectId(InObject, Player.GetSharedPlaybackState());
+}
+
+FGuid FMovieSceneObjectCache::FindCachedObjectId(UObject& InObject, IMovieScenePlayer& Player)
+{
+	return FindCachedObjectId(InObject, Player.GetSharedPlaybackState());
+}
+
+void FMovieSceneObjectCache::Clear(IMovieScenePlayer& Player)
+{
+	Clear(Player.GetSharedPlaybackState());
+}
+
+void FMovieSceneObjectCache::FilterObjectBindings(UObject* PredicateObject, IMovieScenePlayer& Player, TArray<FMovieSceneObjectBindingID>* OutBindings)
+{
+	FilterObjectBindings(PredicateObject, Player.GetSharedPlaybackState(), OutBindings);
+}
+
+UE::MovieScene::TPlaybackCapabilityID<FMovieSceneEvaluationState> FMovieSceneEvaluationState::ID = UE::MovieScene::TPlaybackCapabilityID<FMovieSceneEvaluationState>::Register();
+
 void FMovieSceneEvaluationState::InvalidateExpiredObjects()
 {
 	for (auto& Pair : ObjectCaches)
@@ -570,17 +635,17 @@ void FMovieSceneEvaluationState::SetBindingActivation(const FGuid& InGuid, FMovi
 	GetObjectCache(InSequenceID).SetBindingActivation(InGuid, bActive);
 }
 
-void FMovieSceneEvaluationState::ClearObjectCaches(IMovieScenePlayer& Player)
+void FMovieSceneEvaluationState::ClearObjectCaches(TSharedRef<const FSharedPlaybackState> SharedPlaybackState)
 {
 	for (auto& Pair : ObjectCaches)
 	{
-		Pair.Value.ObjectCache.Clear(Player);
+		Pair.Value.ObjectCache.Clear(SharedPlaybackState);
 	}
 }
 
-void FMovieSceneEvaluationState::AssignSequence(FMovieSceneSequenceIDRef InSequenceID, UMovieSceneSequence& InSequence, IMovieScenePlayer& Player)
+void FMovieSceneEvaluationState::AssignSequence(FMovieSceneSequenceIDRef InSequenceID, UMovieSceneSequence& InSequence, TSharedRef<const FSharedPlaybackState> SharedPlaybackState)
 {
-	GetObjectCache(InSequenceID).SetSequence(InSequence, InSequenceID, Player);
+	GetObjectCache(InSequenceID).SetSequence(InSequence, InSequenceID, SharedPlaybackState);
 }
 
 UMovieSceneSequence* FMovieSceneEvaluationState::FindSequence(FMovieSceneSequenceIDRef InSequenceID) const
@@ -602,25 +667,25 @@ FMovieSceneSequenceID FMovieSceneEvaluationState::FindSequenceId(UMovieSceneSequ
 	return FMovieSceneSequenceID();
 }
 
-FGuid FMovieSceneEvaluationState::FindObjectId(UObject& Object, FMovieSceneSequenceIDRef InSequenceID, IMovieScenePlayer& Player)
+FGuid FMovieSceneEvaluationState::FindObjectId(UObject& Object, FMovieSceneSequenceIDRef InSequenceID, TSharedRef<const FSharedPlaybackState> SharedPlaybackState)
 {
 	FVersionedObjectCache* Cache = ObjectCaches.Find(InSequenceID);
-	return Cache ? Cache->ObjectCache.FindObjectId(Object, Player) : FGuid();
+	return Cache ? Cache->ObjectCache.FindObjectId(Object, SharedPlaybackState) : FGuid();
 }
 
-FGuid FMovieSceneEvaluationState::FindCachedObjectId(UObject& Object, FMovieSceneSequenceIDRef InSequenceID, IMovieScenePlayer& Player)
+FGuid FMovieSceneEvaluationState::FindCachedObjectId(UObject& Object, FMovieSceneSequenceIDRef InSequenceID, TSharedRef<const FSharedPlaybackState> SharedPlaybackState)
 {
 	FVersionedObjectCache* Cache = ObjectCaches.Find(InSequenceID);
-	return Cache ? Cache->ObjectCache.FindCachedObjectId(Object, Player) : FGuid();
+	return Cache ? Cache->ObjectCache.FindCachedObjectId(Object, SharedPlaybackState) : FGuid();
 }
 
-void FMovieSceneEvaluationState::FilterObjectBindings(UObject* PredicateObject, IMovieScenePlayer& Player, TArray<FMovieSceneObjectBindingID>* OutBindings)
+void FMovieSceneEvaluationState::FilterObjectBindings(UObject* PredicateObject, TSharedRef<const FSharedPlaybackState> SharedPlaybackState, TArray<FMovieSceneObjectBindingID>* OutBindings)
 {
 	check(OutBindings);
 
 	for (TTuple<FMovieSceneSequenceID, FVersionedObjectCache>& Cache : ObjectCaches)
 	{
-		Cache.Value.ObjectCache.FilterObjectBindings(PredicateObject, Player, OutBindings);
+		Cache.Value.ObjectCache.FilterObjectBindings(PredicateObject, SharedPlaybackState, OutBindings);
 	}
 }
 
@@ -641,3 +706,49 @@ uint32 FMovieSceneEvaluationState::GetSerialNumber()
 	return SerialNumber;
 }
 
+void FMovieSceneEvaluationState::Initialize(TSharedRef<const FSharedPlaybackState> Owner)
+{
+	UMovieSceneEntitySystemLinker* Linker = Owner->GetLinker();
+	RegisterObjectCacheEvents(Linker, Owner->GetRootInstanceHandle(), MovieSceneSequenceID::Root);
+}
+
+void FMovieSceneEvaluationState::OnSubInstanceCreated(TSharedRef<const FSharedPlaybackState> Owner, const UE::MovieScene::FInstanceHandle InstanceHandle)
+{
+	UMovieSceneEntitySystemLinker* Linker = Owner->GetLinker();
+	const UE::MovieScene::FSequenceInstance& SubInstance = Linker->GetInstanceRegistry()->GetInstance(InstanceHandle);
+	RegisterObjectCacheEvents(Linker, InstanceHandle, SubInstance.GetSequenceID());
+}
+
+void FMovieSceneEvaluationState::RegisterObjectCacheEvents(UMovieSceneEntitySystemLinker* Linker, const UE::MovieScene::FInstanceHandle& InstanceHandle, const FMovieSceneSequenceID SequenceID)
+{
+	FMovieSceneObjectCache& ObjectCache = GetObjectCache(SequenceID);  // Make sure the cache is created...
+	FVersionedObjectCache& VersionedObjectCache = ObjectCaches.FindChecked(SequenceID);
+	VersionedObjectCache.OnInvalidateObjectBindingHandle = ObjectCache.OnBindingInvalidated.AddUObject(
+			Linker, &UMovieSceneEntitySystemLinker::InvalidateObjectBinding, InstanceHandle);
+}
+
+void FMovieSceneEvaluationState::AssignSequence(FMovieSceneSequenceIDRef InSequenceID, UMovieSceneSequence& InSequence, IMovieScenePlayer& Player)
+{
+	AssignSequence(InSequenceID, InSequence, Player.GetSharedPlaybackState());
+}
+
+FGuid FMovieSceneEvaluationState::FindObjectId(UObject& Object, FMovieSceneSequenceIDRef InSequenceID, IMovieScenePlayer& Player)
+{
+	return FindObjectId(Object, InSequenceID, Player.GetSharedPlaybackState());
+}
+
+FGuid FMovieSceneEvaluationState::FindCachedObjectId(UObject& Object, FMovieSceneSequenceIDRef InSequenceID, IMovieScenePlayer& Player)
+{
+	return FindCachedObjectId(Object, InSequenceID, Player.GetSharedPlaybackState());
+}
+
+void FMovieSceneEvaluationState::FilterObjectBindings(UObject* PredicateObject, IMovieScenePlayer& Player, TArray<FMovieSceneObjectBindingID>* OutBindings)
+{
+	FilterObjectBindings(PredicateObject, Player.GetSharedPlaybackState(), OutBindings);
+}
+
+void FMovieSceneEvaluationState::ClearObjectCaches(IMovieScenePlayer& Player)
+{
+	ClearObjectCaches(Player.GetSharedPlaybackState());
+}
+
diff --git a/Engine/Source/Runtime/MovieScene/Private/Evaluation/MovieSceneEvaluationTemplateInstance.cpp b/Engine/Source/Runtime/MovieScene/Private/Evaluation/MovieSceneEvaluationTemplateInstance.cpp
--- a/Engine/Source/Runtime/MovieScene/Private/Evaluation/MovieSceneEvaluationTemplateInstance.cpp
+++ b/Engine/Source/Runtime/MovieScene/Private/Evaluation/MovieSceneEvaluationTemplateInstance.cpp
@@ -165,9 +165,11 @@ void FMovieSceneRootEvaluationTemplateInstance::Initialize(UMovieSceneSequence&
 		FRootInstanceHandle RootInstanceHandle;
 		if (EntitySystemLinker != nullptr && EntitySystemLinker->GetInstanceRegistry())
 		{
+			UObject* PlaybackContext = Player.GetPlaybackContext();
 			FInstanceRegistry* InstanceRegistry = EntitySystemLinker->GetInstanceRegistry();
-			RootInstanceHandle = InstanceRegistry->AllocateRootInstance(&Player, InRootSequence, InRunner, InCompiledDataManager);
+			RootInstanceHandle = InstanceRegistry->AllocateRootInstance(InRootSequence, PlaybackContext, InRunner, InCompiledDataManager);
 			SharedPlaybackState = InstanceRegistry->GetInstance(RootInstanceHandle).GetSharedPlaybackState();
+			Player.InitializeRootInstance(SharedPlaybackState.ToSharedRef());
 		}
 		else
 		{
@@ -217,6 +219,16 @@ bool FMovieSceneRootEvaluationTemplateInstance::IsValid() const
 	return EntitySystemLinker && SharedPlaybackState;
 }
 
+TSharedPtr<UE::MovieScene::FSharedPlaybackState> FMovieSceneRootEvaluationTemplateInstance::GetSharedPlaybackState()
+{
+	return SharedPlaybackState;
+}
+
+TSharedPtr<const UE::MovieScene::FSharedPlaybackState> FMovieSceneRootEvaluationTemplateInstance::GetSharedPlaybackState() const
+{
+	return SharedPlaybackState;
+}
+
 UE::MovieScene::FRootInstanceHandle FMovieSceneRootEvaluationTemplateInstance::GetRootInstanceHandle() const
 {
 	return SharedPlaybackState ? SharedPlaybackState->GetRootInstanceHandle() : UE::MovieScene::FRootInstanceHandle();
@@ -262,19 +274,7 @@ void FMovieSceneRootEvaluationTemplateInstance::EnableGlobalPreAnimatedStateCapt
 
 UMovieSceneSequence* FMovieSceneRootEvaluationTemplateInstance::GetSequence(FMovieSceneSequenceIDRef SequenceID) const
 {
-	if (SequenceID == MovieSceneSequenceID::Root)
-	{
-		return SharedPlaybackState->GetRootSequence();
-	}
-	else if (const FMovieSceneSequenceHierarchy* Hierarchy = SharedPlaybackState->GetHierarchy())
-	{
-		const FMovieSceneSubSequenceData* SubSequenceData = Hierarchy->FindSubData(SequenceID);
-		if (SubSequenceData)
-		{
-			return SubSequenceData->GetSequence();
-		}
-	}
-	return nullptr;
+	return SharedPlaybackState ? SharedPlaybackState->GetSequence(SequenceID) : nullptr;
 }
 
 UMovieSceneEntitySystemLinker* FMovieSceneRootEvaluationTemplateInstance::GetEntitySystemLinker() const
@@ -510,9 +510,13 @@ void FMovieSceneRootEvaluationTemplateInstance::PlaybackContextChanged(IMovieSce
 		Runner->AttachToLinker(EntitySystemLinker);
 	}
 
+	UObject* PlaybackContext = Player.GetPlaybackContext();
 	FInstanceRegistry* InstanceRegistry = EntitySystemLinker->GetInstanceRegistry();
-	const FRootInstanceHandle RootInstanceHandle = InstanceRegistry->AllocateRootInstance(&Player, *RootSequence, Runner, CompiledDataManager);
-	SharedPlaybackState = InstanceRegistry->GetInstance(RootInstanceHandle).GetSharedPlaybackState();
+	const FRootInstanceHandle RootInstanceHandle = InstanceRegistry->AllocateRootInstance(
+			*RootSequence, PlaybackContext, Runner, CompiledDataManager);
+	FSequenceInstance& RootInstance = InstanceRegistry->MutateInstance(RootInstanceHandle);
+	SharedPlaybackState = RootInstance.GetSharedPlaybackState();
+	Player.InitializeRootInstance(SharedPlaybackState.ToSharedRef());
 
 	DirectorInstances.Reset();
 
diff --git a/Engine/Source/Runtime/MovieScene/Private/Evaluation/MovieScenePlaybackCapabilities.cpp b/Engine/Source/Runtime/MovieScene/Private/Evaluation/MovieScenePlaybackCapabilities.cpp
--- a/Engine/Source/Runtime/MovieScene/Private/Evaluation/MovieScenePlaybackCapabilities.cpp
+++ b/Engine/Source/Runtime/MovieScene/Private/Evaluation/MovieScenePlaybackCapabilities.cpp
@@ -36,22 +36,14 @@ FPlaybackCapabilities::~FPlaybackCapabilities()
 	Destroy();
 }
 
+void FPlaybackCapabilities::OnSubInstanceCreated(TSharedRef<const FSharedPlaybackState> Owner, const FInstanceHandle InstanceHandle)
+{
+	ForEachCapabilityInterface([Owner, &InstanceHandle](IPlaybackCapability& Cap) { Cap.OnSubInstanceCreated(Owner, InstanceHandle); });
+}
+
 void FPlaybackCapabilities::InvalidateCachedData(UMovieSceneEntitySystemLinker* Linker)
 {
-	TArrayView<const FPlaybackCapabilityHeader> Headers = GetHeaders();
-	for (int32 Index = 0; Index < Headers.Num(); ++Index)
-	{
-		const FPlaybackCapabilityHeader& Header = Headers[Index];
-		const FPlaybackCapabilityHelpers& ThisHelpers = Helpers[Index];
-		check(ThisHelpers.InterfaceCast != nullptr);
-		{
-			void* Ptr = Header.Capability.Resolve(Memory);
-			if (IPlaybackCapability* Interface = (*ThisHelpers.InterfaceCast)(Ptr))
-			{
-				Interface->InvalidateCachedData(Linker);
-			}
-		}
-	}
+	ForEachCapabilityInterface([Linker](IPlaybackCapability& Cap) { Cap.InvalidateCachedData(Linker); });
 }
 
 void FPlaybackCapabilities::Destroy()
diff --git a/Engine/Source/Runtime/MovieScene/Private/Evaluation/MovieScenePreAnimatedState.cpp b/Engine/Source/Runtime/MovieScene/Private/Evaluation/MovieScenePreAnimatedState.cpp
--- a/Engine/Source/Runtime/MovieScene/Private/Evaluation/MovieScenePreAnimatedState.cpp
+++ b/Engine/Source/Runtime/MovieScene/Private/Evaluation/MovieScenePreAnimatedState.cpp
@@ -29,6 +29,18 @@ IMovieScenePlayer* FRestoreStateParams::GetTerminalPlayer() const
 	return nullptr;
 }
 
+TSharedPtr<const FSharedPlaybackState> FRestoreStateParams::GetTerminalPlaybackState() const
+{
+	if (Linker && TerminalInstanceHandle.IsValid())
+	{
+		const FSequenceInstance& TerminalInstance = Linker->GetInstanceRegistry()->GetInstance(TerminalInstanceHandle);
+		return TerminalInstance.GetSharedPlaybackState().ToSharedPtr();
+	}
+
+	ensureAlways(false);
+	return nullptr;
+}
+
 } // namespace MovieScene
 } // namespace UE
 
diff --git a/Engine/Source/Runtime/MovieScene/Private/Evaluation/SequenceDirectorPlaybackCapability.cpp b/Engine/Source/Runtime/MovieScene/Private/Evaluation/SequenceDirectorPlaybackCapability.cpp
new file mode 100644
--- /dev/null
+++ b/Engine/Source/Runtime/MovieScene/Private/Evaluation/SequenceDirectorPlaybackCapability.cpp
@@ -0,0 +1,16 @@
+// Copyright Epic Games, Inc. All Rights Reserved.
+
+#include "Evaluation/SequenceDirectorPlaybackCapability.h"
+
+#include "EntitySystem/MovieSceneSharedPlaybackState.h"
+#include "Evaluation/MovieSceneSequenceHierarchy.h"
+#include "IMovieScenePlayer.h"
+#include "MovieSceneSequence.h"
+
+namespace UE::MovieScene
+{
+
+TPlaybackCapabilityID<ISequenceDirectorPlaybackCapability> ISequenceDirectorPlaybackCapability::ID = TPlaybackCapabilityID<ISequenceDirectorPlaybackCapability>::Register();
+
+}  // namespace UE::MovieScene
+
diff --git a/Engine/Source/Runtime/MovieScene/Private/MovieSceneObjectBindingID.cpp b/Engine/Source/Runtime/MovieScene/Private/MovieSceneObjectBindingID.cpp
--- a/Engine/Source/Runtime/MovieScene/Private/MovieSceneObjectBindingID.cpp
+++ b/Engine/Source/Runtime/MovieScene/Private/MovieSceneObjectBindingID.cpp
@@ -1,10 +1,11 @@
 // Copyright Epic Games, Inc. All Rights Reserved.
 
 #include "MovieSceneObjectBindingID.h"
+#include "Evaluation/MovieSceneRootOverridePath.h"
 #include "Evaluation/MovieSceneSequenceHierarchy.h"
-#include "Evaluation/MovieSceneEvaluationTemplateInstance.h"
 #include "Compilation/MovieSceneCompiledDataManager.h"
 #include "EntitySystem/MovieSceneSequenceInstance.h"
+#include "EntitySystem/MovieSceneSharedPlaybackState.h"
 #include "IMovieScenePlayer.h"
 
 #include UE_INLINE_GENERATED_CPP_BY_NAME(MovieSceneObjectBindingID)
@@ -17,7 +18,7 @@ namespace MovieScene
 {
 
 
-FMovieSceneSequenceID ResolveExternalSequenceID(FMovieSceneSequenceID SourceSequenceID, int32 RemapSourceParentIndex, FMovieSceneSequenceID TargetSequenceID, IMovieScenePlayer& Player)
+FMovieSceneSequenceID ResolveExternalSequenceID(FMovieSceneSequenceID SourceSequenceID, int32 RemapSourceParentIndex, FMovieSceneSequenceID TargetSequenceID, TSharedRef<const FSharedPlaybackState> SharedPlaybackState)
 {
 	ensureMsgf(RemapSourceParentIndex >= 0, TEXT("Invalid parent index specified"));
 
@@ -26,7 +27,7 @@ FMovieSceneSequenceID ResolveExternalSequenceID(FMovieSceneSequenceID SourceSequ
 		return TargetSequenceID;
 	}
 
-	const FMovieSceneSequenceHierarchy* Hierarchy = Player.GetEvaluationTemplate().GetHierarchy();
+	const FMovieSceneSequenceHierarchy* Hierarchy = SharedPlaybackState->GetHierarchy();
 	if (!ensureMsgf(Hierarchy, TEXT("Sequence being evaluated as a sub sequence without any hierarchy. This indicates a bug with either the player or the compiled data.")))
 	{
 		return TargetSequenceID;
@@ -75,7 +76,12 @@ FRelativeObjectBindingID::FRelativeObjectBindingID(FMovieSceneSequenceID SourceS
 
 FRelativeObjectBindingID::FRelativeObjectBindingID(FMovieSceneSequenceID SourceSequenceID, FMovieSceneSequenceID TargetSequenceID, const FGuid& TargetGuid, IMovieScenePlayer& Player)
 {
-	ConstructInternal(SourceSequenceID, TargetSequenceID, TargetGuid, Player.GetEvaluationTemplate().GetHierarchy());
+	ConstructInternal(SourceSequenceID, TargetSequenceID, TargetGuid, Player.GetSharedPlaybackState()->GetHierarchy());
+}
+
+FRelativeObjectBindingID::FRelativeObjectBindingID(FMovieSceneSequenceID SourceSequenceID, FMovieSceneSequenceID TargetSequenceID, const FGuid& TargetGuid, TSharedRef<const FSharedPlaybackState> SharedPlaybackState)
+{
+	ConstructInternal(SourceSequenceID, TargetSequenceID, TargetGuid, SharedPlaybackState->GetHierarchy());
 }
 
 
@@ -98,7 +104,12 @@ void FRelativeObjectBindingID::ConstructInternal(FMovieSceneSequenceID SourceSeq
 
 FRelativeObjectBindingID FFixedObjectBindingID::ConvertToRelative(FMovieSceneSequenceID SourceSequenceID, IMovieScenePlayer& InPlayer) const
 {
-	return FRelativeObjectBindingID(SourceSequenceID, SequenceID, Guid, InPlayer);
+	return ConvertToRelative(SourceSequenceID, InPlayer.GetSharedPlaybackState());
+}
+
+FRelativeObjectBindingID FFixedObjectBindingID::ConvertToRelative(FMovieSceneSequenceID SourceSequenceID, TSharedRef<const FSharedPlaybackState> SharedPlaybackState) const
+{
+	return FRelativeObjectBindingID(SourceSequenceID, SequenceID, Guid, SharedPlaybackState);
 }
 
 FRelativeObjectBindingID FFixedObjectBindingID::ConvertToRelative(FMovieSceneSequenceID SourceSequenceID, const FMovieSceneSequenceHierarchy* Hierarchy) const
@@ -127,6 +138,11 @@ void FMovieSceneObjectBindingID::PostSerialize(const FArchive& Ar)
 #endif // WITH_EDITORONLY_DATA
 
 FMovieSceneSequenceID FMovieSceneObjectBindingID::ResolveSequenceID(FMovieSceneSequenceID LocalSequenceID, IMovieScenePlayer& Player) const
+{
+	return ResolveSequenceID(LocalSequenceID, Player.GetSharedPlaybackState());
+}
+
+FMovieSceneSequenceID FMovieSceneObjectBindingID::ResolveSequenceID(FMovieSceneSequenceID LocalSequenceID, TSharedRef<const FSharedPlaybackState> SharedPlaybackState) const
 {
 	FMovieSceneSequenceID TargetSequenceID = FMovieSceneSequenceID(SequenceID);
 
@@ -135,7 +151,7 @@ FMovieSceneSequenceID FMovieSceneObjectBindingID::ResolveSequenceID(FMovieSceneS
 		return TargetSequenceID;
 	}
 	
-	return UE::MovieScene::ResolveExternalSequenceID(LocalSequenceID, ResolveParentIndex, TargetSequenceID, Player);
+	return UE::MovieScene::ResolveExternalSequenceID(LocalSequenceID, ResolveParentIndex, TargetSequenceID, SharedPlaybackState);
 }
 
 FMovieSceneSequenceID FMovieSceneObjectBindingID::ResolveSequenceID(FMovieSceneSequenceID LocalSequenceID, const FMovieSceneSequenceHierarchy* Hierarchy) const
@@ -151,6 +167,11 @@ FMovieSceneSequenceID FMovieSceneObjectBindingID::ResolveSequenceID(FMovieSceneS
 
 }
 UE::MovieScene::FFixedObjectBindingID FMovieSceneObjectBindingID::ResolveToFixed(FMovieSceneSequenceID RuntimeSequenceID, IMovieScenePlayer& Player) const
+{
+	return ResolveToFixed(RuntimeSequenceID, Player.GetSharedPlaybackState());
+}
+
+UE::MovieScene::FFixedObjectBindingID FMovieSceneObjectBindingID::ResolveToFixed(FMovieSceneSequenceID RuntimeSequenceID, TSharedRef<const FSharedPlaybackState> SharedPlaybackState) const
 {
 	FMovieSceneSequenceID ThisSequenceID = FMovieSceneSequenceID(SequenceID);
 
@@ -160,7 +181,7 @@ UE::MovieScene::FFixedObjectBindingID FMovieSceneObjectBindingID::ResolveToFixed
 	}
 	else if (ensure(RuntimeSequenceID != MovieSceneSequenceID::Invalid))
 	{
-		ThisSequenceID = UE::MovieScene::ResolveExternalSequenceID(RuntimeSequenceID, ResolveParentIndex, ThisSequenceID, Player);
+		ThisSequenceID = UE::MovieScene::ResolveExternalSequenceID(RuntimeSequenceID, ResolveParentIndex, ThisSequenceID, SharedPlaybackState);
 
 		return UE::MovieScene::FFixedObjectBindingID(Guid, ThisSequenceID);
 	}
@@ -170,21 +191,29 @@ UE::MovieScene::FFixedObjectBindingID FMovieSceneObjectBindingID::ResolveToFixed
 
 TArrayView<TWeakObjectPtr<>> FMovieSceneObjectBindingID::ResolveBoundObjects(FMovieSceneSequenceID LocalSequenceID, IMovieScenePlayer& Player) const
 {
-	if (!IsValid())
+	return ResolveBoundObjects(LocalSequenceID, Player.GetSharedPlaybackState());
+}
+
+TArrayView<TWeakObjectPtr<>> FMovieSceneObjectBindingID::ResolveBoundObjects(FMovieSceneSequenceID LocalSequenceID, TSharedRef<const FSharedPlaybackState> SharedPlaybackState) const
+{
+	if (IsValid())
 	{
-		return TArrayView<TWeakObjectPtr<>>();
+		if (FMovieSceneEvaluationState* State = SharedPlaybackState->FindCapability<FMovieSceneEvaluationState>())
+		{
+			FMovieSceneSequenceID ResolvedSequenceID = ResolveSequenceID(LocalSequenceID, SharedPlaybackState);
+			return State->FindBoundObjects(Guid, ResolvedSequenceID, SharedPlaybackState);
+		}
 	}
 
-	FMovieSceneSequenceID ResolvedSequenceID = ResolveSequenceID(LocalSequenceID, Player);
-	return Player.FindBoundObjects(Guid, ResolvedSequenceID);
+	return TArrayView<TWeakObjectPtr<>>();
 }
 
 TArrayView<TWeakObjectPtr<>> FMovieSceneObjectBindingID::ResolveBoundObjects(const UE::MovieScene::FSequenceInstance& SequenceInstance) const
 {
-	if (!IsValid())
+	if (IsValid())
 	{
-		return TArrayView<TWeakObjectPtr<>>();
+		return ResolveBoundObjects(SequenceInstance.GetSequenceID(), SequenceInstance.GetSharedPlaybackState());
 	}
 
-	return ResolveBoundObjects(SequenceInstance.GetSequenceID(), *SequenceInstance.GetPlayer());
+	return TArrayView<TWeakObjectPtr<>>();
 }
diff --git a/Engine/Source/Runtime/MovieScene/Private/MovieScenePlayer.cpp b/Engine/Source/Runtime/MovieScene/Private/MovieScenePlayer.cpp
--- a/Engine/Source/Runtime/MovieScene/Private/MovieScenePlayer.cpp
+++ b/Engine/Source/Runtime/MovieScene/Private/MovieScenePlayer.cpp
@@ -4,6 +4,8 @@
 #include "IMovieScenePlayer.h"
 #include "Evaluation/MovieSceneEvaluationTemplateInstance.h"
 #include "EntitySystem/MovieSceneSequenceInstance.h"
+#include "EntitySystem/MovieSceneEntitySystemLinker.h"
+#include "IMovieScenePlaybackClient.h"
 #include "Misc/ScopeRWLock.h"
 #include "MovieSceneFwd.h"
 #include "MovieSceneSequence.h"
@@ -41,6 +43,8 @@ uint16 FPlayerIndexPlaybackCapability::GetPlayerIndex(TSharedRef<const FSharedPl
 } // namespace MovieScene
 } // namespace UE
 
+UE::MovieScene::TPlaybackCapabilityID<IMovieScenePlaybackClient> IMovieScenePlaybackClient::ID = UE::MovieScene::TPlaybackCapabilityID<IMovieScenePlaybackClient>::Register();
+
 IMovieScenePlayer::IMovieScenePlayer()
 {
 	FWriteScopeLock ScopeLock(UE::MovieScene::GGlobalPlayerRegistryLock);
@@ -105,6 +109,18 @@ void IMovieScenePlayer::ResolveBoundObjects(const FGuid& InBindingId, FMovieScen
 	Sequence.LocateBoundObjects(InBindingId, ResolutionContext, OutObjects);
 }
 
+TArrayView<TWeakObjectPtr<>> IMovieScenePlayer::FindBoundObjects(const FGuid& ObjectBindingID, FMovieSceneSequenceIDRef SequenceID)
+{
+	using namespace UE::MovieScene;
+
+	if (TSharedPtr<const FSharedPlaybackState> SharedPlaybackState = FindSharedPlaybackState())
+	{
+		return State.FindBoundObjects(ObjectBindingID, SequenceID, SharedPlaybackState.ToSharedRef());
+	}
+
+	return TArrayView<TWeakObjectPtr<>>();
+}
+
 void IMovieScenePlayer::InvalidateCachedData()
 {
 	FMovieSceneRootEvaluationTemplateInstance& Template = GetEvaluationTemplate();
@@ -115,3 +131,52 @@ void IMovieScenePlayer::InvalidateCachedData()
 		RootInstance->InvalidateCachedData();
 	}
 }
+
+TSharedPtr<UE::MovieScene::FSharedPlaybackState> IMovieScenePlayer::FindSharedPlaybackState()
+{
+	const UE::MovieScene::FSequenceInstance* RootInstance = GetEvaluationTemplate().GetRootInstance();
+	if (RootInstance)
+	{
+		return RootInstance->GetSharedPlaybackState();
+	}
+	return nullptr;
+}
+
+TSharedRef<UE::MovieScene::FSharedPlaybackState> IMovieScenePlayer::GetSharedPlaybackState()
+{
+	const UE::MovieScene::FSequenceInstance* RootInstance = GetEvaluationTemplate().GetRootInstance();
+	check(RootInstance);
+	return RootInstance->GetSharedPlaybackState();
+}
+
+void IMovieScenePlayer::ResetDirectorInstances()
+{
+	GetEvaluationTemplate().ResetDirectorInstances();
+}
+
+UObject* IMovieScenePlayer::GetOrCreateDirectorInstance(TSharedRef<const UE::MovieScene::FSharedPlaybackState> SharedPlaybackState, FMovieSceneSequenceIDRef SequenceID)
+{
+	return GetEvaluationTemplate().GetOrCreateDirectorInstance(SequenceID, *this);
+}
+
+void IMovieScenePlayer::InitializeRootInstance(TSharedRef<UE::MovieScene::FSharedPlaybackState> NewSharedPlaybackState)
+{
+	using namespace UE::MovieScene;
+
+	NewSharedPlaybackState->AddCapability<FPlayerIndexPlaybackCapability>(UniqueIndex);
+	NewSharedPlaybackState->AddCapabilityRaw(&State);
+	NewSharedPlaybackState->AddCapabilityRaw(&GetSpawnRegister());
+	NewSharedPlaybackState->AddCapabilityRaw((IObjectBindingNotifyPlaybackCapability*)this);
+	NewSharedPlaybackState->AddCapabilityRaw((IStaticBindingOverridesPlaybackCapability*)this);
+	NewSharedPlaybackState->AddCapabilityRaw((ISequenceDirectorPlaybackCapability*)this);
+
+	if (IMovieScenePlaybackClient* PlaybackClient = GetPlaybackClient())
+	{
+		NewSharedPlaybackState->AddCapabilityRaw(PlaybackClient);
+	}
+
+	FInstanceRegistry* InstanceRegistry = NewSharedPlaybackState->GetLinker()->GetInstanceRegistry();
+	FSequenceInstance& RootInstance = InstanceRegistry->MutateInstance(NewSharedPlaybackState->GetRootInstanceHandle());
+	RootInstance.Initialize();
+}
+
diff --git a/Engine/Source/Runtime/MovieScene/Private/MovieScenePossessable.cpp b/Engine/Source/Runtime/MovieScene/Private/MovieScenePossessable.cpp
--- a/Engine/Source/Runtime/MovieScene/Private/MovieScenePossessable.cpp
+++ b/Engine/Source/Runtime/MovieScene/Private/MovieScenePossessable.cpp
@@ -10,19 +10,24 @@
 
 #include UE_INLINE_GENERATED_CPP_BY_NAME(MovieScenePossessable)
 
-bool FMovieScenePossessable::BindSpawnableObject(FMovieSceneSequenceID SequenceID, UObject* Object, IMovieScenePlayer* Player)
+bool FMovieScenePossessable::BindSpawnableObject(FMovieSceneSequenceID SequenceID, UObject* Object, TSharedRef<const UE::MovieScene::FSharedPlaybackState> SharedPlaybackState)
 {
 	TOptional<FMovieSceneSpawnableAnnotation> Spawnable = FMovieSceneSpawnableAnnotation::Find(Object);
 	if (Spawnable.IsSet())
 	{
 		// Check whether the spawnable is underneath the current sequence, if so, we can remap it to a local sequence ID
-		SetSpawnableObjectBindingID(UE::MovieScene::FRelativeObjectBindingID(SequenceID, Spawnable->SequenceID, Spawnable->ObjectBindingID, *Player));
+		SetSpawnableObjectBindingID(UE::MovieScene::FRelativeObjectBindingID(SequenceID, Spawnable->SequenceID, Spawnable->ObjectBindingID, SharedPlaybackState));
 		return true;
 	}
 
 	return false;
 }
 
+bool FMovieScenePossessable::BindSpawnableObject(FMovieSceneSequenceID SequenceID, UObject* Object, IMovieScenePlayer* Player)
+{
+	return BindSpawnableObject(SequenceID, Object, Player->GetSharedPlaybackState());
+}
+
 void FMovieScenePossessable::SetParent(const FGuid& InParentGuid)
 {
 	ParentGuid = InParentGuid;
diff --git a/Engine/Source/Runtime/MovieScene/Private/MovieSceneSequencePlayer.cpp b/Engine/Source/Runtime/MovieScene/Private/MovieSceneSequencePlayer.cpp
--- a/Engine/Source/Runtime/MovieScene/Private/MovieSceneSequencePlayer.cpp
+++ b/Engine/Source/Runtime/MovieScene/Private/MovieSceneSequencePlayer.cpp
@@ -1111,7 +1111,7 @@ void UMovieSceneSequencePlayer::UpdateTimeCursorPosition_Internal(FFrameTime New
 
 			if (SpawnRegister.IsValid())
 			{
-				SpawnRegister->ForgetExternallyOwnedSpawnedObjects(State, *this);
+				SpawnRegister->ForgetExternallyOwnedSpawnedObjects(GetSharedPlaybackState());
 			}
 
 			// Reset the play position, and generate a new range that gets us to the new frame time
@@ -1397,8 +1397,9 @@ void UMovieSceneSequencePlayer::SetIgnorePlaybackReplication(bool bState)
 TArray<UObject*> UMovieSceneSequencePlayer::GetBoundObjects(FMovieSceneObjectBindingID ObjectBinding)
 {
 	TArray<UObject*> Objects;
+	TSharedRef<const UE::MovieScene::FSharedPlaybackState> SharedPlaybackState = GetSharedPlaybackState();
 
-	for (TWeakObjectPtr<> WeakObject : ObjectBinding.ResolveBoundObjects(MovieSceneSequenceID::Root, *this))
+	for (TWeakObjectPtr<> WeakObject : ObjectBinding.ResolveBoundObjects(MovieSceneSequenceID::Root, SharedPlaybackState))
 	{
 		if (UObject* Object = WeakObject.Get())
 		{
@@ -1411,7 +1412,7 @@ TArray<UObject*> UMovieSceneSequencePlayer::GetBoundObjects(FMovieSceneObjectBin
 TArray<FMovieSceneObjectBindingID> UMovieSceneSequencePlayer::GetObjectBindings(UObject* InObject)
 {
 	TArray<FMovieSceneObjectBindingID> Bindings;
-	State.FilterObjectBindings(InObject, *this, &Bindings);
+	State.FilterObjectBindings(InObject, GetSharedPlaybackState(), &Bindings);
 	return Bindings;
 }
 
diff --git a/Engine/Source/Runtime/MovieScene/Private/MovieSceneSpawnRegister.cpp b/Engine/Source/Runtime/MovieScene/Private/MovieSceneSpawnRegister.cpp
--- a/Engine/Source/Runtime/MovieScene/Private/MovieSceneSpawnRegister.cpp
+++ b/Engine/Source/Runtime/MovieScene/Private/MovieSceneSpawnRegister.cpp
@@ -7,6 +7,8 @@
 #include "MovieSceneDynamicBindingInvoker.h"
 #include "MovieSceneSpawnableAnnotation.h"
 
+UE::MovieScene::TPlaybackCapabilityID<FMovieSceneSpawnRegister> FMovieSceneSpawnRegister::ID = UE::MovieScene::TPlaybackCapabilityID<FMovieSceneSpawnRegister>::Register();
+
 TWeakObjectPtr<> FMovieSceneSpawnRegister::FindSpawnedObject(const FGuid& BindingId, FMovieSceneSequenceIDRef TemplateID) const
 {
 	FMovieSceneSpawnRegisterKey Key(TemplateID, BindingId);
@@ -15,7 +17,7 @@ TWeakObjectPtr<> FMovieSceneSpawnRegister::FindSpawnedObject(const FGuid& Bindin
 	return Existing ? Existing->Object : TWeakObjectPtr<>();
 }
 
-UObject* FMovieSceneSpawnRegister::SpawnObject(const FGuid& BindingId, UMovieScene& MovieScene, FMovieSceneSequenceIDRef TemplateID, IMovieScenePlayer& Player)
+UObject* FMovieSceneSpawnRegister::SpawnObject(const FGuid& BindingId, UMovieScene& MovieScene, FMovieSceneSequenceIDRef TemplateID, TSharedRef<const FSharedPlaybackState> SharedPlaybackState)
 {
 	TWeakObjectPtr<> WeakObjectInstance = FindSpawnedObject(BindingId, TemplateID);
 	UObject*         ObjectInstance     = WeakObjectInstance.Get();
@@ -37,7 +39,7 @@ UObject* FMovieSceneSpawnRegister::SpawnObject(const FGuid& BindingId, UMovieSce
 		return nullptr;
 	}
 
-	UMovieSceneSequence* Sequence = Player.GetEvaluationTemplate().GetSequence(TemplateID);
+	UMovieSceneSequence* Sequence = SharedPlaybackState->GetSequence(TemplateID);
 	if (!ensure(Sequence))
 	{
 		return nullptr;
@@ -45,9 +47,11 @@ UObject* FMovieSceneSpawnRegister::SpawnObject(const FGuid& BindingId, UMovieSce
 
 	UObject* SpawnedActor = nullptr;
 	ESpawnOwnership SpawnOwnership = Spawnable->GetSpawnOwnership();
+	IMovieScenePlayer* Player = UE::MovieScene::FPlayerIndexPlaybackCapability::GetPlayer(SharedPlaybackState);
 
 	// See if there is some dynamic binding logic to invoke, otherwise spawn the actor
-	FMovieSceneDynamicBindingResolveResult ResolveResult = FMovieSceneDynamicBindingInvoker::ResolveDynamicBinding(Player, Sequence, TemplateID, BindingId, Spawnable->DynamicBinding);
+	FMovieSceneDynamicBindingResolveResult ResolveResult = FMovieSceneDynamicBindingInvoker::ResolveDynamicBinding(
+			*Player, Sequence, TemplateID, BindingId, Spawnable->DynamicBinding);
 	if (ResolveResult.Object)
 	{
 		SpawnedActor = ResolveResult.Object;
@@ -58,7 +62,7 @@ UObject* FMovieSceneSpawnRegister::SpawnObject(const FGuid& BindingId, UMovieSce
 	}
 	if (!SpawnedActor)
 	{
-		SpawnedActor = SpawnObject(*Spawnable, TemplateID, Player);
+		SpawnedActor = SpawnObject(*Spawnable, TemplateID, SharedPlaybackState);
 	}
 	
 	if (SpawnedActor)
@@ -68,13 +72,16 @@ UObject* FMovieSceneSpawnRegister::SpawnObject(const FGuid& BindingId, UMovieSce
 		FMovieSceneSpawnRegisterKey Key(TemplateID, BindingId);
 		Register.Add(Key, FSpawnedObject(BindingId, *SpawnedActor, SpawnOwnership));
 
-		Player.State.Invalidate(BindingId, TemplateID);
+		if (FMovieSceneEvaluationState* State = SharedPlaybackState->FindCapability<FMovieSceneEvaluationState>())
+		{
+			State->Invalidate(BindingId, TemplateID);
+		}
 	}
 
 	return SpawnedActor;
 }
 
-bool FMovieSceneSpawnRegister::DestroySpawnedObject(const FGuid& BindingId, FMovieSceneSequenceIDRef TemplateID, IMovieScenePlayer& Player)
+bool FMovieSceneSpawnRegister::DestroySpawnedObject(const FGuid& BindingId, FMovieSceneSequenceIDRef TemplateID, TSharedRef<const FSharedPlaybackState> SharedPlaybackState)
 {
 	TGuardValue<bool> CleaningUp(bCleaningUp, true);
 
@@ -90,12 +97,15 @@ bool FMovieSceneSpawnRegister::DestroySpawnedObject(const FGuid& BindingId, FMov
 
 	Register.Remove(Key);
 
-	Player.State.Invalidate(BindingId, TemplateID);
+	if (FMovieSceneEvaluationState* State = SharedPlaybackState->FindCapability<FMovieSceneEvaluationState>())
+	{
+		State->Invalidate(BindingId, TemplateID);
+	}
 
 	return SpawnedObject != nullptr;
 }
 
-void FMovieSceneSpawnRegister::DestroyObjectsByPredicate(IMovieScenePlayer& Player, const TFunctionRef<bool(const FGuid&, ESpawnOwnership, FMovieSceneSequenceIDRef)>& Predicate)
+void FMovieSceneSpawnRegister::DestroyObjectsByPredicate(TSharedRef<const FSharedPlaybackState> SharedPlaybackState, const TFunctionRef<bool(const FGuid&, ESpawnOwnership, FMovieSceneSequenceIDRef)>& Predicate)
 {
 	for (auto It = Register.CreateIterator(); It; ++It)
 	{
@@ -113,37 +123,98 @@ void FMovieSceneSpawnRegister::DestroyObjectsByPredicate(IMovieScenePlayer& Play
 	}
 }
 
-void FMovieSceneSpawnRegister::ForgetExternallyOwnedSpawnedObjects(FMovieSceneEvaluationState& State, IMovieScenePlayer& Player)
+void FMovieSceneSpawnRegister::ForgetExternallyOwnedSpawnedObjects(TSharedRef<const FSharedPlaybackState> SharedPlaybackState)
 {
+	FMovieSceneEvaluationState* State = SharedPlaybackState->FindCapability<FMovieSceneEvaluationState>();
 	for (auto It = Register.CreateIterator(); It; ++It)
 	{
 		if (It.Value().Ownership == ESpawnOwnership::External)
 		{
-			Player.State.Invalidate(It.Key().BindingId, It.Key().TemplateID);
+			if (State)
+			{
+				State->Invalidate(It.Key().BindingId, It.Key().TemplateID);
+			}
 			It.RemoveCurrent();
 		}
 	}
 }
 
-void FMovieSceneSpawnRegister::CleanUp(IMovieScenePlayer& Player)
+void FMovieSceneSpawnRegister::CleanUp(TSharedRef<const FSharedPlaybackState> SharedPlaybackState)
 {
 	TGuardValue<bool> CleaningUp(bCleaningUp, true);
 
-	DestroyObjectsByPredicate(Player, [&](const FGuid&, ESpawnOwnership, FMovieSceneSequenceIDRef){
+	DestroyObjectsByPredicate(SharedPlaybackState, [&](const FGuid&, ESpawnOwnership, FMovieSceneSequenceIDRef){
 		return true;
 	});
 }
 
-void FMovieSceneSpawnRegister::CleanUpSequence(FMovieSceneSequenceIDRef TemplateID, IMovieScenePlayer& Player)
+void FMovieSceneSpawnRegister::CleanUpSequence(FMovieSceneSequenceIDRef TemplateID, TSharedRef<const FSharedPlaybackState> SharedPlaybackState)
 {
-	DestroyObjectsByPredicate(Player, [&](const FGuid&, ESpawnOwnership, FMovieSceneSequenceIDRef ThisTemplateID){
+	DestroyObjectsByPredicate(SharedPlaybackState, [&](const FGuid&, ESpawnOwnership, FMovieSceneSequenceIDRef ThisTemplateID){
 		return ThisTemplateID == TemplateID;
 	});
 }
 
-void FMovieSceneSpawnRegister::OnSequenceExpired(FMovieSceneSequenceIDRef TemplateID, IMovieScenePlayer& Player)
+void FMovieSceneSpawnRegister::OnSequenceExpired(FMovieSceneSequenceIDRef TemplateID, TSharedRef<const FSharedPlaybackState> SharedPlaybackState)
 {
-	DestroyObjectsByPredicate(Player, [&](const FGuid& ObjectId, ESpawnOwnership Ownership, FMovieSceneSequenceIDRef ThisTemplateID){
+	DestroyObjectsByPredicate(SharedPlaybackState, [&](const FGuid& ObjectId, ESpawnOwnership Ownership, FMovieSceneSequenceIDRef ThisTemplateID){
 		return (Ownership == ESpawnOwnership::InnerSequence) && (TemplateID == ThisTemplateID);
 	});
 }
+
+// Deprecated method redirects
+
+UObject* FMovieSceneSpawnRegister::SpawnObject(const FGuid& BindingId, UMovieScene& MovieScene, FMovieSceneSequenceIDRef Template, IMovieScenePlayer& Player)
+{
+	return SpawnObject(BindingId, MovieScene, Template, Player.GetSharedPlaybackState());
+}
+
+bool FMovieSceneSpawnRegister::DestroySpawnedObject(const FGuid& BindingId, FMovieSceneSequenceIDRef TemplateID, IMovieScenePlayer& Player)
+{
+	return DestroySpawnedObject(BindingId, TemplateID, Player.GetSharedPlaybackState());
+}
+
+void FMovieSceneSpawnRegister::DestroyObjectsByPredicate(IMovieScenePlayer& Player, const TFunctionRef<bool(const FGuid&, ESpawnOwnership, FMovieSceneSequenceIDRef)>& Predicate)
+{
+	DestroyObjectsByPredicate(Player.GetSharedPlaybackState(), Predicate);
+}
+
+void FMovieSceneSpawnRegister::ForgetExternallyOwnedSpawnedObjects(FMovieSceneEvaluationState& State, IMovieScenePlayer& Player)
+{
+	ForgetExternallyOwnedSpawnedObjects(Player.GetSharedPlaybackState());
+}
+
+void FMovieSceneSpawnRegister::CleanUp(IMovieScenePlayer& Player)
+{
+	CleanUp(Player.GetSharedPlaybackState());
+}
+
+void FMovieSceneSpawnRegister::CleanUpSequence(FMovieSceneSequenceIDRef TemplateID, IMovieScenePlayer& Player)
+{
+	CleanUpSequence(TemplateID, Player.GetSharedPlaybackState());
+}
+
+void FMovieSceneSpawnRegister::OnSequenceExpired(FMovieSceneSequenceIDRef TemplateID, IMovieScenePlayer& Player)
+{
+	OnSequenceExpired(TemplateID, Player.GetSharedPlaybackState());
+}
+
+#if WITH_EDITOR
+
+void FMovieSceneSpawnRegister::SaveDefaultSpawnableState(FMovieSceneSpawnable& Spawnable, FMovieSceneSequenceIDRef TemplateID, IMovieScenePlayer& Player)
+{
+	SaveDefaultSpawnableState(Spawnable, TemplateID, Player.GetSharedPlaybackState());
+}
+
+void FMovieSceneSpawnRegister::HandleConvertPossessableToSpawnable(UObject* OldObject, IMovieScenePlayer& Player, TOptional<FTransformData>& OutTransformData)
+{
+	HandleConvertPossessableToSpawnable(OldObject, Player.GetSharedPlaybackState(), OutTransformData);
+}
+
+UObject* FMovieSceneSpawnRegister::SpawnObject(FMovieSceneSpawnable& Spawnable, FMovieSceneSequenceIDRef TemplateID, IMovieScenePlayer& Player)
+{
+	return SpawnObject(Spawnable, TemplateID, Player.GetSharedPlaybackState());
+}
+
+#endif
+
diff --git a/Engine/Source/Runtime/MovieScene/Private/MovieSceneSpawnable.cpp b/Engine/Source/Runtime/MovieScene/Private/MovieSceneSpawnable.cpp
--- a/Engine/Source/Runtime/MovieScene/Private/MovieSceneSpawnable.cpp
+++ b/Engine/Source/Runtime/MovieScene/Private/MovieSceneSpawnable.cpp
@@ -8,6 +8,7 @@
 #include "Misc/StringBuilder.h"
 #include "GameFramework/Actor.h"
 #include "Components/StaticMeshComponent.h"
+#include "EntitySystem/MovieSceneSharedPlaybackState.h"
 
 #include UE_INLINE_GENERATED_CPP_BY_NAME(MovieSceneSpawnable)
 
@@ -54,8 +55,23 @@ void FMovieSceneSpawnable::CopyObjectTemplate(UObject& InSourceObject, UMovieSce
 
 FName FMovieSceneSpawnable::GetNetAddressableName(IMovieScenePlayer& Player, FMovieSceneSequenceID SequenceID) const
 {
-	UObject* PlayerObject = Player.AsUObject();
-	if (!PlayerObject)
+	return GetNetAddressableName(Player.GetSharedPlaybackState(), SequenceID);
+}
+
+FName FMovieSceneSpawnable::GetNetAddressableName(TSharedRef<const UE::MovieScene::FSharedPlaybackState> SharedPlaybackState, FMovieSceneSequenceID SequenceID) const
+{
+	UObject* AddressingContext = nullptr;
+
+	if (IMovieScenePlayer* Player = UE::MovieScene::FPlayerIndexPlaybackCapability::GetPlayer(SharedPlaybackState))
+	{
+		AddressingContext = Player->AsUObject();
+	}
+	else if (UObject* PlaybackContext = SharedPlaybackState->GetPlaybackContext())
+	{
+		AddressingContext = PlaybackContext;
+	}
+
+	if (!AddressingContext)
 	{
 		return NAME_None;
 	}
@@ -72,7 +88,7 @@ FName FMovieSceneSpawnable::GetNetAddressableName(IMovieScenePlayer& Player, FMo
 	AddressableName.Appendf(TEXT("_%08X%08X%08X%08X"), Guid.A, Guid.B, Guid.C, Guid.D);
 
 	// Actor / player Name
-	if (AActor* OuterActor = PlayerObject->GetTypedOuter<AActor>())
+	if (AActor* OuterActor = AddressingContext->GetTypedOuter<AActor>())
 	{
 		AddressableName.AppendChar(TEXT('_'));
 		OuterActor->GetFName().AppendString(AddressableName);
@@ -80,7 +96,7 @@ FName FMovieSceneSpawnable::GetNetAddressableName(IMovieScenePlayer& Player, FMo
 	else
 	{
 		AddressableName.AppendChar(TEXT('_'));
-		PlayerObject->GetFName().AppendString(AddressableName);
+		AddressingContext->GetFName().AppendString(AddressableName);
 	}
 
 	return FName(AddressableName.Len(), AddressableName.GetData());
diff --git a/Engine/Source/Runtime/MovieScene/Private/Tests/MovieScenePlaybackCapabilitiesTests.cpp b/Engine/Source/Runtime/MovieScene/Private/Tests/MovieScenePlaybackCapabilitiesTests.cpp
--- a/Engine/Source/Runtime/MovieScene/Private/Tests/MovieScenePlaybackCapabilitiesTests.cpp
+++ b/Engine/Source/Runtime/MovieScene/Private/Tests/MovieScenePlaybackCapabilitiesTests.cpp
@@ -1,6 +1,7 @@
 // Copyright Epic Games, Inc. All Rights Reserved.
 
 #include "CoreTypes.h"
+#include "Misc/CoreMiscDefines.h"
 #include "MovieSceneFwd.h"
 #include "Misc/AutomationTest.h"
 #include "Evaluation/MovieScenePlaybackCapabilities.h"
@@ -57,6 +58,7 @@ struct FTestCapabilityBase
 
 	virtual int32 GetIntValue() { return 0; }
 	virtual FString GetStringValue() { return FString(); }
+	virtual void DoSomething() = 0;
 
 	static TPlaybackCapabilityID<FTestCapabilityBase> ID;
 };
@@ -80,6 +82,7 @@ struct FTestCapabilityDerived : FTestCapabilityBase
 
 	virtual int32 GetIntValue() override { return IntValue; }
 	virtual FString GetStringValue() override { return StringValue; }
+	virtual void DoSomething() override { ++IntValue; }
 };
 
 int32 FTestCapabilityDerived::TimesDestroyedDerived = 0;
@@ -324,5 +327,79 @@ bool FMovieScenePlaybackCapabilitiesMultipleTest::RunTest(const FString& Paramet
 	return true;
 }
 
+IMPLEMENT_SIMPLE_AUTOMATION_TEST(FMovieScenePlaybackCapabilitiesInheritanceTest, 
+		"System.Engine.Sequencer.Capabilities.Inheritance", 
+		EAutomationTestFlags::EditorContext | EAutomationTestFlags::EngineFilter)
+bool FMovieScenePlaybackCapabilitiesInheritanceTest::RunTest(const FString& Parameters)
+{
+	using namespace UE::MovieScene;
+	using namespace UE::MovieScene::Tests;
+
+	struct ICapabilityAccessor
+	{
+		virtual ~ICapabilityAccessor() {}
+		virtual FTestCapabilityBase& GetCapability() = 0;
+	};
+
+	struct FCapabilityAccessor : ICapabilityAccessor
+	{
+		FTestCapabilityDerived Derived;
+
+		FCapabilityAccessor() : Derived(42, TEXT("Life, the universe, everything")) {}
+
+		virtual FTestCapabilityBase& GetCapability() override { return Derived; }
+	};
+
+	FCapabilityAccessor Accessor;
+	{
+		FPlaybackCapabilities Caps;
+
+		Caps.AddCapabilityRaw(FTestCapabilityBase::ID, &Accessor.GetCapability());
+
+		UTEST_TRUE("Found capability", Caps.HasCapability(FTestCapabilityBase::ID));
+		FTestCapabilityBase* Base = Caps.FindCapability(FTestCapabilityBase::ID);
+		UTEST_NOT_NULL("Got capability pointer", Base);
+		UTEST_EQUAL("Same capability pointer", Base, static_cast<FTestCapabilityBase*>(&Accessor.Derived));
+		UTEST_EQUAL("Checking IntValue", Base->GetIntValue(), 42);
+		Base->DoSomething();
+		UTEST_EQUAL("Checking IntValue", Base->GetIntValue(), 43);
+	}
+
+	return true;
+}
+
+IMPLEMENT_SIMPLE_AUTOMATION_TEST(FMovieScenePlaybackCapabilitiesOverwriteTest, 
+		"System.Engine.Sequencer.Capabilities.Overwrite", 
+		EAutomationTestFlags::EditorContext | EAutomationTestFlags::EngineFilter)
+bool FMovieScenePlaybackCapabilitiesOverwriteTest::RunTest(const FString& Parameters)
+{
+	using namespace UE::MovieScene;
+	using namespace UE::MovieScene::Tests;
+
+	FTestSimpleCapability Cap1(1, "First");
+	FTestSimpleCapability Cap2(2, "Second");
+	{
+		FPlaybackCapabilities Caps;
+
+		Caps.AddCapabilityRaw(FTestSimpleCapability::ID, &Cap1);
+
+		FTestSimpleCapability* ActualCap = Caps.FindCapability(FTestSimpleCapability::ID);
+		UTEST_TRUE("Found capability", Caps.HasCapability(FTestSimpleCapability::ID));
+		UTEST_EQUAL("Same capability pointer", ActualCap, &Cap1);
+		UTEST_EQUAL("Checking IntValue", ActualCap->IntValue, 1);
+		UTEST_EQUAL("Checking StringValue", ActualCap->StringValue, TEXT("First"));
+
+		Caps.OverwriteCapabilityRaw(FTestSimpleCapability::ID, &Cap2);
+		
+		ActualCap = Caps.FindCapability(FTestSimpleCapability::ID);
+		UTEST_TRUE("Found capability", Caps.HasCapability(FTestSimpleCapability::ID));
+		UTEST_EQUAL("Same capability pointer", ActualCap, &Cap2);
+		UTEST_EQUAL("Checking IntValue", ActualCap->IntValue, 2);
+		UTEST_EQUAL("Checking StringValue", ActualCap->StringValue, TEXT("Second"));
+	}
+
+	return true;
+}
+
 #endif // WITH_DEV_AUTOMATION_TESTS
 
diff --git a/Engine/Source/Runtime/MovieScene/Public/EntitySystem/MovieSceneInstanceRegistry.h b/Engine/Source/Runtime/MovieScene/Public/EntitySystem/MovieSceneInstanceRegistry.h
--- a/Engine/Source/Runtime/MovieScene/Public/EntitySystem/MovieSceneInstanceRegistry.h
+++ b/Engine/Source/Runtime/MovieScene/Public/EntitySystem/MovieSceneInstanceRegistry.h
@@ -13,6 +13,7 @@
 
 
 class FMovieSceneEntitySystemRunner;
+class UMovieSceneCompiledDataManager;
 class UMovieSceneEntitySystemLinker;
 class UMovieSceneCompiledDataManager;
 
@@ -83,9 +84,14 @@ struct FInstanceRegistry
 		return GetInstance(InstanceHandle).GetContext();
 	}
 
-	MOVIESCENE_API FRootInstanceHandle AllocateRootInstance(IMovieScenePlayer* Player, UMovieSceneSequence& RootSequence, TSharedPtr<FMovieSceneEntitySystemRunner> Runner, UMovieSceneCompiledDataManager* CompiledDataManager);
+	MOVIESCENE_API FRootInstanceHandle AllocateRootInstance(
+			UMovieSceneSequence& InRootSequence,
+			UObject* InPlaybackContext = nullptr,
+			TSharedPtr<FMovieSceneEntitySystemRunner> InRunner = nullptr,
+			UMovieSceneCompiledDataManager* InCompiledDataManager = nullptr);
 
-	MOVIESCENE_API FInstanceHandle AllocateSubInstance(IMovieScenePlayer* Player, FMovieSceneSequenceID SequenceID, FRootInstanceHandle RootInstance, FInstanceHandle ParentInstanceHandle);
+	MOVIESCENE_API FInstanceHandle AllocateSubInstance(
+			FMovieSceneSequenceID SequenceID, FRootInstanceHandle RootInstance, FInstanceHandle ParentInstanceHandle);
 
 	MOVIESCENE_API void DestroyInstance(FInstanceHandle InstanceHandle);
 
diff --git a/Engine/Source/Runtime/MovieScene/Public/EntitySystem/MovieSceneSequenceInstance.h b/Engine/Source/Runtime/MovieScene/Public/EntitySystem/MovieSceneSequenceInstance.h
--- a/Engine/Source/Runtime/MovieScene/Public/EntitySystem/MovieSceneSequenceInstance.h
+++ b/Engine/Source/Runtime/MovieScene/Public/EntitySystem/MovieSceneSequenceInstance.h
@@ -312,13 +312,13 @@ public:
 public:
 
 	/** Constructor for top level sequences */
-	MOVIESCENE_API explicit FSequenceInstance(TSharedRef<FSharedPlaybackState> PlaybackState, FRootInstanceHandle ThisInstanceHandle);
+	MOVIESCENE_API explicit FSequenceInstance(TSharedRef<FSharedPlaybackState> PlaybackState);
 
 	/** Constructor for sub sequences */
-	MOVIESCENE_API explicit FSequenceInstance(TSharedRef<FSharedPlaybackState> PlaybackState, FInstanceHandle ThisInstanceHandle, FInstanceHandle InParentInstanceHandle, FRootInstanceHandle RootInstanceHandle, FMovieSceneSequenceID InSequenceID);
+	MOVIESCENE_API explicit FSequenceInstance(TSharedRef<FSharedPlaybackState> PlaybackState, FInstanceHandle ThisInstanceHandle, FInstanceHandle InParentInstanceHandle, FMovieSceneSequenceID InSequenceID);
 
-	/** Initialization of the sequence */
-	MOVIESCENE_API void Initialize(IMovieScenePlayer* Player);
+	/** Finish initializing this sequence instance */
+	MOVIESCENE_API void Initialize();
 
 	/** Destructor */
 	MOVIESCENE_API ~FSequenceInstance();
@@ -333,32 +333,32 @@ public:
 
 public:
 
-	UE_DEPRECATED(5.4, "Please use the version of this method that doesn't take a Linker")
-	MOVIESCENE_API bool ConditionalRecompile(UMovieSceneEntitySystemLinker* Linker);
-	UE_DEPRECATED(5.4, "Please use the version of this method that doesn't take a Linker")
-	MOVIESCENE_API void DissectContext(UMovieSceneEntitySystemLinker* Linker, const FMovieSceneContext& Context, TArray<TRange<FFrameTime>>& OutDissections);
-	UE_DEPRECATED(5.4, "Please use the version of this method that doesn't take a Linker")
-	MOVIESCENE_API void Start(UMovieSceneEntitySystemLinker* Linker, const FMovieSceneContext& InContext);
-	UE_DEPRECATED(5.4, "Please use the version of this method that doesn't take a Linker")
-	MOVIESCENE_API void PreEvaluation(UMovieSceneEntitySystemLinker* Linker);
-	UE_DEPRECATED(5.4, "Please use the version of this method that doesn't take a Linker")
-	MOVIESCENE_API void Update(UMovieSceneEntitySystemLinker* Linker, const FMovieSceneContext& InContext);
-	UE_DEPRECATED(5.4, "Please use the version of this method that doesn't take a Linker")
-	MOVIESCENE_API bool CanFinishImmediately(UMovieSceneEntitySystemLinker* Linker) const;
-	UE_DEPRECATED(5.4, "Please use the version of this method that doesn't take a Linker")
-	MOVIESCENE_API void Finish(UMovieSceneEntitySystemLinker* Linker);
-	UE_DEPRECATED(5.4, "Please use the version of this method that doesn't take a Linker")
-	MOVIESCENE_API void PostEvaluation(UMovieSceneEntitySystemLinker* Linker);
-	UE_DEPRECATED(5.4, "Please use the version of this method that doesn't take a Linker")
-	MOVIESCENE_API void InvalidateCachedData(UMovieSceneEntitySystemLinker* Linker);
-	UE_DEPRECATED(5.4, "Please use the version of this method that doesn't take a Linker")
-	MOVIESCENE_API void DestroyImmediately(UMovieSceneEntitySystemLinker* Linker);
-	UE_DEPRECATED(5.4, "Please use the version of this method that doesn't take a Linker")
-	MOVIESCENE_API void OverrideRootSequence(UMovieSceneEntitySystemLinker* Linker, FMovieSceneSequenceID NewRootSequenceID);
+	UE_DEPRECATED(5.4, "Please use the version of this method without a Linker parameter")
+	bool ConditionalRecompile(UMovieSceneEntitySystemLinker* Linker);
+	UE_DEPRECATED(5.4, "Please use the version of this method without a Linker parameter")
+	void DissectContext(UMovieSceneEntitySystemLinker* Linker, const FMovieSceneContext& Context, TArray<TRange<FFrameTime>>& OutDissections);
+	UE_DEPRECATED(5.4, "Please use the version of this method without a Linker parameter")
+	void Start(UMovieSceneEntitySystemLinker* Linker, const FMovieSceneContext& InContext);
+	UE_DEPRECATED(5.4, "Please use the version of this method without a Linker parameter")
+	void PreEvaluation(UMovieSceneEntitySystemLinker* Linker);
+	UE_DEPRECATED(5.4, "Please use the version of this method without a Linker parameter")
+	void Update(UMovieSceneEntitySystemLinker* Linker, const FMovieSceneContext& InContext);
+	UE_DEPRECATED(5.4, "Please use the version of this method without a Linker parameter")
+	bool CanFinishImmediately(UMovieSceneEntitySystemLinker* Linker) const;
+	UE_DEPRECATED(5.4, "Please use the version of this method without a Linker parameter")
+	void Finish(UMovieSceneEntitySystemLinker* Linker);
+	UE_DEPRECATED(5.4, "Please use the version of this method without a Linker parameter")
+	void PostEvaluation(UMovieSceneEntitySystemLinker* Linker);
+	UE_DEPRECATED(5.4, "Please use the version of this method without a Linker parameter")
+	void InvalidateCachedData(UMovieSceneEntitySystemLinker* Linker);
+	UE_DEPRECATED(5.4, "Please use the version of this method without a Linker parameter")
+	void DestroyImmediately(UMovieSceneEntitySystemLinker* Linker);
+	UE_DEPRECATED(5.4, "Please use the version of this method without a Linker parameter")
+	void OverrideRootSequence(UMovieSceneEntitySystemLinker* Linker, FMovieSceneSequenceID NewRootSequenceID);
 
 private:
 
-	MOVIESCENE_API void InitializeLegacyEvaluator();
+	void InitializeLegacyEvaluator();
 
 private:
 
@@ -379,9 +379,6 @@ private:
 	/** Playback state shared by the entire sequence hierarchy */
 	TSharedRef<FSharedPlaybackState> SharedPlaybackState;
 
-
-	/** Delegate Binding for when an object binding is invalidated in this instance. */
-	FDelegateHandle OnInvalidateObjectBindingHandle;
 	/** This sequence instances sequence ID, or MovieSceneSequenceID::Root for top-level sequences. */
 	FMovieSceneSequenceID SequenceID;
 	/** When SequenceID != MovieSceneSequenceID::Root, specifies an ID to override as a simulated root. */
@@ -394,6 +391,8 @@ private:
 	FInstanceHandle ParentInstanceHandle;
 	/** This instance's root handle, if it is a sub sequence. */
 	FRootInstanceHandle RootInstanceHandle;
+	/** Flag that indicates whether this instance was initialized */
+	bool bInitialized : 1;
 	/** Flag that is set when this sequence has or (will be) finished. */
 	bool bFinished : 1;
 	/** Flag that is set if this sequence has ever updated. */
diff --git a/Engine/Source/Runtime/MovieScene/Public/EntitySystem/MovieSceneSequenceUpdaters.h b/Engine/Source/Runtime/MovieScene/Public/EntitySystem/MovieSceneSequenceUpdaters.h
--- a/Engine/Source/Runtime/MovieScene/Public/EntitySystem/MovieSceneSequenceUpdaters.h
+++ b/Engine/Source/Runtime/MovieScene/Public/EntitySystem/MovieSceneSequenceUpdaters.h
@@ -23,7 +23,6 @@ namespace MovieScene
 {
 
 struct FInstanceHandle;
-struct FRootInstanceHandle;
 struct FSharedPlaybackState;
 
 enum class ESequenceInstanceUpdateFlags : uint8;
@@ -58,7 +57,7 @@ public:
 	/**
 	 * Called to initialize the flag structure that denotes what functions need to be called on this updater
 	 */
-	virtual void PopulateUpdateFlags(UMovieSceneEntitySystemLinker* InLinker, TSharedRef<const FSharedPlaybackState> SharedPlaybackState, ESequenceInstanceUpdateFlags& OutUpdateFlags) = 0;
+	virtual void PopulateUpdateFlags(TSharedRef<const FSharedPlaybackState> SharedPlaybackState, ESequenceInstanceUpdateFlags& OutUpdateFlags) = 0;
 
 	/**
 	 * Called before any updates to the sequence to allow this updater to dissect the context into smaller ranges that should be evaluated independently for the purpose of ensuring determinism.
@@ -69,7 +68,7 @@ public:
 	 * @param InContext        The total root-level context for the next evaluation to be dissected
 	 * @param OutDissections   (Out) Array to populate with dissected ranges
 	 */
-	virtual void DissectContext(UMovieSceneEntitySystemLinker* InLinker, TSharedRef<const FSharedPlaybackState> SharedPlaybackState, const FMovieSceneContext& InContext, TArray<TRange<FFrameTime>>& OutDissections) = 0;
+	virtual void DissectContext(TSharedRef<const FSharedPlaybackState> SharedPlaybackState, const FMovieSceneContext& InContext, TArray<TRange<FFrameTime>>& OutDissections) = 0;
 
 
 	/**
@@ -80,7 +79,7 @@ public:
 	 * @param InPlayer         The movie scene player instance playing this sequence
 	 * @param InContext        The root-level context for the current evaluation.
 	 */
-	virtual void Start(UMovieSceneEntitySystemLinker* InLinker, FRootInstanceHandle InInstanceHandle, TSharedRef<const FSharedPlaybackState> SharedPlaybackState, const FMovieSceneContext& InContext) = 0;
+	virtual void Start(TSharedRef<const FSharedPlaybackState> SharedPlaybackState, const FMovieSceneContext& InContext) = 0;
 
 
 	/**
@@ -91,7 +90,7 @@ public:
 	 * @param InPlayer         The movie scene player instance playing this sequence
 	 * @param InContext        The root-level context for the current evaluation.
 	 */
-	virtual void Update(UMovieSceneEntitySystemLinker* InLinker, FRootInstanceHandle InInstanceHandle, TSharedRef<const FSharedPlaybackState> SharedPlaybackState, const FMovieSceneContext& InContext) = 0;
+	virtual void Update(TSharedRef<const FSharedPlaybackState> SharedPlaybackState, const FMovieSceneContext& InContext) = 0;
 
 
 	/**
@@ -101,7 +100,7 @@ public:
 	 * @param RootInstanceHandle  The handle to the root instance
 	 * @return                    Whether the instance can be finished immediately
 	 */
-	virtual bool CanFinishImmediately(UMovieSceneEntitySystemLinker* InLinker, FRootInstanceHandle RootInstanceHandle) const = 0;
+	virtual bool CanFinishImmediately(TSharedRef<const FSharedPlaybackState> SharedPlaybackState) const = 0;
 
 
 	/**
@@ -111,7 +110,7 @@ public:
 	 * @param InInstanceHandle The instance handle for the top level sequence instance that this updater belongs to
 	 * @param InPlayer         The movie scene player instance playing this sequence
 	 */
-	virtual void Finish(UMovieSceneEntitySystemLinker* InLinker, FRootInstanceHandle InInstanceHandle, TSharedRef<const FSharedPlaybackState> SharedPlaybackState) = 0;
+	virtual void Finish(TSharedRef<const FSharedPlaybackState> SharedPlaybackState) = 0;
 
 
 	/**
@@ -120,7 +119,7 @@ public:
 	 * @param InLinker         The linker that is evaluating this sequence
 	 * @param InInstanceHandle The instance handle for the top level sequence instance that this updater belongs to
 	 */
-	virtual void InvalidateCachedData(UMovieSceneEntitySystemLinker* InLinker, FRootInstanceHandle InInstanceHandle) = 0;
+	virtual void InvalidateCachedData(TSharedRef<const FSharedPlaybackState> SharedPlaybackState) = 0;
 
 
 	/**
@@ -128,7 +127,7 @@ public:
 	 *
 	 * @param InLinker         The linker that is owns this sequence
 	 */
-	virtual void Destroy(UMovieSceneEntitySystemLinker* Linker) = 0;
+	virtual void Destroy(TSharedRef<const FSharedPlaybackState> SharedPlaybackState) = 0;
 
 
 	/**
@@ -138,7 +137,7 @@ public:
 	 * @param InstanceHandle              The instance handle for the root sequence
 	 * @param NewRootOverrideSequenceID   The new sequence ID to treat as the root
 	 */
-	virtual void OverrideRootSequence(UMovieSceneEntitySystemLinker* InLinker, FRootInstanceHandle InstanceHandle, FMovieSceneSequenceID NewRootOverrideSequenceID) = 0;
+	virtual void OverrideRootSequence(TSharedRef<const FSharedPlaybackState> SharedPlaybackState, FMovieSceneSequenceID NewRootOverrideSequenceID) = 0;
 
 
 	/**
diff --git a/Engine/Source/Runtime/MovieScene/Public/EntitySystem/MovieSceneSharedPlaybackState.h b/Engine/Source/Runtime/MovieScene/Public/EntitySystem/MovieSceneSharedPlaybackState.h
--- a/Engine/Source/Runtime/MovieScene/Public/EntitySystem/MovieSceneSharedPlaybackState.h
+++ b/Engine/Source/Runtime/MovieScene/Public/EntitySystem/MovieSceneSharedPlaybackState.h
@@ -24,6 +24,8 @@ struct FSharedPlaybackStateCreateParams
 {
 	/**
 	 * The playback context in which the root sequence will be evaluated.
+	 *
+	 * Requires that RootInstanceHandle and Runner are also set.
 	 */
 	UObject* PlaybackContext = nullptr;
 
@@ -31,12 +33,16 @@ struct FSharedPlaybackStateCreateParams
 	 * The handle of the root sequence instance, if the created playback state
 	 * is meant to relate to an instance that has also been created inside
 	 * a runner/linker's instance registry.
+	 *
+	 * Requires that PlaybackContext and Runner are also set.
 	 */
 	FRootInstanceHandle RootInstanceHandle;
 
 	/**
 	 * The runner that will be evaluating the sequence that the created playback
 	 * state relates to.
+	 *
+	 * Requires that PlaybackContext and RootInstanceHandle are also set.
 	 */
 	TSharedPtr<FMovieSceneEntitySystemRunner> Runner;
 
@@ -124,7 +130,9 @@ public:
 	template<typename T, typename ...ArgTypes>
 	T& AddCapability(ArgTypes&&... InArgs)
 	{
-		return Capabilities.AddCapability<T>(T::ID, Forward<ArgTypes>(InArgs)...);
+		T& Cap = Capabilities.AddCapability<T>(T::ID, Forward<ArgTypes>(InArgs)...);
+		MaybeInitialize(Cap);
+		return Cap;
 	}
 
 	/**
@@ -133,7 +141,9 @@ public:
 	template<typename T, typename ...ArgTypes>
 	T& AddCapabilityRaw(T* InPointer)
 	{
-		return Capabilities.AddCapabilityRaw<T>(T::ID, InPointer);
+		T& Cap = Capabilities.AddCapabilityRaw<T>(T::ID, InPointer);
+		MaybeInitialize(Cap);
+		return Cap;
 	}
 
 	/**
@@ -142,12 +152,70 @@ public:
 	template<typename T, typename ...ArgTypes>
 	T& AddCapabilityShared(TSharedRef<T> InSharedRef)
 	{
-		return Capabilities.AddCapabilityShared<T>(T::ID, InSharedRef);
+		T& Cap = Capabilities.AddCapabilityShared<T>(T::ID, InSharedRef);
+		MaybeInitialize(Cap);
+		return Cap;
+	}
+
+	/**
+	 * Adds the specified capability on the root sequence as a raw pointer.
+	 * If the capability already exists, it must be stored as a raw pointer and its
+	 * value will be replaced by the new pointer.
+	 */
+	template<typename T, typename ...ArgTypes>
+	T& SetOrAddCapabilityRaw(T* InPointer)
+	{
+		if (HasCapability<T>())
+		{
+			T& Cap = Capabilities.OverwriteCapabilityRaw<T>(T::ID, InPointer);
+			MaybeInitialize(Cap);
+			return Cap;
+		}
+		else
+		{
+			T& Cap = Capabilities.AddCapabilityRaw<T>(T::ID, InPointer);
+			MaybeInitialize(Cap);
+			return Cap;
+		}
+	}
+
+	/**
+	 * Adds the specified capability on the root sequence as a shared pointer.
+	 * If the capability already exists, it must be stored as a shared pointer and its
+	 * value will be replaced by the new pointer.
+	 */
+	template<typename T, typename ...ArgTypes>
+	T& SetOrAddCapabilityShared(TSharedRef<T> InSharedRef)
+	{
+		if (HasCapability<T>())
+		{
+			T& Cap = Capabilities.OverwriteCapabilityShared<T>(T::ID, InSharedRef);
+			MaybeInitialize(Cap);
+			return Cap;
+		}
+		else
+		{
+			T& Cap = Capabilities.AddCapabilityShared<T>(T::ID, InSharedRef);
+			MaybeInitialize(Cap);
+			return Cap;
+		}
 	}
 
 public:
 
-	MOVIESCENE_API void InvalidateCachedData();
+	void InvalidateCachedData();
+
+private:
+
+	template<typename T>
+	void MaybeInitialize(T& Cap)
+	{
+		if constexpr (TPointerIsConvertibleFromTo<T, IPlaybackCapability>::Value)
+		{
+			IPlaybackCapability* InterfacePtr = static_cast<IPlaybackCapability*>(&Cap);
+			InterfacePtr->Initialize(SharedThis(this));
+		}
+	}
 
 private:
 
diff --git a/Engine/Source/Runtime/MovieScene/Public/Evaluation/IMovieScenePlaybackCapability.h b/Engine/Source/Runtime/MovieScene/Public/Evaluation/IMovieScenePlaybackCapability.h
--- a/Engine/Source/Runtime/MovieScene/Public/Evaluation/IMovieScenePlaybackCapability.h
+++ b/Engine/Source/Runtime/MovieScene/Public/Evaluation/IMovieScenePlaybackCapability.h
@@ -5,6 +5,8 @@
 #include "CoreTypes.h"
 #include "Misc/AssertionMacros.h"
 #include "Misc/EnumClassFlags.h"
+#include "Templates/SharedPointer.h"
+#include "EntitySystem/MovieSceneSequenceInstanceHandle.h"
 
 class IMovieScenePlayer;
 class UMovieSceneEntitySystemLinker;
@@ -12,6 +14,9 @@ class UMovieSceneEntitySystemLinker;
 namespace UE::MovieScene
 {
 
+struct FInstanceHandle;
+struct FSharedPlaybackState;
+
 /**
  * An identifier for a playback capability.
  */
@@ -59,6 +64,11 @@ struct IPlaybackCapability
 {
 	virtual ~IPlaybackCapability() {}
 
+	/** Called after this capability has been added to a shared playback state */
+	virtual void Initialize(TSharedRef<const FSharedPlaybackState> Owner) {}
+	/** Called when a new sequence instance has been created and added to the sequence hierarchy */
+	virtual void OnSubInstanceCreated(TSharedRef<const FSharedPlaybackState> Owner, const FInstanceHandle InstanceHandle) {}
+	/** Called when the root sequence is cleaning cached data */
 	virtual void InvalidateCachedData(UMovieSceneEntitySystemLinker* Linker) {}
 };
 
diff --git a/Engine/Source/Runtime/MovieScene/Public/Evaluation/MovieSceneEvaluationState.h b/Engine/Source/Runtime/MovieScene/Public/Evaluation/MovieSceneEvaluationState.h
--- a/Engine/Source/Runtime/MovieScene/Public/Evaluation/MovieSceneEvaluationState.h
+++ b/Engine/Source/Runtime/MovieScene/Public/Evaluation/MovieSceneEvaluationState.h
@@ -6,11 +6,11 @@
 #include "Containers/ArrayView.h"
 #include "Containers/ContainerAllocationPolicies.h"
 #include "Containers/Map.h"
-#include "Containers/Set.h"
-#include "Containers/SparseArray.h"
 #include "CoreMinimal.h"
 #include "Delegates/Delegate.h"
 #include "Delegates/MulticastDelegateBase.h"
+#include "EntitySystem/MovieSceneSharedPlaybackState.h"
+#include "Evaluation/IMovieScenePlaybackCapability.h"
 #include "Evaluation/MovieSceneEvaluationKey.h"
 #include "Evaluation/PersistentEvaluationData.h"
 #include "HAL/Platform.h"
@@ -33,11 +33,52 @@ struct FMovieSceneObjectBindingID;
 struct FSharedPersistentDataKey;
 struct IPersistentEvaluationData;
 
+namespace UE::MovieScene
+{
+
+	struct FInstanceHandle;
+	struct FSharedPlaybackState;
+
+	/**
+	 * Playback capability for being notified of object bindings changing.
+	 */
+	struct MOVIESCENE_API IObjectBindingNotifyPlaybackCapability
+	{
+		static TPlaybackCapabilityID<IObjectBindingNotifyPlaybackCapability> ID;
+
+		/** Called when multiple object bindings have changed. */
+		virtual void NotifyBindingsChanged() {}
+
+		/** Called when a specific object binding has changed. */
+		virtual void NotifyBindingUpdate(const FGuid& InBindingId, FMovieSceneSequenceIDRef InSequenceID, TArrayView<TWeakObjectPtr<>> BoundObjects) {}
+	};
+
+	/**
+	 * Playback capability for storing static object binding overrides.
+	 */
+	struct MOVIESCENE_API IStaticBindingOverridesPlaybackCapability
+	{
+		static TPlaybackCapabilityID<IStaticBindingOverridesPlaybackCapability> ID;
+
+		virtual ~IStaticBindingOverridesPlaybackCapability() {}
+
+		/** Retrieves any override for the given operand */
+		virtual FMovieSceneEvaluationOperand* GetBindingOverride(const FMovieSceneEvaluationOperand& InOperand) = 0;
+		/** Adds an override for the given operand */
+		virtual void AddBindingOverride(const FMovieSceneEvaluationOperand& InOperand, const FMovieSceneEvaluationOperand& InOverrideOperand) = 0;
+		/** Removes any override set for the given operand */
+		virtual void RemoveBindingOverride(const FMovieSceneEvaluationOperand& InOperand) = 0;
+	};
+
+}  // namespace UE::MovieScene
+
 /**
  * Object cache that looks up, resolves, and caches object bindings for a specific sequence
  */
 struct FMovieSceneObjectCache
 {
+	using FSharedPlaybackState = UE::MovieScene::FSharedPlaybackState;
+
 	DECLARE_MULTICAST_DELEGATE_OneParam(FOnBindingInvalidated, const FGuid&);
 	/** Invoked when a binding is either explicitly invalidated (as a result of a spawnable being spawned, or a binding override being added)
 	 *  or when a previously resolved binding becomes invalid */
@@ -51,7 +92,7 @@ struct FMovieSceneObjectCache
 	 * @param Player			The movie scene player that is playing back the sequence
 	 * @return An iterable type of all objects bound to the specified ID.
 	 */
-	MOVIESCENE_API TArrayView<TWeakObjectPtr<>> FindBoundObjects(const FGuid& InBindingID, IMovieScenePlayer& Player);
+	MOVIESCENE_API TArrayView<TWeakObjectPtr<>> FindBoundObjects(const FGuid& InBindingID, TSharedRef<const FSharedPlaybackState> SharedPlaybackState);
 
 	/**
 	 * Find all objects that are bound to the specified binding ID
@@ -68,7 +109,7 @@ struct FMovieSceneObjectCache
 	 * @param InSequence		The sequence that this cache applies to
 	 * @param InSequenceID		The ID of the sequence within the root sequence
 	 */
-	MOVIESCENE_API void SetSequence(UMovieSceneSequence& InSequence, FMovieSceneSequenceIDRef InSequenceID, IMovieScenePlayer& Player);
+	MOVIESCENE_API void SetSequence(UMovieSceneSequence& InSequence, FMovieSceneSequenceIDRef InSequenceID, TSharedRef<const FSharedPlaybackState> SharedPlaybackState);
 
 	/**
 	 * Attempt deduce the posessable or spawnable that relates to the specified object
@@ -78,7 +119,7 @@ struct FMovieSceneObjectCache
 	 * @param Player			The movie scene player that is playing back the sequence
 	 * @return The object's spawnable or possessable GUID, or a zero GUID if it was not found
 	 */
-	MOVIESCENE_API FGuid FindObjectId(UObject& InObject, IMovieScenePlayer& Player);
+	MOVIESCENE_API FGuid FindObjectId(UObject& InObject, TSharedRef<const FSharedPlaybackState> SharedPlaybackState);
 
 	/**
 	* Attempt deduce the posessable or spawnable that relates to the specified object
@@ -88,7 +129,7 @@ struct FMovieSceneObjectCache
 	* @param Player			The movie scene player that is playing back the sequence
 	* @return The object's spawnable or possessable GUID, or a zero GUID if it was not found
 	*/
-	MOVIESCENE_API FGuid FindCachedObjectId(UObject& InObject, IMovieScenePlayer& Player);
+	MOVIESCENE_API FGuid FindCachedObjectId(UObject& InObject, TSharedRef<const FSharedPlaybackState> SharedPlaybackState);
 
 	/**
 	 * Invalidate any object bindings for objects that have been destroyed
@@ -133,7 +174,7 @@ struct FMovieSceneObjectCache
 	/**
 	 * Completely erase all knowledge of, anc caches for all object bindings
 	 */
-	void Clear(IMovieScenePlayer& Player);
+	void Clear(TSharedRef<const FSharedPlaybackState> SharedPlaybackState);
 
 	/**
 	 * Get the sequence that this cache relates to
@@ -152,21 +193,32 @@ struct FMovieSceneObjectCache
 	 * @param Player				The movie scene player that is playing back the sequence
 	 * @param OutBindings			(mandatory) Array to populate with bindings that relate to the object
 	 */
-	void FilterObjectBindings(UObject* PredicateObject, IMovieScenePlayer& Player, TArray<FMovieSceneObjectBindingID>* OutBindings);
+	void FilterObjectBindings(UObject* PredicateObject, TSharedRef<const FSharedPlaybackState> SharedPlaybackState, TArray<FMovieSceneObjectBindingID>* OutBindings);
+
+public:
+
+	// Backwards compatible API, to be deprecated later
+
+	MOVIESCENE_API TArrayView<TWeakObjectPtr<>> FindBoundObjects(const FGuid& InBindingID, IMovieScenePlayer& Player);
+	MOVIESCENE_API void SetSequence(UMovieSceneSequence& InSequence, FMovieSceneSequenceIDRef InSequenceID, IMovieScenePlayer& Player);
+	MOVIESCENE_API FGuid FindObjectId(UObject& InObject, IMovieScenePlayer& Player);
+	MOVIESCENE_API FGuid FindCachedObjectId(UObject& InObject, IMovieScenePlayer& Player);
+	MOVIESCENE_API void Clear(IMovieScenePlayer& Player);
+	MOVIESCENE_API void FilterObjectBindings(UObject* PredicateObject, IMovieScenePlayer& Player, TArray<FMovieSceneObjectBindingID>* OutBindings);
 
 private:
 	/**
 	 * Update the bindings for the specified GUID
 	 *
 	 * @param InGuid			The object binding ID to update bindings for
-	 * @param Player			The movie scene player that is playing back the sequence
+	 * @param SharedPlaybackState  The playback state for the sequence
 	 */
-	void UpdateBindings(const FGuid& InGuid, IMovieScenePlayer& Player);
+	void UpdateBindings(const FGuid& InGuid, TSharedRef<const FSharedPlaybackState> SharedPlaybackState);
 
 	/**
 	 * Handles optional dynamic binding for a given object binding.
 	 */
-	bool ResolveDynamicBinding(const FGuid& InGuid, const FMovieSceneDynamicBinding& DynamicBinding, IMovieScenePlayer& Player, TArray<UObject*, TInlineAllocator<1>>& OutObjects);
+	bool ResolveDynamicBinding(const FGuid& InGuid, const FMovieSceneDynamicBinding& DynamicBinding, TSharedRef<const FSharedPlaybackState> SharedPlaybackState, TArray<UObject*, TInlineAllocator<1>>& OutObjects);
 
 	/**
 	 * Invokes the custom function used to resolve a given dynamic object binding.
@@ -246,15 +298,19 @@ private:
 /**
  * Provides runtime evaluation functions with the ability to look up state from the main game environment
  */
-struct FMovieSceneEvaluationState
+struct FMovieSceneEvaluationState : public UE::MovieScene::IPlaybackCapability
 {
+	using FSharedPlaybackState = UE::MovieScene::FSharedPlaybackState;
+
+	static UE::MovieScene::TPlaybackCapabilityID<FMovieSceneEvaluationState> ID;
+
 	/**
 	 * Assign a sequence to a specific ID
 	 *
 	 * @param InSequenceID		The sequence ID to assign to
 	 * @param InSequence		The sequence to assign
 	 */
-	MOVIESCENE_API void AssignSequence(FMovieSceneSequenceIDRef InSequenceID, UMovieSceneSequence& InSequence, IMovieScenePlayer& Player);
+	MOVIESCENE_API void AssignSequence(FMovieSceneSequenceIDRef InSequenceID, UMovieSceneSequence& InSequence, TSharedRef<const FSharedPlaybackState> SharedPlaybackState);
 
 	/**
 	 * Attempt to locate a sequence from its ID
@@ -278,7 +334,7 @@ struct FMovieSceneEvaluationState
 	 * @param Player			The movie scene player that is playing back the sequence
 	 * @return The object's spawnable or possessable GUID, or a zero GUID if it was not found
 	 */
-	MOVIESCENE_API FGuid FindObjectId(UObject& Object, FMovieSceneSequenceIDRef InSequenceID, IMovieScenePlayer& Player);
+	MOVIESCENE_API FGuid FindObjectId(UObject& Object, FMovieSceneSequenceIDRef InSequenceID, TSharedRef<const FSharedPlaybackState> SharedPlaybackState);
 
 	/**
 	* Attempt deduce the posessable or spawnable that relates to the specified object
@@ -288,7 +344,7 @@ struct FMovieSceneEvaluationState
 	* @param Player			The movie scene player that is playing back the sequence
 	* @return The object's spawnable or possessable GUID, or a zero GUID if it was not found
 	*/
-	MOVIESCENE_API FGuid FindCachedObjectId(UObject& Object, FMovieSceneSequenceIDRef InSequenceID, IMovieScenePlayer& Player);
+	MOVIESCENE_API FGuid FindCachedObjectId(UObject& Object, FMovieSceneSequenceIDRef InSequenceID, TSharedRef<const FSharedPlaybackState> SharedPlaybackState);
 
 	/**
 	 * Filter all the object bindings in this object cache that contain the specified predicate object
@@ -297,7 +353,7 @@ struct FMovieSceneEvaluationState
 	 * @param Player				The movie scene player that is playing back the sequence
 	 * @param OutBindings			(mandatory) Array to populate with bindings that relate to the object
 	 */
-	MOVIESCENE_API void FilterObjectBindings(UObject* PredicateObject, IMovieScenePlayer& Player, TArray<FMovieSceneObjectBindingID>* OutBindings);
+	MOVIESCENE_API void FilterObjectBindings(UObject* PredicateObject, TSharedRef<const FSharedPlaybackState> SharedPlaybackState, TArray<FMovieSceneObjectBindingID>* OutBindings);
 
 	/**
 	 * Find an object cache pertaining to the specified sequence
@@ -342,6 +398,39 @@ struct FMovieSceneEvaluationState
 		return Cache->ObjectCache;
 	}
 
+	/**
+	 * Locate objects bound to the specified object guid, in the specified sequence
+	 * @note: Objects lists are cached internally until they are invalidate.
+	 *
+	 * @param ObjectBindingID 		The object to resolve
+	 * @param SequenceID 			ID of the sequence to resolve for
+	 *
+	 * @return Iterable list of weak object pointers pertaining to the specified GUID
+	 */
+	TArrayView<TWeakObjectPtr<>> FindBoundObjects(const FGuid& ObjectBindingID, FMovieSceneSequenceIDRef SequenceID, TSharedRef<const FSharedPlaybackState> SharedPlaybackState)
+	{
+		FMovieSceneObjectCache* Cache = FindObjectCache(SequenceID);
+		if (Cache)
+		{
+			return Cache->FindBoundObjects(ObjectBindingID, SharedPlaybackState);
+		}
+
+		return TArrayView<TWeakObjectPtr<>>();
+	}
+
+	/**
+	 * Locate objects bound to the specified sequence operand
+	 * @note: Objects lists are cached internally until they are invalidate.
+	 *
+	 * @param Operand 			The movie scene operand to resolve
+	 *
+	 * @return Iterable list of weak object pointers pertaining to the specified GUID
+	 */
+	TArrayView<TWeakObjectPtr<>> FindBoundObjects(const FMovieSceneEvaluationOperand& Operand, TSharedRef<const FSharedPlaybackState> SharedPlaybackState)
+	{
+		return FindBoundObjects(Operand.ObjectBindingID, Operand.SequenceID, SharedPlaybackState);
+	}
+
 	/**
 	 * Invalidate any object caches that may now contain expired objects
 	 */
@@ -372,7 +461,7 @@ struct FMovieSceneEvaluationState
 	/**
 	 * Forcably clear all object caches
 	 */
-	MOVIESCENE_API void ClearObjectCaches(IMovieScenePlayer& Player);
+	MOVIESCENE_API void ClearObjectCaches(TSharedRef<const FSharedPlaybackState> SharedPlaybackState);
 
 	/**
 	 * Get the serial number for this state.
@@ -385,12 +474,33 @@ struct FMovieSceneEvaluationState
 	/** A map of persistent evaluation data mapped by shared evaluation key. Such data can be accessed from anywhere given an operand and a unique identifier. */
 	TMap<FSharedPersistentDataKey, TUniquePtr<IPersistentEvaluationData>> PersistentSharedData;
 
+public:
+
+	/** IPlaybackCapability members */
+	virtual void Initialize(TSharedRef<const FSharedPlaybackState> Owner) override;
+	virtual void OnSubInstanceCreated(TSharedRef<const FSharedPlaybackState> Owner, const UE::MovieScene::FInstanceHandle InstanceHandle) override;
+
+private:
+
+	void RegisterObjectCacheEvents(UMovieSceneEntitySystemLinker* Linker, const UE::MovieScene::FInstanceHandle& InstanceHandle, const FMovieSceneSequenceID SequenceID);
+
+public:
+
+	// Backwards compatible API, to be deprecated later
+
+	MOVIESCENE_API void AssignSequence(FMovieSceneSequenceIDRef InSequenceID, UMovieSceneSequence& InSequence, IMovieScenePlayer& Player);
+	MOVIESCENE_API FGuid FindObjectId(UObject& Object, FMovieSceneSequenceIDRef InSequenceID, IMovieScenePlayer& Player);
+	MOVIESCENE_API FGuid FindCachedObjectId(UObject& Object, FMovieSceneSequenceIDRef InSequenceID, IMovieScenePlayer& Player);
+	MOVIESCENE_API void FilterObjectBindings(UObject* PredicateObject, IMovieScenePlayer& Player, TArray<FMovieSceneObjectBindingID>* OutBindings);
+	MOVIESCENE_API void ClearObjectCaches(IMovieScenePlayer& Player);
+
 private:
 
 	/** Object cache with a last known serial of it */
 	struct FVersionedObjectCache
 	{
 		FMovieSceneObjectCache ObjectCache;
+		FDelegateHandle OnInvalidateObjectBindingHandle;
 		uint32 LastKnownSerial = 0;
 	};
 
diff --git a/Engine/Source/Runtime/MovieScene/Public/Evaluation/MovieSceneEvaluationTemplateInstance.h b/Engine/Source/Runtime/MovieScene/Public/Evaluation/MovieSceneEvaluationTemplateInstance.h
--- a/Engine/Source/Runtime/MovieScene/Public/Evaluation/MovieSceneEvaluationTemplateInstance.h
+++ b/Engine/Source/Runtime/MovieScene/Public/Evaluation/MovieSceneEvaluationTemplateInstance.h
@@ -74,6 +74,10 @@ public:
 
 	MOVIESCENE_API bool IsValid() const;
 
+	MOVIESCENE_API TSharedPtr<UE::MovieScene::FSharedPlaybackState> GetSharedPlaybackState();
+
+	MOVIESCENE_API TSharedPtr<const UE::MovieScene::FSharedPlaybackState> GetSharedPlaybackState() const;
+
 	MOVIESCENE_API UE::MovieScene::FRootInstanceHandle GetRootInstanceHandle() const;
 
 	MOVIESCENE_API UMovieSceneSequence* GetRootSequence() const;
diff --git a/Engine/Source/Runtime/MovieScene/Public/Evaluation/MovieScenePlaybackCapabilities.h b/Engine/Source/Runtime/MovieScene/Public/Evaluation/MovieScenePlaybackCapabilities.h
--- a/Engine/Source/Runtime/MovieScene/Public/Evaluation/MovieScenePlaybackCapabilities.h
+++ b/Engine/Source/Runtime/MovieScene/Public/Evaluation/MovieScenePlaybackCapabilities.h
@@ -5,8 +5,6 @@
 #include "CoreTypes.h"
 #include "Evaluation/IMovieScenePlaybackCapability.h"
 #include "EntitySystem/RelativePtr.h"
-#include "UObject/WeakObjectPtr.h"
-#include "Misc/TVariant.h"
 #include "Templates/AlignmentTemplates.h"
 #include "Templates/PointerIsConvertibleFromTo.h"
 
@@ -297,12 +295,16 @@ protected:
 			}
 		}
 
-		uint64 RequiredSizeof = 0u;
+		// We'll keep track of where all our new capabilities will go, relative to the buffer start
+		TArray<uint16, TInlineAllocator<16>> NewCapabilityOffsets;
+		NewCapabilityOffsets.SetNum(ExistingNum + 1);
 
 		// Compute the required size of our allocation
+		uint64 RequiredSizeof = 0u;
+
 		{
 			// Allocate space for headers
-			RequiredSizeof += (ExistingNum+1) * sizeof(FPlaybackCapabilityHeader);
+			RequiredSizeof += (ExistingNum + 1) * sizeof(FPlaybackCapabilityHeader);
 
 			int32 Index = 0;
 
@@ -310,11 +312,13 @@ protected:
 			for (; Index < NewCapabilityIndex; ++Index)
 			{
 				RequiredSizeof = Align(RequiredSizeof, (uint64)ExistingHeaders[Index].Alignment);
+				NewCapabilityOffsets[Index] = RequiredSizeof;
 				RequiredSizeof += ExistingHeaders[Index].Sizeof;
 			}
 
 			// Count up the size and alignment for the new capability
 			RequiredSizeof = Align(RequiredSizeof, alignof(StorageType));
+			NewCapabilityOffsets[Index] = RequiredSizeof;
 			RequiredSizeof += sizeof(StorageType);
 			++Index;
 
@@ -322,6 +326,7 @@ protected:
 			for (; Index < ExistingNum+1; ++Index)
 			{
 				RequiredSizeof = Align(RequiredSizeof, (uint64)ExistingHeaders[Index-1].Alignment);
+				NewCapabilityOffsets[Index] = RequiredSizeof;
 				RequiredSizeof += ExistingHeaders[Index-1].Sizeof;
 			}
 		}
@@ -330,7 +335,7 @@ protected:
 
 		/// ----------------------------------------
 
-		uint8* OldAllocation = Memory;
+		uint8* OldMemory = Memory;
 
 		// Make a new allocation if necessary
 		const bool bNeedsReallocation = RequiredAlignment > Alignment || RequiredSizeof > Capacity;
@@ -344,76 +349,153 @@ protected:
 		// We now have an extra entry
 		++Num;
 
-		FPlaybackCapabilityHeader* CurrentHeader = reinterpret_cast<FPlaybackCapabilityHeader*>(Memory) + (Num-1);
-		uint8* CapabilityPtr = Memory + RequiredSizeof;
-
-		auto RelocateCapability = [this, ExistingHeaders, OldAllocation, &CurrentHeader, &CapabilityPtr](int32 OldIndex)
+		// We now need to re-arrange memory carefully, going back to front in order to avoid overlaps.
+		// For instance we can only move headers at the end of the whole operation, otherwise they could overwrite
+		// the memory of the first couple capabilities before we've had a chance to move them. So we do:
+		// 
+		// 1) Relocate capabilities from the last one up to where the new one will go
+		// 2) Allocate the new capability
+		// 3) Relocate capabilities from the one just before the new one, down to the first one
+		// 4) Relocate headers from the last one up to where the new one will go
+		// 5) Allocate the new header
+		// 6) Relocate headers from the one just before the new header, down to the first one
+		//
+		// In order to set the new capability pointers (obtained in steps 1-3) on the new headers (obtained in
+		// steps 4-6), we save these pointers in a temporary array.
+
+		const FPlaybackCapabilityHeader* OldHeaders = ExistingHeaders; // Better named variable for the rest of the steps...
+
+		auto RelocateCapability = [this, OldMemory, OldHeaders, &NewCapabilityOffsets](int32 OldIndex, int32 NewIndex)
 		{
-			// Go back
-			CapabilityPtr -= ExistingHeaders[OldIndex].Sizeof;
-			CapabilityPtr = AlignDown(CapabilityPtr, ExistingHeaders[OldIndex].Alignment);
-
-			// Copy the header
-			new (CurrentHeader) FPlaybackCapabilityHeader(ExistingHeaders[OldIndex]);
-			CurrentHeader->Capability = TRelativePtr<void, uint16>(this->Memory, CapabilityPtr);
-
-			--CurrentHeader;
-
-			void* OldCapability = ExistingHeaders[OldIndex].Capability.Resolve(OldAllocation);
-			FMemory::Memmove(CapabilityPtr, OldCapability, ExistingHeaders[OldIndex].Sizeof);
+			void* NewCapabilityPtr = this->Memory + NewCapabilityOffsets[NewIndex];
+			void* OldCapabilityPtr = OldHeaders[OldIndex].Capability.Resolve(OldMemory);
+			FMemory::Memmove(NewCapabilityPtr, OldCapabilityPtr, OldHeaders[OldIndex].Sizeof);
 		};
 
+		// Step 1
 		int32 Index = static_cast<int32>(Num) - 1;
 		for (; Index > NewCapabilityIndex; --Index)
 		{
-			RelocateCapability(Index-1);
+			RelocateCapability(Index - 1, Index);
 		}
 
-		// Make the new entry
-		CapabilityPtr -= sizeof(StorageType);
-		CapabilityPtr = AlignDown(CapabilityPtr, alignof(StorageType));
-
-		StorageType* NewCapabilityPtr = reinterpret_cast<StorageType*>(CapabilityPtr);
-
-		// Allocate the new type
-		new (NewCapabilityPtr) StorageType (Forward<ArgTypes>(InArgs)...);
+		// Step 2
+		{
+			void* NewCapabilityPtr = this->Memory + NewCapabilityOffsets[Index];
+			StorageType* NewCapabilityTypedPtr = reinterpret_cast<StorageType*>(NewCapabilityPtr);
+			new (NewCapabilityTypedPtr) StorageType (Forward<ArgTypes>(InArgs)...);
+		}
+		--Index;
+		
+		// Step 3
+		for (; Index >= 0; --Index)
+		{
+			RelocateCapability(Index, Index);
+		}
+		
+		FPlaybackCapabilityHeader* NewHeaders = reinterpret_cast<FPlaybackCapabilityHeader*>(Memory);
 
-		static_assert(alignof(StorageType) < 0x7F, "Required alignment of capability must fit in 7 bytes");
+		auto RelocateHeader = [this, OldHeaders, NewHeaders, &NewCapabilityOffsets](int32 OldIndex, int32 NewIndex)
+		{
+			const FPlaybackCapabilityHeader& OldHeader(OldHeaders[OldIndex]);
+			FPlaybackCapabilityHeader* NewHeaderPtr = &NewHeaders[NewIndex];
+			new (NewHeaderPtr) FPlaybackCapabilityHeader(OldHeader);
+
+			// Re-create the relative capability pointer, since the offset has changed (it's one capability
+			// further down), and the base pointer may have changed too (if we re-allocated the buffer)
+			void* NewCapabilityPtr = this->Memory + NewCapabilityOffsets[NewIndex];
+			NewHeaderPtr->Capability = TRelativePtr<void, uint16>(this->Memory, NewCapabilityPtr);
+		};
 
-		using FStorageTraits = TPlaybackCapabilityStorageTraits<StorageType, CapabilityType>;
+		// Step 4
+		Index = static_cast<int32>(Num) - 1;
+		for (; Index > NewCapabilityIndex; --Index)
+		{
+			RelocateHeader(Index - 1, Index);
+		}
 
-		// Construct the header
-		const FPlaybackCapabilityHeader* NewHeader = CurrentHeader;
-		CurrentHeader->Capability.Reset(Memory, NewCapabilityPtr);
-		CurrentHeader->Sizeof = sizeof(StorageType);
-		CurrentHeader->Alignment = alignof(StorageType);
-		CurrentHeader->StorageMode = FStorageTraits::GetStorageMode();
-		// The pointer offset is whatever we need to add to the capability pointer in order to get
-		// a pointer to capability type itself (in case the storage type is a sub-class of it)
-		// We only need it if the capability object is stored inline inside our memory buffer.
-		CurrentHeader->PointerOffset = FStorageTraits::ComputePointerOffset(NewCapabilityPtr);
-
-		--CurrentHeader;
+		// Step 5
+		FPlaybackCapabilityHeader* NewHeaderPtr = &NewHeaders[Index];
+		{
+			using FStorageTraits = TPlaybackCapabilityStorageTraits<StorageType, CapabilityType>;
+			static_assert(alignof(StorageType) < 0x7F, "Required alignment of capability must fit in 7 bytes");
+
+			void* NewCapabilityPtr = Memory + NewCapabilityOffsets[Index];
+			StorageType* NewCapabilityTypedPtr = reinterpret_cast<StorageType*>(NewCapabilityPtr);
+
+			new (NewHeaderPtr) FPlaybackCapabilityHeader(); // Reset to default constructor
+															//
+			NewHeaderPtr->Capability.Reset(Memory, NewCapabilityPtr);
+			NewHeaderPtr->Sizeof = sizeof(StorageType);
+			NewHeaderPtr->Alignment = alignof(StorageType);
+			NewHeaderPtr->StorageMode = FStorageTraits::GetStorageMode();
+			// The poPtrinter offset is whatever we need to add to the capability pointer in order to get
+			// a pointer to capability type itself (in case the storage type is a sub-class of it)
+			// We only need it if the capability object is stored inline inside our memory buffer.
+			NewHeaderPtr->PointerOffset = FStorageTraits::ComputePointerOffset(NewCapabilityTypedPtr);
+		}
 		--Index;
 
-		// Relocate the entries that are before the new one
+		// Step 6
 		for (; Index >= 0; --Index)
 		{
-			RelocateCapability(Index);
+			RelocateHeader(Index, Index);
 		}
 
+		/// ----------------------------------------
+
 		// Tidy up the old allocation. We do not call destructors here because we relocated everything.
-		if (bNeedsReallocation && OldAllocation)
+		if (bNeedsReallocation && OldMemory)
 		{
-			FMemory::Free(OldAllocation);
+			FMemory::Free(OldMemory);
 		}
 
 		// Insert the helpers for the new capability.
 		Helpers.Insert(TPlaybackCapabilityHelpers<StorageType>::GetHelpers(), NewCapabilityIndex);
 
-		// Return the new capability pointer. We call the header's Resolve method here because returning NewCapabilityPtr
-		// would return the derived type pointer. We want the base (capability) pointer.
-		return NewHeader->Resolve(Memory);
+		// Return the new capability pointer. We call the header's Resolve method here because returning 
+		// the pointer of the stored capability would return the derived type pointer. 
+		// We want the base (capability) pointer.
+		return NewHeaderPtr->Resolve(Memory);
+	}
+
+	template<typename StorageType, typename CapabilityType, typename ...ArgTypes>
+	bool OverwriteCapability(uint32 CapabilityBit, ArgTypes&&... InArgs)
+	{
+		if (!ensureMsgf(HasCapability(CapabilityBit), TEXT("The given capability does not exist in this container")))
+		{
+			return false;
+		}
+
+		// Get the header for this capability
+		const int32 Index = GetCapabilityIndex(CapabilityBit);
+		const FPlaybackCapabilityHeader& Header = GetHeader(Index);
+
+		// Check that we are overwriting the same storage mode
+		using FStorageTraits = TPlaybackCapabilityStorageTraits<StorageType, CapabilityType>;
+		EPlaybackCapabilityStorageMode GivenStorageMode = FStorageTraits::GetStorageMode();
+		if (!ensureMsgf(GivenStorageMode == Header.StorageMode, TEXT("The given capability storage mode does not match the existing one")))
+		{
+			return false;
+		}
+
+		// Check that the types and alignments match
+		size_t GivenSizeof = sizeof(StorageType);
+		uint16 GivenAlignment = alignof(StorageType);
+		if (!ensureMsgf(
+					(GivenSizeof == Header.Sizeof && GivenAlignment == Header.Alignment),
+					TEXT("The given capability size and alignment do not match the existing one")))
+		{
+			return false;
+		}
+
+		// Don't use the header's Resolve method, we don't want to offset the pointer, we want
+		// the actual pointer to the actual stored capability.
+		void* CapabilityPtr = Header.Capability.Resolve(Memory);
+		StorageType* TypedCapabilityPtr = reinterpret_cast<StorageType*>(CapabilityPtr);
+		*TypedCapabilityPtr = StorageType(Forward<ArgTypes>(InArgs)...);
+
+		return true;
 	}
 
 	const FPlaybackCapabilityHeader& GetHeader(uint8 Index) const
@@ -532,6 +614,31 @@ struct FPlaybackCapabilities : FPlaybackCapabilitiesImpl
 		return DoAddCapability<TSharedPtr<T>, T>(CapabilityID, InSharedRef);
 	}
 
+	/**
+	 * Overwrites an existing capability, stored as a raw pointer on the container.
+	 */
+	template<typename T>
+	T& OverwriteCapabilityRaw(TPlaybackCapabilityID<T> CapabilityID, T* InPointer)
+	{
+		return DoOverwriteCapability<T*, T>(CapabilityID, InPointer);
+	}
+	
+	/**
+	 * Overwrites an existing capability, stored as a shared pointer on the container.
+	 */
+	template<typename T>
+	T& OverwriteCapabilityShared(TPlaybackCapabilityID<T> CapabilityID, TSharedRef<T> InSharedRef)
+	{
+		return DoOverwriteCapability<TSharedPtr<T>, T>(CapabilityID, InSharedRef);
+	}
+
+public:
+
+	/**
+	 * Calls OnSubInstanceCreated on any capability that implements the IPlaybackCapability interface.
+	 */
+	void OnSubInstanceCreated(TSharedRef<const FSharedPlaybackState> Owner, const FInstanceHandle InstanceHandle);
+
 	/**
 	 * Calls InvalidateCacheData on any capability that implements the IPlaybackCapability interface.
 	 */
@@ -539,6 +646,25 @@ struct FPlaybackCapabilities : FPlaybackCapabilitiesImpl
 
 private:
 
+	template<typename Callback, typename ...ArgTypes>
+	void ForEachCapabilityInterface(Callback&& InCallback, ArgTypes&&... InArgs)
+	{
+		TArrayView<const FPlaybackCapabilityHeader> Headers = GetHeaders();
+		for (int32 Index = 0; Index < Headers.Num(); ++Index)
+		{
+			const FPlaybackCapabilityHeader& Header = Headers[Index];
+			const FPlaybackCapabilityHelpers& ThisHelpers = Helpers[Index];
+			check(ThisHelpers.InterfaceCast != nullptr);
+			{
+				void* Ptr = Header.Capability.Resolve(Memory);
+				if (IPlaybackCapability* Interface = (*ThisHelpers.InterfaceCast)(Ptr))
+				{
+					InCallback(*Interface, Forward<ArgTypes>(InArgs)...);
+				}
+			}
+		}
+	}
+
 	template<typename Impl, typename T, typename ...ArgTypes>
 	T& DoAddCapability(TPlaybackCapabilityID<T> CapabilityID, ArgTypes&&... InArgs)
 	{
@@ -547,6 +673,14 @@ private:
 		return Ptr.ResolveChecked<T>();
 	}
 
+	template<typename Impl, typename T, typename ...ArgTypes>
+	T& DoOverwriteCapability(TPlaybackCapabilityID<T> CapabilityID, ArgTypes&&... InArgs)
+	{
+		uint32 CapabilityBit = 1 << CapabilityID.Index;
+		FPlaybackCapabilitiesImpl::OverwriteCapability<Impl, T>(CapabilityBit, Forward<ArgTypes>(InArgs)...);
+		return GetCapabilityChecked(CapabilityID);
+	}
+
 	void Destroy();
 };
 
diff --git a/Engine/Source/Runtime/MovieScene/Public/Evaluation/PreAnimatedState/MovieSceneRestoreStateParams.h b/Engine/Source/Runtime/MovieScene/Public/Evaluation/PreAnimatedState/MovieSceneRestoreStateParams.h
--- a/Engine/Source/Runtime/MovieScene/Public/Evaluation/PreAnimatedState/MovieSceneRestoreStateParams.h
+++ b/Engine/Source/Runtime/MovieScene/Public/Evaluation/PreAnimatedState/MovieSceneRestoreStateParams.h
@@ -12,6 +12,7 @@ namespace UE
 namespace MovieScene
 {
 
+struct FSharedPlaybackState;
 
 /**
  * Parameters that are passed to IMovieScenePreAnimatedToken::RestoreState and IMovieScenePreAnimatedGlobalToken::RestoreState
@@ -34,6 +35,14 @@ struct FRestoreStateParams
 	 * May be null in test harnesses
 	 */
 	MOVIESCENE_API IMovieScenePlayer* GetTerminalPlayer() const;
+
+	/**
+	 * Retrieve a pointer to the shared playback state that is causing this state to be restored.
+	 * @note: may not relate to the same sequence instance that originally cached the state in the case where
+	 * 2 completely different sequences animated the same object at the same time.
+	 * May be null in test harnesses
+	 */
+	MOVIESCENE_API TSharedPtr<const FSharedPlaybackState> GetTerminalPlaybackState() const;
 };
 
 
diff --git a/Engine/Source/Runtime/MovieScene/Public/Evaluation/SequenceDirectorPlaybackCapability.h b/Engine/Source/Runtime/MovieScene/Public/Evaluation/SequenceDirectorPlaybackCapability.h
new file mode 100644
--- /dev/null
+++ b/Engine/Source/Runtime/MovieScene/Public/Evaluation/SequenceDirectorPlaybackCapability.h
@@ -0,0 +1,33 @@
+// Copyright Epic Games, Inc. All Rights Reserved.
+
+#pragma once
+
+#include "CoreMinimal.h"
+#include "Engine/EngineTypes.h"
+#include "Evaluation/MovieScenePlaybackCapabilities.h"
+#include "MovieSceneSequenceID.h"
+
+namespace UE::MovieScene
+{
+
+struct FSharedPlaybackState;
+
+/**
+ * Playback capability for sequences that have a director blueprint.
+ */
+struct MOVIESCENE_API ISequenceDirectorPlaybackCapability
+{
+	/** Playback capability ID */
+	static TPlaybackCapabilityID<ISequenceDirectorPlaybackCapability> ID;
+
+	virtual ~ISequenceDirectorPlaybackCapability() {}
+
+	/** Remove all director blueprint instances */
+	virtual void ResetDirectorInstances() = 0;
+
+	/** Gets a new or existing director blueprint instance for the given root or sub sequence */
+	virtual UObject* GetOrCreateDirectorInstance(TSharedRef<const FSharedPlaybackState> SharedPlaybackState, FMovieSceneSequenceIDRef SequenceID) = 0;
+};
+
+}  // namespace UE::MovieScene
+
diff --git a/Engine/Source/Runtime/MovieScene/Public/IMovieSceneObjectSpawner.h b/Engine/Source/Runtime/MovieScene/Public/IMovieSceneObjectSpawner.h
--- a/Engine/Source/Runtime/MovieScene/Public/IMovieSceneObjectSpawner.h
+++ b/Engine/Source/Runtime/MovieScene/Public/IMovieSceneObjectSpawner.h
@@ -15,6 +15,11 @@ class USequencerSettings;
 class UActorFactory;
 struct FTransformData;
 
+namespace UE::MovieScene
+{
+	struct FSharedPlaybackState;
+}
+
 /** A delegate which will create an object spawner */
 DECLARE_DELEGATE_RetVal(TSharedRef<IMovieSceneObjectSpawner>, FOnCreateMovieSceneObjectSpawner);
 
@@ -51,7 +56,7 @@ public:
 	 * @param Player 		Movie scene player that is ultimately responsible for spawning the object
 	 * @return the spawned object, or nullptr on failure
 	 */
-	virtual UObject* SpawnObject(FMovieSceneSpawnable& Spawnable, FMovieSceneSequenceIDRef TemplateID, IMovieScenePlayer& Player) = 0;
+	virtual UObject* SpawnObject(FMovieSceneSpawnable& Spawnable, FMovieSceneSequenceIDRef TemplateID, TSharedRef<const UE::MovieScene::FSharedPlaybackState> SharedPlaybackState) = 0;
 
 	/**
 	 * Destroy a specific previously spawned object
@@ -116,4 +121,4 @@ public:
 
 	/** Virtual destructor. */
 	virtual ~IMovieSceneObjectSpawner() { }
-};
\ No newline at end of file
+};
diff --git a/Engine/Source/Runtime/MovieScene/Public/IMovieScenePlaybackClient.h b/Engine/Source/Runtime/MovieScene/Public/IMovieScenePlaybackClient.h
--- a/Engine/Source/Runtime/MovieScene/Public/IMovieScenePlaybackClient.h
+++ b/Engine/Source/Runtime/MovieScene/Public/IMovieScenePlaybackClient.h
@@ -5,6 +5,7 @@
 #include "UObject/Interface.h"
 #include "UObject/ObjectMacros.h"
 #include "Engine/EngineTypes.h"
+#include "Evaluation/IMovieScenePlaybackCapability.h"
 #include "MovieSceneSequenceID.h"
 #include "IMovieScenePlaybackClient.generated.h"
 
@@ -25,6 +26,8 @@ class IMovieScenePlaybackClient
 public:
 	GENERATED_BODY()
 
+	static UE::MovieScene::TPlaybackCapabilityID<IMovieScenePlaybackClient> ID;
+
 	/**
 	 * Locate bound objects that relate to the specified binding ID
 	 *
diff --git a/Engine/Source/Runtime/MovieScene/Public/IMovieScenePlayer.h b/Engine/Source/Runtime/MovieScene/Public/IMovieScenePlayer.h
--- a/Engine/Source/Runtime/MovieScene/Public/IMovieScenePlayer.h
+++ b/Engine/Source/Runtime/MovieScene/Public/IMovieScenePlayer.h
@@ -13,6 +13,7 @@
 #include "Evaluation/MovieSceneEvaluationOperand.h"
 #include "Evaluation/MovieSceneEvaluationState.h"
 #include "Evaluation/MovieScenePreAnimatedState.h"
+#include "Evaluation/SequenceDirectorPlaybackCapability.h"
 #include "MovieSceneSpawnRegister.h"
 
 enum class EMovieSceneBuiltInEasing : uint8;
@@ -106,7 +107,10 @@ using EMovieSceneCameraCutParams = FMovieSceneCameraCutParams;
  * Interface for movie scene players
  * Provides information for playback of a movie scene
  */
-class IMovieScenePlayer
+class IMovieScenePlayer 
+	: public UE::MovieScene::IObjectBindingNotifyPlaybackCapability
+	, public UE::MovieScene::IStaticBindingOverridesPlaybackCapability
+	, public UE::MovieScene::ISequenceDirectorPlaybackCapability
 {
 public:
 	MOVIESCENE_API IMovieScenePlayer();
@@ -183,17 +187,51 @@ public:
 	 * @param InSequenceID	The ID of the sequence in which the object binding resides
 	 * @param Objects		The array of objects that were resolved
 	 */
-	virtual void NotifyBindingUpdate(const FGuid& InGuid, FMovieSceneSequenceIDRef InSequenceID, TArrayView<TWeakObjectPtr<>> Objects) { NotifyBindingsChanged(); }
+	virtual void NotifyBindingUpdate(const FGuid& InGuid, FMovieSceneSequenceIDRef InSequenceID, TArrayView<TWeakObjectPtr<>> Objects) override { NotifyBindingsChanged(); }
 
 	/**
-	 * Called to initialize the flag structure that denotes what functions need to be called on this updater
+	 * Called whenever any object bindings have changed
 	 */
-	MOVIESCENE_API virtual void PopulateUpdateFlags(UE::MovieScene::ESequenceInstanceUpdateFlags& OutFlags);
+	virtual void NotifyBindingsChanged() override {}
+
+	/** 
+	 * Retrieves any override for the given operand
+	 */
+	virtual FMovieSceneEvaluationOperand* GetBindingOverride(const FMovieSceneEvaluationOperand& InOperand) override
+	{
+		return BindingOverrides.Find(InOperand);
+	}
+
+	/** 
+	 * Adds an override for the given operand 
+	 */
+	virtual void AddBindingOverride(const FMovieSceneEvaluationOperand& InOperand, const FMovieSceneEvaluationOperand& InOverrideOperand) override
+	{
+		BindingOverrides.Add(InOperand, InOverrideOperand);
+	}
+
+	/** 
+	 * Removes any override set for the given operand
+	 */
+	virtual void RemoveBindingOverride(const FMovieSceneEvaluationOperand& InOperand) override
+	{
+		BindingOverrides.Remove(InOperand);
+	}
 
 	/**
-	 * Called whenever any object bindings have changed
+	 * Remove all director blueprint instances
 	 */
-	virtual void NotifyBindingsChanged() {}
+	MOVIESCENE_API virtual void ResetDirectorInstances() override;
+
+	/**
+	 * Gets a new or existing director blueprint instance for the given root or sub sequence
+	 */
+	MOVIESCENE_API virtual UObject* GetOrCreateDirectorInstance(TSharedRef<const UE::MovieScene::FSharedPlaybackState> SharedPlaybackState, FMovieSceneSequenceIDRef SequenceID) override;
+
+	/**
+	 * Called to initialize the flag structure that denotes what functions need to be called on this updater
+	 */
+	MOVIESCENE_API virtual void PopulateUpdateFlags(UE::MovieScene::ESequenceInstanceUpdateFlags& OutFlags);
 
 	/**
 	 * Access the playback context for this movie scene player
@@ -246,16 +284,7 @@ public:
 	 *
 	 * @return Iterable list of weak object pointers pertaining to the specified GUID
 	 */
-	TArrayView<TWeakObjectPtr<>> FindBoundObjects(const FGuid& ObjectBindingID, FMovieSceneSequenceIDRef SequenceID)
-	{
-		FMovieSceneObjectCache* Cache = State.FindObjectCache(SequenceID);
-		if (Cache)
-		{
-			return Cache->FindBoundObjects(ObjectBindingID, *this);
-		}
-
-		return TArrayView<TWeakObjectPtr<>>();
-	}
+	MOVIESCENE_API TArrayView<TWeakObjectPtr<>> FindBoundObjects(const FGuid& ObjectBindingID, FMovieSceneSequenceIDRef SequenceID);
 
 	/**
 	 * Locate objects bound to the specified sequence operand
@@ -281,7 +310,7 @@ public:
 	 */
 	FGuid FindObjectId(UObject& InObject, FMovieSceneSequenceIDRef SequenceID)
 	{
-		return State.FindObjectId(InObject, SequenceID, *this);
+		return State.FindObjectId(InObject, SequenceID, GetSharedPlaybackState());
 	}
 
 	/**
@@ -295,7 +324,7 @@ public:
 	*/
 	FGuid FindCachedObjectId(UObject& InObject, FMovieSceneSequenceIDRef SequenceID)
 	{
-		return State.FindCachedObjectId(InObject, SequenceID, *this);
+		return State.FindCachedObjectId(InObject, SequenceID, GetSharedPlaybackState());
 	}
 
 	/**
@@ -329,7 +358,7 @@ public:
 	void RestorePreAnimatedState()
 	{
 		PreAnimatedState.RestorePreAnimatedState();
-		State.ClearObjectCaches(*this);
+		State.ClearObjectCaches(GetSharedPlaybackState());
 	}
 
 	/**
@@ -338,7 +367,7 @@ public:
 	void DiscardPreAnimatedState()
 	{
 		PreAnimatedState.DiscardPreAnimatedState();
-		State.ClearObjectCaches(*this);
+		State.ClearObjectCaches(GetSharedPlaybackState());
 	}
 
 
@@ -355,11 +384,29 @@ public:
 
 	MOVIESCENE_API bool IsEvaluating() const;
 
+	/**
+	 * Returns the evaluated sequence instance's shared playback state, if any.
+	 */
+	MOVIESCENE_API TSharedPtr<UE::MovieScene::FSharedPlaybackState> FindSharedPlaybackState();
+
+	/**
+	 * Returns the evaluated sequence instance's shared playback state, asserts if there is none.
+	 */
+	MOVIESCENE_API TSharedRef<UE::MovieScene::FSharedPlaybackState> GetSharedPlaybackState();
+
 	uint16 GetUniqueIndex() const
 	{
 		return UniqueIndex;
 	}
 
+public:
+
+	/**
+	 * Initializes a new root sequence instance and its shared playback state.
+	 * This adds all the player's playback capabilities to the given state.
+	 */
+	void InitializeRootInstance(TSharedRef<UE::MovieScene::FSharedPlaybackState> NewSharedPlaybackState);
+
 public:
 
 	UE_DEPRECATED(5.4, "Camera cut management has moved to UMovieSceneCameraCutTrackInstance")
diff --git a/Engine/Source/Runtime/MovieScene/Public/MovieSceneObjectBindingID.h b/Engine/Source/Runtime/MovieScene/Public/MovieSceneObjectBindingID.h
--- a/Engine/Source/Runtime/MovieScene/Public/MovieSceneObjectBindingID.h
+++ b/Engine/Source/Runtime/MovieScene/Public/MovieSceneObjectBindingID.h
@@ -29,6 +29,7 @@ namespace MovieScene
 
 
 struct FSequenceInstance;
+struct FSharedPlaybackState;
 
 
 /**
@@ -84,6 +85,8 @@ struct FRelativeObjectBindingID
 	 * @param TargetGuid            The GUID of the object binding within TargetSequenceID that represents the object at runtime
 	 * @param Player                The movie scene player that is currently playing the root sequence
 	 */
+	MOVIESCENE_API FRelativeObjectBindingID(FMovieSceneSequenceID SourceSequenceID, FMovieSceneSequenceID TargetSequenceID, const FGuid& TargetGuid, TSharedRef<const FSharedPlaybackState> SharedPlaybackState);
+
 	MOVIESCENE_API FRelativeObjectBindingID(FMovieSceneSequenceID SourceSequenceID, FMovieSceneSequenceID TargetSequenceID, const FGuid& TargetGuid, IMovieScenePlayer& Player);
 
 	/** GUID of the Object Binding (ie, FMovieSceneBinding::GetObjectGuid) */
@@ -143,6 +146,8 @@ public:
 	 * @param InPlayer            The active movie scene player to retrieve a root sequence hierarchy from
 	 * @return A new binding ID that will resolve relative to the specified sequence
 	 */
+	MOVIESCENE_API FRelativeObjectBindingID ConvertToRelative(FMovieSceneSequenceID SourceSequenceID, TSharedRef<const FSharedPlaybackState> SharedPlaybackState) const;
+
 	MOVIESCENE_API FRelativeObjectBindingID ConvertToRelative(FMovieSceneSequenceID SourceSequenceID, IMovieScenePlayer& InPlayer) const;
 
 	/**
@@ -196,6 +201,8 @@ enum class EMovieSceneObjectBindingSpace : uint8
 USTRUCT(BlueprintType, meta=(HasNativeMake))
 struct FMovieSceneObjectBindingID
 {
+	using FSharedPlaybackState = UE::MovieScene::FSharedPlaybackState;
+
 	GENERATED_BODY()
 
 	/**
@@ -324,7 +331,6 @@ public:
 		return FMovieSceneSequenceID(SequenceID);
 	}
 
-
 	/**
 	 * Resolve this binding ID to a fixed object binding ID resolvable from the root sequence
 	 *
@@ -332,9 +338,9 @@ public:
 	 * @param Player                The movie scene player that is currently playing the root sequence
 	 * @return A fixed binding ID whose SequenceID relates to the root sequence hierarchy
 	 */
-	MOVIESCENE_API UE::MovieScene::FFixedObjectBindingID ResolveToFixed(FMovieSceneSequenceID SourceSequenceID, IMovieScenePlayer& Player) const;
-
+	MOVIESCENE_API UE::MovieScene::FFixedObjectBindingID ResolveToFixed(FMovieSceneSequenceID SourceSequenceID, TSharedRef<const FSharedPlaybackState> SharedPlaybackState) const;
 
+	MOVIESCENE_API UE::MovieScene::FFixedObjectBindingID ResolveToFixed(FMovieSceneSequenceID SourceSequenceID, IMovieScenePlayer& Player) const;
 
 	/**
 	 * Resolve the sequence ID for this binding for the root sequence
@@ -343,8 +349,9 @@ public:
 	 * @param Player                The movie scene player that is currently playing the sequence
 	 * @return The sequence ID that holds the target object
 	 */
-	MOVIESCENE_API FMovieSceneSequenceID ResolveSequenceID(FMovieSceneSequenceID SourceSequenceID, IMovieScenePlayer& Player) const;
+	MOVIESCENE_API FMovieSceneSequenceID ResolveSequenceID(FMovieSceneSequenceID SourceSequenceID, TSharedRef<const FSharedPlaybackState> SharedPlaybackState) const;
 
+	MOVIESCENE_API FMovieSceneSequenceID ResolveSequenceID(FMovieSceneSequenceID SourceSequenceID, IMovieScenePlayer& Player) const;
 
 	/**
 	 * Resolve the sequence ID for this binding for the root sequence
@@ -355,7 +362,6 @@ public:
 	 */
 	MOVIESCENE_API FMovieSceneSequenceID ResolveSequenceID(FMovieSceneSequenceID SourceSequenceID, const FMovieSceneSequenceHierarchy* Hierarchy) const;
 
-
 	/**
 	 * Resolve all the bound objects for this binding ID
 	 *
@@ -363,8 +369,9 @@ public:
 	 * @param Hierarchy             The hierachy to use for computation of the relative path. Many be nullptr where Source and Targets are both MovieSceneSequenceID::Root
 	 * @return An array view of all bound objects
 	 */
-	MOVIESCENE_API TArrayView<TWeakObjectPtr<>> ResolveBoundObjects(FMovieSceneSequenceID SourceSequenceID, IMovieScenePlayer& Player) const;
+	MOVIESCENE_API TArrayView<TWeakObjectPtr<>> ResolveBoundObjects(FMovieSceneSequenceID SourceSequenceID, TSharedRef<const FSharedPlaybackState> SharedPlaybackState) const;
 
+	MOVIESCENE_API TArrayView<TWeakObjectPtr<>> ResolveBoundObjects(FMovieSceneSequenceID SourceSequenceID, IMovieScenePlayer& Player) const;
 
 	/**
 	 * Resolve all the bound objects for this binding ID
@@ -432,4 +439,4 @@ template<> struct TStructOpsTypeTraits<FMovieSceneObjectBindingID> : public TStr
 		WithPostSerialize = true,
 	};
 #endif
-};
\ No newline at end of file
+};
diff --git a/Engine/Source/Runtime/MovieScene/Public/MovieScenePossessable.h b/Engine/Source/Runtime/MovieScene/Public/MovieScenePossessable.h
--- a/Engine/Source/Runtime/MovieScene/Public/MovieScenePossessable.h
+++ b/Engine/Source/Runtime/MovieScene/Public/MovieScenePossessable.h
@@ -22,6 +22,11 @@ class UClass;
 class UMovieScene;
 struct FMovieSceneSequenceID;
 
+namespace UE::MovieScene
+{
+	struct FSharedPlaybackState;
+}
+
 /**
  * MovieScenePossessable is a "typed slot" used to allow the MovieScene to control an already-existing object
  */
@@ -161,6 +166,9 @@ public:
 	}
 
 	/* Bind the potential spawnable object to this possessable by setting the ObjectBindingID */
+	MOVIESCENE_API bool BindSpawnableObject(FMovieSceneSequenceID SequenceID, UObject* Object, TSharedRef<const UE::MovieScene::FSharedPlaybackState> SharedPlaybackState);
+
+	UE_DEPRECATED(5.4, "Please use the FSharedPlaybackState version of this method")
 	MOVIESCENE_API bool BindSpawnableObject(FMovieSceneSequenceID SequenceID, UObject* Object, IMovieScenePlayer* Player);
 
 	/* For sorts so we can search quickly by Guid */
diff --git a/Engine/Source/Runtime/MovieScene/Public/MovieSceneSpawnRegister.h b/Engine/Source/Runtime/MovieScene/Public/MovieSceneSpawnRegister.h
--- a/Engine/Source/Runtime/MovieScene/Public/MovieSceneSpawnRegister.h
+++ b/Engine/Source/Runtime/MovieScene/Public/MovieSceneSpawnRegister.h
@@ -4,6 +4,7 @@
 
 #include "Containers/Map.h"
 #include "CoreMinimal.h"
+#include "Evaluation/IMovieScenePlaybackCapability.h"
 #include "HAL/Platform.h"
 #include "HAL/PlatformCrt.h"
 #include "IMovieSceneObjectSpawner.h"
@@ -35,6 +36,11 @@ struct FMovieSceneEvaluationState;
 struct FMovieSceneSpawnable;
 struct FTransformData;
 
+namespace UE::MovieScene
+{
+	struct FSharedPlaybackState;
+}  // namespace UE::MovieScene
+
 /**
  * Class responsible for managing spawnables in a movie scene
  */
@@ -42,6 +48,10 @@ class FMovieSceneSpawnRegister : public TSharedFromThis<FMovieSceneSpawnRegister
 {
 public:
 
+	using FSharedPlaybackState = UE::MovieScene::FSharedPlaybackState;
+
+	static MOVIESCENE_API UE::MovieScene::TPlaybackCapabilityID<FMovieSceneSpawnRegister> ID;
+
 	/**
 	 * Virtual destructor
 	 */
@@ -65,7 +75,7 @@ public:
 	 * @param Player 		Movie scene player that is ultimately responsible for spawning the object
 	 * @return the spawned object, or nullptr on failure
 	 */
-	MOVIESCENE_API UObject* SpawnObject(const FGuid& BindingId, UMovieScene& MovieScene, FMovieSceneSequenceIDRef Template, IMovieScenePlayer& Player);
+	MOVIESCENE_API UObject* SpawnObject(const FGuid& BindingId, UMovieScene& MovieScene, FMovieSceneSequenceIDRef Template, TSharedRef<const FSharedPlaybackState> SharedPlaybackState);
 
 	/**
 	 * Destroy a specific previously spawned object
@@ -76,7 +86,7 @@ public:
 	 *
 	 * @return True if an object was destroyed, false otherwise
 	 */
-	MOVIESCENE_API bool DestroySpawnedObject(const FGuid& BindingId, FMovieSceneSequenceIDRef TemplateID, IMovieScenePlayer& Player);
+	MOVIESCENE_API bool DestroySpawnedObject(const FGuid& BindingId, FMovieSceneSequenceIDRef TemplateID, TSharedRef<const FSharedPlaybackState> SharedPlaybackState);
 
 	/**
 	 * Destroy a specific previously spawned object, where its binding ID and sequence ID is not known.
@@ -92,30 +102,42 @@ public:
 	 * @param Player 		Movie scene player that is ultimately responsible for destroying the objects
 	 * @param Predicate		Predicate used for testing whether an object should be destroyed. Returns true for destruction, false to skip.
 	 */
-	MOVIESCENE_API void DestroyObjectsByPredicate(IMovieScenePlayer& Player, const TFunctionRef<bool(const FGuid&, ESpawnOwnership, FMovieSceneSequenceIDRef)>& Predicate);
+	MOVIESCENE_API void DestroyObjectsByPredicate(TSharedRef<const FSharedPlaybackState> SharedPlaybackState, const TFunctionRef<bool(const FGuid&, ESpawnOwnership, FMovieSceneSequenceIDRef)>& Predicate);
 
 	/**
 	 * Purge any memory of any objects that are considered externally owned
 
 	 * @param Player 		Movie scene player that is ultimately responsible for destroying the objects
 	 */
-	MOVIESCENE_API void ForgetExternallyOwnedSpawnedObjects(FMovieSceneEvaluationState& State, IMovieScenePlayer& Player);
+	MOVIESCENE_API void ForgetExternallyOwnedSpawnedObjects(TSharedRef<const FSharedPlaybackState> SharedPlaybackState);
 
 public:
 
 	/**
 	 * Called to indiscriminately clean up any spawned objects
 	 */
-	MOVIESCENE_API void CleanUp(IMovieScenePlayer& Player);
+	MOVIESCENE_API void CleanUp(TSharedRef<const FSharedPlaybackState> SharedPlaybackState);
 
 	/**
 	 * Called to clean up any non-externally owned spawnables that were spawned from the specified instance
 	 */
-	MOVIESCENE_API void CleanUpSequence(FMovieSceneSequenceIDRef TemplateID, IMovieScenePlayer& Player);
+	MOVIESCENE_API void CleanUpSequence(FMovieSceneSequenceIDRef TemplateID, TSharedRef<const FSharedPlaybackState> SharedPlaybackState);
 
 	/**
 	 * Called when the current time has moved beyond the specified sequence's play range
 	 */
+	MOVIESCENE_API void OnSequenceExpired(FMovieSceneSequenceIDRef TemplateID, TSharedRef<const FSharedPlaybackState> SharedPlaybackState);
+
+public:
+
+	// Backwards compatible API, to be deprecated later
+
+	MOVIESCENE_API UObject* SpawnObject(const FGuid& BindingId, UMovieScene& MovieScene, FMovieSceneSequenceIDRef Template, IMovieScenePlayer& Player);
+	MOVIESCENE_API bool DestroySpawnedObject(const FGuid& BindingId, FMovieSceneSequenceIDRef TemplateID, IMovieScenePlayer& Player);
+	MOVIESCENE_API void DestroyObjectsByPredicate(IMovieScenePlayer& Player, const TFunctionRef<bool(const FGuid&, ESpawnOwnership, FMovieSceneSequenceIDRef)>& Predicate);
+	MOVIESCENE_API void ForgetExternallyOwnedSpawnedObjects(FMovieSceneEvaluationState& State, IMovieScenePlayer& Player);
+	MOVIESCENE_API void CleanUp(IMovieScenePlayer& Player);
+	MOVIESCENE_API void CleanUpSequence(FMovieSceneSequenceIDRef TemplateID, IMovieScenePlayer& Player);
 	MOVIESCENE_API void OnSequenceExpired(FMovieSceneSequenceIDRef TemplateID, IMovieScenePlayer& Player);
 
 #if WITH_EDITOR
@@ -142,7 +164,7 @@ public:
 	/**
 	 * Called to save the default state of the specified spawnable
 	 */
-	virtual void SaveDefaultSpawnableState(FMovieSceneSpawnable& Spawnable, FMovieSceneSequenceIDRef TemplateID, IMovieScenePlayer& Player) {}
+	virtual void SaveDefaultSpawnableState(FMovieSceneSpawnable& Spawnable, FMovieSceneSequenceIDRef TemplateID, TSharedRef<const FSharedPlaybackState> SharedPlaybackState) {}
 
 	/**
 	 * Setup a new spawnable object with some default tracks and keys
@@ -162,8 +184,7 @@ public:
 	 * @param	Player		The current player
 	 * @param	OutTransformData	Transform data that can be used to recreate objects in the same location
 	 */
-	virtual void HandleConvertPossessableToSpawnable(UObject* OldObject, IMovieScenePlayer& Player, TOptional<FTransformData>& OutTransformData) {}
-
+	virtual void HandleConvertPossessableToSpawnable(UObject* OldObject, TSharedRef<const FSharedPlaybackState> SharedPlaybackState, TOptional<FTransformData>& OutTransformData) {}
 
 	/**
 	 * Check whether the specified Spawnable can become a Possessable.
@@ -171,6 +192,15 @@ public:
 	 * @return whether the conversion from Spawnable to Possessable can occur.
 	 */
 	virtual bool CanConvertSpawnableToPossessable(FMovieSceneSpawnable& Spawnable) const { return true; }
+
+public:
+
+	// Backwards compatible API, to be deprecated later
+
+	MOVIESCENE_API void SaveDefaultSpawnableState(FMovieSceneSpawnable& Spawnable, FMovieSceneSequenceIDRef TemplateID, IMovieScenePlayer& Player);
+	MOVIESCENE_API void HandleConvertPossessableToSpawnable(UObject* OldObject, IMovieScenePlayer& Player, TOptional<FTransformData>& OutTransformData);
+	MOVIESCENE_API UObject* SpawnObject(FMovieSceneSpawnable& Spawnable, FMovieSceneSequenceIDRef TemplateID, IMovieScenePlayer& Player);
+
 #endif
 
 protected:
@@ -183,7 +213,7 @@ protected:
 	 * @param Player 		Movie scene player that is ultimately responsible for spawning the object
 	 * @return the spawned object, or nullptr on failure
 	 */
-	virtual UObject* SpawnObject(FMovieSceneSpawnable& Spawnable, FMovieSceneSequenceIDRef TemplateID, IMovieScenePlayer& Player) = 0;
+	virtual UObject* SpawnObject(FMovieSceneSpawnable& Spawnable, FMovieSceneSequenceIDRef TemplateID, TSharedRef<const FSharedPlaybackState> SharedPlaybackState) = 0;
 
 	/**
 	 * Called right before a spawned object with the specified ID and template ID is destroyed
@@ -258,6 +288,6 @@ protected:
 class FNullMovieSceneSpawnRegister : public FMovieSceneSpawnRegister
 {
 public:
-	virtual UObject* SpawnObject(FMovieSceneSpawnable&, FMovieSceneSequenceIDRef, IMovieScenePlayer&) override { check(false); return nullptr; }
+	virtual UObject* SpawnObject(FMovieSceneSpawnable&, FMovieSceneSequenceIDRef, TSharedRef<const FSharedPlaybackState>) override { check(false); return nullptr; }
 	virtual void DestroySpawnedObject(UObject&) override { }
 };
diff --git a/Engine/Source/Runtime/MovieScene/Public/MovieSceneSpawnable.h b/Engine/Source/Runtime/MovieScene/Public/MovieSceneSpawnable.h
--- a/Engine/Source/Runtime/MovieScene/Public/MovieSceneSpawnable.h
+++ b/Engine/Source/Runtime/MovieScene/Public/MovieSceneSpawnable.h
@@ -9,10 +9,14 @@
 
 #include "MovieSceneSpawnable.generated.h"
 
-struct FMovieSceneSequenceID;
-
 class IMovieScenePlayer;
 class UMovieSceneSequence;
+struct FMovieSceneSequenceID;
+
+namespace UE::MovieScene
+{
+	struct FSharedPlaybackState;
+}
 
 UENUM()
 enum class ESpawnOwnership : uint8
@@ -230,6 +234,9 @@ public:
 	/**
 	 * Get the name to use for spawning this object into a networked level
 	 */
+	MOVIESCENE_API FName GetNetAddressableName(TSharedRef<const UE::MovieScene::FSharedPlaybackState> SharedPlaybackState, FMovieSceneSequenceID SequenceID) const;
+
+	UE_DEPRECATED(5.4, "Please use the FSharedPlaybackState version of this method")
 	MOVIESCENE_API FName GetNetAddressableName(IMovieScenePlayer& Player, FMovieSceneSequenceID SequenceID) const;
 
 	/**
diff --git a/Engine/Source/Runtime/MovieSceneTracks/Private/TrackInstances/MovieSceneCameraCutTrackInstance.cpp b/Engine/Source/Runtime/MovieSceneTracks/Private/TrackInstances/MovieSceneCameraCutTrackInstance.cpp
--- a/Engine/Source/Runtime/MovieSceneTracks/Private/TrackInstances/MovieSceneCameraCutTrackInstance.cpp
+++ b/Engine/Source/Runtime/MovieSceneTracks/Private/TrackInstances/MovieSceneCameraCutTrackInstance.cpp
@@ -155,9 +155,9 @@ struct FCameraCutAnimator
 
 public:
 
-	static UObject* FindBoundObject(FMovieSceneObjectBindingID BindingID, FMovieSceneSequenceID SequenceID, IMovieScenePlayer& Player)
+	static UObject* FindBoundObject(FMovieSceneObjectBindingID BindingID, FMovieSceneSequenceIDRef SequenceID, TSharedRef<const FSharedPlaybackState> SharedPlaybackState)
 	{
-		TArrayView<TWeakObjectPtr<>> Objects = BindingID.ResolveBoundObjects(SequenceID, Player);
+		TArrayView<TWeakObjectPtr<>> Objects = BindingID.ResolveBoundObjects(SequenceID, SharedPlaybackState);
 		if (Objects.Num() > 0)
 		{
 			return Objects[0].Get();
@@ -200,7 +200,7 @@ public:
 
 public:
 
-	void AnimatePreRoll(const FPreRollCameraCut& Params, const FMovieSceneSequenceID& SequenceID, IMovieScenePlayer& Player)
+	void AnimatePreRoll(const FPreRollCameraCut& Params, const FSequenceInstance& SequenceInstance)
 	{
 		if (Params.bHasCutTransform)
 		{
@@ -209,7 +209,7 @@ public:
 		}
 		else
 		{
-			UObject* CameraObject = FindBoundObject(Params.CameraBindingID, SequenceID, Player);
+			UObject* CameraObject = FindBoundObject(Params.CameraBindingID, SequenceInstance.GetSequenceID(), SequenceInstance.GetSharedPlaybackState());
 
 			if (AActor* Actor = Cast<AActor>(CameraObject))
 			{
@@ -224,10 +224,9 @@ public:
 			UMovieSceneEntitySystemLinker* Linker,
 			const FSequenceInstance& SequenceInstance)
 	{
-		IMovieScenePlayer& Player = *SequenceInstance.GetPlayer();
 		const FMovieSceneContext& Context = SequenceInstance.GetContext();
 
-		UObject* CameraActor = FindBoundObject(Params.CameraBindingID, Params.OperandSequenceID, Player);
+		UObject* CameraActor = FindBoundObject(Params.CameraBindingID, Params.OperandSequenceID, SequenceInstance.GetSharedPlaybackState());
 		if (Params.CameraBindingID.IsValid() && CameraActor == nullptr)
 		{
 			// We have an unresolved or incorrect binding.
@@ -241,7 +240,7 @@ public:
 		CameraCutParams.bLockPreviousCamera = Params.bLockPreviousCamera;
 
 #if WITH_EDITOR
-		UObject* PreviousCameraActor = FindBoundObject(Params.PreviousCameraBindingID, Params.PreviousOperandSequenceID, Player);
+		UObject* PreviousCameraActor = FindBoundObject(Params.PreviousCameraBindingID, Params.PreviousOperandSequenceID, SequenceInstance.GetSharedPlaybackState());
 		CameraCutParams.PreviousCameraObject = PreviousCameraActor;
 		CameraCutParams.PreviewBlendFactor = Params.PreviewBlendFactor;
 		CameraCutParams.bCanBlend = Params.bCanBlend;
@@ -448,9 +447,7 @@ void UMovieSceneCameraCutTrackInstance::OnAnimate()
 	{
 		FPreRollCameraCut& CameraCutPreRoll = CameraCutPreRolls.Last();
 		const FSequenceInstance& SequenceInstance = InstanceRegistry->GetInstance(CameraCutPreRoll.InstanceHandle);
-		IMovieScenePlayer* Player = SequenceInstance.GetPlayer();
-		const FMovieSceneSequenceID SequenceID = SequenceInstance.GetSequenceID();
-		Animator.AnimatePreRoll(CameraCutPreRoll, SequenceID, *Player);
+		Animator.AnimatePreRoll(CameraCutPreRoll, SequenceInstance);
 	}
 
 	// For now we only support 2 active camera cuts at most (with blending between them).
@@ -579,9 +576,9 @@ void UMovieSceneCameraCutTrackInstance::OnEndUpdateInputs()
 		else
 		{
 			const FSequenceInstance& SequenceInstance = InstanceRegistry->GetInstance(Input.InstanceHandle);
-			IMovieScenePlayer* Player = SequenceInstance.GetPlayer();
+			TSharedRef<const FSharedPlaybackState> SharedPlaybackState = SequenceInstance.GetSharedPlaybackState();
 
-			if (UObject* PlaybackContext = Player->GetPlaybackContext())
+			if (UObject* PlaybackContext = SharedPlaybackState->GetPlaybackContext())
 			{
 				if (UWorld* World = PlaybackContext->GetWorld())
 				{
