Sequencer: invoke dynamic bindings using the BP director capability instead of the player interface

diff --git a/Engine/Source/Runtime/MovieScene/Private/Evaluation/MovieSceneEvaluationState.cpp b/Engine/Source/Runtime/MovieScene/Private/Evaluation/MovieSceneEvaluationState.cpp
--- a/Engine/Source/Runtime/MovieScene/Private/Evaluation/MovieSceneEvaluationState.cpp
+++ b/Engine/Source/Runtime/MovieScene/Private/Evaluation/MovieSceneEvaluationState.cpp
@@ -430,7 +430,7 @@ void FMovieSceneObjectCache::UpdateBindings(const FGuid& InGuid, TSharedRef<cons
 					}
 					else
 					{
-						FMovieSceneDynamicBindingResolveResult ResolveResult = FMovieSceneDynamicBindingInvoker::ResolveDynamicBinding(*Player, Sequence, SequenceID, InGuid, Possessable->DynamicBinding);
+						FMovieSceneDynamicBindingResolveResult ResolveResult = FMovieSceneDynamicBindingInvoker::ResolveDynamicBinding(SharedPlaybackState, Sequence, SequenceID, InGuid, Possessable->DynamicBinding);
 						if (ResolveResult.Object)
 						{
 							ensureMsgf(ResolveResult.bIsPossessedObject, TEXT("Possessables don't support spawnable-type ownership"));
@@ -473,7 +473,7 @@ void FMovieSceneObjectCache::UpdateBindings(const FGuid& InGuid, TSharedRef<cons
 				}
 				else
 				{
-					FMovieSceneDynamicBindingResolveResult ResolveResult = FMovieSceneDynamicBindingInvoker::ResolveDynamicBinding(*Player, Sequence, SequenceID, InGuid, Possessable->DynamicBinding);
+					FMovieSceneDynamicBindingResolveResult ResolveResult = FMovieSceneDynamicBindingInvoker::ResolveDynamicBinding(SharedPlaybackState, Sequence, SequenceID, InGuid, Possessable->DynamicBinding);
 					if (ResolveResult.Object)
 					{
 						ensureMsgf(ResolveResult.bIsPossessedObject, TEXT("Possessables don't support spawnable-type ownership"));
diff --git a/Engine/Source/Runtime/MovieScene/Private/MovieSceneDynamicBindingInvoker.cpp b/Engine/Source/Runtime/MovieScene/Private/MovieSceneDynamicBindingInvoker.cpp
--- a/Engine/Source/Runtime/MovieScene/Private/MovieSceneDynamicBindingInvoker.cpp
+++ b/Engine/Source/Runtime/MovieScene/Private/MovieSceneDynamicBindingInvoker.cpp
@@ -4,13 +4,17 @@
 
 #include "CoreMinimal.h"
 #include "Engine/World.h"
+#include "EntitySystem/MovieSceneSharedPlaybackState.h"
 #include "Evaluation/MovieSceneEvaluationTemplateInstance.h"
-#include "IMovieScenePlayer.h"
+#include "Evaluation/SequenceDirectorPlaybackCapability.h"
 #include "MovieSceneDynamicBinding.h"
+#include "MovieSceneSequence.h"
 #include "UObject/UnrealType.h"
 
-FMovieSceneDynamicBindingResolveResult FMovieSceneDynamicBindingInvoker::ResolveDynamicBinding(IMovieScenePlayer& Player, UMovieSceneSequence* Sequence, const FMovieSceneSequenceID& SequenceID, const FGuid& InGuid, const FMovieSceneDynamicBinding& DynamicBinding)
+FMovieSceneDynamicBindingResolveResult FMovieSceneDynamicBindingInvoker::ResolveDynamicBinding(TSharedRef<const UE::MovieScene::FSharedPlaybackState> SharedPlaybackState, UMovieSceneSequence* Sequence, const FMovieSceneSequenceID& SequenceID, const FGuid& InGuid, const FMovieSceneDynamicBinding& DynamicBinding)
 {
+	using namespace UE::MovieScene;
+
 	if (!ensure(Sequence))
 	{
 		// Sequence is somehow null... fallback to default behavior.
@@ -24,7 +28,15 @@ FMovieSceneDynamicBindingResolveResult FMovieSceneDynamicBindingInvoker::Resolve
 		return FMovieSceneDynamicBindingResolveResult();
 	}
 
-	UObject* DirectorInstance = Player.GetEvaluationTemplate().GetOrCreateDirectorInstance(SequenceID, Player);
+	ISequenceDirectorPlaybackCapability* DirectorCapability = SharedPlaybackState->FindCapability<ISequenceDirectorPlaybackCapability>();
+	if (!ensureMsgf(
+				DirectorCapability, 
+				TEXT("Can't invoke dynamic bindings, the current evaluation context does not have a sequence director playback capability!")))
+	{
+		return FMovieSceneDynamicBindingResolveResult();
+	}
+
+	UObject* DirectorInstance = DirectorCapability->GetOrCreateDirectorInstance(SharedPlaybackState, SequenceID);
 	if (!DirectorInstance)
 	{
 #if !NO_LOGGING
@@ -59,7 +71,7 @@ FMovieSceneDynamicBindingResolveResult FMovieSceneDynamicBindingInvoker::Resolve
 	FMovieSceneDynamicBindingResolveParams ResolveParams;
 	ResolveParams.ObjectBindingID = InGuid;
 	ResolveParams.Sequence = Sequence;
-	ResolveParams.RootSequence = Player.GetEvaluationTemplate().GetRootSequence();
+	ResolveParams.RootSequence = SharedPlaybackState->GetRootSequence();
 	FMovieSceneDynamicBindingResolveResult Result = InvokeDynamicBinding(DirectorInstance, DynamicBinding, ResolveParams);
 
 	return Result;
diff --git a/Engine/Source/Runtime/MovieScene/Private/MovieSceneDynamicBindingInvoker.h b/Engine/Source/Runtime/MovieScene/Private/MovieSceneDynamicBindingInvoker.h
--- a/Engine/Source/Runtime/MovieScene/Private/MovieSceneDynamicBindingInvoker.h
+++ b/Engine/Source/Runtime/MovieScene/Private/MovieSceneDynamicBindingInvoker.h
@@ -12,13 +12,18 @@ class UMovieSceneSequence;
 struct FGuid;
 struct FMovieSceneSequenceID;
 
+namespace UE::MovieScene
+{
+	struct FSharedPlaybackState;
+}
+
 /**
  * Utility class for invoking dynamic binding endpoints.
  */
 struct FMovieSceneDynamicBindingInvoker
 {
 	/** Invoke the dynamic binding, if any, and return the result */
-	static FMovieSceneDynamicBindingResolveResult ResolveDynamicBinding(IMovieScenePlayer& Player, UMovieSceneSequence* Sequence, const FMovieSceneSequenceID& SequenceID, const FGuid& InGuid, const FMovieSceneDynamicBinding& DynamicBinding);
+	static FMovieSceneDynamicBindingResolveResult ResolveDynamicBinding(TSharedRef<const UE::MovieScene::FSharedPlaybackState> SharedPlaybackState, UMovieSceneSequence* Sequence, const FMovieSceneSequenceID& SequenceID, const FGuid& InGuid, const FMovieSceneDynamicBinding& DynamicBinding);
 
 private:
 	static FMovieSceneDynamicBindingResolveResult InvokeDynamicBinding(UObject* DirectorInstance, const FMovieSceneDynamicBinding& DynamicBinding, const FMovieSceneDynamicBindingResolveParams& Params);
diff --git a/Engine/Source/Runtime/MovieScene/Private/MovieSceneSpawnRegister.cpp b/Engine/Source/Runtime/MovieScene/Private/MovieSceneSpawnRegister.cpp
--- a/Engine/Source/Runtime/MovieScene/Private/MovieSceneSpawnRegister.cpp
+++ b/Engine/Source/Runtime/MovieScene/Private/MovieSceneSpawnRegister.cpp
@@ -51,7 +51,7 @@ UObject* FMovieSceneSpawnRegister::SpawnObject(const FGuid& BindingId, UMovieSce
 
 	// See if there is some dynamic binding logic to invoke, otherwise spawn the actor
 	FMovieSceneDynamicBindingResolveResult ResolveResult = FMovieSceneDynamicBindingInvoker::ResolveDynamicBinding(
-			*Player, Sequence, TemplateID, BindingId, Spawnable->DynamicBinding);
+			SharedPlaybackState, Sequence, TemplateID, BindingId, Spawnable->DynamicBinding);
 	if (ResolveResult.Object)
 	{
 		SpawnedActor = ResolveResult.Object;
