Sequencer: Fixed up deprecation warnings caused by code divergence between main and source branch

diff --git a/Engine/Source/Editor/Sequencer/Private/SequencerUtilities.cpp b/Engine/Source/Editor/Sequencer/Private/SequencerUtilities.cpp
--- a/Engine/Source/Editor/Sequencer/Private/SequencerUtilities.cpp
+++ b/Engine/Source/Editor/Sequencer/Private/SequencerUtilities.cpp
@@ -2558,7 +2558,7 @@ FGuid CreateGenericBinding(TSharedRef<ISequencer> Sequencer, UObject& InObject,
 		FMovieScenePossessable* NewPossessable = OwnerMovieScene->FindPossessable(NewID);
 
 		// If the object is a spawnable, try and bind to that first
-		if (!NewPossessable->BindSpawnableObject(Sequencer->GetFocusedTemplateID(), CurrentObject, &Sequencer.Get()))
+		if (!NewPossessable->BindSpawnableObject(Sequencer->GetFocusedTemplateID(), CurrentObject, Sequencer->GetSharedPlaybackState()))
 		{
 			FUniversalObjectLocator Locator;
 			if (!OwnerSequence->MakeLocatorForObject(CurrentObject, Context, Locator) || Locator.IsEmpty())
diff --git a/Engine/Source/Runtime/MovieScene/Private/Evaluation/MovieSceneEvaluationState.cpp b/Engine/Source/Runtime/MovieScene/Private/Evaluation/MovieSceneEvaluationState.cpp
--- a/Engine/Source/Runtime/MovieScene/Private/Evaluation/MovieSceneEvaluationState.cpp
+++ b/Engine/Source/Runtime/MovieScene/Private/Evaluation/MovieSceneEvaluationState.cpp
@@ -442,7 +442,7 @@ void FMovieSceneObjectCache::UpdateBindings(const FGuid& InGuid, TSharedRef<cons
 						}
 						else
 						{
-							Sequence->LocateBoundObjects(InGuid, ResolutionContext, FoundObjects);
+							Sequence->LocateBoundObjects(InGuid, UE::UniversalObjectLocator::FResolveParams(ResolutionContext), FoundObjects);
 						}
 					}
 					
@@ -485,7 +485,7 @@ void FMovieSceneObjectCache::UpdateBindings(const FGuid& InGuid, TSharedRef<cons
 					}
 					else
 					{
-						Sequence->LocateBoundObjects(InGuid, ResolutionContext, FoundObjects);
+						Sequence->LocateBoundObjects(InGuid, UE::UniversalObjectLocator::FResolveParams(ResolutionContext), FoundObjects);
 					}
 				}
 				
