Sequencer: fix crash when adding new component bindings

This replaces the transient player class definition and instantiation with simply creating a transient shared playback state with the minimal playback capabilities we need (namely the object caches in FMovieSceneEvaluationState). The problem being that players do not create a shared playback state until they are initialized... and we don't need to do the heavily lifting of initializing things here.

There are other places with similar transient player classes that can removed in the same way, which we'll do in future changelists.

diff --git a/Engine/Source/Runtime/MovieScene/Private/MovieSceneSequence.cpp b/Engine/Source/Runtime/MovieScene/Private/MovieSceneSequence.cpp
--- a/Engine/Source/Runtime/MovieScene/Private/MovieSceneSequence.cpp
+++ b/Engine/Source/Runtime/MovieScene/Private/MovieSceneSequence.cpp
@@ -338,24 +338,18 @@ UMovieSceneCompiledData* UMovieSceneSequence::GetOrCreateCompiledData()
 
 FGuid UMovieSceneSequence::FindPossessableObjectId(UObject& Object, UObject* Context) const
 {
-	class FTransientPlayer : public IMovieScenePlayer
-	{
-	public:
-		UObject* PlaybackContext;
-		FMovieSceneRootEvaluationTemplateInstance Template;
-		virtual FMovieSceneRootEvaluationTemplateInstance& GetEvaluationTemplate() override { return Template; }
-		virtual void SetViewportSettings(const TMap<FViewportClient*, EMovieSceneViewportParams>& ViewportParamsMap) override {}
-		virtual void GetViewportSettings(TMap<FViewportClient*, EMovieSceneViewportParams>& ViewportParamsMap) const override {}
-		virtual EMovieScenePlayerStatus::Type GetPlaybackStatus() const { return EMovieScenePlayerStatus::Stopped; }
-		virtual void SetPlaybackStatus(EMovieScenePlayerStatus::Type InPlaybackStatus) override {}
-		virtual UObject* GetPlaybackContext() const override { return PlaybackContext;  }
-	} Player;
+	using namespace UE::MovieScene;
 
+	FSharedPlaybackStateCreateParams CreateParams;
+	CreateParams.PlaybackContext = Context;
 	UMovieSceneSequence* ThisSequence = const_cast<UMovieSceneSequence*>(this);
-	Player.State.AssignSequence(MovieSceneSequenceID::Root, *ThisSequence, Player);
-	Player.PlaybackContext = Context;
+	TSharedRef<FSharedPlaybackState> TransientPlaybackState = MakeShared<FSharedPlaybackState>(*ThisSequence, CreateParams);
+
+	FMovieSceneEvaluationState State;
+	TransientPlaybackState->AddCapabilityRaw(&State);
+	State.AssignSequence(MovieSceneSequenceID::Root, *ThisSequence, TransientPlaybackState);
 
-	FGuid ExistingID = Player.FindObjectId(Object, MovieSceneSequenceID::Root);
+	FGuid ExistingID = State.FindObjectId(Object, MovieSceneSequenceID::Root, TransientPlaybackState);
 	return ExistingID;
 }
 
