Sequencer- Fix further issues with splitting/adding new sections to Binding Lifetime track.

diff --git a/Engine/Source/Runtime/MovieScene/Private/Sections/MovieSceneBindingLifetimeSection.cpp b/Engine/Source/Runtime/MovieScene/Private/Sections/MovieSceneBindingLifetimeSection.cpp
--- a/Engine/Source/Runtime/MovieScene/Private/Sections/MovieSceneBindingLifetimeSection.cpp
+++ b/Engine/Source/Runtime/MovieScene/Private/Sections/MovieSceneBindingLifetimeSection.cpp
@@ -82,10 +82,10 @@ void UMovieSceneBindingLifetimeSection::InitialPlacement(const TArray<UMovieScen
 			else
 			{
 				// We're overlapping an at-least partially infinite section here. Split it.
-				UMovieSceneSection* NewSection = OverlappedSection->SplitSection(FQualifiedFrameTime(InStartTime, GetTypedOuter<UMovieScene>()->GetTickResolution()), false);
+				UMovieSceneSection* NewSection = OverlappedSection->SplitSection(FQualifiedFrameTime(SectionRange.GetLowerBound().GetValue(), GetTypedOuter<UMovieScene>()->GetTickResolution()), false);
 				// The new section will be the one on the right and will still have an open upper bound. Resize it to fit the new section.
 				TRange<FFrameNumber> NewSectionRange = NewSection->GetRange();
-				NewSectionRange.SetLowerBoundValue(InStartTime + Duration);
+				NewSectionRange.SetLowerBoundValue(SectionRange.GetUpperBound().GetValue());
 				NewSection->SetRange(NewSectionRange);
 			}
 		}
