Adding a missing call to resolve binding for the IMovieSceneBoundObjectProxy interface

diff --git a/Engine/Source/Runtime/MovieScene/Private/MovieSceneBindingReferences.cpp b/Engine/Source/Runtime/MovieScene/Private/MovieSceneBindingReferences.cpp
--- a/Engine/Source/Runtime/MovieScene/Private/MovieSceneBindingReferences.cpp
+++ b/Engine/Source/Runtime/MovieScene/Private/MovieSceneBindingReferences.cpp
@@ -80,7 +80,7 @@ void FMovieSceneBindingReferences::ResolveBinding(const FGuid& ObjectId, const U
 	for (int32 Index = StartIndex; Index < Num && SortedReferences[Index].ID == ObjectId; ++Index)
 	{
 		UObject* ResolvedObject = SortedReferences[Index].Locator.Resolve(ResolveParams).SyncGet().Object;
-
+		ResolvedObject = UE::MovieScene::FindBoundObjectProxy(ResolvedObject);
 		if (ResolvedObject)
 		{
 			OutObjects.Add(ResolvedObject);
