Fixed missing enum size declaration due to bug in UHT

#rnx

diff --git a/Engine/Source/Runtime/MovieScene/Public/Evaluation/MovieSceneSectionParameters.h b/Engine/Source/Runtime/MovieScene/Public/Evaluation/MovieSceneSectionParameters.h
--- a/Engine/Source/Runtime/MovieScene/Public/Evaluation/MovieSceneSectionParameters.h
+++ b/Engine/Source/Runtime/MovieScene/Public/Evaluation/MovieSceneSectionParameters.h
@@ -12,7 +12,7 @@
  * behaviors for the nested sub-sequence.
  */
 UENUM(BlueprintType)
-enum class EMovieSceneSubSectionFlags
+enum class EMovieSceneSubSectionFlags : uint8
 {
 	None = 0,
 
