Sequencer: Fix subsection trimming so that the start frame offset is transformed to the same space as the trim time

diff --git a/Engine/Source/Runtime/MovieScene/Private/Sections/MovieSceneSubSection.cpp b/Engine/Source/Runtime/MovieScene/Private/Sections/MovieSceneSubSection.cpp
--- a/Engine/Source/Runtime/MovieScene/Private/Sections/MovieSceneSubSection.cpp
+++ b/Engine/Source/Runtime/MovieScene/Private/Sections/MovieSceneSubSection.cpp
@@ -407,8 +407,9 @@ void UMovieSceneSubSection::TrimSection( FQualifiedFrameTime TrimTime, bool bTri
 		// The new first loop start offset is where the trim time fell inside the sub-sequence (this time is already
 		// normalized in the case of looping sub-sequences).
 		FFrameTime LocalTrimTime = OuterToInnerTransform().TransformTime(LocalTickResolutionTrimTime);
-		FFrameNumber NewStartOffset = LocalTrimTime.FrameNumber - Parameters.StartFrameOffset;
-		
+		FFrameTime LocalStartFrameOffset = OuterToInnerTransform().TransformTime(Parameters.StartFrameOffset);
+		FFrameNumber NewStartOffset = LocalTrimTime.FrameNumber - LocalStartFrameOffset.FrameNumber;
+
 		// Make sure we don't have negative offsets (this shouldn't happen, though).
 		NewStartOffset = FMath::Max(FFrameNumber(0), NewStartOffset);
 		
