Sequencer: handle play-rate in replicated playback synchronization

diff --git a/Engine/Source/Runtime/MovieScene/Private/MovieSceneSequencePlayer.cpp b/Engine/Source/Runtime/MovieScene/Private/MovieSceneSequencePlayer.cpp
--- a/Engine/Source/Runtime/MovieScene/Private/MovieSceneSequencePlayer.cpp
+++ b/Engine/Source/Runtime/MovieScene/Private/MovieSceneSequencePlayer.cpp
@@ -1435,7 +1435,7 @@ FFrameTime UMovieSceneSequencePlayer::UpdateServerTimeSamples()
 	// We need to reproject the samples to the current wall-clock time, based on when they were taken
 	const double CurrentWallClock = FPlatformTime::Seconds();
 	const double Lifetime         = CurrentWallClock - float(GSequencerMaxSmoothedNetSyncSampleAge) / 1000.f;
-	const float PlaybackMultiplier = bReversePlayback ? -1.f : 1.f;
+	const float PlaybackMultiplier = bReversePlayback ? -PlaybackSettings.PlayRate : PlaybackSettings.PlayRate;
 
 	float TimeDilation = 1.0f;
 	if (const UWorld* World = GetPlaybackWorld())
@@ -1743,6 +1743,8 @@ void UMovieSceneSequencePlayer::PostNetReceive()
 	const bool bHasChangedStatus  = NetSyncProps.LastKnownStatus   != Status;
 	const bool bHasChangedTime    = NetSyncProps.LastKnownPosition != PlayPosition.GetCurrentPosition();
 
+	const float PlayRate = PlaybackSettings.PlayRate;
+
 	float TimeDilation = 1.0f;
 	if (const UWorld* World = GetPlaybackWorld())
 	{
@@ -1753,11 +1755,11 @@ void UMovieSceneSequencePlayer::PostNetReceive()
 	}
 
 	const float PingMs            = GetPing();
-	const FFrameTime PingLag      = (PingMs/1000.f) * PlayPosition.GetInputRate() * TimeDilation;
+	const FFrameTime PingLag      = (PingMs/1000.f) * PlayPosition.GetInputRate() * PlayRate * TimeDilation;
 	//const FFrameTime LagThreshold = 0.2f * PlayPosition.GetInputRate();
 	//const FFrameTime LagDisparity = FMath::Abs(PlayPosition.GetCurrentPosition() - NetSyncProps.LastKnownPosition);
 
-	const FFrameTime LagThreshold = (GSequencerNetSyncThresholdMS * 0.001f) * PlayPosition.GetInputRate() * TimeDilation;
+	const FFrameTime LagThreshold = (GSequencerNetSyncThresholdMS * 0.001f) * PlayPosition.GetInputRate() * PlayRate * TimeDilation;
 
 	if (!bHasChangedStatus && !bHasChangedTime)
 	{
@@ -1856,7 +1858,10 @@ void UMovieSceneSequencePlayer::UpdateNetworkSync()
 	// Only process net playback synchronization if we are still Playing.
 	if (Status == EMovieScenePlayerStatus::Playing)
 	{
-		const float PingMs            = GetPing();	
+		const float PingMs = GetPing();
+
+		const float PlayRate = PlaybackSettings.PlayRate;
+
 		float TimeDilation = 1.0f;
 		if (const UWorld* World = GetPlaybackWorld())
 		{
@@ -1865,8 +1870,9 @@ void UMovieSceneSequencePlayer::UpdateNetworkSync()
 				TimeDilation = WorldSettings->GetEffectiveTimeDilation();
 			}
 		}
-		const FFrameTime PingLag      = (PingMs/1000.f) * PlayPosition.GetInputRate() * TimeDilation;
-		const FFrameTime LagThreshold = (GSequencerNetSyncThresholdMS * 0.001f) * PlayPosition.GetInputRate() * TimeDilation;
+
+		const FFrameTime PingLag      = (PingMs/1000.f) * PlayPosition.GetInputRate() * PlayRate * TimeDilation;
+		const FFrameTime LagThreshold = (GSequencerNetSyncThresholdMS * 0.001f) * PlayPosition.GetInputRate() * PlayRate * TimeDilation;
 		
 		// When the server has looped back to the start but a client is near the end (and is thus about to loop), we don't want to forcibly synchronize the time unless
 		// the *real* difference in time is above the threshold. We compute the real-time difference by adding SequenceDuration*LoopCountDifference to the server position:
@@ -1879,7 +1885,7 @@ void UMovieSceneSequencePlayer::UpdateNetworkSync()
 		//			   OffsetServerTime = srv_time + FrameDuration*LoopOffset = 1 + 20*1 = 21
 		//			   Difference = 21 - 18 = 3 frames
 		const int32        LoopOffset       = (NetSyncProps.LastKnownNumLoops - CurrentNumLoops) * (bReversePlayback ? -1 : 1);
-		const FFrameTime   OffsetServerTime = (NetSyncProps.LastKnownPosition + PingLag) + GetFrameDuration()*LoopOffset;
+		const FFrameTime   OffsetServerTime = (NetSyncProps.LastKnownPosition + PingLag) + GetFrameDuration() * LoopOffset;
 
 		if (LoopOffset != 0)
 		{
