Sequencer: add some comments to last change

diff --git a/Engine/Source/Runtime/MovieScene/Private/MovieSceneSequencePlayer.cpp b/Engine/Source/Runtime/MovieScene/Private/MovieSceneSequencePlayer.cpp
--- a/Engine/Source/Runtime/MovieScene/Private/MovieSceneSequencePlayer.cpp
+++ b/Engine/Source/Runtime/MovieScene/Private/MovieSceneSequencePlayer.cpp
@@ -1478,7 +1478,7 @@ FFrameTime UMovieSceneSequencePlayer::UpdateServerTimeSamples()
 
 	auto UpdateSamplesForChangedTimeDilation = [&]()
 	{
-		// Project all server time samples back based on the new time dilation so future updates will be accurate
+		// Project all server time samples back based on the new play-rate and time dilation so future updates will be accurate
 		if (LastEffectiveTimeDilation != TimeDilation)
 		{
 			for (FServerTimeSample& Sample : ServerTimeSamples)
@@ -1743,6 +1743,9 @@ void UMovieSceneSequencePlayer::PostNetReceive()
 	const bool bHasChangedStatus  = NetSyncProps.LastKnownStatus   != Status;
 	const bool bHasChangedTime    = NetSyncProps.LastKnownPosition != PlayPosition.GetCurrentPosition();
 
+	// We need to take play-rate into account when determining how many frames we can lag behind the server.
+	// For instance, if we play 3 times faster than normal (play-rate = 3), we should be able to lag 3 times as
+	// many frames behind as normal before we force a re-sync.
 	const float PlayRate = PlaybackSettings.PlayRate;
 
 	float TimeDilation = 1.0f;
@@ -1860,6 +1863,9 @@ void UMovieSceneSequencePlayer::UpdateNetworkSync()
 	{
 		const float PingMs = GetPing();
 
+		// We need to take play-rate into account when determining how many frames we can lag behind the server.
+		// For instance, if we play 3 times faster than normal (play-rate = 3), we should be able to lag 3 times as
+		// many frames behind as normal before we force a re-sync.
 		const float PlayRate = PlaybackSettings.PlayRate;
 
 		float TimeDilation = 1.0f;
