Sequencer: add ability to replace inline capabilities

diff --git a/Engine/Source/Runtime/MovieScene/Private/Tests/MovieScenePlaybackCapabilitiesTests.cpp b/Engine/Source/Runtime/MovieScene/Private/Tests/MovieScenePlaybackCapabilitiesTests.cpp
--- a/Engine/Source/Runtime/MovieScene/Private/Tests/MovieScenePlaybackCapabilitiesTests.cpp
+++ b/Engine/Source/Runtime/MovieScene/Private/Tests/MovieScenePlaybackCapabilitiesTests.cpp
@@ -102,6 +102,24 @@ struct FTestInvalidatableCapability : IPlaybackCapability
 
 TPlaybackCapabilityID<FTestInvalidatableCapability> FTestInvalidatableCapability::ID = TPlaybackCapabilityID<FTestInvalidatableCapability>::Register();
 
+// Derived version of the simple test capability, but with specific destructor counter
+struct FTestSimpleCapabilityDerived : FTestSimpleCapability
+{
+	static int32 TimesDestroyedDerived;
+
+	FTestSimpleCapabilityDerived()
+	{}
+	FTestSimpleCapabilityDerived(int32 InIntValue, const FString& InStringValue)
+		: FTestSimpleCapability(InIntValue, InStringValue)
+	{}
+	~FTestSimpleCapabilityDerived()
+	{
+		++TimesDestroyedDerived;
+	}
+};
+
+int32 FTestSimpleCapabilityDerived::TimesDestroyedDerived = 0;
+
 // Derived version of the simple test capability, but with interface
 struct FTestSimpleCapabilityDerivedWithInterface : FTestSimpleCapability, IPlaybackCapability
 {
@@ -126,15 +144,15 @@ bool FMovieScenePlaybackCapabilitiesSimpleTest::RunTest(const FString& Parameter
 	FTestSimpleCapability::TimesDestroyed = 0;
 	{
 		FPlaybackCapabilities Caps;
-		UTEST_FALSE("No capability", Caps.HasCapability(FTestSimpleCapability::ID));
+		UTEST_FALSE("No capability", Caps.HasCapability<FTestSimpleCapability>());
 
-		FTestSimpleCapability& SimpleCap = Caps.AddCapability(FTestSimpleCapability::ID, 42, TEXT("Just a test"));
+		FTestSimpleCapability& SimpleCap = Caps.AddCapability<FTestSimpleCapability>(42, TEXT("Just a test"));
 		UTEST_EQUAL("Int value", SimpleCap.IntValue, 42);
 		UTEST_EQUAL("String value", SimpleCap.StringValue, TEXT("Just a test"));
 
-		UTEST_TRUE("Found capability", Caps.HasCapability(FTestSimpleCapability::ID));
+		UTEST_TRUE("Found capability", Caps.HasCapability<FTestSimpleCapability>());
 
-		FTestSimpleCapability* SimpleCap2 = Caps.FindCapability(FTestSimpleCapability::ID);
+		FTestSimpleCapability* SimpleCap2 = Caps.FindCapability<FTestSimpleCapability>();
 		UTEST_NOT_NULL("Retrived test capacity", SimpleCap2);
 		UTEST_EQUAL("Same capacity?", SimpleCap2, &SimpleCap);
 	}
@@ -155,16 +173,15 @@ bool FMovieScenePlaybackCapabilitiesInlineInterfaceTest::RunTest(const FString&
 	FTestCapabilityDerived::TimesDestroyedDerived = 0;
 	{
 		FPlaybackCapabilities Caps;
-		UTEST_FALSE("No capability", Caps.HasCapability(FTestCapabilityBase::ID));
+		UTEST_FALSE("No capability", Caps.HasCapability<FTestCapabilityBase>());
 
-		FTestCapabilityBase& InterfaceCap = Caps.AddCapabilityImplementation<FTestCapabilityDerived>(
-				FTestCapabilityBase::ID, 42, TEXT("Just a test"));
+		FTestCapabilityBase& InterfaceCap = Caps.AddCapability<FTestCapabilityDerived>(42, TEXT("Just a test"));
 		UTEST_EQUAL("Int value", InterfaceCap.GetIntValue(), 42);
 		UTEST_EQUAL("String value", InterfaceCap.GetStringValue(), TEXT("Just a test"));
 
-		UTEST_TRUE("Found capability", Caps.HasCapability(FTestCapabilityBase::ID));
+		UTEST_TRUE("Found capability", Caps.HasCapability<FTestCapabilityBase>());
 
-		FTestCapabilityBase* InterfaceCap2 = Caps.FindCapability(FTestCapabilityBase::ID);
+		FTestCapabilityBase* InterfaceCap2 = Caps.FindCapability<FTestCapabilityBase>();
 		UTEST_NOT_NULL("Retrived interface capacity", InterfaceCap2);
 		UTEST_EQUAL("Same capacity?", InterfaceCap2, &InterfaceCap);
 	}
@@ -187,12 +204,12 @@ bool FMovieScenePlaybackCapabilitiesRawPointerTest::RunTest(const FString& Param
 		FTestSimpleCapability ActualCap;
 		{
 			FPlaybackCapabilities Caps;
-			UTEST_FALSE("No capability", Caps.HasCapability(FTestSimpleCapability::ID));
+			UTEST_FALSE("No capability", Caps.HasCapability<FTestSimpleCapability>());
 
-			FTestSimpleCapability& PtrCap = Caps.AddCapabilityRaw(FTestSimpleCapability::ID, &ActualCap);
-			UTEST_TRUE("Found capability", Caps.HasCapability(FTestSimpleCapability::ID));
+			FTestSimpleCapability& PtrCap = Caps.AddCapabilityRaw<FTestSimpleCapability>(&ActualCap);
+			UTEST_TRUE("Found capability", Caps.HasCapability<FTestSimpleCapability>());
 
-			FTestSimpleCapability* PtrCap2 = Caps.FindCapability(FTestSimpleCapability::ID);
+			FTestSimpleCapability* PtrCap2 = Caps.FindCapability<FTestSimpleCapability>();
 			UTEST_NOT_NULL("Retrived test capacity", PtrCap2);
 			UTEST_EQUAL("Same capacity?", PtrCap2, &PtrCap);
 			UTEST_EQUAL("Same capacity?", PtrCap2, &ActualCap);
@@ -218,13 +235,13 @@ bool FMovieScenePlaybackCapabilitiesSharedTest::RunTest(const FString& Parameter
 		UTEST_EQUAL("Ref count", ActualCap.GetSharedReferenceCount(), 1);
 		{
 			FPlaybackCapabilities Caps;
-			UTEST_FALSE("No capability", Caps.HasCapability(FTestSimpleCapability::ID));
+			UTEST_FALSE("No capability", Caps.HasCapability<FTestSimpleCapability>());
 
-			FTestSimpleCapability& PtrCap = Caps.AddCapabilityShared(FTestSimpleCapability::ID, ActualCap);
-			UTEST_TRUE("Found capability", Caps.HasCapability(FTestSimpleCapability::ID));
+			FTestSimpleCapability& PtrCap = Caps.AddCapabilityShared<FTestSimpleCapability>(ActualCap);
+			UTEST_TRUE("Found capability", Caps.HasCapability<FTestSimpleCapability>());
 			UTEST_EQUAL("Ref count", ActualCap.GetSharedReferenceCount(), 2);
 
-			FTestSimpleCapability* PtrCap2 = Caps.FindCapability(FTestSimpleCapability::ID);
+			FTestSimpleCapability* PtrCap2 = Caps.FindCapability<FTestSimpleCapability>();
 			UTEST_NOT_NULL("Retrived test capacity", PtrCap2);
 			UTEST_EQUAL("Same capacity?", PtrCap2, &PtrCap);
 			UTEST_EQUAL("Same capacity?", PtrCap2, ActualCap.ToSharedPtr().Get());
@@ -248,7 +265,7 @@ bool FMovieScenePlaybackCapabilitiesInvalidateTest::RunTest(const FString& Param
 	// Test simple inline
 	{
 		FPlaybackCapabilities Caps;
-		FTestInvalidatableCapability& Cap = Caps.AddCapability(FTestInvalidatableCapability::ID);
+		FTestInvalidatableCapability& Cap = Caps.AddCapability<FTestInvalidatableCapability>();
 		UTEST_EQUAL("Not invalidated", Cap.TimesInvalidated, 0);
 		Caps.InvalidateCachedData(nullptr);
 		UTEST_EQUAL("Invalidated", Cap.TimesInvalidated, 1);
@@ -257,7 +274,7 @@ bool FMovieScenePlaybackCapabilitiesInvalidateTest::RunTest(const FString& Param
 	// Test inline subclass implementing interface when base class doesn't
 	{
 		FPlaybackCapabilities Caps;
-		FTestSimpleCapability& Cap = Caps.AddCapabilityImplementation<FTestSimpleCapabilityDerivedWithInterface>(FTestSimpleCapability::ID);
+		FTestSimpleCapability& Cap = Caps.AddCapability<FTestSimpleCapabilityDerivedWithInterface>();
 		FTestSimpleCapabilityDerivedWithInterface* ActualCap = static_cast<FTestSimpleCapabilityDerivedWithInterface*>(&Cap);
 		UTEST_EQUAL("Not invalidated", ActualCap->TimesInvalidated, 0);
 		Caps.InvalidateCachedData(nullptr);
@@ -269,7 +286,7 @@ bool FMovieScenePlaybackCapabilitiesInvalidateTest::RunTest(const FString& Param
 		FTestInvalidatableCapability ActualCap;
 
 		FPlaybackCapabilities Caps;
-		Caps.AddCapabilityRaw(FTestInvalidatableCapability::ID, &ActualCap);
+		Caps.AddCapabilityRaw<FTestInvalidatableCapability>(&ActualCap);
 		UTEST_EQUAL("Not invalidated", ActualCap.TimesInvalidated, 0);
 		Caps.InvalidateCachedData(nullptr);
 		UTEST_EQUAL("Invalidated", ActualCap.TimesInvalidated, 1);
@@ -280,7 +297,7 @@ bool FMovieScenePlaybackCapabilitiesInvalidateTest::RunTest(const FString& Param
 		TSharedRef<FTestInvalidatableCapability> ActualCap = MakeShared<FTestInvalidatableCapability>();
 
 		FPlaybackCapabilities Caps;
-		Caps.AddCapabilityShared(FTestInvalidatableCapability::ID, ActualCap);
+		Caps.AddCapabilityShared<FTestInvalidatableCapability>(ActualCap);
 		UTEST_EQUAL("Not invalidated", ActualCap->TimesInvalidated, 0);
 		Caps.InvalidateCachedData(nullptr);
 		UTEST_EQUAL("Invalidated", ActualCap->TimesInvalidated, 1);
@@ -302,22 +319,21 @@ bool FMovieScenePlaybackCapabilitiesMultipleTest::RunTest(const FString& Paramet
 	{
 		FPlaybackCapabilities Caps;
 
-		Caps.AddCapabilityRaw(FTestOtherSimpleCapability::ID, &RawPtrCap);
-		Caps.AddCapabilityShared(FTestInvalidatableCapability::ID, SharedCap);
+		Caps.AddCapabilityRaw<FTestOtherSimpleCapability>(&RawPtrCap);
+		Caps.AddCapabilityShared<FTestInvalidatableCapability>(SharedCap);
 
-		UTEST_FALSE("Found capability", Caps.HasCapability(FTestSimpleCapability::ID));
-		UTEST_FALSE("Found capability", Caps.HasCapability(FTestCapabilityBase::ID));
-		UTEST_TRUE("Found capability", Caps.HasCapability(FTestOtherSimpleCapability::ID));
-		UTEST_TRUE("Found capability", Caps.HasCapability(FTestInvalidatableCapability::ID));
+		UTEST_FALSE("Found capability", Caps.HasCapability<FTestSimpleCapability>());
+		UTEST_FALSE("Found capability", Caps.HasCapability<FTestCapabilityBase>());
+		UTEST_TRUE("Found capability", Caps.HasCapability<FTestOtherSimpleCapability>());
+		UTEST_TRUE("Found capability", Caps.HasCapability<FTestInvalidatableCapability>());
 
-		Caps.AddCapability(FTestSimpleCapability::ID, 8, TEXT("Other simple test"));
-		Caps.AddCapabilityImplementation<FTestCapabilityDerived>(
-				FTestCapabilityBase::ID, 12, TEXT("Another implementation test"));
+		Caps.AddCapability<FTestSimpleCapability>(8, TEXT("Other simple test"));
+		Caps.AddCapability<FTestCapabilityDerived>(12, TEXT("Another implementation test"));
 
-		UTEST_TRUE("Found capability", Caps.HasCapability(FTestSimpleCapability::ID));
-		UTEST_TRUE("Found capability", Caps.HasCapability(FTestCapabilityBase::ID));
-		UTEST_TRUE("Found capability", Caps.HasCapability(FTestOtherSimpleCapability::ID));
-		UTEST_TRUE("Found capability", Caps.HasCapability(FTestInvalidatableCapability::ID));
+		UTEST_TRUE("Found capability", Caps.HasCapability<FTestSimpleCapability>());
+		UTEST_TRUE("Found capability", Caps.HasCapability<FTestCapabilityBase>());
+		UTEST_TRUE("Found capability", Caps.HasCapability<FTestOtherSimpleCapability>());
+		UTEST_TRUE("Found capability", Caps.HasCapability<FTestInvalidatableCapability>());
 
 		UTEST_EQUAL("Before invalidation", SharedCap->TimesInvalidated, 0);
 		Caps.InvalidateCachedData(nullptr);
@@ -354,10 +370,10 @@ bool FMovieScenePlaybackCapabilitiesInheritanceTest::RunTest(const FString& Para
 	{
 		FPlaybackCapabilities Caps;
 
-		Caps.AddCapabilityRaw(FTestCapabilityBase::ID, &Accessor.GetCapability());
+		Caps.AddCapabilityRaw<FTestCapabilityBase>(&Accessor.GetCapability());
 
-		UTEST_TRUE("Found capability", Caps.HasCapability(FTestCapabilityBase::ID));
-		FTestCapabilityBase* Base = Caps.FindCapability(FTestCapabilityBase::ID);
+		UTEST_TRUE("Found capability", Caps.HasCapability<FTestCapabilityBase>());
+		FTestCapabilityBase* Base = Caps.FindCapability<FTestCapabilityBase>();
 		UTEST_NOT_NULL("Got capability pointer", Base);
 		UTEST_EQUAL("Same capability pointer", Base, static_cast<FTestCapabilityBase*>(&Accessor.Derived));
 		UTEST_EQUAL("Checking IntValue", Base->GetIntValue(), 42);
@@ -381,25 +397,78 @@ bool FMovieScenePlaybackCapabilitiesOverwriteTest::RunTest(const FString& Parame
 	{
 		FPlaybackCapabilities Caps;
 
-		Caps.AddCapabilityRaw(FTestSimpleCapability::ID, &Cap1);
+		Caps.AddCapabilityRaw<FTestSimpleCapability>(&Cap1);
 
-		FTestSimpleCapability* ActualCap = Caps.FindCapability(FTestSimpleCapability::ID);
-		UTEST_TRUE("Found capability", Caps.HasCapability(FTestSimpleCapability::ID));
+		FTestSimpleCapability* ActualCap = Caps.FindCapability<FTestSimpleCapability>();
+		UTEST_TRUE("Found capability", Caps.HasCapability<FTestSimpleCapability>());
 		UTEST_EQUAL("Same capability pointer", ActualCap, &Cap1);
 		UTEST_EQUAL("Checking IntValue", ActualCap->IntValue, 1);
 		UTEST_EQUAL("Checking StringValue", ActualCap->StringValue, TEXT("First"));
 
-		Caps.OverwriteCapabilityRaw(FTestSimpleCapability::ID, &Cap2);
+		Caps.OverwriteCapabilityRaw<FTestSimpleCapability>(&Cap2);
 		
-		ActualCap = Caps.FindCapability(FTestSimpleCapability::ID);
-		UTEST_TRUE("Found capability", Caps.HasCapability(FTestSimpleCapability::ID));
+		ActualCap = Caps.FindCapability<FTestSimpleCapability>();
+		UTEST_TRUE("Found capability", Caps.HasCapability<FTestSimpleCapability>());
 		UTEST_EQUAL("Same capability pointer", ActualCap, &Cap2);
 		UTEST_EQUAL("Checking IntValue", ActualCap->IntValue, 2);
 		UTEST_EQUAL("Checking StringValue", ActualCap->StringValue, TEXT("Second"));
 	}
 
+	FTestSimpleCapability::TimesDestroyed = 0;
+	FTestSimpleCapabilityDerived::TimesDestroyedDerived = 0;
+	{
+		FPlaybackCapabilities Caps;
+
+		// Add the base, overwrite with the derived
+		// (this is only allowed because they both have the same size)
+		Caps.AddCapability<FTestSimpleCapability>(1, TEXT("Base"));
+
+		FTestSimpleCapability* BaseCap = Caps.FindCapability<FTestSimpleCapability>();
+		UTEST_NOT_NULL("Found capabality", BaseCap);
+		UTEST_EQUAL("Checking IntValue", BaseCap->IntValue, 1);
+		UTEST_EQUAL("Checking StringValue", BaseCap->StringValue, TEXT("Base"));
+
+		Caps.OverwriteCapability<FTestSimpleCapabilityDerived>(2, TEXT("Derived"));
+
+		UTEST_EQUAL("Times destroyed", FTestSimpleCapability::TimesDestroyed, 1);
+		UTEST_EQUAL("Times destroyed", FTestSimpleCapabilityDerived::TimesDestroyedDerived, 0);
+
+		FTestSimpleCapability* DerivedCap = Caps.FindCapability<FTestSimpleCapability>();
+		UTEST_NOT_NULL("Found capabality", DerivedCap);
+		UTEST_EQUAL("Checking IntValue", DerivedCap->IntValue, 2);
+		UTEST_EQUAL("Checking StringValue", DerivedCap->StringValue, TEXT("Derived"));
+	}
+	UTEST_EQUAL("Times destroyed", FTestSimpleCapability::TimesDestroyed, 2);
+	UTEST_EQUAL("Times destroyed", FTestSimpleCapabilityDerived::TimesDestroyedDerived, 1);
+
+	FTestSimpleCapability::TimesDestroyed = 0;
+	FTestSimpleCapabilityDerived::TimesDestroyedDerived = 0;
+	{
+		FPlaybackCapabilities Caps;
+
+		// Add the derived, overwrite with the base
+		// (this is only allowed because they both have the same size)
+		Caps.AddCapability<FTestSimpleCapabilityDerived>(2, TEXT("Derived"));
+
+		FTestSimpleCapability* DerivedCap = Caps.FindCapability<FTestSimpleCapability>();
+		UTEST_NOT_NULL("Found capabality", DerivedCap);
+		UTEST_EQUAL("Checking IntValue", DerivedCap->IntValue, 2);
+		UTEST_EQUAL("Checking StringValue", DerivedCap->StringValue, TEXT("Derived"));
+
+		Caps.OverwriteCapability<FTestSimpleCapability>(1, TEXT("Base"));
+
+		UTEST_EQUAL("Times destroyed", FTestSimpleCapability::TimesDestroyed, 1);
+		UTEST_EQUAL("Times destroyed", FTestSimpleCapabilityDerived::TimesDestroyedDerived, 1);
+
+		FTestSimpleCapability* BaseCap = Caps.FindCapability<FTestSimpleCapability>();
+		UTEST_NOT_NULL("Found capabality", BaseCap);
+		UTEST_EQUAL("Checking IntValue", BaseCap->IntValue, 1);
+		UTEST_EQUAL("Checking StringValue", BaseCap->StringValue, TEXT("Base"));
+	}
+	UTEST_EQUAL("Times destroyed", FTestSimpleCapability::TimesDestroyed, 2);
+	UTEST_EQUAL("Times destroyed", FTestSimpleCapabilityDerived::TimesDestroyedDerived, 1);
+
 	return true;
 }
 
 #endif // WITH_DEV_AUTOMATION_TESTS
-
diff --git a/Engine/Source/Runtime/MovieScene/Public/EntitySystem/MovieSceneSharedPlaybackState.h b/Engine/Source/Runtime/MovieScene/Public/EntitySystem/MovieSceneSharedPlaybackState.h
--- a/Engine/Source/Runtime/MovieScene/Public/EntitySystem/MovieSceneSharedPlaybackState.h
+++ b/Engine/Source/Runtime/MovieScene/Public/EntitySystem/MovieSceneSharedPlaybackState.h
@@ -112,7 +112,7 @@ public:
 	template<typename T>
 	bool HasCapability() const
 	{
-		return Capabilities.HasCapability(T::ID);
+		return Capabilities.HasCapability<T>();
 	}
 
 	/**
@@ -121,7 +121,7 @@ public:
 	template<typename T>
 	T* FindCapability() const
 	{
-		return Capabilities.FindCapability<T>(T::ID);
+		return Capabilities.FindCapability<T>();
 	}
 
 	/**
@@ -130,7 +130,7 @@ public:
 	template<typename T, typename ...ArgTypes>
 	T& AddCapability(ArgTypes&&... InArgs)
 	{
-		T& Cap = Capabilities.AddCapability<T>(T::ID, Forward<ArgTypes>(InArgs)...);
+		T& Cap = Capabilities.AddCapability<T>(Forward<ArgTypes>(InArgs)...);
 		MaybeInitialize(Cap);
 		return Cap;
 	}
@@ -138,10 +138,10 @@ public:
 	/**
 	 * Adds the specified capability on the root sequence as a raw pointer.
 	 */
-	template<typename T, typename ...ArgTypes>
+	template<typename T>
 	T& AddCapabilityRaw(T* InPointer)
 	{
-		T& Cap = Capabilities.AddCapabilityRaw<T>(T::ID, InPointer);
+		T& Cap = Capabilities.AddCapabilityRaw<T>(InPointer);
 		MaybeInitialize(Cap);
 		return Cap;
 	}
@@ -149,14 +149,39 @@ public:
 	/**
 	 * Adds the specified capability on the root sequence as a shared pointer.
 	 */
-	template<typename T, typename ...ArgTypes>
+	template<typename T>
 	T& AddCapabilityShared(TSharedRef<T> InSharedRef)
 	{
-		T& Cap = Capabilities.AddCapabilityShared<T>(T::ID, InSharedRef);
+		T& Cap = Capabilities.AddCapabilityShared<T>(InSharedRef);
 		MaybeInitialize(Cap);
 		return Cap;
 	}
 
+	/**
+	 * Adds the specified capability on the root sequence.
+	 * If the capability already exists, it must be stored inline, and its
+	 * value will be replaced by the new object.
+	 * If the template parameter is a sub-class of the playback capability, the previously
+	 * stored playback capability must not only have been stored inline, but must have 
+	 * been of the same sub-class (or a sub-class with the exact same size and alignment).
+	 */
+	template<typename T, typename ...ArgTypes>
+	T& SetOrAddCapability(ArgTypes&&... InArgs)
+	{
+		if (HasCapability<T>())
+		{
+			T& Cap = Capabilities.OverwriteCapability<T>(Forward<ArgTypes>(InArgs)...);
+			MaybeInitialize(Cap);
+			return Cap;
+		}
+		else
+		{
+			T& Cap = Capabilities.AddCapability<T>(Forward<ArgTypes>(InArgs)...);
+			MaybeInitialize(Cap);
+			return Cap;
+		}
+	}
+
 	/**
 	 * Adds the specified capability on the root sequence as a raw pointer.
 	 * If the capability already exists, it must be stored as a raw pointer and its
@@ -167,13 +192,13 @@ public:
 	{
 		if (HasCapability<T>())
 		{
-			T& Cap = Capabilities.OverwriteCapabilityRaw<T>(T::ID, InPointer);
+			T& Cap = Capabilities.OverwriteCapabilityRaw<T>(InPointer);
 			MaybeInitialize(Cap);
 			return Cap;
 		}
 		else
 		{
-			T& Cap = Capabilities.AddCapabilityRaw<T>(T::ID, InPointer);
+			T& Cap = Capabilities.AddCapabilityRaw<T>(InPointer);
 			MaybeInitialize(Cap);
 			return Cap;
 		}
@@ -189,13 +214,13 @@ public:
 	{
 		if (HasCapability<T>())
 		{
-			T& Cap = Capabilities.OverwriteCapabilityShared<T>(T::ID, InSharedRef);
+			T& Cap = Capabilities.OverwriteCapabilityShared<T>(InSharedRef);
 			MaybeInitialize(Cap);
 			return Cap;
 		}
 		else
 		{
-			T& Cap = Capabilities.AddCapabilityShared<T>(T::ID, InSharedRef);
+			T& Cap = Capabilities.AddCapabilityShared<T>(InSharedRef);
 			MaybeInitialize(Cap);
 			return Cap;
 		}
diff --git a/Engine/Source/Runtime/MovieScene/Public/Evaluation/IMovieScenePlaybackCapability.h b/Engine/Source/Runtime/MovieScene/Public/Evaluation/IMovieScenePlaybackCapability.h
--- a/Engine/Source/Runtime/MovieScene/Public/Evaluation/IMovieScenePlaybackCapability.h
+++ b/Engine/Source/Runtime/MovieScene/Public/Evaluation/IMovieScenePlaybackCapability.h
@@ -42,6 +42,9 @@ protected:
 template<typename T>
 struct TPlaybackCapabilityID : FPlaybackCapabilityID
 {
+public:
+	using CapabilityType = T;
+
 private:
 	// Only T should construct this (to ensure safe construction over DLL boundaries)
 	friend T;
diff --git a/Engine/Source/Runtime/MovieScene/Public/Evaluation/MovieScenePlaybackCapabilities.h b/Engine/Source/Runtime/MovieScene/Public/Evaluation/MovieScenePlaybackCapabilities.h
--- a/Engine/Source/Runtime/MovieScene/Public/Evaluation/MovieScenePlaybackCapabilities.h
+++ b/Engine/Source/Runtime/MovieScene/Public/Evaluation/MovieScenePlaybackCapabilities.h
@@ -492,8 +492,22 @@ protected:
 		// Don't use the header's Resolve method, we don't want to offset the pointer, we want
 		// the actual pointer to the actual stored capability.
 		void* CapabilityPtr = Header.Capability.Resolve(Memory);
+
+		// If we have inline storage, call the destructor on the previous capability.
+		if (Header.StorageMode == EPlaybackCapabilityStorageMode::Inline)
+		{
+			Helpers[Index].Destructor(CapabilityPtr);
+		}
+
+		// Allocate the new capability.
 		StorageType* TypedCapabilityPtr = reinterpret_cast<StorageType*>(CapabilityPtr);
-		*TypedCapabilityPtr = StorageType(Forward<ArgTypes>(InArgs)...);
+		new (TypedCapabilityPtr) StorageType(Forward<ArgTypes>(InArgs)...);
+
+		// If we have inline storage, we could potentially have changed from a stored base-class to
+		// a stored child-class, or vice-versa, or from one child-class to another child-class. In all these
+		// cases, the helper functions (destructor, interface cast, etc) have changed, so let's store the 
+		// new helper functions on the header.
+		Helpers[Index] = TPlaybackCapabilityHelpers<StorageType>::GetHelpers();
 
 		return true;
 	}
@@ -546,16 +560,23 @@ struct FPlaybackCapabilities : FPlaybackCapabilitiesImpl
 	MOVIESCENE_API ~FPlaybackCapabilities();
 
 	/** Checks whether this container has the given capability */
-	bool HasCapability(FPlaybackCapabilityID CapabilityID) const
+	template<typename T>
+	bool HasCapability() const
 	{
-		uint32 CapabilityBit = 1 << CapabilityID.Index;
+		uint32 CapabilityBit = 1 << T::ID.Index;
 		return FPlaybackCapabilitiesImpl::HasCapability(CapabilityBit);
 	}
 
 	/** Finds the specified capability within the container, if present */
 	template<typename T>
-	T* FindCapability(TPlaybackCapabilityID<T> CapabilityID) const
+	T* FindCapability() const
 	{
+		// T must be the base capability type, and not a sub-class, because we are returning a pointer to it,
+		// and we can't check what sort of sub-class we have or not.
+		using CapabilityIDType = decltype(T::ID);
+		static_assert(std::is_same<T, typename CapabilityIDType::CapabilityType>::value, "You must pass the actual playback capability type as a template parameter, not a sub-class.");
+
+		const TPlaybackCapabilityID<T> CapabilityID(T::ID);
 		uint32 CapabilityBit = 1 << CapabilityID.Index;
 		FPlaybackCapabilityPtr Ptr = FPlaybackCapabilitiesImpl::FindCapability(CapabilityBit);
 		return Ptr.ResolveOptional<T>();
@@ -563,33 +584,34 @@ struct FPlaybackCapabilities : FPlaybackCapabilitiesImpl
 
 	/** Returns the specified capability within the container, asserts if not found */
 	template<typename T>
-	T& GetCapabilityChecked(TPlaybackCapabilityID<T> CapabilityID) const
+	T& GetCapabilityChecked() const
 	{
+		// T must be the base capability type, and not a sub-class, because we are returning a reference to it,
+		// and we can't check what sort of sub-class we have or not.
+		using CapabilityIDType = decltype(T::ID);
+		static_assert(std::is_same<T, typename CapabilityIDType::CapabilityType>::value, "You must pass the actual playback capability type as a template parameter, not a sub-class.");
+
+		const TPlaybackCapabilityID<T> CapabilityID(T::ID);
 		uint32 CapabilityBit = 1 << CapabilityID.Index;
 		FPlaybackCapabilityPtr Ptr = FPlaybackCapabilitiesImpl::GetCapabilityChecked(CapabilityBit);
 		return Ptr.ResolveChecked<T>();
 	}
 
 	/**
-	 * Adds the specified capability to the container, using the supplied arguments to construct it
+	 * Adds the specified capability to the container, using the supplied arguments to construct it.
 	 * The capability object will be stored inline and owned by this container. It will be destroyed when the
 	 * container itself is destroyed.
+	 * If the template parameter is a sub-class of the playback capability class, that sub-class will be
+	 * created and stored inline.
 	 */
 	template<typename T, typename ...ArgTypes>
-	T& AddCapability(TPlaybackCapabilityID<T> CapabilityID, ArgTypes&&... InArgs)
+	T& AddCapability(ArgTypes&&... InArgs)
 	{
-		return DoAddCapability<T, T>(CapabilityID, Forward<ArgTypes>(InArgs)...);
-	}
+		using CapabilityIDType = decltype(T::ID);
+		using CapabilityType = typename CapabilityIDType::CapabilityType;
 
-	/**
-	 * Adds the specified capability to the container, using the supplied arguments to construct a sub-class of it 
-	 * The capability object of the specified sub-class will be stored inline and owned by this container. It will
-	 * be destroyed when the container itself is destroyed.
-	 */
-	template<typename Impl, typename T, typename ...ArgTypes>
-	T& AddCapabilityImplementation(TPlaybackCapabilityID<T> CapabilityID, ArgTypes&&... InArgs)
-	{
-		return DoAddCapability<Impl, T>(CapabilityID, Forward<ArgTypes>(InArgs)...);
+		CapabilityType& NewCapability = DoAddCapability<T>(T::ID, Forward<ArgTypes>(InArgs)...);
+		return static_cast<T&>(NewCapability);
 	}
 
 	/**
@@ -598,9 +620,13 @@ struct FPlaybackCapabilities : FPlaybackCapabilitiesImpl
 	 * store a raw pointer.
 	 */
 	template<typename T>
-	T& AddCapabilityRaw(TPlaybackCapabilityID<T> CapabilityID, T* InPointer)
+	T& AddCapabilityRaw(T* InPointer)
 	{
-		return DoAddCapability<T*, T>(CapabilityID, InPointer);
+		using CapabilityIDType = decltype(T::ID);
+		using CapabilityType = typename CapabilityIDType::CapabilityType;
+
+		CapabilityType& NewCapability = DoAddCapability<CapabilityType*>(T::ID, static_cast<CapabilityType*>(InPointer));
+		return static_cast<T&>(NewCapability);
 	}
 	
 	/**
@@ -609,27 +635,58 @@ struct FPlaybackCapabilities : FPlaybackCapabilitiesImpl
 	 * only maintains one such shared pointer until the container is destroyed.
 	 */
 	template<typename T>
-	T& AddCapabilityShared(TPlaybackCapabilityID<T> CapabilityID, TSharedRef<T> InSharedRef)
+	T& AddCapabilityShared(TSharedRef<T> InSharedRef)
+	{
+		using CapabilityIDType = decltype(T::ID);
+		using CapabilityType = typename CapabilityIDType::CapabilityType;
+
+		CapabilityType& NewCapability = DoAddCapability<TSharedPtr<CapabilityType>>(T::ID, StaticCastSharedRef<CapabilityType>(InSharedRef));
+		return static_cast<T&>(NewCapability);
+	}
+
+	/**
+	 * Overwrites an existing capability, stored inline and owned by this container.
+	 * The previous storage mode of the capability must also be inline.
+	 * If the template parameter is a sub-class of the playback capability, the previously stored playback capability
+	 * must not only have been stored inline, but must have been of the same sub-class (or a sub-class with the exact
+	 * same size and alignment).
+	 */
+	template<typename T, typename ...ArgTypes>
+	T& OverwriteCapability(ArgTypes&&... InArgs)
 	{
-		return DoAddCapability<TSharedPtr<T>, T>(CapabilityID, InSharedRef);
+		using CapabilityIDType = decltype(T::ID);
+		using CapabilityType = typename CapabilityIDType::CapabilityType;
+
+		CapabilityType& NewCapability = DoOverwriteCapability<T>(T::ID, Forward<ArgTypes>(InArgs)...);
+		return static_cast<T&>(NewCapability);
 	}
 
 	/**
 	 * Overwrites an existing capability, stored as a raw pointer on the container.
+	 * The previous storage mode of the capability must also be a raw pointer.
 	 */
 	template<typename T>
-	T& OverwriteCapabilityRaw(TPlaybackCapabilityID<T> CapabilityID, T* InPointer)
+	T& OverwriteCapabilityRaw(T* InPointer)
 	{
-		return DoOverwriteCapability<T*, T>(CapabilityID, InPointer);
+		using CapabilityIDType = decltype(T::ID);
+		using CapabilityType = typename CapabilityIDType::CapabilityType;
+
+		CapabilityType& NewCapability = DoOverwriteCapability<CapabilityType*>(T::ID, static_cast<CapabilityType*>(InPointer));
+		return static_cast<T&>(NewCapability);
 	}
 	
 	/**
 	 * Overwrites an existing capability, stored as a shared pointer on the container.
+	 * The previous storage mode of the capability must also be a shared pointer.
 	 */
 	template<typename T>
-	T& OverwriteCapabilityShared(TPlaybackCapabilityID<T> CapabilityID, TSharedRef<T> InSharedRef)
+	T& OverwriteCapabilityShared(TSharedRef<T> InSharedRef)
 	{
-		return DoOverwriteCapability<TSharedPtr<T>, T>(CapabilityID, InSharedRef);
+		using CapabilityIDType = decltype(T::ID);
+		using CapabilityType = typename CapabilityIDType::CapabilityType;
+
+		CapabilityType& NewCapability = DoOverwriteCapability<TSharedPtr<CapabilityType>>(T::ID, StaticCastSharedRef<CapabilityType>(InSharedRef));
+		return static_cast<T&>(NewCapability);
 	}
 
 public:
@@ -678,7 +735,7 @@ private:
 	{
 		uint32 CapabilityBit = 1 << CapabilityID.Index;
 		FPlaybackCapabilitiesImpl::OverwriteCapability<Impl, T>(CapabilityBit, Forward<ArgTypes>(InArgs)...);
-		return GetCapabilityChecked(CapabilityID);
+		return GetCapabilityChecked<T>();
 	}
 
 	void Destroy();
