Sequencer: change how the sequence observer is queried to solve networking bugs

The observer is now checked at all public API entry points for whether a sequence can be observed or not. Before, it would only be checked when updating the sequence, i.e. a non-observable sequence would still actually be playing... it would just not be updating. Now, a non-observable sequence is not playing at all.

This has the effect of letting us try again in PostNetReceive if needed. If the server started playing and the client wants to start playing too, it will now check if the sequence is observable.

1. If it is observable, it starts playing.

2. If it isn't observable, it doesn't start playing, but on the next PostNetReceive it will try again. For a truly non-observable sequence, this means trying again all the time and always refusing to. This isn't any worse than before, when the sequence was playing but always refusing to update every frame. In fact, it's a bit less worse, since PostNetReceive doesn't necessarily happen every frame.

3. If it is observable, but network replication of instigator properties hasn't happened yet, then the observer mistakenly tells the client that the sequence isn't observable, and nothing happens... but eventually, network replication does come through, and the client finally starts playing, with a delay.

diff --git a/Engine/Source/Runtime/MovieScene/Private/MovieSceneSequencePlayer.cpp b/Engine/Source/Runtime/MovieScene/Private/MovieSceneSequencePlayer.cpp
--- a/Engine/Source/Runtime/MovieScene/Private/MovieSceneSequencePlayer.cpp
+++ b/Engine/Source/Runtime/MovieScene/Private/MovieSceneSequencePlayer.cpp
@@ -242,6 +242,11 @@ void UMovieSceneSequencePlayer::PlayLooping(int32 NumLoops)
 
 void UMovieSceneSequencePlayer::PlayInternal()
 {
+	if (Observer && !Observer->CanObserveSequence())
+	{
+		return;
+	}
+
 	if (NeedsQueueLatentAction())
 	{
 		QueueLatentAction(FMovieSceneSequenceLatentActionDelegate::CreateUObject(this, &UMovieSceneSequencePlayer::PlayInternal));
@@ -321,6 +326,11 @@ void UMovieSceneSequencePlayer::PlayInternal()
 
 void UMovieSceneSequencePlayer::Pause()
 {
+	if (Observer && !Observer->CanObserveSequence())
+	{
+		return;
+	}
+
 	if (NeedsQueueLatentAction())
 	{
 		QueueLatentAction(FMovieSceneSequenceLatentActionDelegate::CreateUObject(this, &UMovieSceneSequencePlayer::Pause));
@@ -352,12 +362,8 @@ void UMovieSceneSequencePlayer::Pause()
 			}
 		};
 
-		if (Observer && !Observer->CanObserveSequence())
-		{
-			FinishPause();
-		}
 		// Evaluate the sequence at its current time, with a status of 'stopped' to ensure that animated state pauses correctly. (ie. audio sounds should stop/pause)
-		else if (TSharedPtr<FMovieSceneEntitySystemRunner> Runner = RootTemplateInstance.GetRunner())
+		if (TSharedPtr<FMovieSceneEntitySystemRunner> Runner = RootTemplateInstance.GetRunner())
 		{
 			FMovieSceneEvaluationRange CurrentTimeRange = PlayPosition.GetCurrentPositionAsRange();
 			const FMovieSceneContext Context(CurrentTimeRange, EMovieScenePlayerStatus::Stopped);
@@ -391,6 +397,11 @@ void UMovieSceneSequencePlayer::StopAtCurrentTime()
 
 void UMovieSceneSequencePlayer::StopInternal(FFrameTime TimeToResetTo)
 {
+	if (Observer && !Observer->CanObserveSequence())
+	{
+		return;
+	}
+
 	if (NeedsQueueLatentAction())
 	{
 		QueueLatentAction(FMovieSceneSequenceLatentActionDelegate::CreateUObject(this, &UMovieSceneSequencePlayer::StopInternal, TimeToResetTo));
@@ -459,11 +470,7 @@ void UMovieSceneSequencePlayer::StopInternal(FFrameTime TimeToResetTo)
 
 		TSharedPtr<FMovieSceneEntitySystemRunner> Runner = RootTemplateInstance.GetRunner();
 
-		if (Observer && !Observer->CanObserveSequence())
-		{
-			OnFlushed();
-		}
-		else if (Runner)
+		if (Runner)
 		{
 			// Finish but do not destroy
 			if (Runner->QueueFinalUpdate(RootTemplateInstance.GetRootInstanceHandle(), FSimpleDelegate::CreateWeakLambda(this, OnFlushed)))
@@ -625,6 +632,11 @@ void UMovieSceneSequencePlayer::PlayTo(FMovieSceneSequencePlaybackParams InPlayb
 
 void UMovieSceneSequencePlayer::SetPlaybackPosition(FMovieSceneSequencePlaybackParams InPlaybackParams)
 {
+	if (Observer && !Observer->CanObserveSequence())
+	{
+		return;
+	}
+
 	if (NeedsQueueLatentAction())
 	{
 		QueueLatentAction(FMovieSceneSequenceLatentActionDelegate::CreateUObject(this, &UMovieSceneSequencePlayer::SetPlaybackPosition, InPlaybackParams));
@@ -1235,6 +1247,7 @@ void UMovieSceneSequencePlayer::UpdateMovieSceneInstance(FMovieSceneEvaluationRa
 {
 	if (Observer && !Observer->CanObserveSequence())
 	{
+		UE_LOG(LogMovieScene, Error, TEXT("Refusing to update an unobservable sequence! Did it become unobservable during playback?"));
 		return;
 	}
 
@@ -1761,6 +1774,23 @@ void UMovieSceneSequencePlayer::PostNetReceive()
 		return;
 	}
 
+	if (Observer && !Observer->CanObserveSequence())
+	{
+		// We shouldn't do anything.
+#if !NO_LOGGING
+		if (UE_LOG_ACTIVE(LogMovieSceneRepl, Verbose))
+		{
+			const FFrameTime CurrentTime = PlayPosition.GetCurrentPosition();
+			const FString SequenceName = GetSequenceName(true);
+			UE_LOG(LogMovieSceneRepl, Verbose, TEXT("Ignoring network update for unobservable sequence %s %s @ %s. Server is %s @ %s."),
+				*SequenceName,
+				*UEnum::GetValueAsString(Status.GetValue()), *LexToString(CurrentTime),
+				*UEnum::GetValueAsString(NetSyncProps.LastKnownStatus.GetValue()), *LexToString(NetSyncProps.LastKnownPosition));
+		}
+#endif
+		return;
+	}
+
 #if !NO_LOGGING
 	if (UE_LOG_ACTIVE(LogMovieSceneRepl, VeryVerbose))
 	{
diff --git a/Engine/Source/Runtime/MovieScene/Public/MovieSceneSequencePlayer.h b/Engine/Source/Runtime/MovieScene/Public/MovieSceneSequencePlayer.h
--- a/Engine/Source/Runtime/MovieScene/Public/MovieSceneSequencePlayer.h
+++ b/Engine/Source/Runtime/MovieScene/Public/MovieSceneSequencePlayer.h
@@ -165,7 +165,7 @@ class UMovieSceneSequencePlayer
 public:
 	GENERATED_BODY()
 
-	/** Obeserver interface used for controlling whether the effects of this sequence can be seen even when it is playing back. */
+	/** Obeserver interface used for controlling whether this sequence can be played. */
 	UPROPERTY(replicated)
 	TScriptInterface<IMovieSceneSequencePlayerObserver> Observer;
 
