Sequencer: Fixed non-deterministic cook caused by whether possessable's classes were loaded or not

  - The problem code was inside FMovieScenePropertyTrackEntityImportHelper::PopulateEvaluationField that was using the possessed object class to attempt to locate EditConditions on properties. This would fail or succeed non-deterministically depending on whether the possessed class was loaded or not.

diff --git a/Engine/Source/Runtime/MovieScene/Public/MovieScenePossessable.h b/Engine/Source/Runtime/MovieScene/Public/MovieScenePossessable.h
--- a/Engine/Source/Runtime/MovieScene/Public/MovieScenePossessable.h
+++ b/Engine/Source/Runtime/MovieScene/Public/MovieScenePossessable.h
@@ -109,7 +109,7 @@ public:
 	 */
 	const UClass* GetPossessedObjectClass() const
 	{
-		return (const UClass*)(PossessedObjectClass.Get());
+		return (const UClass*)(PossessedObjectClass.LoadSynchronous());
 	}
 
 	/**
