Sequencer: cleanup pre-animated state classes

The entity capture source predated the templated class that implements it all in a generic way, so this change deletes all the entity-specific code and makes the entity capture source use the templated class.

This also removes dependencies on the player interface from a few existing pieces of code, like the evaluation hook system.

diff --git a/Engine/Source/Runtime/MovieScene/Private/EntitySystem/MovieSceneEvaluationHookSystem.cpp b/Engine/Source/Runtime/MovieScene/Private/EntitySystem/MovieSceneEvaluationHookSystem.cpp
--- a/Engine/Source/Runtime/MovieScene/Private/EntitySystem/MovieSceneEvaluationHookSystem.cpp
+++ b/Engine/Source/Runtime/MovieScene/Private/EntitySystem/MovieSceneEvaluationHookSystem.cpp
@@ -6,9 +6,8 @@
 #include "EntitySystem/MovieSceneSpawnablesSystem.h"
 #include "EntitySystem/MovieSceneEntitySystemTask.h"
 #include "Evaluation/PreAnimatedState/MovieScenePreAnimatedCaptureSource.h"
-
+#include "Evaluation/PreAnimatedState/MovieScenePreAnimatedCaptureSources.h"
 #include "Evaluation/MovieSceneEvaluationTemplateInstance.h"
-#include "IMovieScenePlayer.h"
 
 #include UE_INLINE_GENERATED_CPP_BY_NAME(MovieSceneEvaluationHookSystem)
 
@@ -47,6 +46,7 @@ struct FEvaluationHookUpdater
 			NewEvent.Hook          = Hooks[Index];
 			NewEvent.Type          = EEvaluationHookEvent::Update;
 			NewEvent.RootTime      = EvalTimes[Index] * SequenceInstance.GetContext().GetSequenceToRootSequenceTransform();
+			NewEvent.RootInstanceHandle = SequenceInstance.GetRootInstanceHandle();
 			NewEvent.SequenceID    = SequenceInstance.GetSequenceID();
 			NewEvent.bRestoreState = bRestoreState;
 
@@ -157,6 +157,7 @@ void UMovieSceneEvaluationHookSystem::UpdateHooks()
 			NewEvent.Hook          = Hooks[Index];
 			NewEvent.Type          = EEvaluationHookEvent::Begin;
 			NewEvent.RootTime      = EvalTimes[Index] * SequenceInstance.GetContext().GetSequenceToRootSequenceTransform();
+			NewEvent.RootInstanceHandle = SequenceInstance.GetRootInstanceHandle();
 			NewEvent.SequenceID    = SequenceInstance.GetSequenceID();
 			NewEvent.bRestoreState = bRestoreState;
 
@@ -177,6 +178,7 @@ void UMovieSceneEvaluationHookSystem::UpdateHooks()
 			NewEvent.Hook          = Hooks[Index];
 			NewEvent.Type          = EEvaluationHookEvent::End;
 			NewEvent.RootTime      = EvalTimes[Index] * SequenceInstance.GetContext().GetSequenceToRootSequenceTransform();
+			NewEvent.RootInstanceHandle = SequenceInstance.GetRootInstanceHandle();
 			NewEvent.SequenceID    = SequenceInstance.GetSequenceID();
 			NewEvent.bRestoreState = bRestoreState;
 
@@ -239,10 +241,12 @@ void UMovieSceneEvaluationHookSystem::TriggerAllEvents()
 
 		IMovieScenePlayer* Player      = SequenceInstance.GetPlayer();
 		FMovieSceneContext RootContext = SequenceInstance.GetContext();
+		TSharedRef<const FSharedPlaybackState> SharedPlaybackState = SequenceInstance.GetSharedPlaybackState();
+		FPreAnimatedEvaluationHookCaptureSources* EvaluationHookMetaData = Linker->PreAnimatedState.GetEvaluationHookMetaData();
 
 		for (const FMovieSceneEvaluationHookEvent& Event : Pair.Value.Events)
 		{
-			FScopedPreAnimatedCaptureSource CaptureSource(&Player->PreAnimatedState, Event.Hook.Interface.GetObject(), Event.SequenceID, Event.bRestoreState);
+			FScopedPreAnimatedCaptureSource CaptureSource(SharedPlaybackState, Event.Hook.Interface.GetObject(), Event.SequenceID, Event.bRestoreState);
 
 			FEvaluationHookParams Params = {
 				Event.Hook.ObjectBindingID, RootContext, Event.SequenceID, Event.TriggerIndex
@@ -267,7 +271,10 @@ void UMovieSceneEvaluationHookSystem::TriggerAllEvents()
 					break;
 				case EEvaluationHookEvent::End:
 					Event.Hook.Interface->End(Player, Params);
-					Player->PreAnimatedState.OnFinishedEvaluating(Event.Hook.Interface.GetObject(), Event.SequenceID);
+					if (EvaluationHookMetaData)
+					{
+						EvaluationHookMetaData->StopTrackingCaptureSource(Event.Hook.Interface.GetObject(), Event.RootInstanceHandle, Event.SequenceID);
+					}
 					break;
 
 				case EEvaluationHookEvent::Trigger:
diff --git a/Engine/Source/Runtime/MovieScene/Private/EntitySystem/MovieScenePreAnimatedStateSystem.cpp b/Engine/Source/Runtime/MovieScene/Private/EntitySystem/MovieScenePreAnimatedStateSystem.cpp
--- a/Engine/Source/Runtime/MovieScene/Private/EntitySystem/MovieScenePreAnimatedStateSystem.cpp
+++ b/Engine/Source/Runtime/MovieScene/Private/EntitySystem/MovieScenePreAnimatedStateSystem.cpp
@@ -8,7 +8,7 @@
 
 #include "Evaluation/PreAnimatedState/MovieScenePreAnimatedStateExtension.h"
 #include "Evaluation/PreAnimatedState/MovieScenePreAnimatedStateStorage.h"
-#include "Evaluation/PreAnimatedState/MovieScenePreAnimatedEntityCaptureSource.h"
+#include "Evaluation/PreAnimatedState/MovieScenePreAnimatedCaptureSources.h"
 
 #include UE_INLINE_GENERATED_CPP_BY_NAME(MovieScenePreAnimatedStateSystem)
 
diff --git a/Engine/Source/Runtime/MovieScene/Private/EntitySystem/MovieSceneSequenceInstance.cpp b/Engine/Source/Runtime/MovieScene/Private/EntitySystem/MovieSceneSequenceInstance.cpp
--- a/Engine/Source/Runtime/MovieScene/Private/EntitySystem/MovieSceneSequenceInstance.cpp
+++ b/Engine/Source/Runtime/MovieScene/Private/EntitySystem/MovieSceneSequenceInstance.cpp
@@ -396,10 +396,11 @@ void FSequenceInstance::PostEvaluation()
 
 void FSequenceInstance::DestroyImmediately()
 {
-	if (!Ledger.IsEmpty())
+	UMovieSceneEntitySystemLinker* Linker = SharedPlaybackState->GetLinker();
+	
+	if (!Ledger.IsEmpty() && ensure(Linker))
 	{
 		UE_LOG(LogMovieSceneECS, Verbose, TEXT("Instance being destroyed without first having been finished by calling Finish()"));
-		UMovieSceneEntitySystemLinker* Linker = SharedPlaybackState->GetLinker();
 		Ledger.UnlinkEverything(Linker, EUnlinkEverythingMode::CleanGarbage);
 	}
 
diff --git a/Engine/Source/Runtime/MovieScene/Private/Evaluation/Instances/MovieSceneTrackEvaluator.cpp b/Engine/Source/Runtime/MovieScene/Private/Evaluation/Instances/MovieSceneTrackEvaluator.cpp
--- a/Engine/Source/Runtime/MovieScene/Private/Evaluation/Instances/MovieSceneTrackEvaluator.cpp
+++ b/Engine/Source/Runtime/MovieScene/Private/Evaluation/Instances/MovieSceneTrackEvaluator.cpp
@@ -7,6 +7,7 @@
 #include "Sections/MovieSceneSubSection.h"
 #include "Compilation/MovieSceneCompiledDataManager.h"
 #include "Evaluation/MovieSceneEvaluationTemplateInstance.h"
+#include "Evaluation/PreAnimatedState/MovieScenePreAnimatedCaptureSources.h"
 #include "Stats/Stats2.h"
 
 #include "IMovieSceneModule.h"
@@ -38,9 +39,14 @@ struct FDelayedPreAnimatedStateRestore
 
 	void RestoreNow()
 	{
+		using namespace UE::MovieScene;
+
+		UMovieSceneEntitySystemLinker* Linker = Player.GetSharedPlaybackState()->GetLinker();
+		FPreAnimatedTemplateCaptureSources* TemplateMetaData = Linker->PreAnimatedState.GetTemplateMetaData();
+		const FRootInstanceHandle RootInstanceHandle = Player.GetEvaluationTemplate().GetRootInstanceHandle();
 		for (FMovieSceneEvaluationKey Key : KeysToRestore)
 		{
-			Player.PreAnimatedState.OnFinishedEvaluating(Key);
+			TemplateMetaData->StopTrackingCaptureSource(Key, RootInstanceHandle);
 		}
 		KeysToRestore.Reset();
 	}
@@ -323,9 +329,13 @@ void FMovieSceneTrackEvaluator::EvaluateGroup(const FMovieSceneEvaluationGroup&
 
 void FMovieSceneTrackEvaluator::CallSetupTearDown(IMovieScenePlayer& Player, FDelayedPreAnimatedStateRestore* DelayedRestore)
 {
+	using namespace UE::MovieScene;
+
 	MOVIESCENE_DETAILED_SCOPE_CYCLE_COUNTER(MovieSceneEval_CallSetupTearDown);
 
 	UMovieSceneEntitySystemLinker* Linker = Player.GetEvaluationTemplate().GetEntitySystemLinker();
+	FPreAnimatedTemplateCaptureSources* TemplateMetaData = Linker->PreAnimatedState.GetTemplateMetaData();
+	const FRootInstanceHandle RootInstanceHandle = Player.GetEvaluationTemplate().GetRootInstanceHandle();
 
 	FPersistentEvaluationData PersistentDataProxy(Player);
 
@@ -370,17 +380,17 @@ void FMovieSceneTrackEvaluator::CallSetupTearDown(IMovieScenePlayer& Player, FDe
 			{
 				DelayedRestore->Add(Key);
 			}
-			else
+			else if (TemplateMetaData)
 			{
-				Player.PreAnimatedState.OnFinishedEvaluating(Key);
+				TemplateMetaData->StopTrackingCaptureSource(Key, RootInstanceHandle);
 			}
 		}
-		else
+		else if (TemplateMetaData)
 		{
 			// If the track has been destroyed since it was last evaluated, we can still restore state for anything it made
 			// In particular this is needed for movie renders, where it will enable/disable shots between cuts in order
 			// to render handle frames
-			Player.PreAnimatedState.OnFinishedEvaluating(Key);
+			TemplateMetaData->StopTrackingCaptureSource(Key, RootInstanceHandle);
 		}
 	}
 
diff --git a/Engine/Source/Runtime/MovieScene/Private/Evaluation/MovieScenePreAnimatedState.cpp b/Engine/Source/Runtime/MovieScene/Private/Evaluation/MovieScenePreAnimatedState.cpp
--- a/Engine/Source/Runtime/MovieScene/Private/Evaluation/MovieScenePreAnimatedState.cpp
+++ b/Engine/Source/Runtime/MovieScene/Private/Evaluation/MovieScenePreAnimatedState.cpp
@@ -9,43 +9,9 @@
 #include "Evaluation/PreAnimatedState/MovieScenePreAnimatedCaptureSources.h"
 #include "Evaluation/PreAnimatedState/MovieScenePreAnimatedObjectTokenStorage.h"
 #include "Evaluation/PreAnimatedState/MovieScenePreAnimatedRootTokenStorage.h"
-#include "Evaluation/PreAnimatedState/MovieScenePreAnimatedEntityCaptureSource.h"
 
 DECLARE_CYCLE_STAT(TEXT("Save Pre Animated State"), MovieSceneEval_SavePreAnimatedState, STATGROUP_MovieSceneEval);
 
-namespace UE
-{
-namespace MovieScene
-{
-
-IMovieScenePlayer* FRestoreStateParams::GetTerminalPlayer() const
-{
-	if (Linker && TerminalInstanceHandle.IsValid())
-	{
-		return Linker->GetInstanceRegistry()->GetInstance(TerminalInstanceHandle).GetPlayer();
-	}
-
-	ensureAlways(false);
-	return nullptr;
-}
-
-TSharedPtr<const FSharedPlaybackState> FRestoreStateParams::GetTerminalPlaybackState() const
-{
-	if (Linker && TerminalInstanceHandle.IsValid())
-	{
-		const FSequenceInstance& TerminalInstance = Linker->GetInstanceRegistry()->GetInstance(TerminalInstanceHandle);
-		return TerminalInstance.GetSharedPlaybackState().ToSharedPtr();
-	}
-
-	ensureAlways(false);
-	return nullptr;
-}
-
-} // namespace MovieScene
-} // namespace UE
-
-
-
 FMovieScenePreAnimatedState::~FMovieScenePreAnimatedState()
 {
 	// Ensure that the global state capture request is removed
@@ -71,9 +37,6 @@ void FMovieScenePreAnimatedState::Initialize(UMovieSceneEntitySystemLinker* Link
 
 	bCapturingGlobalPreAnimatedState = false;
 
-	TemplateMetaData = nullptr;
-	EvaluationHookMetaData = nullptr;
-
 	WeakLinker = Linker;
 	InstanceHandle = InInstanceHandle;
 }
@@ -158,22 +121,6 @@ void FMovieScenePreAnimatedState::DiscardPreAnimatedState()
 	}
 }
 
-void FMovieScenePreAnimatedState::OnFinishedEvaluating(const FMovieSceneEvaluationKey& Key)
-{
-	if (TemplateMetaData)
-	{
-		TemplateMetaData->StopTrackingCaptureSource(Key);
-	}
-}
-
-void FMovieScenePreAnimatedState::OnFinishedEvaluating(const UObject* EvaluationHook, FMovieSceneSequenceID SequenceID)
-{
-	if (EvaluationHookMetaData)
-	{
-		EvaluationHookMetaData->StopTrackingCaptureSource(EvaluationHook, SequenceID);
-	}
-}
-
 void FMovieScenePreAnimatedState::RestorePreAnimatedState(UObject& Object)
 {
 	using namespace UE::MovieScene;
diff --git a/Engine/Source/Runtime/MovieScene/Private/Evaluation/PreAnimatedState/MovieScenePreAnimatedCaptureSource.cpp b/Engine/Source/Runtime/MovieScene/Private/Evaluation/PreAnimatedState/MovieScenePreAnimatedCaptureSource.cpp
--- a/Engine/Source/Runtime/MovieScene/Private/Evaluation/PreAnimatedState/MovieScenePreAnimatedCaptureSource.cpp
+++ b/Engine/Source/Runtime/MovieScene/Private/Evaluation/PreAnimatedState/MovieScenePreAnimatedCaptureSource.cpp
@@ -4,13 +4,34 @@
 #include "Evaluation/PreAnimatedState/MovieScenePreAnimatedCaptureSources.inl"
 #include "Evaluation/MovieScenePreAnimatedState.h"
 #include "EntitySystem/MovieSceneEntitySystemLinker.h"
+#include "EntitySystem/MovieSceneSharedPlaybackState.h"
 #include "MovieSceneSection.h"
 #include "MovieSceneSequence.h"
 
+FScopedPreAnimatedCaptureSource::FScopedPreAnimatedCaptureSource(TSharedRef<const FSharedPlaybackState> SharedPlaybackState, const FMovieSceneEvaluationKey& InEvalKey, bool bInWantsRestoreState)
+	: Variant(TInPlaceType<FPreAnimatedEvaluationKeyType>(), FPreAnimatedEvaluationKeyType{ InEvalKey, SharedPlaybackState->GetRootInstanceHandle() })
+	, RootInstanceHandle(SharedPlaybackState->GetRootInstanceHandle())
+	, WeakLinker(SharedPlaybackState->GetLinker())
+	, bWantsRestoreState(bInWantsRestoreState)
+{
+	FScopedPreAnimatedCaptureSource*& CaptureSourcePtr = GetCaptureSourcePtr();
+	PrevCaptureSource = CaptureSourcePtr;
+	CaptureSourcePtr = this;
+}
 FScopedPreAnimatedCaptureSource::FScopedPreAnimatedCaptureSource(FMovieScenePreAnimatedState* InPreAnimatedState, const FMovieSceneEvaluationKey& InEvalKey, bool bInWantsRestoreState)
-	: Variant(TInPlaceType<FMovieSceneEvaluationKey>(), InEvalKey)
-	, WeakLinker(InPreAnimatedState ? InPreAnimatedState->GetLinker() : nullptr)
-	, OptionalSequencePreAnimatedState(InPreAnimatedState)
+	: Variant(TInPlaceType<FPreAnimatedEvaluationKeyType>(), FPreAnimatedEvaluationKeyType{ InEvalKey, InPreAnimatedState->InstanceHandle })
+	, RootInstanceHandle(InPreAnimatedState->InstanceHandle)
+	, WeakLinker(InPreAnimatedState->GetLinker())
+	, bWantsRestoreState(bInWantsRestoreState)
+{
+	FScopedPreAnimatedCaptureSource*& CaptureSourcePtr = GetCaptureSourcePtr();
+	PrevCaptureSource = CaptureSourcePtr;
+	CaptureSourcePtr = this;
+}
+FScopedPreAnimatedCaptureSource::FScopedPreAnimatedCaptureSource(TSharedRef<const FSharedPlaybackState> SharedPlaybackState, const UObject* InEvalHook, FMovieSceneSequenceID InSequenceID, bool bInWantsRestoreState)
+	: Variant(TInPlaceType<FPreAnimatedEvalHookKeyType>(), FPreAnimatedEvalHookKeyType{ InEvalHook, SharedPlaybackState->GetRootInstanceHandle(), InSequenceID } )
+	, RootInstanceHandle(SharedPlaybackState->GetRootInstanceHandle())
+	, WeakLinker(SharedPlaybackState->GetLinker())
 	, bWantsRestoreState(bInWantsRestoreState)
 {
 	FScopedPreAnimatedCaptureSource*& CaptureSourcePtr = GetCaptureSourcePtr();
@@ -18,9 +39,9 @@ FScopedPreAnimatedCaptureSource::FScopedPreAnimatedCaptureSource(FMovieScenePreA
 	CaptureSourcePtr = this;
 }
 FScopedPreAnimatedCaptureSource::FScopedPreAnimatedCaptureSource(FMovieScenePreAnimatedState* InPreAnimatedState, const UObject* InEvalHook, FMovieSceneSequenceID InSequenceID, bool bInWantsRestoreState)
-	: Variant(TInPlaceType<FEvalHookType>(), FEvalHookType{InEvalHook, InSequenceID} )
+	: Variant(TInPlaceType<FPreAnimatedEvalHookKeyType>(), FPreAnimatedEvalHookKeyType{InEvalHook, InPreAnimatedState->InstanceHandle, InSequenceID} )
+	, RootInstanceHandle(InPreAnimatedState->InstanceHandle)
 	, WeakLinker(InPreAnimatedState ? InPreAnimatedState->GetLinker() : nullptr)
-	, OptionalSequencePreAnimatedState(InPreAnimatedState)
 	, bWantsRestoreState(bInWantsRestoreState)
 {
 	FScopedPreAnimatedCaptureSource*& CaptureSourcePtr = GetCaptureSourcePtr();
@@ -30,7 +51,6 @@ FScopedPreAnimatedCaptureSource::FScopedPreAnimatedCaptureSource(FMovieScenePreA
 FScopedPreAnimatedCaptureSource::FScopedPreAnimatedCaptureSource(UMovieSceneEntitySystemLinker* InLinker, UMovieSceneTrackInstance* InTrackInstance, bool bInWantsRestoreState)
 	: Variant(TInPlaceType<UMovieSceneTrackInstance*>(), InTrackInstance)
 	, WeakLinker(InLinker)
-	, OptionalSequencePreAnimatedState(nullptr)
 	, bWantsRestoreState(bInWantsRestoreState)
 {
 	FScopedPreAnimatedCaptureSource*& CaptureSourcePtr = GetCaptureSourcePtr();
@@ -40,7 +60,6 @@ FScopedPreAnimatedCaptureSource::FScopedPreAnimatedCaptureSource(UMovieSceneEnti
 FScopedPreAnimatedCaptureSource::FScopedPreAnimatedCaptureSource(UMovieSceneEntitySystemLinker* InLinker, const FMovieSceneTrackInstanceInput& TrackInstanceInput)
 	: Variant(TInPlaceType<FMovieSceneTrackInstanceInput>(), TrackInstanceInput)
 	, WeakLinker(InLinker)
-	, OptionalSequencePreAnimatedState(nullptr)
 {
 	EMovieSceneCompletionMode CompletionMode = TrackInstanceInput.Section->GetCompletionMode();
 	if (CompletionMode == EMovieSceneCompletionMode::ProjectDefault)
@@ -69,11 +88,10 @@ FScopedPreAnimatedCaptureSource*& FScopedPreAnimatedCaptureSource::GetCaptureSou
 
 UE::MovieScene::FRootInstanceHandle FScopedPreAnimatedCaptureSource::GetRootInstanceHandle(UMovieSceneEntitySystemLinker* Linker) const
 {
-	if (OptionalSequencePreAnimatedState)
+	if (RootInstanceHandle.IsValid())
 	{
-		return OptionalSequencePreAnimatedState->InstanceHandle;
+		return RootInstanceHandle;
 	}
-
 	if (const FMovieSceneTrackInstanceInput* TrackInstanceInput = Variant.TryGet<FMovieSceneTrackInstanceInput>())
 	{
 		return Linker->GetInstanceRegistry()->GetInstance(TrackInstanceInput->InstanceHandle).GetRootInstanceHandle();
@@ -88,36 +106,27 @@ void FScopedPreAnimatedCaptureSource::BeginTracking(const UE::MovieScene::FPreAn
 	ensureMsgf(!WeakLinker.IsValid() || WeakLinker.Get() == Linker, 
 			TEXT("Attempting to track state on a capture source related to a different linker. Are you missing setting a scoped capture source?"));
 
-	if (FMovieSceneEvaluationKey* EvalKey = Variant.TryGet<FMovieSceneEvaluationKey>())
+	// All capture sources' meta-data is shared between all players
+
+	if (FPreAnimatedEvaluationKeyType* EvalKey = Variant.TryGet<FPreAnimatedEvaluationKeyType>())
 	{
-		check(OptionalSequencePreAnimatedState);
-		// Make the association to this track template key
-		if (!OptionalSequencePreAnimatedState->TemplateMetaData)
-		{
-			OptionalSequencePreAnimatedState->TemplateMetaData = MakeShared<FPreAnimatedTemplateCaptureSources>(&Linker->PreAnimatedState);
-			Linker->PreAnimatedState.AddWeakCaptureSource(OptionalSequencePreAnimatedState->TemplateMetaData);
-		}
-		OptionalSequencePreAnimatedState->TemplateMetaData->BeginTrackingCaptureSource(*EvalKey, MetaData);
+		FPreAnimatedTemplateCaptureSources* TemplateMetaData = Linker->PreAnimatedState.GetOrCreateTemplateMetaData();
+		ensureMsgf(EvalKey->RootInstanceHandle == MetaData.RootInstanceHandle, TEXT("Mismatched root handles between the scoped evaluation key and the given metadata"));
+		TemplateMetaData->BeginTrackingCaptureSource(*EvalKey, MetaData);
 	}
-	else if (FScopedPreAnimatedCaptureSource::FEvalHookType* EvalHook = Variant.TryGet<FScopedPreAnimatedCaptureSource::FEvalHookType>())
+	else if (FPreAnimatedEvalHookKeyType* EvalHook = Variant.TryGet<FPreAnimatedEvalHookKeyType>())
 	{
-		check(OptionalSequencePreAnimatedState);
-		if (!OptionalSequencePreAnimatedState->EvaluationHookMetaData)
-		{
-			OptionalSequencePreAnimatedState->EvaluationHookMetaData = MakeShared<FPreAnimatedEvaluationHookCaptureSources>(&Linker->PreAnimatedState);
-			Linker->PreAnimatedState.AddWeakCaptureSource(OptionalSequencePreAnimatedState->EvaluationHookMetaData);
-		}
-		OptionalSequencePreAnimatedState->EvaluationHookMetaData->BeginTrackingCaptureSource(EvalHook->EvalHook, EvalHook->SequenceID, MetaData);
+		FPreAnimatedEvaluationHookCaptureSources* EvaluationHookMetaData = Linker->PreAnimatedState.GetOrCreateEvaluationHookMetaData();
+		ensureMsgf(EvalHook->RootInstanceHandle == MetaData.RootInstanceHandle, TEXT("Mismatched root handles between the scoped evaluation hook and the given metadata"));
+		EvaluationHookMetaData->BeginTrackingCaptureSource(*EvalHook, MetaData);
 	}
 	else if (UMovieSceneTrackInstance* const * TrackInstance = Variant.TryGet<UMovieSceneTrackInstance*>())
 	{
-		// Track instance meta-data is shared between all players
 		FPreAnimatedTrackInstanceCaptureSources* TrackInstanceMetaData = Linker->PreAnimatedState.GetOrCreateTrackInstanceMetaData();
 		TrackInstanceMetaData->BeginTrackingCaptureSource(*TrackInstance, MetaData);
 	}
 	else if (const FMovieSceneTrackInstanceInput* TrackInstanceInput = Variant.TryGet<FMovieSceneTrackInstanceInput>())
 	{
-		// Track instance meta-data is shared between all players
 		FPreAnimatedTrackInstanceInputCaptureSources* TrackInstanceInputMetaData = Linker->PreAnimatedState.GetOrCreateTrackInstanceInputMetaData();
 		TrackInstanceInputMetaData->BeginTrackingCaptureSource(*TrackInstanceInput, MetaData);
 	}
diff --git a/Engine/Source/Runtime/MovieScene/Private/Evaluation/PreAnimatedState/MovieScenePreAnimatedCaptureSources.cpp b/Engine/Source/Runtime/MovieScene/Private/Evaluation/PreAnimatedState/MovieScenePreAnimatedCaptureSources.cpp
--- a/Engine/Source/Runtime/MovieScene/Private/Evaluation/PreAnimatedState/MovieScenePreAnimatedCaptureSources.cpp
+++ b/Engine/Source/Runtime/MovieScene/Private/Evaluation/PreAnimatedState/MovieScenePreAnimatedCaptureSources.cpp
@@ -15,8 +15,33 @@ namespace MovieScene
 
 // Explicit, exported template instantiation for these classes
 template struct MOVIESCENE_API TPreAnimatedCaptureSources<FObjectKey>;
-template struct MOVIESCENE_API TPreAnimatedCaptureSources<FMovieSceneEvaluationKey>;
+template struct MOVIESCENE_API TPreAnimatedCaptureSources<FPreAnimatedEvaluationKeyType>;
 template struct MOVIESCENE_API TPreAnimatedCaptureSources<FPreAnimatedEvalHookKeyType>;
+template struct MOVIESCENE_API TPreAnimatedCaptureSources<FMovieSceneEntityID>;
+
+FPreAnimatedEntityCaptureSource::FPreAnimatedEntityCaptureSource(FPreAnimatedStateExtension* InOwner)
+	: TPreAnimatedCaptureSources<FMovieSceneEntityID>(InOwner)
+{}
+
+void FPreAnimatedEntityCaptureSource::BeginTrackingEntity(const FPreAnimatedStateEntry& Entry, FMovieSceneEntityID EntityID, FRootInstanceHandle RootInstanceHandle, bool bWantsRestoreState)
+{
+	FPreAnimatedStateMetaData MetaData{ Entry, RootInstanceHandle, bWantsRestoreState };
+	BeginTrackingCaptureSource(EntityID, MetaData);
+}
+
+FPreAnimatedTemplateCaptureSources::FPreAnimatedTemplateCaptureSources(FPreAnimatedStateExtension* InOwner)
+	: TPreAnimatedCaptureSources<FPreAnimatedEvaluationKeyType>(InOwner)
+{}
+
+EPreAnimatedCaptureSourceState FPreAnimatedTemplateCaptureSources::BeginTrackingCaptureSource(const FMovieSceneEvaluationKey& EvaluationKey, const FPreAnimatedStateMetaData& MetaData)
+{
+	return TPreAnimatedCaptureSources<FPreAnimatedEvaluationKeyType>::BeginTrackingCaptureSource(FPreAnimatedEvaluationKeyType{ EvaluationKey, MetaData.RootInstanceHandle }, MetaData);
+}
+
+void FPreAnimatedTemplateCaptureSources::StopTrackingCaptureSource(const FMovieSceneEvaluationKey& EvaluationKey, FRootInstanceHandle RootInstanceHandle)
+{
+	TPreAnimatedCaptureSources<FPreAnimatedEvaluationKeyType>::StopTrackingCaptureSource(FPreAnimatedEvaluationKeyType{ EvaluationKey, RootInstanceHandle });
+}
 
 FPreAnimatedEvaluationHookCaptureSources::FPreAnimatedEvaluationHookCaptureSources(FPreAnimatedStateExtension* InOwner)
 	: TPreAnimatedCaptureSources<FPreAnimatedEvalHookKeyType>(InOwner)
@@ -24,16 +49,14 @@ FPreAnimatedEvaluationHookCaptureSources::FPreAnimatedEvaluationHookCaptureSourc
 
 EPreAnimatedCaptureSourceState FPreAnimatedEvaluationHookCaptureSources::BeginTrackingCaptureSource(const UObject* Hook, FMovieSceneSequenceID SequenceID, const FPreAnimatedStateMetaData& MetaData)
 {
-	return TPreAnimatedCaptureSources<FPreAnimatedEvalHookKeyType>::BeginTrackingCaptureSource(FPreAnimatedEvalHookKeyType{ Hook, SequenceID }, MetaData);
+	return TPreAnimatedCaptureSources<FPreAnimatedEvalHookKeyType>::BeginTrackingCaptureSource(FPreAnimatedEvalHookKeyType{ Hook, MetaData.RootInstanceHandle, SequenceID }, MetaData);
 }
 
-void FPreAnimatedEvaluationHookCaptureSources::StopTrackingCaptureSource(const UObject* Hook, FMovieSceneSequenceID SequenceID)
+void FPreAnimatedEvaluationHookCaptureSources::StopTrackingCaptureSource(const UObject* Hook, FRootInstanceHandle RootInstanceHandle, FMovieSceneSequenceID SequenceID)
 {
-	TPreAnimatedCaptureSources<FPreAnimatedEvalHookKeyType>::StopTrackingCaptureSource(FPreAnimatedEvalHookKeyType{ Hook, SequenceID });
+	TPreAnimatedCaptureSources<FPreAnimatedEvalHookKeyType>::StopTrackingCaptureSource(FPreAnimatedEvalHookKeyType{ Hook, RootInstanceHandle, SequenceID });
 }
 
-
-
 FPreAnimatedTrackInstanceCaptureSources::FPreAnimatedTrackInstanceCaptureSources(FPreAnimatedStateExtension* InOwner)
 	: TPreAnimatedCaptureSources<FObjectKey>(InOwner)
 {}
@@ -54,10 +77,6 @@ FPreAnimatedTrackInstanceInputCaptureSources::FPreAnimatedTrackInstanceInputCapt
 	: TPreAnimatedCaptureSources<FMovieSceneTrackInstanceInput>(InOwner)
 {}
 
-FPreAnimatedTemplateCaptureSources::FPreAnimatedTemplateCaptureSources(FPreAnimatedStateExtension* InOwner)
-	: TPreAnimatedCaptureSources<FMovieSceneEvaluationKey>(InOwner)
-{}
-
 
 } // namespace MovieScene
 } // namespace UE
diff --git a/Engine/Source/Runtime/MovieScene/Private/Evaluation/PreAnimatedState/MovieScenePreAnimatedCaptureSources.inl b/Engine/Source/Runtime/MovieScene/Private/Evaluation/PreAnimatedState/MovieScenePreAnimatedCaptureSources.inl
--- a/Engine/Source/Runtime/MovieScene/Private/Evaluation/PreAnimatedState/MovieScenePreAnimatedCaptureSources.inl
+++ b/Engine/Source/Runtime/MovieScene/Private/Evaluation/PreAnimatedState/MovieScenePreAnimatedCaptureSources.inl
@@ -39,6 +39,29 @@ EPreAnimatedCaptureSourceState TPreAnimatedCaptureSources<KeyType>::BeginTrackin
 	return EPreAnimatedCaptureSourceState::New;
 }
 
+template<typename KeyType>
+void TPreAnimatedCaptureSources<KeyType>::StopTrackingCaptureSource(const KeyType& InKey, FPreAnimatedStorageID InStorageID)
+{
+	FPreAnimatedStateMetaDataArray* Array = KeyToMetaData.Find(InKey);
+	if (Array)
+	{
+		for (int32 Index = Array->Num()-1; Index >= 0; --Index)
+		{
+			const FPreAnimatedStateMetaData& MetaData = (*Array)[Index];
+			if (MetaData.Entry.ValueHandle.TypeID == InStorageID)
+			{
+				Owner->RemoveMetaData(MetaData);
+				Array->RemoveAt(Index, 1, false);
+			}
+		}
+
+		if (Array->Num() == 0)
+		{
+			KeyToMetaData.Remove(InKey);
+		}
+	}
+}
+
 template<typename KeyType>
 void TPreAnimatedCaptureSources<KeyType>::StopTrackingCaptureSource(const KeyType& InKey)
 {
@@ -139,6 +162,29 @@ void TPreAnimatedCaptureSources<KeyType>::GatherAndRemoveMetaDataForStorage(FPre
 	}
 }
 
+template<typename KeyType>
+void TPreAnimatedCaptureSources<KeyType>::GatherAndRemoveMetaDataForRootInstance(FRootInstanceHandle InstanceHandle, TArray<FPreAnimatedStateMetaData>& OutExpiredMetaData)
+{
+	for (auto It = KeyToMetaData.CreateIterator(); It; ++It)
+	{
+		FPreAnimatedStateMetaDataArray& Array = It.Value();
+		for (int32 Index = Array.Num()-1; Index >= 0; --Index)
+		{
+			const FPreAnimatedStateMetaData& MetaData = Array[Index];
+			if (MetaData.RootInstanceHandle == InstanceHandle)
+			{
+				OutExpiredMetaData.Add(MetaData);
+				Array.RemoveAt(Index, 1, false);
+			}
+		}
+
+		if (Array.Num() == 0)
+		{
+			It.RemoveCurrent();
+		}
+	}
+}
+
 template<typename KeyType>
 bool TPreAnimatedCaptureSources<KeyType>::ContainsInstanceHandle(FRootInstanceHandle RootInstanceHandle) const
 {
diff --git a/Engine/Source/Runtime/MovieScene/Private/Evaluation/PreAnimatedState/MovieScenePreAnimatedEntityCaptureSource.cpp b/Engine/Source/Runtime/MovieScene/Private/Evaluation/PreAnimatedState/MovieScenePreAnimatedEntityCaptureSource.cpp
deleted file mode 100644
diff --git a/Engine/Source/Runtime/MovieScene/Private/Evaluation/PreAnimatedState/MovieScenePreAnimatedStateExtension.cpp b/Engine/Source/Runtime/MovieScene/Private/Evaluation/PreAnimatedState/MovieScenePreAnimatedStateExtension.cpp
--- a/Engine/Source/Runtime/MovieScene/Private/Evaluation/PreAnimatedState/MovieScenePreAnimatedStateExtension.cpp
+++ b/Engine/Source/Runtime/MovieScene/Private/Evaluation/PreAnimatedState/MovieScenePreAnimatedStateExtension.cpp
@@ -1,15 +1,16 @@
 // Copyright Epic Games, Inc. All Rights Reserved.
 
 #include "Evaluation/PreAnimatedState/MovieScenePreAnimatedStateExtension.h"
-#include "Evaluation/PreAnimatedState/MovieScenePreAnimatedEntityCaptureSource.h"
 #include "Evaluation/PreAnimatedState/MovieScenePreAnimatedCaptureSources.h"
 #include "Evaluation/PreAnimatedState/MovieSceneRestoreStateParams.h"
 #include "Evaluation/PreAnimatedState/MovieScenePreAnimatedCaptureSource.h"
 #include "Evaluation/PreAnimatedState/MovieScenePreAnimatedObjectTokenStorage.h"
 #include "Evaluation/PreAnimatedState/MovieScenePreAnimatedRootTokenStorage.h"
-#include "EntitySystem/MovieSceneEntitySystemTask.h"
-#include "EntitySystem/MovieSceneComponentPtr.h"
 #include "EntitySystem/BuiltInComponentTypes.h"
+#include "EntitySystem/MovieSceneComponentPtr.h"
+#include "EntitySystem/MovieSceneEntitySystemTask.h"
+#include "EntitySystem/MovieSceneSharedPlaybackState.h"
+#include "MovieSceneSequence.h"
 #include "Algo/Find.h"
 #include "Algo/IndexOf.h"
 #include "Algo/BinarySearch.h"
@@ -276,6 +277,34 @@
 	return TrackInstanceInputCaptureSource.Get();
 }
 
+FPreAnimatedTemplateCaptureSources* FPreAnimatedStateExtension::GetTemplateMetaData() const
+{
+	return TemplateCaptureSource.Get();
+}
+
+FPreAnimatedTemplateCaptureSources* FPreAnimatedStateExtension::GetOrCreateTemplateMetaData()
+{
+	if (!TemplateCaptureSource)
+	{
+		TemplateCaptureSource = MakeUnique<FPreAnimatedTemplateCaptureSources>(this);
+	}
+	return TemplateCaptureSource.Get();
+}
+
+FPreAnimatedEvaluationHookCaptureSources* FPreAnimatedStateExtension::GetEvaluationHookMetaData() const
+{
+	return EvaluationHookCaptureSource.Get();
+}
+
+FPreAnimatedEvaluationHookCaptureSources* FPreAnimatedStateExtension::GetOrCreateEvaluationHookMetaData()
+{
+	if (!EvaluationHookCaptureSource)
+	{
+		EvaluationHookCaptureSource = MakeUnique<FPreAnimatedEvaluationHookCaptureSources>(this);
+	}
+	return EvaluationHookCaptureSource.Get();
+}
+
 void FPreAnimatedStateExtension::AddWeakCaptureSource(TWeakPtr<IPreAnimatedCaptureSource> InWeakCaptureSource)
 {
 	WeakExternalCaptureSources.Add(InWeakCaptureSource);
@@ -319,6 +348,16 @@
 		TrackInstanceInputMetaData->GatherAndRemoveExpiredMetaData(Params, ExpiredMetaData);
 	}
 
+	if (FPreAnimatedTemplateCaptureSources* TemplateMetaData = GetTemplateMetaData())
+	{
+		TemplateMetaData->GatherAndRemoveExpiredMetaData(Params, ExpiredMetaData);
+	}
+
+	if (FPreAnimatedEvaluationHookCaptureSources* EvaluationHookMetaData = GetEvaluationHookMetaData())
+	{
+		EvaluationHookMetaData->GatherAndRemoveExpiredMetaData(Params, ExpiredMetaData);
+	}
+
 	for (int32 Index = WeakExternalCaptureSources.Num()-1; Index >= 0; --Index)
 	{
 		TSharedPtr<IPreAnimatedCaptureSource> MetaData = WeakExternalCaptureSources[Index].Pin();
@@ -333,80 +372,11 @@
 		}
 	}
 
-	// Remove all contributions
-	for (const FPreAnimatedStateMetaData& MetaData : ExpiredMetaData)
-	{
-		FAggregatePreAnimatedStateMetaData* Aggregate = FindMetaData(MetaData.Entry);
-		if (ensure(Aggregate))
-		{
-			const int32 TotalNum = --Aggregate->NumContributors;
-			if (MetaData.bWantsRestoreState)
+	HandleMetaDataToRemove(Params, ExpiredMetaData,
+			[Params](IPreAnimatedStorage& Storage, FPreAnimatedStorageIndex StorageIndex)
 			{
-				--Aggregate->NumRestoreContributors;
-			}
-
-			if (TotalNum == 0)
-			{
-				Aggregate->bWantedRestore = false;
-				Aggregate->TerminalInstanceHandle = MetaData.RootInstanceHandle;
-			}
-		}
-	}
-
-	// Ensure that the entries are restored in strictly the reverse order they were cached in
-	for (int32 Index = 0; Index < GroupMetaData.GetMaxIndex(); ++Index)
-	{
-		if (!GroupMetaData.IsAllocated(Index))
-		{
-			continue;
-		}
-
-		FPreAnimatedGroupMetaData& Group = GroupMetaData[Index];
-
-		for (int32 AggregateIndex = Group.AggregateMetaData.Num()-1; AggregateIndex >= 0; --AggregateIndex)
-		{
-			FAggregatePreAnimatedStateMetaData& Aggregate = Group.AggregateMetaData[AggregateIndex];
-			if (Aggregate.NumContributors == 0 && (!Aggregate.TerminalInstanceHandle.IsValid() || Aggregate.TerminalInstanceHandle == Params.TerminalInstanceHandle))
-			{
-				TSharedPtr<IPreAnimatedStorage> Storage = GetStorageChecked(Aggregate.ValueHandle.TypeID);
-				Storage->RestorePreAnimatedStateStorage(Aggregate.ValueHandle.StorageIndex, EPreAnimatedStorageRequirement::Persistent, EPreAnimatedStorageRequirement::None, Params);
-
-				Group.AggregateMetaData.RemoveAt(AggregateIndex, 1, false);
-			}
-
-			if (Group.AggregateMetaData.Num() == 0)
-			{
-				// Remove at will not re-allocate the array or shuffle items within the sparse array, so this is safe
-				Group.GroupManagerPtr->OnGroupDestroyed(Index);
-				GroupMetaData.RemoveAt(Index);
-			}
-		}
-	}
-
-	for (auto UngroupedIt = UngroupedMetaData.CreateIterator(); UngroupedIt; ++UngroupedIt)
-	{
-		FAggregatePreAnimatedStateMetaData& Aggregate = UngroupedIt.Value();
-		if (Aggregate.NumContributors == 0 && (!Aggregate.TerminalInstanceHandle.IsValid() || Aggregate.TerminalInstanceHandle == Params.TerminalInstanceHandle))
-		{
-			TSharedPtr<IPreAnimatedStorage> Storage = GetStorageChecked(Aggregate.ValueHandle.TypeID);
-			Storage->RestorePreAnimatedStateStorage(Aggregate.ValueHandle.StorageIndex, EPreAnimatedStorageRequirement::Persistent, EPreAnimatedStorageRequirement::None, Params);
-
-			UngroupedIt.RemoveCurrent();
-		}
-	}
-
-	GroupMetaData.Shrink();
-
-	// Invalidate cached data for any sequence instance that belongs to the terminal instance
-	if (Params.TerminalInstanceHandle.IsValid() && 
-			ensureMsgf(
-				Linker->GetInstanceRegistry()->IsHandleValid(Params.TerminalInstanceHandle),
-				TEXT("Terminal instance handle is not valid anymore, was the sequence destroyed?")))
-	{
-		Linker->GetInstanceRegistry()->MutateInstance(Params.TerminalInstanceHandle).InvalidateCachedData();
-	}
-
-	bEntriesInvalidated = true;
+				Storage.RestorePreAnimatedStateStorage(StorageIndex, EPreAnimatedStorageRequirement::Persistent, EPreAnimatedStorageRequirement::None, Params);
+			});
 }
 
 void FPreAnimatedStateExtension::DiscardStaleObjectState()
@@ -449,6 +419,16 @@
 		TrackInstanceInputMetaData->GatherAndRemoveExpiredMetaData(Params, ExpiredMetaData);
 	}
 
+	if (FPreAnimatedTemplateCaptureSources* TemplateMetaData = GetTemplateMetaData())
+	{
+		TemplateMetaData->GatherAndRemoveExpiredMetaData(Params, ExpiredMetaData);
+	}
+
+	if (FPreAnimatedEvaluationHookCaptureSources* EvaluationHookMetaData = GetEvaluationHookMetaData())
+	{
+		EvaluationHookMetaData->GatherAndRemoveExpiredMetaData(Params, ExpiredMetaData);
+	}
+
 	for (int32 Index = WeakExternalCaptureSources.Num() - 1; Index >= 0; --Index)
 	{
 		TSharedPtr<IPreAnimatedCaptureSource> MetaData = WeakExternalCaptureSources[Index].Pin();
@@ -463,77 +443,11 @@
 		}
 	}
 
-	// Remove all contributions
-	for (const FPreAnimatedStateMetaData& MetaData : ExpiredMetaData)
-	{
-		FAggregatePreAnimatedStateMetaData* Aggregate = FindMetaData(MetaData.Entry);
-		if (ensure(Aggregate))
-		{
-			const int32 TotalNum = --Aggregate->NumContributors;
-			if (MetaData.bWantsRestoreState)
+	HandleMetaDataToRemove(Params, ExpiredMetaData,
+			[](IPreAnimatedStorage& Storage, FPreAnimatedStorageIndex StorageIndex)
 			{
-				--Aggregate->NumRestoreContributors;
-			}
-
-			if (TotalNum == 0)
-			{
-				Aggregate->bWantedRestore = false;
-				Aggregate->TerminalInstanceHandle = MetaData.RootInstanceHandle;
-			}
-		}
-	}
-
-	// Ensure that the entries are discarded in strictly the reverse order they were cached in
-	for (int32 Index = 0; Index < GroupMetaData.GetMaxIndex(); ++Index)
-	{
-		if (!GroupMetaData.IsAllocated(Index))
-		{
-			continue;
-		}
-
-		FPreAnimatedGroupMetaData& Group = GroupMetaData[Index];
-
-		for (int32 AggregateIndex = Group.AggregateMetaData.Num() - 1; AggregateIndex >= 0; --AggregateIndex)
-		{
-			FAggregatePreAnimatedStateMetaData& Aggregate = Group.AggregateMetaData[AggregateIndex];
-			if (Aggregate.NumContributors == 0 && (!Aggregate.TerminalInstanceHandle.IsValid() || Aggregate.TerminalInstanceHandle == Params.TerminalInstanceHandle))
-			{
-				TSharedPtr<IPreAnimatedStorage> Storage = GetStorageChecked(Aggregate.ValueHandle.TypeID);
-				Storage->DiscardPreAnimatedStateStorage(Aggregate.ValueHandle.StorageIndex, EPreAnimatedStorageRequirement::Persistent);
-
-				Group.AggregateMetaData.RemoveAt(AggregateIndex, 1, false);
-			}
-
-			if (Group.AggregateMetaData.Num() == 0)
-			{
-				// Remove at will not re-allocate the array or shuffle items within the sparse array, so this is safe
-				Group.GroupManagerPtr->OnGroupDestroyed(Index);
-				GroupMetaData.RemoveAt(Index);
-			}
-		}
-	}
-
-	for (auto UngroupedIt = UngroupedMetaData.CreateIterator(); UngroupedIt; ++UngroupedIt)
-	{
-		FAggregatePreAnimatedStateMetaData& Aggregate = UngroupedIt.Value();
-		if (Aggregate.NumContributors == 0 && (!Aggregate.TerminalInstanceHandle.IsValid() || Aggregate.TerminalInstanceHandle == Params.TerminalInstanceHandle))
-		{
-			TSharedPtr<IPreAnimatedStorage> Storage = GetStorageChecked(Aggregate.ValueHandle.TypeID);
-			Storage->DiscardPreAnimatedStateStorage(Aggregate.ValueHandle.StorageIndex, EPreAnimatedStorageRequirement::Persistent);
-
-			UngroupedIt.RemoveCurrent();
-		}
-	}
-
-	GroupMetaData.Shrink();
-
-	// Invalidate cached data for any sequence instance that belongs to the terminal instance
-	if (Params.TerminalInstanceHandle.IsValid())
-	{
-		Linker->GetInstanceRegistry()->MutateInstance(Params.TerminalInstanceHandle).InvalidateCachedData();
-	}
-
-	bEntriesInvalidated = true;
+				Storage.DiscardPreAnimatedStateStorage(StorageIndex, EPreAnimatedStorageRequirement::Persistent);
+			});
 }
 
 void FPreAnimatedStateExtension::DiscardTransientState()
@@ -553,6 +467,16 @@
 		TrackInstanceInputMetaData->Reset();
 	}
 
+	if (FPreAnimatedTemplateCaptureSources* TemplateMetaData = GetTemplateMetaData())
+	{
+		TemplateMetaData->Reset();
+	}
+
+	if (FPreAnimatedEvaluationHookCaptureSources* EvaluationHookMetaData = GetEvaluationHookMetaData())
+	{
+		EvaluationHookMetaData->Reset();
+	}
+
 	for (int32 Index = WeakExternalCaptureSources.Num()-1; Index >= 0; --Index)
 	{
 		if (TSharedPtr<IPreAnimatedCaptureSource> MetaData = WeakExternalCaptureSources[Index].Pin())
@@ -603,6 +527,16 @@
 		TrackInstanceInputMetaData->GatherAndRemoveMetaDataForGroup(GroupHandle, MetaDataToRemove);
 	}
 
+	if (FPreAnimatedTemplateCaptureSources* TemplateMetaData = GetTemplateMetaData())
+	{
+		TemplateMetaData->GatherAndRemoveMetaDataForGroup(GroupHandle, MetaDataToRemove);
+	}
+
+	if (FPreAnimatedEvaluationHookCaptureSources* EvaluationHookMetaData = GetEvaluationHookMetaData())
+	{
+		EvaluationHookMetaData->GatherAndRemoveMetaDataForGroup(GroupHandle, MetaDataToRemove);
+	}
+
 	for (int32 Index = WeakExternalCaptureSources.Num()-1; Index >= 0; --Index)
 	{
 		if (TSharedPtr<IPreAnimatedCaptureSource> MetaData = WeakExternalCaptureSources[Index].Pin())
@@ -645,6 +579,16 @@
 		TrackInstanceInputMetaData->GatherAndRemoveMetaDataForStorage(StorageID, StorageIndex, MetaDataToRemove);
 	}
 
+	if (FPreAnimatedTemplateCaptureSources* TemplateMetaData = GetTemplateMetaData())
+	{
+		TemplateMetaData->GatherAndRemoveMetaDataForStorage(StorageID, StorageIndex, MetaDataToRemove);
+	}
+
+	if (FPreAnimatedEvaluationHookCaptureSources* EvaluationHookMetaData = GetEvaluationHookMetaData())
+	{
+		EvaluationHookMetaData->GatherAndRemoveMetaDataForStorage(StorageID, StorageIndex, MetaDataToRemove);
+	}
+
 	for (int32 Index = WeakExternalCaptureSources.Num()-1; Index >= 0; --Index)
 	{
 		if (TSharedPtr<IPreAnimatedCaptureSource> MetaData = WeakExternalCaptureSources[Index].Pin())
@@ -735,6 +679,38 @@
 		}
 	}
 
+	if (FPreAnimatedTrackInstanceCaptureSources* TrackInstanceMetaData = GetTrackInstanceMetaData())
+	{
+		if (TrackInstanceMetaData->ContainsInstanceHandle(RootInstanceHandle))
+		{
+			return true;
+		}
+	}
+
+	if (FPreAnimatedTrackInstanceInputCaptureSources* TrackInstanceInputMetaData = GetTrackInstanceInputMetaData())
+	{
+		if (TrackInstanceInputMetaData->ContainsInstanceHandle(RootInstanceHandle))
+		{
+			return true;
+		}
+	}
+
+	if (FPreAnimatedTemplateCaptureSources* TemplateMetaData = GetTemplateMetaData())
+	{
+		if (TemplateMetaData->ContainsInstanceHandle(RootInstanceHandle))
+		{
+			return true;
+		}
+	}
+
+	if (FPreAnimatedEvaluationHookCaptureSources* EvaluationHookMetaData = GetEvaluationHookMetaData())
+	{
+		if (EvaluationHookMetaData->ContainsInstanceHandle(RootInstanceHandle))
+		{
+			return true;
+		}
+	}
+
 	for (int32 Index = WeakExternalCaptureSources.Num()-1; Index >= 0; --Index)
 	{
 		TSharedPtr<IPreAnimatedCaptureSource> MetaData = WeakExternalCaptureSources[Index].Pin();
@@ -747,6 +723,84 @@
 	return false;
 }
 
+void FPreAnimatedStateExtension::HandleMetaDataToRemove(const FRestoreStateParams& Params, TArrayView<FPreAnimatedStateMetaData> MetaDataToRemove, FContributionRemover RemoveFunc)
+{
+	// Remove all contributions
+	for (const FPreAnimatedStateMetaData& MetaData : MetaDataToRemove)
+	{
+		FAggregatePreAnimatedStateMetaData* Aggregate = FindMetaData(MetaData.Entry);
+		if (ensure(Aggregate))
+		{
+			const int32 TotalNum = --Aggregate->NumContributors;
+			if (MetaData.bWantsRestoreState)
+			{
+				--Aggregate->NumRestoreContributors;
+			}
+
+			if (TotalNum == 0)
+			{
+				Aggregate->bWantedRestore = false;
+				Aggregate->TerminalInstanceHandle = MetaData.RootInstanceHandle;
+			}
+		}
+	}
+
+	// Ensure that the entries are removed in strictly the reverse order they were cached in
+	for (int32 Index = 0; Index < GroupMetaData.GetMaxIndex(); ++Index)
+	{
+		if (!GroupMetaData.IsAllocated(Index))
+		{
+			continue;
+		}
+
+		FPreAnimatedGroupMetaData& Group = GroupMetaData[Index];
+
+		for (int32 AggregateIndex = Group.AggregateMetaData.Num() - 1; AggregateIndex >= 0; --AggregateIndex)
+		{
+			FAggregatePreAnimatedStateMetaData& Aggregate = Group.AggregateMetaData[AggregateIndex];
+			if (Aggregate.NumContributors == 0 && (!Aggregate.TerminalInstanceHandle.IsValid() || Aggregate.TerminalInstanceHandle == Params.TerminalInstanceHandle))
+			{
+				TSharedPtr<IPreAnimatedStorage> Storage = GetStorageChecked(Aggregate.ValueHandle.TypeID);
+				RemoveFunc(*Storage.Get(), Aggregate.ValueHandle.StorageIndex);
+
+				Group.AggregateMetaData.RemoveAt(AggregateIndex, 1, false);
+			}
+
+			if (Group.AggregateMetaData.Num() == 0)
+			{
+				// Remove at will not re-allocate the array or shuffle items within the sparse array, so this is safe
+				Group.GroupManagerPtr->OnGroupDestroyed(Index);
+				GroupMetaData.RemoveAt(Index);
+			}
+		}
+	}
+
+	for (auto UngroupedIt = UngroupedMetaData.CreateIterator(); UngroupedIt; ++UngroupedIt)
+	{
+		FAggregatePreAnimatedStateMetaData& Aggregate = UngroupedIt.Value();
+		if (Aggregate.NumContributors == 0 && (!Aggregate.TerminalInstanceHandle.IsValid() || Aggregate.TerminalInstanceHandle == Params.TerminalInstanceHandle))
+		{
+			TSharedPtr<IPreAnimatedStorage> Storage = GetStorageChecked(Aggregate.ValueHandle.TypeID);
+			RemoveFunc(*Storage.Get(), Aggregate.ValueHandle.StorageIndex);
+
+			UngroupedIt.RemoveCurrent();
+		}
+	}
+
+	GroupMetaData.Shrink();
+
+	// Invalidate cached data for any sequence instance that belongs to the terminal instance
+	if (Params.TerminalInstanceHandle.IsValid() && 
+			ensureMsgf(
+				Linker->GetInstanceRegistry()->IsHandleValid(Params.TerminalInstanceHandle),
+				TEXT("Terminal instance handle is not valid anymore, was the sequence destroyed?")))
+	{
+		Linker->GetInstanceRegistry()->MutateInstance(Params.TerminalInstanceHandle).InvalidateCachedData();
+	}
+
+	bEntriesInvalidated = true;
+}
+
 bool FPreAnimatedStateExtension::HasActiveCaptureSource() const
 {
 	FScopedPreAnimatedCaptureSource* CaptureSource = FScopedPreAnimatedCaptureSource::GetCaptureSourcePtr();
diff --git a/Engine/Source/Runtime/MovieScene/Private/Evaluation/PreAnimatedState/MovieSceneRestoreStateParams.cpp b/Engine/Source/Runtime/MovieScene/Private/Evaluation/PreAnimatedState/MovieSceneRestoreStateParams.cpp
new file mode 100644
--- /dev/null
+++ b/Engine/Source/Runtime/MovieScene/Private/Evaluation/PreAnimatedState/MovieSceneRestoreStateParams.cpp
@@ -0,0 +1,38 @@
+// Copyright Epic Games, Inc. All Rights Reserved.
+
+#include "Evaluation/PreAnimatedState/MovieSceneRestoreStateParams.h"
+
+#include "EntitySystem/MovieSceneEntitySystemLinker.h"
+#include "EntitySystem/MovieSceneInstanceRegistry.h"
+#include "EntitySystem/MovieSceneSequenceInstance.h"
+
+namespace UE
+{
+namespace MovieScene
+{
+
+IMovieScenePlayer* FRestoreStateParams::GetTerminalPlayer() const
+{
+	if (Linker && TerminalInstanceHandle.IsValid())
+	{
+		return Linker->GetInstanceRegistry()->GetInstance(TerminalInstanceHandle).GetPlayer();
+	}
+
+	ensureAlways(false);
+	return nullptr;
+}
+
+TSharedPtr<const FSharedPlaybackState> FRestoreStateParams::GetTerminalPlaybackState() const
+{
+	if (Linker && TerminalInstanceHandle.IsValid())
+	{
+		const FSequenceInstance& TerminalInstance = Linker->GetInstanceRegistry()->GetInstance(TerminalInstanceHandle);
+		return TerminalInstance.GetSharedPlaybackState().ToSharedPtr();
+	}
+
+	ensureAlways(false);
+	return nullptr;
+}
+
+} // namespace MovieScene
+} // namespace UE
diff --git a/Engine/Source/Runtime/MovieScene/Private/Tests/MovieScenePreAnimatedStateTests.cpp b/Engine/Source/Runtime/MovieScene/Private/Tests/MovieScenePreAnimatedStateTests.cpp
--- a/Engine/Source/Runtime/MovieScene/Private/Tests/MovieScenePreAnimatedStateTests.cpp
+++ b/Engine/Source/Runtime/MovieScene/Private/Tests/MovieScenePreAnimatedStateTests.cpp
@@ -9,6 +9,7 @@
 #include "Evaluation/MovieSceneEvaluationTemplateInstance.h"
 #include "EntitySystem/MovieSceneEntitySystemLinker.h"
 #include "EntitySystem/MovieSceneEntitySystemRunner.h"
+#include "Evaluation/PreAnimatedState/MovieScenePreAnimatedCaptureSources.h"
 #include "Evaluation/PreAnimatedState/MovieScenePreAnimatedStateExtension.h"
 #include "MovieSceneTestObjects.h"
 #include "UObject/Package.h"
@@ -94,6 +95,16 @@ namespace Impl
 		}
 	}
 
+	void OnFinishedEvaluating(UMovieSceneEntitySystemLinker* InLinker, const FMovieSceneEvaluationKey& InKey, const UE::MovieScene::FRootInstanceHandle RootInstanceHandle)
+	{
+		using namespace UE::MovieScene;
+
+		if (FPreAnimatedTemplateCaptureSources* TemplateMetaData = InLinker->PreAnimatedState.GetTemplateMetaData())
+		{
+			TemplateMetaData->StopTrackingCaptureSource(InKey, RootInstanceHandle);
+		}
+	}
+
 	void ResetValues()
 	{
 		TestValue1 = TestMagicNumber;
@@ -199,7 +210,7 @@ bool FMovieScenePreAnimatedStateEntityTest::RunTest(const FString& Parameters)
 
 	TestValue1 = 50;
 
-	State.OnFinishedEvaluating(SectionKey1);
+	OnFinishedEvaluating(Linker, SectionKey1, FRootInstanceHandle());
 	Assert(this, TestValue1, TestMagicNumber, TEXT("Section did not restore correctly."));
 
 	TestValue1 = 100;
@@ -262,15 +273,15 @@ bool FMovieScenePreAnimatedStateOverlappingEntitiesTest::RunTest(const FString&
 	}
 
 	// Restore the section first - ensure it does not restore the value (because the track is still animating it)
-	State.OnFinishedEvaluating(SectionKey1);
+	OnFinishedEvaluating(Linker, SectionKey1, FRootInstanceHandle());
 	Assert(this, TestValue1, 150, TEXT("Section 1 should not have restored."));
 
 	// Restore the track - it should not restore either, because section 2 is still active
-	State.OnFinishedEvaluating(TrackKey1);
+	OnFinishedEvaluating(Linker, TrackKey1, FRootInstanceHandle());
 	Assert(this, TestValue1, 150, TEXT("Track should not have restored."));
 
 	// Restore the section - since it's the last entity animating the object with 'RestoreState' it should restore to the orignal value
-	State.OnFinishedEvaluating(SectionKey2);
+	OnFinishedEvaluating(Linker, SectionKey2, FRootInstanceHandle());
 	Assert(this, TestValue1, 0, TEXT("Section 2 did not restore correctly."));
 
 	// Restore globally - ensure that test value goes back to the original value
@@ -311,7 +322,7 @@ bool FMovieScenePreAnimatedStateKeepThenRestoreEntityTest::RunTest(const FString
 	TestValue1 = 50;
 
 	// Restore state for the entity only - this should not do anything since we specified KeepState above
-	State.OnFinishedEvaluating(SectionKey1);
+	OnFinishedEvaluating(Linker, SectionKey1, FRootInstanceHandle());
 	Assert(this, TestValue1, 50, TEXT("Section should not have restored state."));
 
 	{
@@ -325,7 +336,7 @@ bool FMovieScenePreAnimatedStateKeepThenRestoreEntityTest::RunTest(const FString
 	TestValue1 = 100;
 
 	// Restoring section key 2 here should result in the test value being the same value that was set while section 1 was evaluating (50)
-	State.OnFinishedEvaluating(SectionKey2);
+	OnFinishedEvaluating(Linker, SectionKey2, FRootInstanceHandle());
 	Assert(this, TestValue1, 50, TEXT("Section 2 did not restore to the correct value. It should restore back to the value that was set in section 1 (it doesn't restore state)."));
 
 	// We should still have the global state of the object cached which will restore it to the original state
diff --git a/Engine/Source/Runtime/MovieScene/Public/EntitySystem/MovieSceneEvaluationHookSystem.h b/Engine/Source/Runtime/MovieScene/Public/EntitySystem/MovieSceneEvaluationHookSystem.h
--- a/Engine/Source/Runtime/MovieScene/Public/EntitySystem/MovieSceneEvaluationHookSystem.h
+++ b/Engine/Source/Runtime/MovieScene/Public/EntitySystem/MovieSceneEvaluationHookSystem.h
@@ -30,6 +30,8 @@ struct FMovieSceneEvaluationHookEvent
 	UPROPERTY()
 	FMovieSceneEvaluationHookComponent Hook;
 
+	UE::MovieScene::FRootInstanceHandle RootInstanceHandle;
+
 	FMovieSceneSequenceID SequenceID;
 
 	int32 TriggerIndex;
diff --git a/Engine/Source/Runtime/MovieScene/Public/Evaluation/MovieScenePreAnimatedState.h b/Engine/Source/Runtime/MovieScene/Public/Evaluation/MovieScenePreAnimatedState.h
--- a/Engine/Source/Runtime/MovieScene/Public/Evaluation/MovieScenePreAnimatedState.h
+++ b/Engine/Source/Runtime/MovieScene/Public/Evaluation/MovieScenePreAnimatedState.h
@@ -21,18 +21,6 @@ struct IMovieScenePreAnimatedGlobalTokenProducer;
 struct IMovieScenePreAnimatedTokenProducer;
 template <typename FuncType> class TFunctionRef;
 
-namespace UE
-{
-namespace MovieScene
-{
-
-struct FPreAnimatedEvaluationHookCaptureSources;
-struct FPreAnimatedStateEntry;
-struct FPreAnimatedTemplateCaptureSources;
-
-}
-}
-
 /**
  * Class that caches pre-animated state for objects that were manipulated by sequencer
  */
@@ -78,12 +66,6 @@ public:
 	 */
 	MOVIESCENE_API void SavePreAnimatedState(FMovieSceneAnimTypeID InTokenType, const IMovieScenePreAnimatedGlobalTokenProducer& Producer);
 
-public:
-
-	MOVIESCENE_API void OnFinishedEvaluating(const FMovieSceneEvaluationKey& Key);
-
-	MOVIESCENE_API void OnFinishedEvaluating(const UObject* EvaluationHook, FMovieSceneSequenceID SequenceID);
-
 public:
 
 	MOVIESCENE_API void RestorePreAnimatedState();
@@ -133,11 +115,6 @@ private:
 	/** Weak pointer to the linker that we're associated with */
 	TWeakObjectPtr<UMovieSceneEntitySystemLinker> WeakLinker;
 
-	/** Meta-data ledger for any pre-animated state that originates from track templates */
-	TSharedPtr<UE::MovieScene::FPreAnimatedTemplateCaptureSources> TemplateMetaData;
-	/** Meta-data ledger for any pre-animated state that originates from evaluation hooks */
-	TSharedPtr<UE::MovieScene::FPreAnimatedEvaluationHookCaptureSources> EvaluationHookMetaData;
-
 	/** The instance handle for the root sequence instance */
 	UE::MovieScene::FRootInstanceHandle InstanceHandle;
 
diff --git a/Engine/Source/Runtime/MovieScene/Public/Evaluation/PreAnimatedState/IMovieScenePreAnimatedCaptureSource.h b/Engine/Source/Runtime/MovieScene/Public/Evaluation/PreAnimatedState/IMovieScenePreAnimatedCaptureSource.h
--- a/Engine/Source/Runtime/MovieScene/Public/Evaluation/PreAnimatedState/IMovieScenePreAnimatedCaptureSource.h
+++ b/Engine/Source/Runtime/MovieScene/Public/Evaluation/PreAnimatedState/IMovieScenePreAnimatedCaptureSource.h
@@ -27,6 +27,7 @@ struct IPreAnimatedCaptureSource
 	virtual void GatherAndRemoveExpiredMetaData(const FRestoreStateParams& Params, TArray<FPreAnimatedStateMetaData>& OutExpiredMetaData) = 0;
 	virtual void GatherAndRemoveMetaDataForGroup(FPreAnimatedStorageGroupHandle Group, TArray<FPreAnimatedStateMetaData>& OutExpiredMetaData) = 0;
 	virtual void GatherAndRemoveMetaDataForStorage(FPreAnimatedStorageID StorageID, FPreAnimatedStorageIndex StorageIndex, TArray<FPreAnimatedStateMetaData>& OutExpiredMetaData) = 0;
+	virtual void GatherAndRemoveMetaDataForRootInstance(FRootInstanceHandle InstanceHandle, TArray<FPreAnimatedStateMetaData>& OutExpiredMetaData) = 0;
 };
 
 
diff --git a/Engine/Source/Runtime/MovieScene/Public/Evaluation/PreAnimatedState/MovieScenePreAnimatedCaptureSource.h b/Engine/Source/Runtime/MovieScene/Public/Evaluation/PreAnimatedState/MovieScenePreAnimatedCaptureSource.h
--- a/Engine/Source/Runtime/MovieScene/Public/Evaluation/PreAnimatedState/MovieScenePreAnimatedCaptureSource.h
+++ b/Engine/Source/Runtime/MovieScene/Public/Evaluation/PreAnimatedState/MovieScenePreAnimatedCaptureSource.h
@@ -7,6 +7,7 @@
 #include "EntitySystem/TrackInstance/MovieSceneTrackInstance.h"
 #include "Evaluation/MovieSceneEvaluationKey.h"
 #include "Evaluation/PreAnimatedState/MovieScenePreAnimatedStateTypes.h"
+#include "Evaluation/PreAnimatedState/MovieScenePreAnimatedCaptureSources.h"
 #include "Misc/TVariant.h"
 #include "MovieSceneSequenceID.h"
 #include "UObject/WeakObjectPtrTemplates.h"
@@ -26,6 +27,7 @@ namespace MovieScene
 
 struct FPreAnimatedStateExtension;
 struct FPreAnimatedStateMetaData;
+struct FSharedPlaybackState;
 
 }
 }
@@ -36,15 +38,28 @@ struct FPreAnimatedStateMetaData;
  */
 struct FScopedPreAnimatedCaptureSource
 {
+	using FRootInstanceHandle = UE::MovieScene::FRootInstanceHandle;
+	using FSharedPlaybackState = UE::MovieScene::FSharedPlaybackState;
+
+	/**
+	 * Construct this capture source from a template (FMovieSceneEvalTemplate) evaluation key, and whether this should restore state when the template is finished
+	 */
+	MOVIESCENE_API explicit FScopedPreAnimatedCaptureSource(TSharedRef<const FSharedPlaybackState> SharedPlaybackState, const FMovieSceneEvaluationKey& InEvalKey, bool bInWantsRestoreState);
+
 	/**
 	 * Construct this capture source from a template (FMovieSceneEvalTemplate) evaluation key, and whether this should restore state when the template is finished
 	 */
 	MOVIESCENE_API explicit FScopedPreAnimatedCaptureSource(FMovieScenePreAnimatedState* InPreAnimatedState, const FMovieSceneEvaluationKey& InEvalKey, bool bInWantsRestoreState);
 
 	/**
-	 * Construct this capture source from an evaluation hook (UMovieSceneEvaluationHookSection), its SequenceID, and whether this should restore state when the template is finished
+	 * Construct this capture source from an evaluation hook (UMovieSceneEvaluationHookSection), its instance handle, and whether this should restore state when the template is finished
 	 */
-	MOVIESCENE_API explicit FScopedPreAnimatedCaptureSource(FMovieScenePreAnimatedState* InPreAnimatedState, const UObject* InEvalHook, FMovieSceneSequenceID SequenceID, bool bInWantsRestoreState);
+	MOVIESCENE_API explicit FScopedPreAnimatedCaptureSource(TSharedRef<const FSharedPlaybackState> SharedPlaybackState, const UObject* InEvalHook, FMovieSceneSequenceID InSequenceID, bool bInWantsRestoreState);
+
+	/**
+	 * Construct this capture source from an evaluation hook (UMovieSceneEvaluationHookSection), its instance handle, and whether this should restore state when the template is finished
+	 */
+	MOVIESCENE_API explicit FScopedPreAnimatedCaptureSource(FMovieScenePreAnimatedState* InPreAnimatedState, const UObject* InEvalHook, FMovieSceneSequenceID InSequenceID, bool bInWantsRestoreState);
 
 	/**
 	 * Construct this capture source from a track instance (UMovieSceneTrackInstance) and whether this should restore state when the template is finished
@@ -77,18 +92,15 @@ private:
 	static FScopedPreAnimatedCaptureSource*& GetCaptureSourcePtr();
 
 	void BeginTracking(const UE::MovieScene::FPreAnimatedStateMetaData& MetaData, UMovieSceneEntitySystemLinker* Linker);
-	UE::MovieScene::FRootInstanceHandle GetRootInstanceHandle(UMovieSceneEntitySystemLinker* Linker) const;
+	FRootInstanceHandle GetRootInstanceHandle(UMovieSceneEntitySystemLinker* Linker) const;
 
-	struct FEvalHookType
-	{
-		const UObject* EvalHook;
-		FMovieSceneSequenceID SequenceID;
-	};
-	using CaptureSourceType = TVariant<FMovieSceneEvaluationKey, FEvalHookType, UMovieSceneTrackInstance*, FMovieSceneTrackInstanceInput>;
+	using FPreAnimatedEvaluationKeyType = UE::MovieScene::FPreAnimatedEvaluationKeyType;
+	using FPreAnimatedEvalHookKeyType = UE::MovieScene::FPreAnimatedEvalHookKeyType;
+	using CaptureSourceType = TVariant<FPreAnimatedEvaluationKeyType, FPreAnimatedEvalHookKeyType, UMovieSceneTrackInstance*, FMovieSceneTrackInstanceInput>;
 
 	CaptureSourceType Variant;
+	FRootInstanceHandle RootInstanceHandle;
 	TWeakObjectPtr<UMovieSceneEntitySystemLinker> WeakLinker;
-	FMovieScenePreAnimatedState* OptionalSequencePreAnimatedState;
 	FScopedPreAnimatedCaptureSource* PrevCaptureSource;
 	bool bWantsRestoreState;
 };
diff --git a/Engine/Source/Runtime/MovieScene/Public/Evaluation/PreAnimatedState/MovieScenePreAnimatedCaptureSources.h b/Engine/Source/Runtime/MovieScene/Public/Evaluation/PreAnimatedState/MovieScenePreAnimatedCaptureSources.h
--- a/Engine/Source/Runtime/MovieScene/Public/Evaluation/PreAnimatedState/MovieScenePreAnimatedCaptureSources.h
+++ b/Engine/Source/Runtime/MovieScene/Public/Evaluation/PreAnimatedState/MovieScenePreAnimatedCaptureSources.h
@@ -12,7 +12,6 @@
 #include "Evaluation/MovieSceneEvaluationKey.h"
 #include "Evaluation/PreAnimatedState/IMovieScenePreAnimatedCaptureSource.h"
 #include "Evaluation/PreAnimatedState/MovieScenePreAnimatedStateTypes.h"
-#include "MovieSceneSequenceID.h"
 #include "UObject/ObjectKey.h"
 
 class UMovieSceneTrackInstance;
@@ -23,7 +22,6 @@ namespace UE
 namespace MovieScene
 {
 
-struct FPreAnimatedEvalHookKeyType;
 struct FPreAnimatedStateExtension;
 struct FRestoreStateParams;
 
@@ -37,13 +35,37 @@ enum class EPreAnimatedCaptureSourceState
 	UpToDate,
 };
 
+/** Key type for pre-animated state associated to evaluation hooks */
 struct FPreAnimatedEvalHookKeyType
 {
 	FObjectKey Hook;
+	FRootInstanceHandle RootInstanceHandle;
 	FMovieSceneSequenceID SequenceID;
 
-	friend uint32 GetTypeHash(const FPreAnimatedEvalHookKeyType& In) { return GetTypeHash(In.Hook) ^ GetTypeHash(In.SequenceID); }
-	friend bool operator==(const FPreAnimatedEvalHookKeyType& A, const FPreAnimatedEvalHookKeyType& B) { return A.Hook == B.Hook && A.SequenceID == B.SequenceID; }
+	friend uint32 GetTypeHash(const FPreAnimatedEvalHookKeyType& In)
+	{
+		return GetTypeHash(In.Hook) ^ GetTypeHash(In.RootInstanceHandle) ^ GetTypeHash(In.SequenceID);
+	}
+	friend bool operator==(const FPreAnimatedEvalHookKeyType& A, const FPreAnimatedEvalHookKeyType& B)
+	{
+		return A.Hook == B.Hook && A.RootInstanceHandle == B.RootInstanceHandle && A.SequenceID == B.SequenceID;
+	}
+};
+
+/** Key type for pre-animated state associated to evaluation templates */
+struct FPreAnimatedEvaluationKeyType
+{
+	FMovieSceneEvaluationKey EvaluationKey;
+	FRootInstanceHandle RootInstanceHandle;
+
+	friend uint32 GetTypeHash(const FPreAnimatedEvaluationKeyType& In)
+	{
+		return GetTypeHash(In.EvaluationKey) ^ GetTypeHash(In.RootInstanceHandle);
+	}
+	friend bool operator==(const FPreAnimatedEvaluationKeyType& A, const FPreAnimatedEvaluationKeyType& B)
+	{
+		return A.EvaluationKey == B.EvaluationKey && A.RootInstanceHandle == B.RootInstanceHandle; 
+	}
 };
 
 /**
@@ -55,21 +77,29 @@ struct TPreAnimatedCaptureSources : IPreAnimatedCaptureSource
 	TPreAnimatedCaptureSources(FPreAnimatedStateExtension* InOwner);
 
 	/**
-	 * Make an association for the specifieid evaluation key to the value specified by Entry,
+	 * Make an association for the specified evaluation key to the value specified by Entry,
 	 * along with an indication of whether it should be restored on completion
 	 */
 	EPreAnimatedCaptureSourceState BeginTrackingCaptureSource(const KeyType& InKey, const FPreAnimatedStateMetaData& MetaData);
 
+	/**
+	 * Stop tracking the evaluation key in question, where associated with the given storage value ID
+	 * This restores the value if it was captured with bWantsRestoreState
+	 */
+	void StopTrackingCaptureSource(const KeyType& InKey, FPreAnimatedStorageID InStorageID);
+
 	/**
 	 * Stop tracking the evaluation key in question, restoring the value if it was captured with bWantsRestoreState
 	 */
 	void StopTrackingCaptureSource(const KeyType& InKey);
 
+	// IPreAnimatedCaptureSource members
 	void Reset() override;
 	bool ContainsInstanceHandle(FRootInstanceHandle RootInstanceHandle) const override;
 	void GatherAndRemoveExpiredMetaData(const FRestoreStateParams& Params, TArray<FPreAnimatedStateMetaData>& OutExpiredMetaData) override;
 	void GatherAndRemoveMetaDataForGroup(FPreAnimatedStorageGroupHandle Group, TArray<FPreAnimatedStateMetaData>& OutExpiredMetaData) override;
 	void GatherAndRemoveMetaDataForStorage(FPreAnimatedStorageID StorageID, FPreAnimatedStorageIndex StorageIndex, TArray<FPreAnimatedStateMetaData>& OutExpiredMetaData) override;
+	void GatherAndRemoveMetaDataForRootInstance(FRootInstanceHandle InstanceHandle, TArray<FPreAnimatedStateMetaData>& OutExpiredMetaData) override;
 
 private:
 
@@ -78,12 +108,47 @@ private:
 };
 
 
+/**
+ * Structure responsible for tracking contributions to pre-animated state entries that originate from ECS data (ie, from FMovieSceneEntityIDs)
+ */
+struct FPreAnimatedEntityCaptureSource : TPreAnimatedCaptureSources<FMovieSceneEntityID>
+{
+	MOVIESCENE_API FPreAnimatedEntityCaptureSource(FPreAnimatedStateExtension* InOwner);
+
+	MOVIESCENE_API void BeginTrackingEntity(const FPreAnimatedStateEntry& Entry, FMovieSceneEntityID EntityID, FRootInstanceHandle RootInstanceHandle, bool bWantsRestoreState);
+
+	void StopTrackingEntity(FMovieSceneEntityID EntityID, FPreAnimatedStorageID StorageID)
+	{
+		StopTrackingCaptureSource(EntityID, StorageID);
+	}
+
+	void StopTrackingEntity(FMovieSceneEntityID EntityID)
+	{
+		StopTrackingCaptureSource(EntityID);
+	}
+};
+
 /**
  * Structure responsible for tracking contributions to pre-eanimated state entries that originate from track templates (ie, from an IMovieSceneExecutionToken::Execute)
  */
-struct FPreAnimatedTemplateCaptureSources : TPreAnimatedCaptureSources<FMovieSceneEvaluationKey>
+struct FPreAnimatedTemplateCaptureSources : TPreAnimatedCaptureSources<FPreAnimatedEvaluationKeyType>
 {
 	MOVIESCENE_API FPreAnimatedTemplateCaptureSources(FPreAnimatedStateExtension* InOwner);
+
+	/**
+	 * Make an association for the specified evaluation key to the value specified by Entry,
+	 * along with an indication of whether it should be restored on completion
+	 */
+	MOVIESCENE_API EPreAnimatedCaptureSourceState BeginTrackingCaptureSource(const FMovieSceneEvaluationKey& EvaluationKey, const FPreAnimatedStateMetaData& MetaData);
+
+	/**
+	 * Stop tracking the evaluation key in question, restoring the value if it was captured with bWantsRestoreState
+	 */
+	MOVIESCENE_API void StopTrackingCaptureSource(const FMovieSceneEvaluationKey& EvaluationKey, FRootInstanceHandle RootInstanceHandle);
+
+	// Un-shadow base class methods (it's OK, they're not virtual)
+	using TPreAnimatedCaptureSources<FPreAnimatedEvaluationKeyType>::BeginTrackingCaptureSource;
+	using TPreAnimatedCaptureSources<FPreAnimatedEvaluationKeyType>::StopTrackingCaptureSource;
 };
 
 /**
@@ -94,7 +159,7 @@ struct FPreAnimatedEvaluationHookCaptureSources : TPreAnimatedCaptureSources<FPr
 	MOVIESCENE_API FPreAnimatedEvaluationHookCaptureSources(FPreAnimatedStateExtension* InOwner);
 
 	/**
-	 * Make an association for the specifieid evaluation key to the value specified by Entry,
+	 * Make an association for the specified evaluation key to the value specified by Entry,
 	 * along with an indication of whether it should be restored on completion
 	 */
 	MOVIESCENE_API EPreAnimatedCaptureSourceState BeginTrackingCaptureSource(const UObject* Hook, FMovieSceneSequenceID SequenceID, const FPreAnimatedStateMetaData& MetaData);
@@ -102,7 +167,11 @@ struct FPreAnimatedEvaluationHookCaptureSources : TPreAnimatedCaptureSources<FPr
 	/**
 	 * Stop tracking the evaluation key in question, restoring the value if it was captured with bWantsRestoreState
 	 */
-	MOVIESCENE_API void StopTrackingCaptureSource(const UObject* Hook, FMovieSceneSequenceID SequenceID);
+	MOVIESCENE_API void StopTrackingCaptureSource(const UObject* Hook, FRootInstanceHandle RootInstanceHandle, FMovieSceneSequenceID SequenceID);
+
+	// Un-shadow base class methods (it's OK, they're not virtual)
+	using TPreAnimatedCaptureSources<FPreAnimatedEvalHookKeyType>::BeginTrackingCaptureSource;
+	using TPreAnimatedCaptureSources<FPreAnimatedEvalHookKeyType>::StopTrackingCaptureSource;
 };
 
 
@@ -114,7 +183,7 @@ struct FPreAnimatedTrackInstanceCaptureSources : TPreAnimatedCaptureSources<FObj
 	MOVIESCENE_API FPreAnimatedTrackInstanceCaptureSources(FPreAnimatedStateExtension* InOwner);
 
 	/**
-	 * Make an association for the specifieid evaluation key to the value specified by Entry,
+	 * Make an association for the specified evaluation key to the value specified by Entry,
 	 * along with an indication of whether it should be restored on completion
 	 */
 	MOVIESCENE_API EPreAnimatedCaptureSourceState BeginTrackingCaptureSource(UMovieSceneTrackInstance* TrackInstance, const FPreAnimatedStateMetaData& MetaData);
diff --git a/Engine/Source/Runtime/MovieScene/Public/Evaluation/PreAnimatedState/MovieScenePreAnimatedEntityCaptureSource.h b/Engine/Source/Runtime/MovieScene/Public/Evaluation/PreAnimatedState/MovieScenePreAnimatedEntityCaptureSource.h
--- a/Engine/Source/Runtime/MovieScene/Public/Evaluation/PreAnimatedState/MovieScenePreAnimatedEntityCaptureSource.h
+++ b/Engine/Source/Runtime/MovieScene/Public/Evaluation/PreAnimatedState/MovieScenePreAnimatedEntityCaptureSource.h
@@ -2,50 +2,7 @@
 
 #pragma once
 
-#include "Containers/Array.h"
-#include "Containers/Map.h"
-#include "CoreTypes.h"
-#include "EntitySystem/MovieSceneEntityIDs.h"
-#include "EntitySystem/MovieSceneSequenceInstanceHandle.h"
-#include "Evaluation/PreAnimatedState/MovieScenePreAnimatedStateTypes.h"
-#include "Evaluation/PreAnimatedState/MovieScenePreAnimatedStorageID.h"
+#include "Evaluation/PreAnimatedState/MovieScenePreAnimatedCaptureSources.h"
 
-namespace UE
-{
-namespace MovieScene
-{
+COMPILE_WARNING("FPreAnimatedEntityCaptureSource has been moved to MovieScenePreAnimatedCaptureSources.h")
 
-struct FPreAnimatedStateExtension;
-struct FPreAnimatedStorageID;
-struct FRestoreStateParams;
-struct FRootInstanceHandle;
-
-/**
- * Structure responsible for tracking contributions to pre-eanimated state entries that originate from ECS data (ie, from FMovieSceneEntityIDs)
- */
-struct FPreAnimatedEntityCaptureSource
-{
-	FPreAnimatedEntityCaptureSource(FPreAnimatedStateExtension* InOwner);
-
-	void Reset();
-
-	MOVIESCENE_API void BeginTrackingEntity(const FPreAnimatedStateEntry& Entry, FMovieSceneEntityID EntityID, FRootInstanceHandle RootInstanceHandle, bool bWantsRestoreState);
-	MOVIESCENE_API void StopTrackingEntity(FMovieSceneEntityID EntityID, FPreAnimatedStorageID StorageID);
-	MOVIESCENE_API void StopTrackingEntity(FMovieSceneEntityID EntityID);
-
-	bool ContainsInstanceHandle(FRootInstanceHandle RootInstanceHandle) const;
-	void GatherAndRemoveExpiredMetaData(const FRestoreStateParams& Params, TArray<FPreAnimatedStateMetaData>& OutExpiredMetaData);
-	void GatherAndRemoveMetaDataForGroup(FPreAnimatedStorageGroupHandle Group, TArray<FPreAnimatedStateMetaData>& OutExpiredMetaData);
-	void GatherAndRemoveMetaDataForStorage(FPreAnimatedStorageID StorageID, FPreAnimatedStorageIndex StorageIndex, TArray<FPreAnimatedStateMetaData>& OutExpiredMetaData);
-
-private:
-
-	TMap<FMovieSceneEntityID, FPreAnimatedStateMetaDataArray> KeyToMetaData;
-	FPreAnimatedStateExtension* Owner;
-};
-
-
-
-
-} // namespace MovieScene
-} // namespace UE
diff --git a/Engine/Source/Runtime/MovieScene/Public/Evaluation/PreAnimatedState/MovieScenePreAnimatedObjectStorage.h b/Engine/Source/Runtime/MovieScene/Public/Evaluation/PreAnimatedState/MovieScenePreAnimatedObjectStorage.h
--- a/Engine/Source/Runtime/MovieScene/Public/Evaluation/PreAnimatedState/MovieScenePreAnimatedObjectStorage.h
+++ b/Engine/Source/Runtime/MovieScene/Public/Evaluation/PreAnimatedState/MovieScenePreAnimatedObjectStorage.h
@@ -7,7 +7,7 @@
 #include "Evaluation/PreAnimatedState/IMovieScenePreAnimatedStorage.h"
 #include "Evaluation/PreAnimatedState/MovieScenePreAnimatedStateStorage.h"
 #include "Evaluation/PreAnimatedState/MovieScenePreAnimatedObjectGroupManager.h"
-#include "Evaluation/PreAnimatedState/MovieScenePreAnimatedEntityCaptureSource.h"
+#include "Evaluation/PreAnimatedState/MovieScenePreAnimatedCaptureSources.h"
 #include "EntitySystem/BuiltInComponentTypes.h"
 
 
diff --git a/Engine/Source/Runtime/MovieScene/Public/Evaluation/PreAnimatedState/MovieScenePreAnimatedPropertyStorage.h b/Engine/Source/Runtime/MovieScene/Public/Evaluation/PreAnimatedState/MovieScenePreAnimatedPropertyStorage.h
--- a/Engine/Source/Runtime/MovieScene/Public/Evaluation/PreAnimatedState/MovieScenePreAnimatedPropertyStorage.h
+++ b/Engine/Source/Runtime/MovieScene/Public/Evaluation/PreAnimatedState/MovieScenePreAnimatedPropertyStorage.h
@@ -15,7 +15,7 @@
 #include "Evaluation/PreAnimatedState/MovieSceneRestoreStateParams.h"
 #include "Evaluation/PreAnimatedState/MovieScenePreAnimatedStorageID.inl"
 #include "Evaluation/PreAnimatedState/MovieScenePreAnimatedObjectStorage.h"
-#include "Evaluation/PreAnimatedState/MovieScenePreAnimatedEntityCaptureSource.h"
+#include "Evaluation/PreAnimatedState/MovieScenePreAnimatedCaptureSources.h"
 
 
 namespace UE
diff --git a/Engine/Source/Runtime/MovieScene/Public/Evaluation/PreAnimatedState/MovieScenePreAnimatedStateExtension.h b/Engine/Source/Runtime/MovieScene/Public/Evaluation/PreAnimatedState/MovieScenePreAnimatedStateExtension.h
--- a/Engine/Source/Runtime/MovieScene/Public/Evaluation/PreAnimatedState/MovieScenePreAnimatedStateExtension.h
+++ b/Engine/Source/Runtime/MovieScene/Public/Evaluation/PreAnimatedState/MovieScenePreAnimatedStateExtension.h
@@ -33,10 +33,13 @@ namespace UE
 namespace MovieScene
 {
 
-struct FAnimTypePreAnimatedStateRootStorage;
 struct FAnimTypePreAnimatedStateObjectStorage;
+struct FAnimTypePreAnimatedStateRootStorage;
 struct FPreAnimatedEntityCaptureSource;
+struct FPreAnimatedEvaluationHookCaptureSources;
+struct FPreAnimatedStateEntry;
 struct FPreAnimatedStateExtension;
+struct FPreAnimatedTemplateCaptureSources;
 struct FPreAnimatedTrackInstanceCaptureSources;
 struct FPreAnimatedTrackInstanceInputCaptureSources;
 struct FRestoreStateParams;
@@ -298,6 +301,12 @@ public:
 	MOVIESCENE_API FPreAnimatedTrackInstanceInputCaptureSources* GetTrackInstanceInputMetaData() const;
 	MOVIESCENE_API FPreAnimatedTrackInstanceInputCaptureSources* GetOrCreateTrackInstanceInputMetaData();
 
+	MOVIESCENE_API FPreAnimatedTemplateCaptureSources* GetTemplateMetaData() const;
+	MOVIESCENE_API FPreAnimatedTemplateCaptureSources* GetOrCreateTemplateMetaData();
+
+	MOVIESCENE_API FPreAnimatedEvaluationHookCaptureSources* GetEvaluationHookMetaData() const;
+	MOVIESCENE_API FPreAnimatedEvaluationHookCaptureSources* GetOrCreateEvaluationHookMetaData();
+
 	MOVIESCENE_API bool HasActiveCaptureSource() const;
 
 	MOVIESCENE_API void AddWeakCaptureSource(TWeakPtr<IPreAnimatedCaptureSource> InWeakMetaData);
@@ -320,6 +329,12 @@ private:
 
 	MOVIESCENE_API void AddReferencedObjects(UMovieSceneEntitySystemLinker*, FReferenceCollector& ReferenceCollector);
 
+	using FContributionRemover = TFunctionRef<void(IPreAnimatedStorage&, FPreAnimatedStorageIndex)>;
+	void HandleMetaDataToRemove(
+			const FRestoreStateParams& Params, 
+			TArrayView<FPreAnimatedStateMetaData> MetaDataToRemove, 
+			FContributionRemover RemoveFunc);
+
 public:
 
 	/** The number of requests that have been made to capture global state - only one should exist per playing sequence */
@@ -336,6 +351,12 @@ private:
 	/** Meta-data pertaining to pre-animated state originating from track instances from a specific input */
 	TUniquePtr<FPreAnimatedTrackInstanceInputCaptureSources> TrackInstanceInputCaptureSource;
 
+	/** Meta-data ledger for any pre-animated state that originates from track templates */
+	TUniquePtr<UE::MovieScene::FPreAnimatedTemplateCaptureSources> TemplateCaptureSource;
+
+	/** Meta-data ledger for any pre-animated state that originates from evaluation hooks */
+	TUniquePtr<UE::MovieScene::FPreAnimatedEvaluationHookCaptureSources> EvaluationHookCaptureSource;
+
 	/** Weakly held meta data provided by FMovieScenePreAnimatedState for various other origins */
 	TArray<TWeakPtr<IPreAnimatedCaptureSource>> WeakExternalCaptureSources;
 
@@ -346,7 +367,7 @@ private:
 
 private:
 
-	friend struct FScopedPreAnimatedCaptureSource;
+	friend struct ::FScopedPreAnimatedCaptureSource;
 
 	struct FAggregatePreAnimatedStateMetaData
 	{
diff --git a/Engine/Source/Runtime/MovieScene/Public/Evaluation/PreAnimatedState/MovieScenePreAnimatedStateStorage.h b/Engine/Source/Runtime/MovieScene/Public/Evaluation/PreAnimatedState/MovieScenePreAnimatedStateStorage.h
--- a/Engine/Source/Runtime/MovieScene/Public/Evaluation/PreAnimatedState/MovieScenePreAnimatedStateStorage.h
+++ b/Engine/Source/Runtime/MovieScene/Public/Evaluation/PreAnimatedState/MovieScenePreAnimatedStateStorage.h
@@ -16,7 +16,7 @@
 #include "EntitySystem/MovieSceneEntitySystemTask.h"
 #include "EntitySystem/MovieSceneEntitySystemTypes.h"
 #include "Evaluation/PreAnimatedState/IMovieScenePreAnimatedStorage.h"
-#include "Evaluation/PreAnimatedState/MovieScenePreAnimatedEntityCaptureSource.h"
+#include "Evaluation/PreAnimatedState/MovieScenePreAnimatedCaptureSources.h"
 #include "Evaluation/PreAnimatedState/MovieScenePreAnimatedStateExtension.h"
 #include "Evaluation/PreAnimatedState/MovieScenePreAnimatedStateTypes.h"
 #include "Evaluation/PreAnimatedState/MovieScenePreAnimatedStorageID.h"
diff --git a/Engine/Source/Runtime/MovieSceneTracks/Private/Systems/MovieSceneComponentMobilitySystem.cpp b/Engine/Source/Runtime/MovieSceneTracks/Private/Systems/MovieSceneComponentMobilitySystem.cpp
--- a/Engine/Source/Runtime/MovieSceneTracks/Private/Systems/MovieSceneComponentMobilitySystem.cpp
+++ b/Engine/Source/Runtime/MovieSceneTracks/Private/Systems/MovieSceneComponentMobilitySystem.cpp
@@ -12,7 +12,7 @@
 
 #include "Evaluation/PreAnimatedState/MovieScenePreAnimatedObjectGroupManager.h"
 #include "Evaluation/PreAnimatedState/MovieScenePreAnimatedStorageID.inl"
-#include "Evaluation/PreAnimatedState/MovieScenePreAnimatedEntityCaptureSource.h"
+#include "Evaluation/PreAnimatedState/MovieScenePreAnimatedCaptureSources.h"
 
 #include "Systems/MovieSceneComponentTransformSystem.h"
 #include "Systems/MovieScenePropertyInstantiator.h"
diff --git a/Engine/Source/Runtime/MovieSceneTracks/Private/Systems/MovieSceneDataLayerSystem.cpp b/Engine/Source/Runtime/MovieSceneTracks/Private/Systems/MovieSceneDataLayerSystem.cpp
--- a/Engine/Source/Runtime/MovieSceneTracks/Private/Systems/MovieSceneDataLayerSystem.cpp
+++ b/Engine/Source/Runtime/MovieSceneTracks/Private/Systems/MovieSceneDataLayerSystem.cpp
@@ -6,7 +6,7 @@
 #include "EntitySystem/MovieSceneBindingLifetimeSystem.h"
 #include "EntitySystem/MovieSceneEntitySystemLinker.h"
 #include "EntitySystem/BuiltInComponentTypes.h"
-#include "Evaluation/PreAnimatedState/MovieScenePreAnimatedEntityCaptureSource.h"
+#include "Evaluation/PreAnimatedState/MovieScenePreAnimatedCaptureSources.h"
 #include "Evaluation/PreAnimatedState/MovieScenePreAnimatedObjectStorage.h"
 #include "Evaluation/PreAnimatedState/MovieScenePreAnimatedStateStorage.h"
 #include "Evaluation/PreAnimatedState/MovieScenePreAnimatedStorageID.inl"
diff --git a/Engine/Source/Runtime/MovieSceneTracks/Private/Systems/MovieScenePropertySystem.cpp b/Engine/Source/Runtime/MovieSceneTracks/Private/Systems/MovieScenePropertySystem.cpp
--- a/Engine/Source/Runtime/MovieSceneTracks/Private/Systems/MovieScenePropertySystem.cpp
+++ b/Engine/Source/Runtime/MovieSceneTracks/Private/Systems/MovieScenePropertySystem.cpp
@@ -5,7 +5,7 @@
 
 #include "EntitySystem/MovieSceneEntitySystemLinker.h"
 #include "EntitySystem/Interrogation/MovieSceneInterrogationLinker.h"
-#include "Evaluation/PreAnimatedState/MovieScenePreAnimatedEntityCaptureSource.h"
+#include "Evaluation/PreAnimatedState/MovieScenePreAnimatedCaptureSources.h"
 
 #include UE_INLINE_GENERATED_CPP_BY_NAME(MovieScenePropertySystem)
 
