Sequencer: re-add null check that got forgotten in the latest pre-animated state refactor

diff --git a/Engine/Source/Runtime/MovieScene/Private/Evaluation/Instances/MovieSceneTrackEvaluator.cpp b/Engine/Source/Runtime/MovieScene/Private/Evaluation/Instances/MovieSceneTrackEvaluator.cpp
--- a/Engine/Source/Runtime/MovieScene/Private/Evaluation/Instances/MovieSceneTrackEvaluator.cpp
+++ b/Engine/Source/Runtime/MovieScene/Private/Evaluation/Instances/MovieSceneTrackEvaluator.cpp
@@ -43,10 +43,13 @@ struct FDelayedPreAnimatedStateRestore
 
 		UMovieSceneEntitySystemLinker* Linker = Player.GetSharedPlaybackState()->GetLinker();
 		FPreAnimatedTemplateCaptureSources* TemplateMetaData = Linker->PreAnimatedState.GetTemplateMetaData();
-		const FRootInstanceHandle RootInstanceHandle = Player.GetEvaluationTemplate().GetRootInstanceHandle();
-		for (FMovieSceneEvaluationKey Key : KeysToRestore)
+		if (TemplateMetaData)
 		{
-			TemplateMetaData->StopTrackingCaptureSource(Key, RootInstanceHandle);
+			const FRootInstanceHandle RootInstanceHandle = Player.GetEvaluationTemplate().GetRootInstanceHandle();
+			for (FMovieSceneEvaluationKey Key : KeysToRestore)
+			{
+				TemplateMetaData->StopTrackingCaptureSource(Key, RootInstanceHandle);
+			}
 		}
 		KeysToRestore.Reset();
 	}
