Ensure instead of crashing in IMovieScenePlayer::InitializeRootInstance

#rnx

diff --git a/Engine/Source/Runtime/MovieScene/Private/MovieScenePlayer.cpp b/Engine/Source/Runtime/MovieScene/Private/MovieScenePlayer.cpp
--- a/Engine/Source/Runtime/MovieScene/Private/MovieScenePlayer.cpp
+++ b/Engine/Source/Runtime/MovieScene/Private/MovieScenePlayer.cpp
@@ -221,8 +221,17 @@ void IMovieScenePlayer::InitializeRootInstance(TSharedRef<UE::MovieScene::FShare
 		NewSharedPlaybackState->AddCapabilityRaw(PlaybackClient);
 	}
 
-	FInstanceRegistry* InstanceRegistry = NewSharedPlaybackState->GetLinker()->GetInstanceRegistry();
-	FSequenceInstance& RootInstance = InstanceRegistry->MutateInstance(NewSharedPlaybackState->GetRootInstanceHandle());
-	RootInstance.Initialize();
+	UMovieSceneEntitySystemLinker* Linker = NewSharedPlaybackState->GetLinker();
+
+	if (ensure(Linker))
+	{
+		FInstanceRegistry* InstanceRegistry = Linker->GetInstanceRegistry();
+
+		if (ensure(InstanceRegistry))
+		{
+			FSequenceInstance& RootInstance = InstanceRegistry->MutateInstance(NewSharedPlaybackState->GetRootInstanceHandle());
+			RootInstance.Initialize();
+		}
+	}
 }
 
