Sequencer: throw up an assert if playback capabilities are accessed outside of the game-thread.

The playback capabilities container isn't thread-safe, so we don't want someone to be accessing something while someone else is adding a new capability. This is especially problematic since adding a capability can relocate previously existing capabilities.

diff --git a/Engine/Source/Runtime/MovieScene/Public/EntitySystem/MovieSceneSharedPlaybackState.h b/Engine/Source/Runtime/MovieScene/Public/EntitySystem/MovieSceneSharedPlaybackState.h
--- a/Engine/Source/Runtime/MovieScene/Public/EntitySystem/MovieSceneSharedPlaybackState.h
+++ b/Engine/Source/Runtime/MovieScene/Public/EntitySystem/MovieSceneSharedPlaybackState.h
@@ -99,12 +99,26 @@ public:
 	/**
 	 * Gets the capabilities container.
 	 */
-	FPlaybackCapabilities& GetCapabilities() { return Capabilities; }
+	FPlaybackCapabilities& GetCapabilities()
+	{
+#if !UE_BUILD_SHIPPING
+		ensureMsgf(IsInGameThread(), 
+				TEXT("Playback capabilities aren't meant to be thread-safe. Do not modify or access their container outside of the game thread."));
+#endif
+		return Capabilities;
+	}
 
 	/**
 	 * Gets the capabilities container.
 	 */
-	const FPlaybackCapabilities& GetCapabilities() const { return Capabilities; }
+	const FPlaybackCapabilities& GetCapabilities() const
+	{
+#if !UE_BUILD_SHIPPING
+		ensureMsgf(IsInGameThread(), 
+				TEXT("Playback capabilities aren't meant to be thread-safe. Do not modify or access their container outside of the game thread."));
+#endif
+		return Capabilities;
+	}
 
 	/**
 	 * Returns whether the root sequence has the specified capability.
@@ -112,6 +126,10 @@ public:
 	template<typename T>
 	bool HasCapability() const
 	{
+#if !UE_BUILD_SHIPPING
+		ensureMsgf(IsInGameThread(), 
+				TEXT("Playback capabilities aren't meant to be thread-safe. Do not modify or access their container outside of the game thread."));
+#endif
 		return Capabilities.HasCapability<T>();
 	}
 
@@ -121,6 +139,10 @@ public:
 	template<typename T>
 	T* FindCapability() const
 	{
+#if !UE_BUILD_SHIPPING
+		ensureMsgf(IsInGameThread(), 
+				TEXT("Playback capabilities aren't meant to be thread-safe. Do not modify or access their container outside of the game thread."));
+#endif
 		return Capabilities.FindCapability<T>();
 	}
 
@@ -130,6 +152,10 @@ public:
 	template<typename T, typename ...ArgTypes>
 	T& AddCapability(ArgTypes&&... InArgs)
 	{
+#if !UE_BUILD_SHIPPING
+		ensureMsgf(IsInGameThread(), 
+				TEXT("Playback capabilities aren't meant to be thread-safe. Do not modify or access their container outside of the game thread."));
+#endif
 		T& Cap = Capabilities.AddCapability<T>(Forward<ArgTypes>(InArgs)...);
 		MaybeInitialize(Cap);
 		return Cap;
@@ -141,6 +167,10 @@ public:
 	template<typename T>
 	T& AddCapabilityRaw(T* InPointer)
 	{
+#if !UE_BUILD_SHIPPING
+		ensureMsgf(IsInGameThread(), 
+				TEXT("Playback capabilities aren't meant to be thread-safe. Do not modify or access their container outside of the game thread."));
+#endif
 		T& Cap = Capabilities.AddCapabilityRaw<T>(InPointer);
 		MaybeInitialize(Cap);
 		return Cap;
@@ -152,6 +182,10 @@ public:
 	template<typename T>
 	T& AddCapabilityShared(TSharedRef<T> InSharedRef)
 	{
+#if !UE_BUILD_SHIPPING
+		ensureMsgf(IsInGameThread(), 
+				TEXT("Playback capabilities aren't meant to be thread-safe. Do not modify or access their container outside of the game thread."));
+#endif
 		T& Cap = Capabilities.AddCapabilityShared<T>(InSharedRef);
 		MaybeInitialize(Cap);
 		return Cap;
@@ -168,6 +202,10 @@ public:
 	template<typename T, typename ...ArgTypes>
 	T& SetOrAddCapability(ArgTypes&&... InArgs)
 	{
+#if !UE_BUILD_SHIPPING
+		ensureMsgf(IsInGameThread(), 
+				TEXT("Playback capabilities aren't meant to be thread-safe. Do not modify or access their container outside of the game thread."));
+#endif
 		if (HasCapability<T>())
 		{
 			T& Cap = Capabilities.OverwriteCapability<T>(Forward<ArgTypes>(InArgs)...);
@@ -190,6 +228,10 @@ public:
 	template<typename T, typename ...ArgTypes>
 	T& SetOrAddCapabilityRaw(T* InPointer)
 	{
+#if !UE_BUILD_SHIPPING
+		ensureMsgf(IsInGameThread(), 
+				TEXT("Playback capabilities aren't meant to be thread-safe. Do not modify or access their container outside of the game thread."));
+#endif
 		if (HasCapability<T>())
 		{
 			T& Cap = Capabilities.OverwriteCapabilityRaw<T>(InPointer);
@@ -212,6 +254,10 @@ public:
 	template<typename T, typename ...ArgTypes>
 	T& SetOrAddCapabilityShared(TSharedRef<T> InSharedRef)
 	{
+#if !UE_BUILD_SHIPPING
+		ensureMsgf(IsInGameThread(), 
+				TEXT("Playback capabilities aren't meant to be thread-safe. Do not modify or access their container outside of the game thread."));
+#endif
 		if (HasCapability<T>())
 		{
 			T& Cap = Capabilities.OverwriteCapabilityShared<T>(InSharedRef);
diff --git a/Engine/Source/Runtime/MovieSceneTracks/Private/Systems/MovieSceneSkeletalAnimationSystem.cpp b/Engine/Source/Runtime/MovieSceneTracks/Private/Systems/MovieSceneSkeletalAnimationSystem.cpp
--- a/Engine/Source/Runtime/MovieSceneTracks/Private/Systems/MovieSceneSkeletalAnimationSystem.cpp
+++ b/Engine/Source/Runtime/MovieSceneTracks/Private/Systems/MovieSceneSkeletalAnimationSystem.cpp
@@ -357,7 +357,6 @@ struct FGatherSkeletalAnimations
 			// Get the full context, so we can get both the current and previous evaluation times.
 			const FSequenceInstance& SequenceInstance = InstanceRegistry->GetInstance(InstanceHandle);
 			const FMovieSceneContext& Context = SequenceInstance.GetContext();
-			IMovieScenePlayer* Player = SequenceInstance.GetPlayer();
 
 			// Calculate the time at which to evaluate the animation
 			const UMovieSceneSkeletalAnimationSection* AnimSection = SkeletalAnimation.Section;
@@ -376,7 +375,7 @@ struct FGatherSkeletalAnimations
 			const FMovieSceneContext& RootContext = RootInstance.GetContext();
 			const double RootDeltaTime = (RootContext.HasJumped() ? FFrameTime(0) : RootContext.GetRange().Size<FFrameTime>() ) / RootContext.GetFrameRate();
 
-			const EMovieScenePlayerStatus::Type PlayerStatus = Player->GetPlaybackStatus();
+			const EMovieScenePlayerStatus::Type PlayerStatus = Context.GetStatus();
 
 			const bool bResetDynamics = PlayerStatus == EMovieScenePlayerStatus::Stepping || 
 				PlayerStatus == EMovieScenePlayerStatus::Jumping || 
