Sequencer: fix splitting/trimming of looping sub-sections

Splitting a section is, by default, done by cloning the section, trimming the original on the right, and trimming the clone on the left.
Trimming looping sub-sections on the right is straight-forward. Trimming them on the left requires adjusting the FirstLoopStartFrameOffset to preserve where the loops happen. This was previously done directly in the SplitSection method, which meant that it was roughly OK for splitting, but that trimming in-and-of-itself was broken. This was later fixed, but the code in SplitSection remained, which meant that we were fixing up FirstLoopStartFrameOffset twice. This change therefore just removes the code in SplitSection.

In addition to that, the code in TrimSection had a bug: the StartFrameOffset property is already defined in the inner sub-sequence's time space, so there is no need to transform it using the OuterToInnerTransform function.

diff --git a/Engine/Source/Runtime/MovieScene/Private/Sections/MovieSceneSubSection.cpp b/Engine/Source/Runtime/MovieScene/Private/Sections/MovieSceneSubSection.cpp
--- a/Engine/Source/Runtime/MovieScene/Private/Sections/MovieSceneSubSection.cpp
+++ b/Engine/Source/Runtime/MovieScene/Private/Sections/MovieSceneSubSection.cpp
@@ -310,58 +310,6 @@ void UMovieSceneSubSection::PostEditChangeProperty(FPropertyChangedEvent& Proper
 }
 #endif
 
-UMovieSceneSection* UMovieSceneSubSection::SplitSection( FQualifiedFrameTime SplitTime, bool bDeleteKeys)
-{
-	// GetRange is in owning sequence resolution so we check against the incoming SplitTime without converting it.
-	TRange<FFrameNumber> InitialRange = GetRange();
-	if ( !InitialRange.Contains(SplitTime.Time.FrameNumber) )
-	{
-		return nullptr;
-	}
-
-	FFrameNumber InitialStartOffset = Parameters.StartFrameOffset;
-
-	UMovieSceneSubSection* NewSection = Cast<UMovieSceneSubSection>( UMovieSceneSection::SplitSection( SplitTime, bDeleteKeys ) );
-	if ( NewSection )
-	{
-		if (InitialRange.GetLowerBound().IsClosed())
-		{
-			// Sections need their offsets calculated in their local resolution. Different sequences can have different tick resolutions 
-			// so we need to transform from the parent resolution to the local one before splitting them.
-			FFrameRate LocalTickResolution;
-			if (GetSequence())
-			{
-				LocalTickResolution = GetSequence()->GetMovieScene()->GetTickResolution();
-			}
-			else
-			{
-				UMovieScene* OuterScene = GetTypedOuter<UMovieScene>();
-				if (OuterScene)
-				{
-					LocalTickResolution = OuterScene->GetTickResolution();
-				}
-			}
-
-			FFrameNumber LocalResolutionStartOffset = FFrameRate::TransformTime(SplitTime.Time.GetFrame() - UE::MovieScene::DiscreteInclusiveLower(InitialRange), SplitTime.Rate, LocalTickResolution).FrameNumber;
-
-			FFrameNumber NewStartOffset = LocalResolutionStartOffset * Parameters.TimeScale;
-			NewStartOffset += InitialStartOffset;
-
-			if (NewStartOffset >= 0)
-			{
-				NewSection->Parameters.StartFrameOffset = NewStartOffset.Value;
-			}
-		}
-
-		return NewSection;
-	}
-
-	// Restore original offset modified by splitting
-	Parameters.StartFrameOffset = InitialStartOffset;
-
-	return nullptr;
-}
-
 TOptional<TRange<FFrameNumber> > UMovieSceneSubSection::GetAutoSizeRange() const
 {
 	UMovieScene* MovieScene = SubSequence ? SubSequence->GetMovieScene() : nullptr;
@@ -406,9 +354,9 @@ void UMovieSceneSubSection::TrimSection( FQualifiedFrameTime TrimTime, bool bTri
 
 		// The new first loop start offset is where the trim time fell inside the sub-sequence (this time is already
 		// normalized in the case of looping sub-sequences).
-		FFrameTime LocalTrimTime = OuterToInnerTransform().TransformTime(LocalTickResolutionTrimTime);
-		FFrameTime LocalStartFrameOffset = OuterToInnerTransform().TransformTime(Parameters.StartFrameOffset);
-		FFrameNumber NewStartOffset = LocalTrimTime.FrameNumber - LocalStartFrameOffset.FrameNumber;
+		const FMovieSceneSequenceTransform OuterToInner(OuterToInnerTransform());
+		FFrameTime LocalTrimTime = OuterToInner.TransformTime(LocalTickResolutionTrimTime);
+		FFrameNumber NewStartOffset = LocalTrimTime.FrameNumber - Parameters.StartFrameOffset;
 
 		// Make sure we don't have negative offsets (this shouldn't happen, though).
 		NewStartOffset = FMath::Max(FFrameNumber(0), NewStartOffset);
diff --git a/Engine/Source/Runtime/MovieScene/Public/Sections/MovieSceneSubSection.h b/Engine/Source/Runtime/MovieScene/Public/Sections/MovieSceneSubSection.h
--- a/Engine/Source/Runtime/MovieScene/Public/Sections/MovieSceneSubSection.h
+++ b/Engine/Source/Runtime/MovieScene/Public/Sections/MovieSceneSubSection.h
@@ -147,7 +147,6 @@ public:
 
 	//~ UMovieSceneSection interface
 	MOVIESCENE_API virtual TOptional<TRange<FFrameNumber> > GetAutoSizeRange() const override;
-	MOVIESCENE_API virtual UMovieSceneSection* SplitSection( FQualifiedFrameTime SplitTime, bool bDeleteKeys ) override;
 	MOVIESCENE_API virtual void TrimSection( FQualifiedFrameTime TrimTime, bool bTrimLeft, bool bDeleteKeys) override;
 	virtual TOptional<FFrameTime> GetOffsetTime() const override { return TOptional<FFrameTime>(FFrameTime(Parameters.StartFrameOffset)); }
 	MOVIESCENE_API virtual void GetSnapTimes(TArray<FFrameNumber>& OutSnapTimes, bool bGetSectionBorders) const override;
