Sequencer: Fix objects being bound to PIE worlds even when we don't want them to.

This used to be in LevelSequenceBindingReference::Resolve

diff --git a/Engine/Source/Runtime/MovieScene/Private/MovieSceneBindingReferences.cpp b/Engine/Source/Runtime/MovieScene/Private/MovieSceneBindingReferences.cpp
--- a/Engine/Source/Runtime/MovieScene/Private/MovieSceneBindingReferences.cpp
+++ b/Engine/Source/Runtime/MovieScene/Private/MovieSceneBindingReferences.cpp
@@ -6,6 +6,8 @@
 #include "UniversalObjectLocatorResolveParams.h"
 #include "Engine/World.h"
 #include "Evaluation/MovieSceneEvaluationState.h"
+#include "UObject/Package.h"
+#include "UnrealEngine.h"
 
 #include UE_INLINE_GENERATED_CPP_BY_NAME(MovieSceneBindingReferences)
 
@@ -92,6 +94,15 @@ void FMovieSceneBindingReferences::RemoveBinding(const FGuid& ObjectId)
 
 void FMovieSceneBindingReferences::ResolveBinding(const FGuid& ObjectId, const UE::UniversalObjectLocator::FResolveParams& ResolveParams, TArray<UObject*, TInlineAllocator<1>>& OutObjects) const
 {
+#if WITH_EDITORONLY_DATA
+	// Sequencer is explicit about providing a resolution context for its bindings. We never want to resolve to objects
+	// with a different PIE instance ID, even if the current callstack is being executed inside a different GPlayInEditorID
+	// scope. Since ResolveObject will always call FixupForPIE in editor based on GPlayInEditorID, we always override the current
+	// GPlayInEditorID to be the current PIE instance of the provided context.
+	const int32 ContextPlayInEditorID = ResolveParams.Context ? ResolveParams.Context->GetOutermost()->GetPIEInstanceID() : INDEX_NONE;
+	FTemporaryPlayInEditorIDOverride PIEGuard(ContextPlayInEditorID);
+#endif
+
 	const int32 StartIndex = Algo::LowerBoundBy(SortedReferences, ObjectId, &FMovieSceneBindingReference::ID);
 	const int32 Num = SortedReferences.Num();
 
