Sequencer: Added missing BlueprintType to MovieSceneFolder

diff --git a/Engine/Source/Runtime/MovieScene/Public/MovieSceneFolder.h b/Engine/Source/Runtime/MovieScene/Public/MovieSceneFolder.h
--- a/Engine/Source/Runtime/MovieScene/Public/MovieSceneFolder.h
+++ b/Engine/Source/Runtime/MovieScene/Public/MovieSceneFolder.h
@@ -27,7 +27,7 @@ class UMovieSceneTrack;
 namespace UE { namespace MovieScene { class IFolderEventHandler; } }
 
 /** Represents a folder used for organizing objects in tracks in a movie scene. */
-UCLASS(DefaultToInstanced, MinimalAPI)
+UCLASS(DefaultToInstanced, MinimalAPI, BlueprintType)
 class UMovieSceneFolder : public UObject
 {
 	GENERATED_UCLASS_BODY()
