Sequencer: fix sub-section trimming not taking into account non-zero playback range starts

#changelist validated
#virtualized

diff --git a/Engine/Source/Runtime/MovieScene/Private/Sections/MovieSceneSubSection.cpp b/Engine/Source/Runtime/MovieScene/Private/Sections/MovieSceneSubSection.cpp
--- a/Engine/Source/Runtime/MovieScene/Private/Sections/MovieSceneSubSection.cpp
+++ b/Engine/Source/Runtime/MovieScene/Private/Sections/MovieSceneSubSection.cpp
@@ -355,8 +355,8 @@ void UMovieSceneSubSection::TrimSection( FQualifiedFrameTime TrimTime, bool bTri
 		// The new first loop start offset is where the trim time fell inside the sub-sequence (this time is already
 		// normalized in the case of looping sub-sequences).
 		const FMovieSceneSequenceTransform OuterToInner(OuterToInnerTransform());
-		FFrameTime LocalTrimTime = OuterToInner.TransformTime(LocalTickResolutionTrimTime);
-		FFrameNumber NewStartOffset = LocalTrimTime.FrameNumber - Parameters.StartFrameOffset;
+		const FFrameTime LocalTrimTime = OuterToInner.TransformTime(LocalTickResolutionTrimTime);
+		FFrameNumber NewStartOffset = LocalTrimTime.FrameNumber + Parameters.StartFrameOffset;
 
 		// Make sure we don't have negative offsets (this shouldn't happen, though).
 		NewStartOffset = FMath::Max(FFrameNumber(0), NewStartOffset);
