Sequencer: add logging and checks to diagnose unknown crashes around sequence instance creation

#changelist validated
#virtualized

diff --git a/Engine/Source/Runtime/MovieScene/Private/Compilation/MovieSceneCompiledDataManager.cpp b/Engine/Source/Runtime/MovieScene/Private/Compilation/MovieSceneCompiledDataManager.cpp
--- a/Engine/Source/Runtime/MovieScene/Private/Compilation/MovieSceneCompiledDataManager.cpp
+++ b/Engine/Source/Runtime/MovieScene/Private/Compilation/MovieSceneCompiledDataManager.cpp
@@ -666,6 +666,26 @@ bool UMovieSceneCompiledDataManager::IsDirty(UMovieSceneSequence* Sequence) cons
 	return true;
 }
 
+bool UMovieSceneCompiledDataManager::ValidateEntry(FMovieSceneCompiledDataID DataID, UMovieSceneSequence* Sequence) const
+{
+	if (!ensureMsgf(
+			CompiledDataEntries.IsValidIndex(DataID.Value),
+			TEXT("Given DataID %d is not valid! (%d entries in the data manager)"), DataID.Value, CompiledDataEntries.Num()))
+	{
+		return false;
+	}
+
+	const FMovieSceneCompiledDataEntry& Entry = CompiledDataEntries[DataID.Value];
+	UMovieSceneSequence* EntrySequence = Entry.GetSequence();
+	if (!ensureMsgf(
+			EntrySequence == Sequence,
+			TEXT("Unexpected sequence for data ID! Expected '%s', but data manager has '%s'."), *GetNameSafe(Sequence), *GetNameSafe(EntrySequence)))
+	{
+		return false;
+	}
+
+	return true;
+}
 
 void UMovieSceneCompiledDataManager::Compile(FMovieSceneCompiledDataID DataID)
 {
diff --git a/Engine/Source/Runtime/MovieScene/Private/EntitySystem/MovieSceneSequenceInstance.cpp b/Engine/Source/Runtime/MovieScene/Private/EntitySystem/MovieSceneSequenceInstance.cpp
--- a/Engine/Source/Runtime/MovieScene/Private/EntitySystem/MovieSceneSequenceInstance.cpp
+++ b/Engine/Source/Runtime/MovieScene/Private/EntitySystem/MovieSceneSequenceInstance.cpp
@@ -162,15 +162,44 @@ void FSequenceInstance::InitializeLegacyEvaluator()
 
 void FSequenceInstance::InvalidateCachedData()
 {
-	ensureMsgf(bInitialized, TEXT("This instance hasn't been initialized yet!"));
+	ensureMsgf(bInitialized, TEXT("Sequence instance hasn't been initialized yet!"));
+
+	UMovieSceneSequence* RootSequence = SharedPlaybackState->GetRootSequence();
+	if (!ensureMsgf(RootSequence, TEXT("Sequence instance has a null root sequence!")))
+	{
+		return;
+	}
+
+	UMovieSceneCompiledDataManager* CompiledDataManager = SharedPlaybackState->GetCompiledDataManager();
+	if (!ensureMsgf(
+				CompiledDataManager, 
+				TEXT("Sequence instance (%s) has no compiled data manager! Re-building a default one."),
+				*RootSequence->GetPathName()))
+	{
+		CompiledDataManager = UMovieSceneCompiledDataManager::GetPrecompiledData();
+	}
+
+	FMovieSceneCompiledDataID RootCompiledDataID = SharedPlaybackState->GetRootCompiledDataID();
+	if (!ensureMsgf(
+				RootCompiledDataID.IsValid(), 
+				TEXT("Sequence instance (%s) has invalid data ID for root sequence! Re-building it."),
+				*RootSequence->GetPathName()))
+	{
+		RootCompiledDataID = CompiledDataManager->GetDataID(RootSequence);
+	}
+
+	if (!ensureMsgf(
+				CompiledDataManager->ValidateEntry(RootCompiledDataID, RootSequence),
+				TEXT("Sequence instance (%s) has invalid data ID for root sequence! Aborting invalidation of cached data."),
+				*RootSequence->GetPathName()))
+	{
+		return;
+	}
 
 	Ledger.Invalidate();
 
 	UpdateFlags = ESequenceInstanceUpdateFlags::None;
 
-	UMovieSceneCompiledDataManager* CompiledDataManager = SharedPlaybackState->GetCompiledDataManager();
-	const FMovieSceneCompiledDataID RootCompiledDataID = SharedPlaybackState->GetRootCompiledDataID();
-
 	FMovieSceneEvaluationState* State = SharedPlaybackState->FindCapability<FMovieSceneEvaluationState>();
 
 	if (SequenceID == MovieSceneSequenceID::Root)
diff --git a/Engine/Source/Runtime/MovieScene/Public/Compilation/MovieSceneCompiledDataManager.h b/Engine/Source/Runtime/MovieScene/Public/Compilation/MovieSceneCompiledDataManager.h
--- a/Engine/Source/Runtime/MovieScene/Public/Compilation/MovieSceneCompiledDataManager.h
+++ b/Engine/Source/Runtime/MovieScene/Public/Compilation/MovieSceneCompiledDataManager.h
@@ -248,6 +248,8 @@ public:
 
 	MOVIESCENE_API bool IsDirty(const FMovieSceneCompiledDataEntry& Entry) const;
 
+	MOVIESCENE_API bool ValidateEntry(FMovieSceneCompiledDataID DataID, UMovieSceneSequence* Sequence) const;
+
 	/**
 	 * Return a reference to a compiled data entry.
 	 * WARNING: This reference will become invalid if any sequence in this manager is (re)compiled
