Sequencer: Bypass degenerate ranges when checking for overlaps

diff --git a/Engine/Source/Runtime/MovieScene/Private/Tracks/MovieSceneBindingLifetimeTrack.cpp b/Engine/Source/Runtime/MovieScene/Private/Tracks/MovieSceneBindingLifetimeTrack.cpp
--- a/Engine/Source/Runtime/MovieScene/Private/Tracks/MovieSceneBindingLifetimeTrack.cpp
+++ b/Engine/Source/Runtime/MovieScene/Private/Tracks/MovieSceneBindingLifetimeTrack.cpp
@@ -135,7 +135,7 @@ TArray<FFrameNumberRange> UMovieSceneBindingLifetimeTrack::CalculateInverseLifet
 	TArray<TRange<FFrameNumber>> SectionRanges;
 	for (const FFrameNumberRange& Range : Ranges)
 	{
-		if (Algo::AnyOf(SectionRanges, [&](const TRange<FFrameNumber> OtherRange) { return OtherRange.Overlaps(Range); }))
+		if (Algo::AnyOf(SectionRanges, [&](const TRange<FFrameNumber> OtherRange) { return !OtherRange.IsDegenerate() && !Range.IsDegenerate() && OtherRange.Overlaps(Range); }))
 		{
 			// Lifetime Range sections have managed to overlap, which should not be allowed.
 			ensure(false);
