Sequencer: change ensure to error message

diff --git a/Engine/Source/Runtime/MovieScene/Private/Evaluation/MovieSceneEvaluationState.cpp b/Engine/Source/Runtime/MovieScene/Private/Evaluation/MovieSceneEvaluationState.cpp
--- a/Engine/Source/Runtime/MovieScene/Private/Evaluation/MovieSceneEvaluationState.cpp
+++ b/Engine/Source/Runtime/MovieScene/Private/Evaluation/MovieSceneEvaluationState.cpp
@@ -483,7 +483,12 @@ void FMovieSceneObjectCache::UpdateBindings(const FGuid& InGuid, TSharedRef<cons
 						FMovieSceneDynamicBindingResolveResult ResolveResult = FMovieSceneDynamicBindingInvoker::ResolveDynamicBinding(SharedPlaybackState, Sequence, SequenceID, InGuid, Possessable->DynamicBinding);
 						if (ResolveResult.Object)
 						{
-							ensureMsgf(ResolveResult.bIsPossessedObject, TEXT("Possessables don't support spawnable-type ownership"));
+							if (!ResolveResult.bIsPossessedObject)
+							{
+								UE_LOG(LogMovieScene, Error, 
+									TEXT("Possessable '%s' (dynamically resolved to '%s') can't have spawnable-type ownership. The user-defined director blueprint endpoint should set bIsPossessedObject to true."),
+									*LexToString(Possessable->GetName()), *ResolveResult.Object->GetName());
+							}
 							FoundObjects.Add(ResolveResult.Object);
 						}
 						else
@@ -532,7 +537,12 @@ void FMovieSceneObjectCache::UpdateBindings(const FGuid& InGuid, TSharedRef<cons
 					FMovieSceneDynamicBindingResolveResult ResolveResult = FMovieSceneDynamicBindingInvoker::ResolveDynamicBinding(SharedPlaybackState, Sequence, SequenceID, InGuid, Possessable->DynamicBinding);
 					if (ResolveResult.Object)
 					{
-						ensureMsgf(ResolveResult.bIsPossessedObject, TEXT("Possessables don't support spawnable-type ownership"));
+						if (!ResolveResult.bIsPossessedObject)
+						{
+							UE_LOG(LogMovieScene, Error,
+								TEXT("Possessable '%s' (dynamically resolved to '%s') can't have spawnable-type ownership. The user-defined director blueprint endpoint should set bIsPossessedObject to true."),
+								*LexToString(Possessable->GetName()), *ResolveResult.Object->GetName());
+						}
 						FoundObjects.Add(ResolveResult.Object);
 					}
 					else
