Sequencer: better handle incorrectly resolved object in dynamic bindings

Spawnables and top-level possessables shoudl be resolved to actors. Children possessables should be resolved to components. When a dynamic binding resolves to something unexpected, log and error and fallback to the default behavior.

diff --git a/Engine/Source/Runtime/MovieScene/Private/Evaluation/MovieSceneEvaluationState.cpp b/Engine/Source/Runtime/MovieScene/Private/Evaluation/MovieSceneEvaluationState.cpp
--- a/Engine/Source/Runtime/MovieScene/Private/Evaluation/MovieSceneEvaluationState.cpp
+++ b/Engine/Source/Runtime/MovieScene/Private/Evaluation/MovieSceneEvaluationState.cpp
@@ -480,7 +480,7 @@ void FMovieSceneObjectCache::UpdateBindings(const FGuid& InGuid, TSharedRef<cons
 					}
 					else
 					{
-						FMovieSceneDynamicBindingResolveResult ResolveResult = FMovieSceneDynamicBindingInvoker::ResolveDynamicBinding(SharedPlaybackState, Sequence, SequenceID, InGuid, Possessable->DynamicBinding);
+						FMovieSceneDynamicBindingResolveResult ResolveResult = FMovieSceneDynamicBindingInvoker::ResolveDynamicBinding(SharedPlaybackState, Sequence, SequenceID, *Possessable);
 						if (ResolveResult.Object)
 						{
 							if (!ResolveResult.bIsPossessedObject)
@@ -534,7 +534,7 @@ void FMovieSceneObjectCache::UpdateBindings(const FGuid& InGuid, TSharedRef<cons
 				}
 				else
 				{
-					FMovieSceneDynamicBindingResolveResult ResolveResult = FMovieSceneDynamicBindingInvoker::ResolveDynamicBinding(SharedPlaybackState, Sequence, SequenceID, InGuid, Possessable->DynamicBinding);
+					FMovieSceneDynamicBindingResolveResult ResolveResult = FMovieSceneDynamicBindingInvoker::ResolveDynamicBinding(SharedPlaybackState, Sequence, SequenceID, *Possessable);
 					if (ResolveResult.Object)
 					{
 						if (!ResolveResult.bIsPossessedObject)
diff --git a/Engine/Source/Runtime/MovieScene/Private/MovieSceneDynamicBindingInvoker.cpp b/Engine/Source/Runtime/MovieScene/Private/MovieSceneDynamicBindingInvoker.cpp
--- a/Engine/Source/Runtime/MovieScene/Private/MovieSceneDynamicBindingInvoker.cpp
+++ b/Engine/Source/Runtime/MovieScene/Private/MovieSceneDynamicBindingInvoker.cpp
@@ -8,10 +8,50 @@
 #include "Evaluation/MovieSceneEvaluationTemplateInstance.h"
 #include "Evaluation/SequenceDirectorPlaybackCapability.h"
 #include "MovieSceneDynamicBinding.h"
+#include "MovieScenePossessable.h"
 #include "MovieSceneSequence.h"
+#include "MovieSceneSpawnable.h"
 #include "UObject/UnrealType.h"
 
-FMovieSceneDynamicBindingResolveResult FMovieSceneDynamicBindingInvoker::ResolveDynamicBinding(TSharedRef<const UE::MovieScene::FSharedPlaybackState> SharedPlaybackState, UMovieSceneSequence* Sequence, const FMovieSceneSequenceID& SequenceID, const FGuid& InGuid, const FMovieSceneDynamicBinding& DynamicBinding)
+FMovieSceneDynamicBindingResolveResult FMovieSceneDynamicBindingInvoker::ResolveDynamicBinding(TSharedRef<const FSharedPlaybackState> SharedPlaybackState, UMovieSceneSequence* Sequence, const FMovieSceneSequenceID& SequenceID, const FMovieScenePossessable& Possessable)
+{
+	FMovieSceneDynamicBindingResolveResult Result = ResolveDynamicBinding(SharedPlaybackState, Sequence, SequenceID, Possessable.GetGuid(), Possessable.DynamicBinding);
+	// Top-level possessables must be actors. Children possessables must be components.
+	if (Result.Object && !Possessable.GetParent().IsValid() && !Result.Object.IsA<AActor>())
+	{
+		UE_LOG(
+			LogMovieScene, Error, 
+			TEXT("Possessable '%s' ('%s') must be bound to an actor, but its dynamic binding resolved it to '%s' which is a '%s'. Aborting dynamic binding and falling back to defaults."),
+			*Possessable.GetName(), *LexToString(Possessable.GetGuid()), *Result.Object->GetName(), *Result.Object->GetClass()->GetName());
+		return FMovieSceneDynamicBindingResolveResult();
+	}
+	else if (Result.Object && Possessable.GetParent().IsValid() && !Result.Object.IsA<UActorComponent>())
+	{
+		UE_LOG(
+			LogMovieScene, Error, 
+			TEXT("Possessable '%s' ('%s') must be bound to a component, but its dynamic binding resolved it to '%s' which is a '%s'. Aborting dynamic binding and falling back to defaults."),
+			*Possessable.GetName(), *LexToString(Possessable.GetGuid()), *Result.Object->GetName(), *Result.Object->GetClass()->GetName());
+		return FMovieSceneDynamicBindingResolveResult();
+	}
+	return Result;
+}
+
+FMovieSceneDynamicBindingResolveResult FMovieSceneDynamicBindingInvoker::ResolveDynamicBinding(TSharedRef<const FSharedPlaybackState> SharedPlaybackState, UMovieSceneSequence* Sequence, const FMovieSceneSequenceID& SequenceID, const FMovieSceneSpawnable& Spawnable)
+{
+	FMovieSceneDynamicBindingResolveResult Result = ResolveDynamicBinding(SharedPlaybackState, Sequence, SequenceID, Spawnable.GetGuid(), Spawnable.DynamicBinding);
+	// Spawnables must be actors.
+	if (Result.Object && !Result.Object.IsA<AActor>())
+	{
+		UE_LOG(
+			LogMovieScene, Error, 
+			TEXT("Spawnable '%s' ('%s') must be bound to an actor, but its dynamic binding resolved it to '%s' which is a '%s'. Aborting dynamic binding and falling back to defaults."),
+			*Spawnable.GetName(), *LexToString(Spawnable.GetGuid()), *Result.Object->GetName(), *Result.Object->GetClass()->GetName());
+		return FMovieSceneDynamicBindingResolveResult();
+	}
+	return Result;
+}
+
+FMovieSceneDynamicBindingResolveResult FMovieSceneDynamicBindingInvoker::ResolveDynamicBinding(TSharedRef<const FSharedPlaybackState> SharedPlaybackState, UMovieSceneSequence* Sequence, const FMovieSceneSequenceID& SequenceID, const FGuid& InGuid, const FMovieSceneDynamicBinding& DynamicBinding)
 {
 	using namespace UE::MovieScene;
 
diff --git a/Engine/Source/Runtime/MovieScene/Private/MovieSceneDynamicBindingInvoker.h b/Engine/Source/Runtime/MovieScene/Private/MovieSceneDynamicBindingInvoker.h
--- a/Engine/Source/Runtime/MovieScene/Private/MovieSceneDynamicBindingInvoker.h
+++ b/Engine/Source/Runtime/MovieScene/Private/MovieSceneDynamicBindingInvoker.h
@@ -10,7 +10,9 @@ class UObject;
 class IMovieScenePlayer;
 class UMovieSceneSequence;
 struct FGuid;
+struct FMovieScenePossessable;
 struct FMovieSceneSequenceID;
+struct FMovieSceneSpawnable;
 
 namespace UE::MovieScene
 {
@@ -22,10 +24,16 @@ namespace UE::MovieScene
  */
 struct FMovieSceneDynamicBindingInvoker
 {
+	using FSharedPlaybackState = UE::MovieScene::FSharedPlaybackState;
+
+	/** Invoke the dynamic binding, if any, and return the result */
+	static FMovieSceneDynamicBindingResolveResult ResolveDynamicBinding(TSharedRef<const FSharedPlaybackState> SharedPlaybackState, UMovieSceneSequence* Sequence, const FMovieSceneSequenceID& SequenceID, const FMovieScenePossessable& Possessable);
+
 	/** Invoke the dynamic binding, if any, and return the result */
-	static FMovieSceneDynamicBindingResolveResult ResolveDynamicBinding(TSharedRef<const UE::MovieScene::FSharedPlaybackState> SharedPlaybackState, UMovieSceneSequence* Sequence, const FMovieSceneSequenceID& SequenceID, const FGuid& InGuid, const FMovieSceneDynamicBinding& DynamicBinding);
+	static FMovieSceneDynamicBindingResolveResult ResolveDynamicBinding(TSharedRef<const FSharedPlaybackState> SharedPlaybackState, UMovieSceneSequence* Sequence, const FMovieSceneSequenceID& SequenceID, const FMovieSceneSpawnable& Spawnable);
 
 private:
+	static FMovieSceneDynamicBindingResolveResult ResolveDynamicBinding(TSharedRef<const FSharedPlaybackState> SharedPlaybackState, UMovieSceneSequence* Sequence, const FMovieSceneSequenceID& SequenceID, const FGuid& InGuid, const FMovieSceneDynamicBinding& DynamicBinding);
 	static FMovieSceneDynamicBindingResolveResult InvokeDynamicBinding(UObject* DirectorInstance, const FMovieSceneDynamicBinding& DynamicBinding, const FMovieSceneDynamicBindingResolveParams& Params);
 };
 
diff --git a/Engine/Source/Runtime/MovieScene/Private/MovieSceneSpawnRegister.cpp b/Engine/Source/Runtime/MovieScene/Private/MovieSceneSpawnRegister.cpp
--- a/Engine/Source/Runtime/MovieScene/Private/MovieSceneSpawnRegister.cpp
+++ b/Engine/Source/Runtime/MovieScene/Private/MovieSceneSpawnRegister.cpp
@@ -50,8 +50,7 @@ UObject* FMovieSceneSpawnRegister::SpawnObject(const FGuid& BindingId, UMovieSce
 	IMovieScenePlayer* Player = UE::MovieScene::FPlayerIndexPlaybackCapability::GetPlayer(SharedPlaybackState);
 
 	// See if there is some dynamic binding logic to invoke, otherwise spawn the actor
-	FMovieSceneDynamicBindingResolveResult ResolveResult = FMovieSceneDynamicBindingInvoker::ResolveDynamicBinding(
-			SharedPlaybackState, Sequence, TemplateID, BindingId, Spawnable->DynamicBinding);
+	FMovieSceneDynamicBindingResolveResult ResolveResult = FMovieSceneDynamicBindingInvoker::ResolveDynamicBinding(SharedPlaybackState, Sequence, TemplateID, *Spawnable);
 	if (ResolveResult.Object)
 	{
 		SpawnedActor = ResolveResult.Object;
