Sequencer: account for non-zero playback start times when trimming/splitting sub-sequences.

diff --git a/Engine/Source/Runtime/MovieScene/Private/Sections/MovieSceneSubSection.cpp b/Engine/Source/Runtime/MovieScene/Private/Sections/MovieSceneSubSection.cpp
--- a/Engine/Source/Runtime/MovieScene/Private/Sections/MovieSceneSubSection.cpp
+++ b/Engine/Source/Runtime/MovieScene/Private/Sections/MovieSceneSubSection.cpp
@@ -356,7 +356,11 @@ void UMovieSceneSubSection::TrimSection( FQualifiedFrameTime TrimTime, bool bTri
 		// normalized in the case of looping sub-sequences).
 		const FMovieSceneSequenceTransform OuterToInner(OuterToInnerTransform());
 		const FFrameTime LocalTrimTime = OuterToInner.TransformTime(LocalTickResolutionTrimTime);
-		FFrameNumber NewStartOffset = LocalTrimTime.FrameNumber + Parameters.StartFrameOffset;
+		// LocalTrimTime is now in the inner sequence timespace, but StartFrameOffset is an offset from the inner sequence's own
+		// playback start time, so we need to account for that.
+		TRange<FFrameNumber> LocalPlaybackRange = LocalMovieScene->GetPlaybackRange();
+		const FFrameNumber LocalPlaybackStart = LocalPlaybackRange.HasLowerBound() ? LocalPlaybackRange.GetLowerBoundValue() : FFrameNumber(0);
+		FFrameNumber NewStartOffset = LocalTrimTime.FrameNumber - LocalPlaybackStart;
 
 		// Make sure we don't have negative offsets (this shouldn't happen, though).
 		NewStartOffset = FMath::Max(FFrameNumber(0), NewStartOffset);
