Sequencer: linker/runner refactor

The original design for runners vs linkers, which was to have a central linker with many runners (one for each player) isn't valid anymore now that runners are also global (owned by the tick manager). This change therefore makes this official by making runners created and owned by the linkers themselves.

diff --git a/Engine/Plugins/Experimental/Avalanche/Source/AvalancheSequence/Private/AvaSequence.cpp b/Engine/Plugins/Experimental/Avalanche/Source/AvalancheSequence/Private/AvaSequence.cpp
--- a/Engine/Plugins/Experimental/Avalanche/Source/AvalancheSequence/Private/AvaSequence.cpp
+++ b/Engine/Plugins/Experimental/Avalanche/Source/AvalancheSequence/Private/AvaSequence.cpp
@@ -36,9 +36,8 @@ namespace UE::AvaSequence::Private
 	{
 	public:
 		FTransientPlayer(UAvaSequence& InSequence)
-			: Runner(MakeShared<FMovieSceneEntitySystemRunner>())
 		{
-			Template.Initialize(InSequence, *this, nullptr, Runner);
+			Template.Initialize(InSequence, *this, nullptr);
 			State.AssignSequence(MovieSceneSequenceID::Root, InSequence, *this);
 		}
 
@@ -51,8 +50,6 @@ namespace UE::AvaSequence::Private
 		//~ End IMovieScenePlayer
 
 	private:
-		TSharedRef<FMovieSceneEntitySystemRunner> Runner;
-
 		FMovieSceneRootEvaluationTemplateInstance Template;
 	};
 }
diff --git a/Engine/Plugins/Experimental/PropertyAnimator/Source/PropertyAnimator/Private/MovieScene/Easing/PropertyAnimatorEasingDoubleChannelEvaluatorSystem.cpp b/Engine/Plugins/Experimental/PropertyAnimator/Source/PropertyAnimator/Private/MovieScene/Easing/PropertyAnimatorEasingDoubleChannelEvaluatorSystem.cpp
--- a/Engine/Plugins/Experimental/PropertyAnimator/Source/PropertyAnimator/Private/MovieScene/Easing/PropertyAnimatorEasingDoubleChannelEvaluatorSystem.cpp
+++ b/Engine/Plugins/Experimental/PropertyAnimator/Source/PropertyAnimator/Private/MovieScene/Easing/PropertyAnimatorEasingDoubleChannelEvaluatorSystem.cpp
@@ -78,11 +78,7 @@ void UPropertyAnimatorEasingDoubleChannelEvaluatorSystem::OnRun(FSystemTaskPrere
 	const FBuiltInComponentTypes* BuiltInComponents = FBuiltInComponentTypes::Get();
 	const FPropertyAnimatorComponentTypes* PropertyAnimatorComponents = FPropertyAnimatorComponentTypes::Get();
 
-	FMovieSceneEntitySystemRunner* Runner = Linker->GetActiveRunner();
-	if (!Runner)
-	{
-		return;
-	}
+	TSharedRef<FMovieSceneEntitySystemRunner> Runner = Linker->GetRunner();
 
 	if (Runner->GetCurrentPhase() == ESystemPhase::Instantiation)
 	{
diff --git a/Engine/Plugins/Experimental/PropertyAnimator/Source/PropertyAnimator/Private/MovieScene/Wave/PropertyAnimatorWaveDoubleChannelEvaluatorSystem.cpp b/Engine/Plugins/Experimental/PropertyAnimator/Source/PropertyAnimator/Private/MovieScene/Wave/PropertyAnimatorWaveDoubleChannelEvaluatorSystem.cpp
--- a/Engine/Plugins/Experimental/PropertyAnimator/Source/PropertyAnimator/Private/MovieScene/Wave/PropertyAnimatorWaveDoubleChannelEvaluatorSystem.cpp
+++ b/Engine/Plugins/Experimental/PropertyAnimator/Source/PropertyAnimator/Private/MovieScene/Wave/PropertyAnimatorWaveDoubleChannelEvaluatorSystem.cpp
@@ -77,11 +77,7 @@ void UPropertyAnimatorWaveDoubleChannelEvaluatorSystem::OnRun(FSystemTaskPrerequ
 	const FBuiltInComponentTypes* BuiltInComponents = FBuiltInComponentTypes::Get();
 	const FPropertyAnimatorComponentTypes* PropertyAnimatorComponents = FPropertyAnimatorComponentTypes::Get();
 
-	FMovieSceneEntitySystemRunner* Runner = Linker->GetActiveRunner();
-	if (!Runner)
-	{
-		return;
-	}
+	TSharedRef<FMovieSceneEntitySystemRunner> Runner = Linker->GetRunner();
 
 	if (Runner->GetCurrentPhase() == ESystemPhase::Instantiation)
 	{
diff --git a/Engine/Plugins/MovieScene/ReplayTracks/Source/ReplayTracks/Private/Systems/MovieSceneReplaySystem.cpp b/Engine/Plugins/MovieScene/ReplayTracks/Source/ReplayTracks/Private/Systems/MovieSceneReplaySystem.cpp
--- a/Engine/Plugins/MovieScene/ReplayTracks/Source/ReplayTracks/Private/Systems/MovieSceneReplaySystem.cpp
+++ b/Engine/Plugins/MovieScene/ReplayTracks/Source/ReplayTracks/Private/Systems/MovieSceneReplaySystem.cpp
@@ -84,8 +84,8 @@ void UMovieSceneReplaySystem::OnRun(FSystemTaskPrerequisites& InPrerequisites, F
 
 	using namespace UE::MovieScene;
 
-	FMovieSceneEntitySystemRunner* ActiveRunner = Linker->GetActiveRunner();
-	ESystemPhase CurrentPhase = ActiveRunner->GetCurrentPhase();
+	TSharedRef<FMovieSceneEntitySystemRunner> Runner = Linker->GetRunner();
+	ESystemPhase CurrentPhase = Runner->GetCurrentPhase();
 
 	if (CurrentPhase == ESystemPhase::Instantiation)
 	{
diff --git a/Engine/Plugins/MovieScene/TemplateSequence/Source/TemplateSequence/Private/CameraAnimationSequencePlayer.cpp b/Engine/Plugins/MovieScene/TemplateSequence/Source/TemplateSequence/Private/CameraAnimationSequencePlayer.cpp
--- a/Engine/Plugins/MovieScene/TemplateSequence/Source/TemplateSequence/Private/CameraAnimationSequencePlayer.cpp
+++ b/Engine/Plugins/MovieScene/TemplateSequence/Source/TemplateSequence/Private/CameraAnimationSequencePlayer.cpp
@@ -374,9 +374,7 @@ void UCameraAnimationSequencePlayer::Initialize(UMovieSceneSequence* InSequence,
 
 	UCameraAnimationSequenceSubsystem* Subsystem = UCameraAnimationSequenceSubsystem::GetCameraAnimationSequenceSubsystem(GetWorld());
 	ensureMsgf(Subsystem, TEXT("Unable to locate a valid camera animation sub-system. Camera anim sequences will not play."));
-
-	TSharedPtr<FMovieSceneEntitySystemRunner> Runner = Subsystem ? Subsystem->GetRunner() : nullptr;
-	RootTemplateInstance.Initialize(*Sequence, *this, nullptr, Runner);
+	RootTemplateInstance.Initialize(*Sequence, *this, nullptr);
 }
 
 void UCameraAnimationSequencePlayer::Play(bool bLoop, bool bRandomStartTime)
diff --git a/Engine/Plugins/MovieScene/TemplateSequence/Source/TemplateSequence/Private/CameraAnimationSequenceSubsystem.cpp b/Engine/Plugins/MovieScene/TemplateSequence/Source/TemplateSequence/Private/CameraAnimationSequenceSubsystem.cpp
--- a/Engine/Plugins/MovieScene/TemplateSequence/Source/TemplateSequence/Private/CameraAnimationSequenceSubsystem.cpp
+++ b/Engine/Plugins/MovieScene/TemplateSequence/Source/TemplateSequence/Private/CameraAnimationSequenceSubsystem.cpp
@@ -211,18 +211,11 @@ UCameraAnimationSequenceSubsystem::~UCameraAnimationSequenceSubsystem()
 
 void UCameraAnimationSequenceSubsystem::Initialize(FSubsystemCollectionBase& Collection)
 {
-	Runner = MakeShared<FMovieSceneEntitySystemRunner>();
 }
 
 void UCameraAnimationSequenceSubsystem::Deinitialize()
 {
-	// We check if the runner still has a valid pointer on the linker because the linker could
-	// have been GC'ed just now, which would make DetachFromLinker complain.
-	if (Runner->GetLinker())
-	{
-		Runner->DetachFromLinker();
-		Runner = nullptr;
-	}
+	Runner = nullptr;
 	Linker = nullptr;
 
 	Super::Deinitialize();
@@ -233,7 +226,7 @@ UMovieSceneEntitySystemLinker* UCameraAnimationSequenceSubsystem::GetLinker(bool
 	if (!Linker && bAutoCreate)
 	{
 		Linker = CreateLinker(this, TEXT("CameraAnimationSequenceSubsystemLinker"));
-		Runner->AttachToLinker(Linker);
+		Runner = Linker->GetRunner();
 	}
 	return Linker;
 }
diff --git a/Engine/Source/Editor/Sequencer/Private/SSequencer.cpp b/Engine/Source/Editor/Sequencer/Private/SSequencer.cpp
--- a/Engine/Source/Editor/Sequencer/Private/SSequencer.cpp
+++ b/Engine/Source/Editor/Sequencer/Private/SSequencer.cpp
@@ -2055,7 +2055,7 @@ void SSequencer::FillAdvancedMenu(FMenuBuilder& MenuBuilder)
 				// Since sequencer owns its own compiled data manager, it's ok to override the mask here and reset everything
 				Template.GetCompiledDataManager()->SetEmulatedNetworkMask(InMode);
 				// Reinitialize the template again
-				Template.Initialize(*RootSequence, *SequencerPin, Template.GetCompiledDataManager(), Template.GetRunner());
+				Template.Initialize(*RootSequence, *SequencerPin, Template.GetCompiledDataManager());
 			}
 		};
 		auto IsNetworkModeChecked = [WeakSequencer = SequencerPtr](EMovieSceneServerClientMask InMode)
diff --git a/Engine/Source/Editor/Sequencer/Private/Sequencer.cpp b/Engine/Source/Editor/Sequencer/Private/Sequencer.cpp
--- a/Engine/Source/Editor/Sequencer/Private/Sequencer.cpp
+++ b/Engine/Source/Editor/Sequencer/Private/Sequencer.cpp
@@ -485,8 +485,8 @@ void FSequencer::InitSequencer(const FSequencerInitParams& InitParams, const TSh
 	ActiveTemplateIDs.Add(MovieSceneSequenceID::Root);
 	ActiveTemplateStates.Add(true);
 
-	Runner = MakeShared<FMovieSceneEntitySystemRunner>();
-	RootTemplateInstance.Initialize(*InitParams.RootSequence, *this, CompiledDataManager, Runner);
+	RootTemplateInstance.Initialize(*InitParams.RootSequence, *this, CompiledDataManager);
+	Runner = RootTemplateInstance.GetRunner();
 
 	RootTemplateInstance.EnableGlobalPreAnimatedStateCapture();
 
@@ -716,7 +716,7 @@ FSequencer::FSequencer()
 
 FSequencer::~FSequencer()
 {
-	if (Runner->IsAttachedToLinker())
+	if (Runner)
 	{
 		Runner->QueueFinalUpdate(RootTemplateInstance.GetRootInstanceHandle());
 		Runner->Flush();
@@ -759,7 +759,7 @@ void FSequencer::Close()
 		OldMaxTickRate.Reset();
 	}
 
-	if (Runner->IsAttachedToLinker())
+	if (Runner)
 	{
 		Runner->QueueFinalUpdate(RootTemplateInstance.GetRootInstanceHandle());
 		Runner->Flush();
@@ -1110,7 +1110,7 @@ void FSequencer::ResetToNewRootSequence(UMovieSceneSequence& NewSequence)
 		}
 	}
 
-	if (Runner->IsAttachedToLinker())
+	if (Runner)
 	{
 		Runner->QueueFinalUpdate(RootTemplateInstance.GetRootInstanceHandle());
 		Runner->Flush();
@@ -1121,7 +1121,7 @@ void FSequencer::ResetToNewRootSequence(UMovieSceneSequence& NewSequence)
 	ActiveTemplateStates.Reset();
 	ActiveTemplateStates.Add(true);
 
-	RootTemplateInstance.Initialize(NewSequence, *this, CompiledDataManager, Runner);
+	RootTemplateInstance.Initialize(NewSequence, *this, CompiledDataManager);
 	RootTemplateInstance.EnableGlobalPreAnimatedStateCapture();
 
 	RootToLocalTransform = FMovieSceneSequenceTransform();
diff --git a/Engine/Source/Runtime/MovieScene/Private/EntitySystem/MovieSceneEntitySystemLinker.cpp b/Engine/Source/Runtime/MovieScene/Private/EntitySystem/MovieSceneEntitySystemLinker.cpp
--- a/Engine/Source/Runtime/MovieScene/Private/EntitySystem/MovieSceneEntitySystemLinker.cpp
+++ b/Engine/Source/Runtime/MovieScene/Private/EntitySystem/MovieSceneEntitySystemLinker.cpp
@@ -163,6 +163,8 @@ UMovieSceneEntitySystemLinker::UMovieSceneEntitySystemLinker(const FObjectInitia
 
 		InstanceRegistry.Reset(new FInstanceRegistry(this));
 
+		Runner = MakeShared<FMovieSceneEntitySystemRunner>(this);
+
 #if WITH_EDITOR
 		FCoreUObjectDelegates::OnObjectsReplaced.AddUObject(this, &UMovieSceneEntitySystemLinker::OnObjectsReplaced);
 #endif
@@ -314,15 +316,13 @@ bool UMovieSceneEntitySystemLinker::HasStructureChangedSinceLastRun() const
 	return EntityManager.HasStructureChangedSince(LastInstantiationVersion);
 }
 
-bool UMovieSceneEntitySystemLinker::StartEvaluation(FMovieSceneEntitySystemRunner& InRunner)
+bool UMovieSceneEntitySystemLinker::StartEvaluation()
 {
-
-	if (ActiveRunners.Num() == 0 || ActiveRunnerReentrancyFlags[ActiveRunners.Num()-1])
+	if (RunnerReentrancyFlags.Num() == 0 || RunnerReentrancyFlags[RunnerReentrancyFlags.Num() - 1])
 	{
 		// Default to re-entrancy being forbidden. The runner will allow re-entrancy at specific spots
 		// in the evaluation loop, via a "re-entrancy window".
-		ActiveRunners.Emplace(&InRunner);
-		ActiveRunnerReentrancyFlags.Add(false);
+		RunnerReentrancyFlags.Add(false);
 		return true;
 	}
 		
@@ -330,33 +330,25 @@ bool UMovieSceneEntitySystemLinker::StartEvaluation(FMovieSceneEntitySystemRunne
 	return false;
 }
 
-FMovieSceneEntitySystemRunner* UMovieSceneEntitySystemLinker::GetActiveRunner() const
+TSharedRef<FMovieSceneEntitySystemRunner> UMovieSceneEntitySystemLinker::GetRunner() const
 {
-	if (ActiveRunners.Num() > 0)
-	{
-		return ActiveRunners.Last();
-	}
-	return nullptr;
+	// The runner might lose us if we are GC'ed while someone else is keeping the runner's ref-count up,
+	// but we can never lose our runner, since we have at least one reference.
+	return Runner.ToSharedRef();
 }
 
-void UMovieSceneEntitySystemLinker::PostInstantation(FMovieSceneEntitySystemRunner& InRunner)
+void UMovieSceneEntitySystemLinker::PostInstantation()
 {
 	LastInstantiationVersion = EntityManager.GetSystemSerial();
 
 	GetInstanceRegistry()->PostInstantation();
 }
 
-void UMovieSceneEntitySystemLinker::EndEvaluation(FMovieSceneEntitySystemRunner& InRunner)
+void UMovieSceneEntitySystemLinker::EndEvaluation()
 {
-	if (ensureMsgf((ActiveRunners.Num() > 0 && ActiveRunners.Last() == &InRunner),
-				TEXT("Trying end the evaluation of a runner that's not the latest one to run.")))
-	{
-		const int32 LastIndex = ActiveRunners.Num()-1;
-		ensureAlways(ActiveRunnerReentrancyFlags[LastIndex] == false);
-
-		ActiveRunners.Pop();
-		ActiveRunnerReentrancyFlags.RemoveAt(LastIndex);
-	}
+	const int32 LastIndex = RunnerReentrancyFlags.Num()-1;
+	ensureAlways(RunnerReentrancyFlags[LastIndex] == false);
+	RunnerReentrancyFlags.RemoveAt(LastIndex);
 }
 
 void UMovieSceneEntitySystemLinker::HandlePreGarbageCollection()
@@ -458,16 +450,13 @@ void UMovieSceneEntitySystemLinker::CleanGarbage()
 	InstanceRegistry->CleanupLinkerEntities(FreedEntities);
 
 	// If we have any runners part-way through an evaluation, we need to reset them so that they re-evaluate from the start
-	ResetActiveRunners();
+	ResetRunner();
 }
 
-void UMovieSceneEntitySystemLinker::ResetActiveRunners()
+void UMovieSceneEntitySystemLinker::ResetRunner()
 {
 	// If we have any runners part-way through an evaluation, we need to reset them so that they re-evaluate from the start
-	for (int32 Index = ActiveRunners.Num()-1; Index >= 0; --Index)
-	{
-		ActiveRunners[Index]->ResetFlushState();
-	}
+	Runner->ResetFlushState();
 }
 
 void UMovieSceneEntitySystemLinker::DestroyInstanceImmediately(UE::MovieScene::FRootInstanceHandle Instance)
@@ -478,7 +467,7 @@ void UMovieSceneEntitySystemLinker::DestroyInstanceImmediately(UE::MovieScene::F
 	// Destroy the instance and any sub sequences. Any pre-existing NeedsLink entities will be forcibly made NeedsUnlink and cleaned as garbage
 	GetInstanceRegistry()->DestroyInstance(Instance);
 	CleanGarbage();
-	ResetActiveRunners();
+	ResetRunner();
 }
 
 void UMovieSceneEntitySystemLinker::OnObjectsReplaced(const TMap<UObject*, UObject*>& ReplacementMap)
diff --git a/Engine/Source/Runtime/MovieScene/Private/EntitySystem/MovieSceneEntitySystemRunner.cpp b/Engine/Source/Runtime/MovieScene/Private/EntitySystem/MovieSceneEntitySystemRunner.cpp
--- a/Engine/Source/Runtime/MovieScene/Private/EntitySystem/MovieSceneEntitySystemRunner.cpp
+++ b/Engine/Source/Runtime/MovieScene/Private/EntitySystem/MovieSceneEntitySystemRunner.cpp
@@ -13,6 +13,7 @@
 #include "Algo/Reverse.h"
 #include "Algo/Sort.h"
 #include "ProfilingDebugging/CountersTrace.h"
+#include "UObject/ObjectMacros.h"
 
 DECLARE_CYCLE_STAT(TEXT("Runner Flush"), 				MovieSceneEval_RunnerFlush, 				STATGROUP_MovieSceneEval);
 
@@ -50,21 +51,21 @@ struct FMovieSceneEntitySystemEvaluationReentrancyWindow
 {
 	FMovieSceneEntitySystemRunner* Runner;
 	UMovieSceneEntitySystemLinker* Linker;
-	int32 RunnerIndex;
+	int32 ReentrancyLevel;
 	UE::MovieScene::ERunnerFlushState CachedFlushState;
 	UE::MovieScene::ERunnerFlushState CachedCurrentFlushState;
 
 	FMovieSceneEntitySystemEvaluationReentrancyWindow(FMovieSceneEntitySystemRunner* InRunner, UMovieSceneEntitySystemLinker* InLinker)
 		: Runner(InRunner)
 		, Linker(InLinker)
-		, RunnerIndex(Linker->ActiveRunners.Num() - 1)
+		, ReentrancyLevel(Linker->RunnerReentrancyFlags.Num() - 1)
 		, CachedFlushState(Runner->FlushState)
 		, CachedCurrentFlushState(Runner->CurrentFlushState)
 	{
-		check(RunnerIndex >= 0);
-		checkf(Linker->ActiveRunnerReentrancyFlags[RunnerIndex] == false, TEXT("Nested FMovieSceneEntitySystemEvaluationReentrancyWindows are not supported for the same active runner"));
+		check(ReentrancyLevel >= 0);
+		checkf(Linker->RunnerReentrancyFlags[ReentrancyLevel] == false, TEXT("Nested FMovieSceneEntitySystemEvaluationReentrancyWindows are not supported for the same active runner"));
 
-		Linker->ActiveRunnerReentrancyFlags[RunnerIndex] = true;
+		Linker->RunnerReentrancyFlags[ReentrancyLevel] = true;
 
 		// Clear the current flush state to prevent inner scopes from flushing
 		//      our pending states and set the flag that allows us to be re-entrant
@@ -82,14 +83,13 @@ struct FMovieSceneEntitySystemEvaluationReentrancyWindow
 
 		Runner->FlushState = CachedFlushState;
 		Runner->CurrentFlushState = CachedCurrentFlushState;
-		if (ensure(Linker->ActiveRunnerReentrancyFlags.IsValidIndex(RunnerIndex)))
+		if (ensure(Linker->RunnerReentrancyFlags.IsValidIndex(ReentrancyLevel)))
 		{
-			Linker->ActiveRunnerReentrancyFlags[RunnerIndex] = false;
+			Linker->RunnerReentrancyFlags[ReentrancyLevel] = false;
 		}
 	}
 };
 
-
 FMovieSceneEntitySystemRunner::FMovieSceneEntitySystemRunner()
 	: GameThread(ENamedThreads::GameThread_Local)
 	, CurrentPhase(UE::MovieScene::ESystemPhase::None)
@@ -100,57 +100,25 @@ FMovieSceneEntitySystemRunner::FMovieSceneEntitySystemRunner()
 {
 }
 
-FMovieSceneEntitySystemRunner::~FMovieSceneEntitySystemRunner()
-{
-	if (IsAttachedToLinker())
-	{
-		DetachFromLinker();
-	}
-}
-
-void FMovieSceneEntitySystemRunner::AttachToLinker(UMovieSceneEntitySystemLinker* InLinker)
+FMovieSceneEntitySystemRunner::FMovieSceneEntitySystemRunner(UMovieSceneEntitySystemLinker* InLinker)
+	: WeakLinker(InLinker)
+	, GameThread(ENamedThreads::GameThread_Local)
+	, CurrentPhase(UE::MovieScene::ESystemPhase::None)
+	, FlushState(UE::MovieScene::ERunnerFlushState::None)
+	, CurrentFlushState(UE::MovieScene::ERunnerFlushState::None)
+	, bRequireFullFlush(false)
+	, bIsUpdatingSequence(false)
 {
-	if (!ensureMsgf(InLinker, TEXT("Can't attach to a null linker!")))
-	{
-		return;
-	}
-	if (!ensureMsgf(WeakLinker.IsExplicitlyNull(), TEXT("This runner is already attached to a linker")))
-	{
-		if (ensureMsgf(WeakLinker.IsValid(), TEXT("Our previous linker isn't valid anymore! We will permit attaching to a new one.")))
-		{
-			return;
-		}
-	}
-
-	WeakLinker = InLinker;
+	// We need to handle the case of our linker being GC'ed, in case someone else is holding
+	// our ref-count up in addition to the linker.
 	InLinker->Events.AbandonLinker.AddRaw(this, &FMovieSceneEntitySystemRunner::OnLinkerAbandon);
 }
 
-bool FMovieSceneEntitySystemRunner::IsAttachedToLinker() const
-{
-	return !WeakLinker.IsExplicitlyNull();
-}
-
-void FMovieSceneEntitySystemRunner::DetachFromLinker()
-{
-	if (ensureMsgf(!WeakLinker.IsExplicitlyNull(), TEXT("This runner is not attached to any linker")))
-	{
-		// Abandon our previous linker. We need to do so even for a linker that is pending kill, otherwise
-		// we will later get a OnLinkerAbandon call, which _could_ come _after_ we've been re-attached to a
-		// new valid linker. This would in turn trip the ensure in OnLinkerAbandon that checks that we are
-		// abandoning the linker we have, instead of an unrelated linker.
-		if (UMovieSceneEntitySystemLinker* Linker = WeakLinker.Get(true))
-		{
-			OnLinkerAbandon(Linker);
-		}
-	}
-
-	WeakLinker.Reset();
-}
-
-UMovieSceneEntitySystemLinker* FMovieSceneEntitySystemRunner::GetLinker() const
+FMovieSceneEntitySystemRunner::~FMovieSceneEntitySystemRunner()
 {
-	return WeakLinker.Get();
+	// We don't need to unregister AbandonLinker here since we are either being destroyed along with
+	// our linker, or someone kept us alive longer than the linker and we already unregistered the
+	// event inisde OnLinkerAbandon.
 }
 
 UE::MovieScene::FEntityManager* FMovieSceneEntitySystemRunner::GetEntityManager() const
@@ -321,7 +289,7 @@ UE::MovieScene::ERunnerFlushResult FMovieSceneEntitySystemRunner::StartEvaluatio
 		return ERunnerFlushResult::Break;
 	}
 
-	if (!Linker->StartEvaluation(*this))
+	if (!Linker->StartEvaluation())
 	{
 		return ERunnerFlushResult::Break;
 	}
@@ -342,7 +310,7 @@ UE::MovieScene::ERunnerFlushResult FMovieSceneEntitySystemRunner::StartEvaluatio
 void FMovieSceneEntitySystemRunner::EndEvaluation(UMovieSceneEntitySystemLinker* Linker)
 {
 	using namespace UE::MovieScene;
-	Linker->EndEvaluation(*this);
+	Linker->EndEvaluation();
 }
 
 UE::MovieScene::ERunnerFlushResult FMovieSceneEntitySystemRunner::FlushNext(UMovieSceneEntitySystemLinker* Linker)
@@ -486,10 +454,9 @@ void FMovieSceneEntitySystemRunner::FlushOutstanding(double BudgetMs, UE::MovieS
 		return;
 	}
 
-	UMovieSceneEntitySystemLinker* Linker = GetLinker();
-
-	// Check that we are attached to a linker that allows starting a new evaluation.
-	if (!ensureMsgf(Linker, TEXT("Runner isn't attached to a valid linker")))
+	// If this runner's linker has been destroyed, early return
+	UMovieSceneEntitySystemLinker* Linker = WeakLinker.Get();
+	if (!Linker)
 	{
 		return;
 	}
@@ -739,9 +706,9 @@ UE::MovieScene::ERunnerFlushResult FMovieSceneEntitySystemRunner::GameThread_Upd
 			FSequenceInstance& Instance = InstanceRegistry->MutateInstance(Request.Params.InstanceHandle);
 			if (!Instance.IsRootSequence())
 			{
-				FMovieSceneRootEvaluationTemplateInstance& Template = Instance.GetPlayer()->GetEvaluationTemplate();
-				UMovieSceneSequence* RootSequence = Template.GetRootSequence();
-				UMovieSceneSequence* SubSequence  = Template.GetSequence(Instance.GetSequenceID());
+				TSharedRef<const FSharedPlaybackState> SharedPlaybackState = Instance.GetSharedPlaybackState();
+				UMovieSceneSequence* RootSequence = SharedPlaybackState->GetRootSequence();
+				UMovieSceneSequence* SubSequence  = SharedPlaybackState->GetSequence(Instance.GetSequenceID());
 
 				ensureMsgf(Instance.IsRootSequence(), TEXT("Update request received for a non-root sequence ID 0x%08X (%s) in root-sequence %s. This is not supported."),
 					Instance.GetSequenceID().GetInternalValue(),
@@ -923,7 +890,7 @@ UE::MovieScene::ERunnerFlushResult FMovieSceneEntitySystemRunner::GameThread_Rei
 
 	TGuardValue<bool> IsUpdatingSequenceGuard(bIsUpdatingSequence, true);
 
-	FInstanceRegistry* InstanceRegistry = WeakLinker.Get()->GetInstanceRegistry();
+	FInstanceRegistry* InstanceRegistry = Linker->GetInstanceRegistry();
 
 	// NOTE: IncrementSystemSerial must be called before any instance updates are made
 	//       to ensure that up-to-date versions are used inside FEntityManager::OnStructureChanged
@@ -1067,7 +1034,7 @@ UE::MovieScene::ERunnerFlushResult FMovieSceneEntitySystemRunner::GameThread_Pos
 
 	check(GameThread == ENamedThreads::GameThread || GameThread == ENamedThreads::GameThread_Local);
 
-	Linker->PostInstantation(*this);
+	Linker->PostInstantation();
 
 	FEntityManager& EntityManager = Linker->EntityManager;
 	FBuiltInComponentTypes* BuiltInComponentTypes = FBuiltInComponentTypes::Get();
@@ -1358,8 +1325,8 @@ bool FMovieSceneEntitySystemRunner::FlushSingleEvaluationPhase()
 		return false;
 	}
 
-	UMovieSceneEntitySystemLinker* Linker = GetLinker();
-	if (!ensureMsgf(Linker, TEXT("Runner isn't attached to a valid linker")))
+	UMovieSceneEntitySystemLinker* Linker = WeakLinker.Get();
+	if (!ensureMsgf(Linker, TEXT("Runner doesn't have a valid linker")))
 	{
 		return false;
 	}
diff --git a/Engine/Source/Runtime/MovieScene/Private/EntitySystem/MovieSceneEvaluationHookSystem.cpp b/Engine/Source/Runtime/MovieScene/Private/EntitySystem/MovieSceneEvaluationHookSystem.cpp
--- a/Engine/Source/Runtime/MovieScene/Private/EntitySystem/MovieSceneEvaluationHookSystem.cpp
+++ b/Engine/Source/Runtime/MovieScene/Private/EntitySystem/MovieSceneEvaluationHookSystem.cpp
@@ -107,11 +107,7 @@ void UMovieSceneEvaluationHookSystem::OnRun(FSystemTaskPrerequisites& InPrerequi
 {
 	using namespace UE::MovieScene;
 
-	FMovieSceneEntitySystemRunner* Runner = Linker->GetActiveRunner();
-	if (!ensure(Runner))
-	{
-		return;
-	}
+	TSharedRef<FMovieSceneEntitySystemRunner> Runner = Linker->GetRunner();
 
 	ESystemPhase CurrentPhase = Runner->GetCurrentPhase();
 	if (CurrentPhase == ESystemPhase::Instantiation)
diff --git a/Engine/Source/Runtime/MovieScene/Private/EntitySystem/MovieSceneInstanceRegistry.cpp b/Engine/Source/Runtime/MovieScene/Private/EntitySystem/MovieSceneInstanceRegistry.cpp
--- a/Engine/Source/Runtime/MovieScene/Private/EntitySystem/MovieSceneInstanceRegistry.cpp
+++ b/Engine/Source/Runtime/MovieScene/Private/EntitySystem/MovieSceneInstanceRegistry.cpp
@@ -60,7 +60,6 @@ FInstanceHandle FInstanceRegistry::FindRelatedInstanceHandle(FInstanceHandle Ins
 FRootInstanceHandle FInstanceRegistry::AllocateRootInstance(
 		UMovieSceneSequence& InRootSequence,
 		UObject* InPlaybackContext,
-		TSharedPtr<FMovieSceneEntitySystemRunner> InRunner,
 		UMovieSceneCompiledDataManager* InCompiledDataManager)
 {
 	check(Instances.Num() < 65535);
@@ -70,11 +69,6 @@ FRootInstanceHandle FInstanceRegistry::AllocateRootInstance(
 	FSparseArrayAllocationInfo NewAllocation = Instances.AddUninitialized();
 	FRootInstanceHandle InstanceHandle { (uint16)NewAllocation.Index, InstanceSerial };
 
-	if (!InRunner)
-	{
-		FMovieSceneEntitySystemRunner* ActiveRunner = Linker->GetActiveRunner();
-		InRunner = ActiveRunner ? ActiveRunner->AsShared() : TSharedPtr<FMovieSceneEntitySystemRunner>();
-	}
 	if (!InCompiledDataManager)
 	{
 		InCompiledDataManager = UMovieSceneCompiledDataManager::GetPrecompiledData();
@@ -83,7 +77,7 @@ FRootInstanceHandle FInstanceRegistry::AllocateRootInstance(
 	FSharedPlaybackStateCreateParams PlaybackStateCreateParams;
 	PlaybackStateCreateParams.PlaybackContext = InPlaybackContext;
 	PlaybackStateCreateParams.RootInstanceHandle = InstanceHandle;
-	PlaybackStateCreateParams.Runner = InRunner;
+	PlaybackStateCreateParams.Linker = Linker;
 	PlaybackStateCreateParams.CompiledDataManager = InCompiledDataManager;
 
 	TSharedRef<FSharedPlaybackState> NewPlaybackState = MakeShared<FSharedPlaybackState>(InRootSequence, PlaybackStateCreateParams);
diff --git a/Engine/Source/Runtime/MovieScene/Private/EntitySystem/MovieSceneSequenceUpdaters.cpp b/Engine/Source/Runtime/MovieScene/Private/EntitySystem/MovieSceneSequenceUpdaters.cpp
--- a/Engine/Source/Runtime/MovieScene/Private/EntitySystem/MovieSceneSequenceUpdaters.cpp
+++ b/Engine/Source/Runtime/MovieScene/Private/EntitySystem/MovieSceneSequenceUpdaters.cpp
@@ -752,8 +752,7 @@ void FSequenceUpdater_Hierarchical::Update(TSharedRef<const FSharedPlaybackState
 		}
 	}
 
-	FMovieSceneEntitySystemRunner* Runner = Linker->GetActiveRunner();
-	check(Runner);
+	TSharedRef<FMovieSceneEntitySystemRunner> Runner = Linker->GetRunner();
 
 	for (auto InstanceIt = SequenceInstances.CreateIterator(); InstanceIt; ++InstanceIt)
 	{
diff --git a/Engine/Source/Runtime/MovieScene/Private/EntitySystem/MovieSceneSharedPlaybackState.cpp b/Engine/Source/Runtime/MovieScene/Private/EntitySystem/MovieSceneSharedPlaybackState.cpp
--- a/Engine/Source/Runtime/MovieScene/Private/EntitySystem/MovieSceneSharedPlaybackState.cpp
+++ b/Engine/Source/Runtime/MovieScene/Private/EntitySystem/MovieSceneSharedPlaybackState.cpp
@@ -18,7 +18,7 @@ FSharedPlaybackState::FSharedPlaybackState(
 		const FSharedPlaybackStateCreateParams& CreateParams)
 	: WeakRootSequence(&InRootSequence)
 	, WeakPlaybackContext(CreateParams.PlaybackContext)
-	, WeakRunner(CreateParams.Runner)
+	, WeakLinker(CreateParams.Linker)
 	, CompiledDataManager(CreateParams.CompiledDataManager)
 	, RootInstanceHandle(CreateParams.RootInstanceHandle)
 {
@@ -28,11 +28,11 @@ FSharedPlaybackState::FSharedPlaybackState(
 	}
 }
 
-UMovieSceneEntitySystemLinker* FSharedPlaybackState::GetLinker() const
+TSharedPtr<FMovieSceneEntitySystemRunner> FSharedPlaybackState::GetRunner() const
 {
-	if (TSharedPtr<FMovieSceneEntitySystemRunner> Runner = WeakRunner.Pin())
+	if (UMovieSceneEntitySystemLinker* Linker = WeakLinker.Get())
 	{
-		return Runner->GetLinker();
+		return Linker->GetRunner();
 	}
 	return nullptr;
 }
diff --git a/Engine/Source/Runtime/MovieScene/Private/Evaluation/MovieSceneEvaluationTemplateInstance.cpp b/Engine/Source/Runtime/MovieScene/Private/Evaluation/MovieSceneEvaluationTemplateInstance.cpp
--- a/Engine/Source/Runtime/MovieScene/Private/Evaluation/MovieSceneEvaluationTemplateInstance.cpp
+++ b/Engine/Source/Runtime/MovieScene/Private/Evaluation/MovieSceneEvaluationTemplateInstance.cpp
@@ -67,7 +67,7 @@ UMovieSceneEntitySystemLinker* FMovieSceneRootEvaluationTemplateInstance::Constr
 	return UMovieSceneEntitySystemLinker::FindOrCreateLinker(PlaybackContext, UE::MovieScene::EEntitySystemLinkerRole::Standalone, TEXT("DefaultEntitySystemLinker"));
 }
 
-void FMovieSceneRootEvaluationTemplateInstance::Initialize(UMovieSceneSequence& InRootSequence, IMovieScenePlayer& Player, UMovieSceneCompiledDataManager* InCompiledDataManager, TSharedPtr<FMovieSceneEntitySystemRunner> InRunner)
+void FMovieSceneRootEvaluationTemplateInstance::Initialize(UMovieSceneSequence& InRootSequence, IMovieScenePlayer& Player, UMovieSceneCompiledDataManager* InCompiledDataManager)
 {
 	using namespace UE::MovieScene;
 
@@ -111,9 +111,7 @@ void FMovieSceneRootEvaluationTemplateInstance::Initialize(UMovieSceneSequence&
 	bReinitialize |= (PreviousCompiledDataManager != InCompiledDataManager);
 
 	// Reinitialize if the runner has changed.
-	TSharedPtr<FMovieSceneEntitySystemRunner> PreviousRunner = SharedPlaybackState ?
-		SharedPlaybackState->GetRunner() : nullptr;
-	bReinitialize |= (PreviousRunner != InRunner);
+	TSharedPtr<FMovieSceneEntitySystemRunner> PreviousRunner = SharedPlaybackState ? SharedPlaybackState->GetRunner() : nullptr;
 
 	// Reinitialize if the root sequence has changed.
 	UMovieSceneSequence* PreviousRootSequence = SharedPlaybackState ? 
@@ -148,26 +146,13 @@ void FMovieSceneRootEvaluationTemplateInstance::Initialize(UMovieSceneSequence&
 		// when initializing the new root instance and shared playback state.
 		EntitySystemLinker = ConstructEntityLinker(Player);
 
-		if (ensure(InRunner) && EntitySystemLinker)
-		{
-			if (InRunner->IsAttachedToLinker() && InRunner->GetLinker() != EntitySystemLinker)
-			{
-				InRunner->DetachFromLinker();
-			}
-
-			if (!InRunner->IsAttachedToLinker())
-			{
-				InRunner->AttachToLinker(EntitySystemLinker);
-			}
-		}
-
 		// Create the new root instance and save its new shared playback state.
 		FRootInstanceHandle RootInstanceHandle;
 		if (EntitySystemLinker != nullptr && EntitySystemLinker->GetInstanceRegistry())
 		{
 			UObject* PlaybackContext = Player.GetPlaybackContext();
 			FInstanceRegistry* InstanceRegistry = EntitySystemLinker->GetInstanceRegistry();
-			RootInstanceHandle = InstanceRegistry->AllocateRootInstance(InRootSequence, PlaybackContext, InRunner, InCompiledDataManager);
+			RootInstanceHandle = InstanceRegistry->AllocateRootInstance(InRootSequence, PlaybackContext, InCompiledDataManager);
 			SharedPlaybackState = InstanceRegistry->GetInstance(RootInstanceHandle).GetSharedPlaybackState();
 			Player.InitializeRootInstance(SharedPlaybackState.ToSharedRef());
 		}
@@ -443,7 +428,6 @@ void FMovieSceneRootEvaluationTemplateInstance::PlaybackContextChanged(IMovieSce
 
 	// Only the playback context changed, so we keep the same sequence, runner, and compiled data manager.
 	UMovieSceneSequence* RootSequence = SharedPlaybackState->GetRootSequence();
-	TSharedPtr<FMovieSceneEntitySystemRunner> Runner = SharedPlaybackState->GetRunner();
 	UMovieSceneCompiledDataManager* CompiledDataManager = SharedPlaybackState->GetCompiledDataManager();
 
 	const bool bGlobalCapture = Player.PreAnimatedState.IsCapturingGlobalPreAnimatedState();
@@ -457,6 +441,7 @@ void FMovieSceneRootEvaluationTemplateInstance::PlaybackContextChanged(IMovieSce
 	{
 		EntitySystemLinker->CleanupInvalidBoundObjects();
 
+		TSharedPtr<FMovieSceneEntitySystemRunner> Runner = SharedPlaybackState->GetRunner();
 		if (Runner)
 		{
 			if (Runner->QueueFinalUpdate(PreviousRootInstanceHandle))
@@ -473,20 +458,11 @@ void FMovieSceneRootEvaluationTemplateInstance::PlaybackContextChanged(IMovieSce
 	}
 
 	EntitySystemLinker = ConstructEntityLinker(Player);
-	if (Runner)
-	{
-		if (Runner->IsAttachedToLinker())
-		{
-			Runner->DetachFromLinker();
-		}
-
-		Runner->AttachToLinker(EntitySystemLinker);
-	}
 
 	UObject* PlaybackContext = Player.GetPlaybackContext();
 	FInstanceRegistry* InstanceRegistry = EntitySystemLinker->GetInstanceRegistry();
 	const FRootInstanceHandle RootInstanceHandle = InstanceRegistry->AllocateRootInstance(
-			*RootSequence, PlaybackContext, Runner, CompiledDataManager);
+			*RootSequence, PlaybackContext, CompiledDataManager);
 	FSequenceInstance& RootInstance = InstanceRegistry->MutateInstance(RootInstanceHandle);
 	SharedPlaybackState = RootInstance.GetSharedPlaybackState();
 	Player.InitializeRootInstance(SharedPlaybackState.ToSharedRef());
diff --git a/Engine/Source/Runtime/MovieScene/Private/MovieSceneSequencePlayer.cpp b/Engine/Source/Runtime/MovieScene/Private/MovieSceneSequencePlayer.cpp
--- a/Engine/Source/Runtime/MovieScene/Private/MovieSceneSequencePlayer.cpp
+++ b/Engine/Source/Runtime/MovieScene/Private/MovieSceneSequencePlayer.cpp
@@ -975,15 +975,7 @@ void UMovieSceneSequencePlayer::Initialize(UMovieSceneSequence* InSequence)
 
 	RegisteredTickInterval = TickInterval;
 
-	TSharedPtr<FMovieSceneEntitySystemRunner> RunnerToUse = TickManager->GetRunner(RegisteredTickInterval.GetValue());
-	if (EnumHasAnyFlags(Sequence->GetFlags(), EMovieSceneSequenceFlags::BlockingEvaluation))
-	{
-		SynchronousRunner = MakeShared<FMovieSceneEntitySystemRunner>();
-		RunnerToUse = SynchronousRunner;
-	}
-
-	check(RunnerToUse);
-	RootTemplateInstance.Initialize(*Sequence, *this, nullptr, RunnerToUse);
+	RootTemplateInstance.Initialize(*Sequence, *this, nullptr);
 
 	if (!PlaybackSettings.bDynamicWeighting)
 	{
diff --git a/Engine/Source/Runtime/MovieScene/Private/MovieSceneSequenceTickManager.cpp b/Engine/Source/Runtime/MovieScene/Private/MovieSceneSequenceTickManager.cpp
--- a/Engine/Source/Runtime/MovieScene/Private/MovieSceneSequenceTickManager.cpp
+++ b/Engine/Source/Runtime/MovieScene/Private/MovieSceneSequenceTickManager.cpp
@@ -172,8 +172,7 @@ void UMovieSceneSequenceTickManager::RegisterTickClient(const FMovieSceneSequenc
 
 		FLinkerGroup& NewGroup = LinkerGroups[LinkerIndex];
 		NewGroup.Linker = Linker;
-		NewGroup.Runner = MakeShared<FMovieSceneEntitySystemRunner>();
-		NewGroup.Runner->AttachToLinker(Linker);
+		NewGroup.Runner = Linker->GetRunner();
 		NewGroup.RoundedTickIntervalMs = DesiredTickIntervalMs;
 		NewGroup.FrameBudgetMs = DesiredBudgetMs;
 		NewGroup.NumClients = 1;
diff --git a/Engine/Source/Runtime/MovieScene/Private/Tests/MovieScenePreAnimatedStateTests.cpp b/Engine/Source/Runtime/MovieScene/Private/Tests/MovieScenePreAnimatedStateTests.cpp
--- a/Engine/Source/Runtime/MovieScene/Private/Tests/MovieScenePreAnimatedStateTests.cpp
+++ b/Engine/Source/Runtime/MovieScene/Private/Tests/MovieScenePreAnimatedStateTests.cpp
@@ -393,10 +393,8 @@ bool FMovieScenePreAnimatedStateTrackTypesTest::RunTest(const FString& Parameter
 
 	UMovieSceneCompiledDataManager* CompiledDataManager = UMovieSceneCompiledDataManager::GetPrecompiledData();
 
-	TSharedPtr<FMovieSceneEntitySystemRunner> Runner = MakeShared<FMovieSceneEntitySystemRunner>();
-
 	FMovieSceneCompiledDataID DataID = CompiledDataManager->Compile(Sequence);
-	TestPlayer.Template.Initialize(*Sequence, TestPlayer, CompiledDataManager, Runner);
+	TestPlayer.Template.Initialize(*Sequence, TestPlayer, CompiledDataManager);
 	TestPlayer.Template.EnableGlobalPreAnimatedStateCapture();
 
 	// Test the keep state section
@@ -475,10 +473,8 @@ bool FMovieScenePreAnimatedStateContextChangedTest::RunTest(const FString& Param
 
 	UMovieSceneCompiledDataManager* CompiledDataManager = UMovieSceneCompiledDataManager::GetPrecompiledData();
 
-	TSharedPtr<FMovieSceneEntitySystemRunner> Runner = MakeShared<FMovieSceneEntitySystemRunner>();
-
 	FMovieSceneCompiledDataID DataID = CompiledDataManager->Compile(Sequence);
-	TestPlayer.Template.Initialize(*Sequence, TestPlayer, CompiledDataManager, Runner);
+	TestPlayer.Template.Initialize(*Sequence, TestPlayer, CompiledDataManager);
 	TestPlayer.Template.EnableGlobalPreAnimatedStateCapture();
 
 	// ---------------------------------------------
diff --git a/Engine/Source/Runtime/MovieScene/Public/EntitySystem/MovieSceneEntitySystemLinker.h b/Engine/Source/Runtime/MovieScene/Public/EntitySystem/MovieSceneEntitySystemLinker.h
--- a/Engine/Source/Runtime/MovieScene/Public/EntitySystem/MovieSceneEntitySystemLinker.h
+++ b/Engine/Source/Runtime/MovieScene/Public/EntitySystem/MovieSceneEntitySystemLinker.h
@@ -337,15 +337,18 @@ public:
 	MOVIESCENE_API void InvalidateObjectBinding(const FGuid& ObjectBindingID, FInstanceHandle InstanceHandle);
 	MOVIESCENE_API void CleanupInvalidBoundObjects();
 
-	MOVIESCENE_API bool StartEvaluation(FMovieSceneEntitySystemRunner& InRunner);
-	MOVIESCENE_API FMovieSceneEntitySystemRunner* GetActiveRunner() const;
-	MOVIESCENE_API void PostInstantation(FMovieSceneEntitySystemRunner& InRunner);
-	MOVIESCENE_API void EndEvaluation(FMovieSceneEntitySystemRunner& InRunner);
+	MOVIESCENE_API bool StartEvaluation();
+	MOVIESCENE_API TSharedRef<FMovieSceneEntitySystemRunner> GetRunner() const;
+	MOVIESCENE_API void PostInstantation();
+	MOVIESCENE_API void EndEvaluation();
 
-	MOVIESCENE_API void ResetActiveRunners();
+	MOVIESCENE_API void ResetRunner();
 
 	MOVIESCENE_API void DestroyInstanceImmediately(UE::MovieScene::FRootInstanceHandle Instance);
 
+	UE_DEPRECATED(5.5, "Please use GetRunner()")
+	FMovieSceneEntitySystemRunner* GetActiveRunner() const { return &GetRunner().Get(); }
+
 private:
 
 	MOVIESCENE_API UMovieSceneEntitySystem* LinkSystemImpl(TSubclassOf<UMovieSceneEntitySystem> InClassType);
@@ -370,8 +373,8 @@ private:
 	TSparseArray<UMovieSceneEntitySystem*> EntitySystemsByGlobalGraphID;
 	TMap<TObjectPtr<UClass>, TObjectPtr<UMovieSceneEntitySystem>> EntitySystemsRecyclingPool;
 
-	TArray<FMovieSceneEntitySystemRunner*> ActiveRunners;
-	TBitArray<> ActiveRunnerReentrancyFlags;
+	TSharedPtr<FMovieSceneEntitySystemRunner> Runner;
+	TBitArray<> RunnerReentrancyFlags;
 
 	TSparseArray<void*> ExtensionsByID;
 
diff --git a/Engine/Source/Runtime/MovieScene/Public/EntitySystem/MovieSceneEntitySystemRunner.h b/Engine/Source/Runtime/MovieScene/Public/EntitySystem/MovieSceneEntitySystemRunner.h
--- a/Engine/Source/Runtime/MovieScene/Public/EntitySystem/MovieSceneEntitySystemRunner.h
+++ b/Engine/Source/Runtime/MovieScene/Public/EntitySystem/MovieSceneEntitySystemRunner.h
@@ -79,18 +79,11 @@ public:
 	using FInstanceRegistry = UE::MovieScene::FInstanceRegistry;
 
 public:
-	/** Creates an unbound runner */
-	MOVIESCENE_API FMovieSceneEntitySystemRunner();
+
 	/** Destructor */
 	MOVIESCENE_API ~FMovieSceneEntitySystemRunner();
 
-	/** Attach this runner to a linker */
-	MOVIESCENE_API void AttachToLinker(UMovieSceneEntitySystemLinker* InLinker);
-	/** Returns whether this runner is attached to a linker */
-	MOVIESCENE_API bool IsAttachedToLinker() const;
-	/** Detaches this runner from a linker */
-	MOVIESCENE_API void DetachFromLinker();
-
+	/** Gets how many instances are in the update queue. */
 	MOVIESCENE_API int32 GetQueuedUpdateCount() const;
 	/** Returns whether this runner has any outstanding updates. */
 	MOVIESCENE_API bool HasQueuedUpdates() const;
@@ -168,17 +161,33 @@ public:
 
 public:
 
-	MOVIESCENE_API UMovieSceneEntitySystemLinker* GetLinker() const;
+	UMovieSceneEntitySystemLinker* GetLinker() const { return WeakLinker.Get(); }
+	bool IsValid() const { return WeakLinker.IsValid(); }
+
 	MOVIESCENE_API FEntityManager* GetEntityManager() const;
 	MOVIESCENE_API FInstanceRegistry* GetInstanceRegistry() const;
 
+public:
+
+	UE_DEPRECATED(5.5, "Runners are now owned by linkers, and always attached")
+	void AttachToLinker(UMovieSceneEntitySystemLinker* InLinker) {}
+	UE_DEPRECATED(5.5, "Runners are now owned by linkers, and always attached")
+	bool IsAttachedToLinker() const { return IsValid(); }
+	UE_DEPRECATED(5.5, "Runners are now owned by linkers, and always attached")
+	void DetachFromLinker() {}
+
 public:
 	
 	// Internal API
 
-	MOVIESCENE_API void MarkForUpdate(FInstanceHandle InInstanceHandle, UE::MovieScene::ERunnerUpdateFlags UpdateFlags);
+	void MarkForUpdate(FInstanceHandle InInstanceHandle, UE::MovieScene::ERunnerUpdateFlags UpdateFlags);
 	MOVIESCENE_API FMovieSceneEntitySystemEventTriggers& GetQueuedEventTriggers();
 
+	/** Creates an unbound runner */
+	FMovieSceneEntitySystemRunner();
+	/** Create a bound runner */
+	FMovieSceneEntitySystemRunner(UMovieSceneEntitySystemLinker* InLinker);
+
 private:
 
 	MOVIESCENE_API void OnLinkerAbandon(UMovieSceneEntitySystemLinker* Linker);
@@ -251,7 +260,7 @@ private:
 		UE::MovieScene::ERunnerUpdateFlags UpdateFlags = UE::MovieScene::ERunnerUpdateFlags::None;
 	};
 
-	/** Owner linker */
+	/** Owning linker */
 	TWeakObjectPtr<UMovieSceneEntitySystemLinker> WeakLinker;
 
 	struct FQueuedUpdateParams
diff --git a/Engine/Source/Runtime/MovieScene/Public/EntitySystem/MovieSceneInstanceRegistry.h b/Engine/Source/Runtime/MovieScene/Public/EntitySystem/MovieSceneInstanceRegistry.h
--- a/Engine/Source/Runtime/MovieScene/Public/EntitySystem/MovieSceneInstanceRegistry.h
+++ b/Engine/Source/Runtime/MovieScene/Public/EntitySystem/MovieSceneInstanceRegistry.h
@@ -12,7 +12,6 @@
 #include "MovieSceneSequenceID.h"
 
 
-class FMovieSceneEntitySystemRunner;
 class UMovieSceneCompiledDataManager;
 class UMovieSceneEntitySystemLinker;
 class UMovieSceneCompiledDataManager;
@@ -87,7 +86,6 @@ struct FInstanceRegistry
 	MOVIESCENE_API FRootInstanceHandle AllocateRootInstance(
 			UMovieSceneSequence& InRootSequence,
 			UObject* InPlaybackContext = nullptr,
-			TSharedPtr<FMovieSceneEntitySystemRunner> InRunner = nullptr,
 			UMovieSceneCompiledDataManager* InCompiledDataManager = nullptr);
 
 	MOVIESCENE_API FInstanceHandle AllocateSubInstance(
diff --git a/Engine/Source/Runtime/MovieScene/Public/EntitySystem/MovieSceneSharedPlaybackState.h b/Engine/Source/Runtime/MovieScene/Public/EntitySystem/MovieSceneSharedPlaybackState.h
--- a/Engine/Source/Runtime/MovieScene/Public/EntitySystem/MovieSceneSharedPlaybackState.h
+++ b/Engine/Source/Runtime/MovieScene/Public/EntitySystem/MovieSceneSharedPlaybackState.h
@@ -11,6 +11,7 @@
 
 class FMovieSceneEntitySystemRunner;
 class UMovieSceneCompiledDataManager;
+class UMovieSceneEntitySystemLinker;
 struct FMovieSceneObjectCache;
 struct FMovieSceneSequenceHierarchy;
 
@@ -25,7 +26,7 @@ struct FSharedPlaybackStateCreateParams
 	/**
 	 * The playback context in which the root sequence will be evaluated.
 	 *
-	 * Requires that RootInstanceHandle and Runner are also set.
+	 * Requires that RootInstanceHandle and Linker are also set.
 	 */
 	UObject* PlaybackContext = nullptr;
 
@@ -34,17 +35,17 @@ struct FSharedPlaybackStateCreateParams
 	 * is meant to relate to an instance that has also been created inside
 	 * a runner/linker's instance registry.
 	 *
-	 * Requires that PlaybackContext and Runner are also set.
+	 * Requires that PlaybackContext and Linker are also set.
 	 */
 	FRootInstanceHandle RootInstanceHandle;
 
 	/**
-	 * The runner that will be evaluating the sequence that the created playback
+	 * The linker that will be evaluating the sequence that the created playback
 	 * state relates to.
 	 *
 	 * Requires that PlaybackContext and RootInstanceHandle are also set.
 	 */
-	TSharedPtr<FMovieSceneEntitySystemRunner> Runner;
+	TObjectPtr<UMovieSceneEntitySystemLinker> Linker;
 
 	/**
 	 * The compiled data manager with which the root sequence was compiled, or
@@ -73,8 +74,8 @@ public:
 	/** Gets the root sequence */
 	UMovieSceneSequence* GetRootSequence() const { return WeakRootSequence.Get(); }
 
-	/** Gets the runner evaluating this root sequence */
-	TSharedPtr<FMovieSceneEntitySystemRunner> GetRunner() const { return WeakRunner.Pin(); }
+	/** Gets the linker evaluating this root sequence */
+	UMovieSceneEntitySystemLinker* GetLinker() const { return WeakLinker.Get(); }
 
 	/** Gets the compiled data manager that contains the data for the root sequence */
 	TObjectPtr<UMovieSceneCompiledDataManager> GetCompiledDataManager() const { return CompiledDataManager; }
@@ -89,9 +90,12 @@ public:
 
 	// General utility methods
 
-	UMovieSceneEntitySystemLinker* GetLinker() const;
+	/** Gets the runner evaluating this root sequence */
+	TSharedPtr<FMovieSceneEntitySystemRunner> GetRunner() const;
 
+	/** Gets the hierarchy (if any) for this root sequence */
 	const FMovieSceneSequenceHierarchy* GetHierarchy() const;
+	/** Gets a sub-sequence given an ID */
 	UMovieSceneSequence* GetSequence(FMovieSceneSequenceIDRef SequenceID) const;
 
 public:
@@ -296,8 +300,8 @@ private:
 	/** The playback context */
 	TWeakObjectPtr<UObject> WeakPlaybackContext;
 
-	/** The runner evaluating this root sequence */
-	TWeakPtr<FMovieSceneEntitySystemRunner> WeakRunner;
+	/** The linker evaluating this root sequence */
+	TWeakObjectPtr<UMovieSceneEntitySystemLinker> WeakLinker;
 
 	/** The compiled data manager that contains the data for the root sequence */
 	TObjectPtr<UMovieSceneCompiledDataManager> CompiledDataManager;
diff --git a/Engine/Source/Runtime/MovieScene/Public/Evaluation/MovieSceneEvaluationTemplateInstance.h b/Engine/Source/Runtime/MovieScene/Public/Evaluation/MovieSceneEvaluationTemplateInstance.h
--- a/Engine/Source/Runtime/MovieScene/Public/Evaluation/MovieSceneEvaluationTemplateInstance.h
+++ b/Engine/Source/Runtime/MovieScene/Public/Evaluation/MovieSceneEvaluationTemplateInstance.h
@@ -50,7 +50,7 @@ public:
 	 * @param Player					The player responsible for playback
 	 * @param TemplateStore				Template store responsible for supplying templates for a given sequence
 	 */
-	MOVIESCENE_API void Initialize(UMovieSceneSequence& RootSequence, IMovieScenePlayer& Player, UMovieSceneCompiledDataManager* InCompiledDataManager, TSharedPtr<FMovieSceneEntitySystemRunner> InRunner);
+	MOVIESCENE_API void Initialize(UMovieSceneSequence& RootSequence, IMovieScenePlayer& Player, UMovieSceneCompiledDataManager* InCompiledDataManager);
 
 	/**
 	 * Evaluate this sequence in a synchronous way.
diff --git a/Engine/Source/Runtime/MovieSceneTracks/Private/EntitySystem/Interrogation/MovieSceneInterrogationLinker.cpp b/Engine/Source/Runtime/MovieSceneTracks/Private/EntitySystem/Interrogation/MovieSceneInterrogationLinker.cpp
--- a/Engine/Source/Runtime/MovieSceneTracks/Private/EntitySystem/Interrogation/MovieSceneInterrogationLinker.cpp
+++ b/Engine/Source/Runtime/MovieSceneTracks/Private/EntitySystem/Interrogation/MovieSceneInterrogationLinker.cpp
@@ -844,9 +844,8 @@ void FSystemInterrogator::Update()
 			}
 		});
 
-	FMovieSceneEntitySystemRunner Runner;
-	Runner.AttachToLinker(Linker);
-	Runner.Flush();
+	TSharedPtr<FMovieSceneEntitySystemRunner> Runner = Linker->GetRunner();
+	Runner->Flush();
 
 	Linker->EntityManager.IncrementSystemSerial();
 }
diff --git a/Engine/Source/Runtime/MovieSceneTracks/Private/Systems/DoubleChannelEvaluatorSystem.cpp b/Engine/Source/Runtime/MovieSceneTracks/Private/Systems/DoubleChannelEvaluatorSystem.cpp
--- a/Engine/Source/Runtime/MovieSceneTracks/Private/Systems/DoubleChannelEvaluatorSystem.cpp
+++ b/Engine/Source/Runtime/MovieSceneTracks/Private/Systems/DoubleChannelEvaluatorSystem.cpp
@@ -205,11 +205,7 @@ void UDoubleChannelEvaluatorSystem::OnRun(FSystemTaskPrerequisites& InPrerequisi
 {
 	using namespace UE::MovieScene;
 
-	FMovieSceneEntitySystemRunner* Runner = Linker->GetActiveRunner();
-	if (!Runner)
-	{
-		return;
-	}
+	TSharedRef<FMovieSceneEntitySystemRunner> Runner = Linker->GetRunner();
 
 	FBuiltInComponentTypes* BuiltInComponents = FBuiltInComponentTypes::Get();
 
diff --git a/Engine/Source/Runtime/MovieSceneTracks/Private/Systems/DoublePerlinNoiseChannelEvaluatorSystem.cpp b/Engine/Source/Runtime/MovieSceneTracks/Private/Systems/DoublePerlinNoiseChannelEvaluatorSystem.cpp
--- a/Engine/Source/Runtime/MovieSceneTracks/Private/Systems/DoublePerlinNoiseChannelEvaluatorSystem.cpp
+++ b/Engine/Source/Runtime/MovieSceneTracks/Private/Systems/DoublePerlinNoiseChannelEvaluatorSystem.cpp
@@ -83,11 +83,7 @@ void UDoublePerlinNoiseChannelEvaluatorSystem::OnRun(FSystemTaskPrerequisites& I
 	const FBuiltInComponentTypes* BuiltInComponents = FBuiltInComponentTypes::Get();
 	const FMovieSceneTracksComponentTypes* TrackComponents = FMovieSceneTracksComponentTypes::Get();
 
-	FMovieSceneEntitySystemRunner* Runner = Linker->GetActiveRunner();
-	if (!Runner)
-	{
-		return;
-	}
+	TSharedRef<FMovieSceneEntitySystemRunner> Runner = Linker->GetRunner();
 
 	if (Runner->GetCurrentPhase() == ESystemPhase::Instantiation)
 	{
diff --git a/Engine/Source/Runtime/MovieSceneTracks/Private/Systems/FloatChannelEvaluatorSystem.cpp b/Engine/Source/Runtime/MovieSceneTracks/Private/Systems/FloatChannelEvaluatorSystem.cpp
--- a/Engine/Source/Runtime/MovieSceneTracks/Private/Systems/FloatChannelEvaluatorSystem.cpp
+++ b/Engine/Source/Runtime/MovieSceneTracks/Private/Systems/FloatChannelEvaluatorSystem.cpp
@@ -218,11 +218,7 @@ void UFloatChannelEvaluatorSystem::OnRun(FSystemTaskPrerequisites& InPrerequisit
 {
 	using namespace UE::MovieScene;
 
-	FMovieSceneEntitySystemRunner* Runner = Linker->GetActiveRunner();
-	if (!Runner)
-	{
-		return;
-	}
+	TSharedRef<FMovieSceneEntitySystemRunner> Runner = Linker->GetRunner();
 
 	FBuiltInComponentTypes* BuiltInComponents = FBuiltInComponentTypes::Get();
 
diff --git a/Engine/Source/Runtime/MovieSceneTracks/Private/Systems/FloatPerlinNoiseChannelEvaluatorSystem.cpp b/Engine/Source/Runtime/MovieSceneTracks/Private/Systems/FloatPerlinNoiseChannelEvaluatorSystem.cpp
--- a/Engine/Source/Runtime/MovieSceneTracks/Private/Systems/FloatPerlinNoiseChannelEvaluatorSystem.cpp
+++ b/Engine/Source/Runtime/MovieSceneTracks/Private/Systems/FloatPerlinNoiseChannelEvaluatorSystem.cpp
@@ -84,11 +84,7 @@ void UFloatPerlinNoiseChannelEvaluatorSystem::OnRun(FSystemTaskPrerequisites& In
 	const FBuiltInComponentTypes* BuiltInComponents = FBuiltInComponentTypes::Get();
 	const FMovieSceneTracksComponentTypes* TrackComponents = FMovieSceneTracksComponentTypes::Get();
 
-	FMovieSceneEntitySystemRunner* Runner = Linker->GetActiveRunner();
-	if (!Runner)
-	{
-		return;
-	}
+	TSharedRef<FMovieSceneEntitySystemRunner> Runner = Linker->GetRunner();
 
 	if (Runner->GetCurrentPhase() == ESystemPhase::Instantiation)
 	{
diff --git a/Engine/Source/Runtime/MovieSceneTracks/Private/Systems/MovieSceneConstraintSystem.cpp b/Engine/Source/Runtime/MovieSceneTracks/Private/Systems/MovieSceneConstraintSystem.cpp
--- a/Engine/Source/Runtime/MovieSceneTracks/Private/Systems/MovieSceneConstraintSystem.cpp
+++ b/Engine/Source/Runtime/MovieSceneTracks/Private/Systems/MovieSceneConstraintSystem.cpp
@@ -254,8 +254,8 @@ void UMovieSceneConstraintSystem::OnRun(FSystemTaskPrerequisites& InPrerequisite
 {
 	using namespace UE::MovieScene;
 
-	FMovieSceneEntitySystemRunner* ActiveRunner = Linker->GetActiveRunner();
-	ESystemPhase CurrentPhase = ActiveRunner->GetCurrentPhase();
+	TSharedRef<FMovieSceneEntitySystemRunner> Runner = Linker->GetRunner();
+	ESystemPhase CurrentPhase = Runner->GetCurrentPhase();
 
 	if (CurrentPhase == ESystemPhase::Instantiation)
 	{
diff --git a/Engine/Source/Runtime/MovieSceneTracks/Private/Systems/MovieSceneCustomPrimitiveDataSystem.cpp b/Engine/Source/Runtime/MovieSceneTracks/Private/Systems/MovieSceneCustomPrimitiveDataSystem.cpp
--- a/Engine/Source/Runtime/MovieSceneTracks/Private/Systems/MovieSceneCustomPrimitiveDataSystem.cpp
+++ b/Engine/Source/Runtime/MovieSceneTracks/Private/Systems/MovieSceneCustomPrimitiveDataSystem.cpp
@@ -431,13 +431,9 @@ void UMovieSceneCustomPrimitiveDataSystem::OnRun(FSystemTaskPrerequisites& InPre
 {
 	using namespace UE::MovieScene;
 
-	FMovieSceneEntitySystemRunner* ActiveRunner = Linker->GetActiveRunner();
-	if (!ensure(ActiveRunner))
-	{
-		return;
-	}
+	TSharedRef<FMovieSceneEntitySystemRunner> Runner = Linker->GetRunner();
 
-	ESystemPhase CurrentPhase = ActiveRunner->GetCurrentPhase();
+	ESystemPhase CurrentPhase = Runner->GetCurrentPhase();
 	if (CurrentPhase == ESystemPhase::Instantiation)
 	{
 		OnInstantiation();
diff --git a/Engine/Source/Runtime/MovieSceneTracks/Private/Systems/MovieSceneDeferredComponentMovementSystem.cpp b/Engine/Source/Runtime/MovieSceneTracks/Private/Systems/MovieSceneDeferredComponentMovementSystem.cpp
--- a/Engine/Source/Runtime/MovieSceneTracks/Private/Systems/MovieSceneDeferredComponentMovementSystem.cpp
+++ b/Engine/Source/Runtime/MovieSceneTracks/Private/Systems/MovieSceneDeferredComponentMovementSystem.cpp
@@ -145,8 +145,8 @@ void UMovieSceneDeferredComponentMovementSystem::OnRun(FSystemTaskPrerequisites&
 {
 	using namespace UE::MovieScene;
 
-	FMovieSceneEntitySystemRunner* ActiveRunner = Linker->GetActiveRunner();
-	ESystemPhase CurrentPhase = ActiveRunner->GetCurrentPhase();
+	TSharedRef<FMovieSceneEntitySystemRunner> Runner = Linker->GetRunner();
+	ESystemPhase CurrentPhase = Runner->GetCurrentPhase();
 
 	FBuiltInComponentTypes*          BuiltInComponents = FBuiltInComponentTypes::Get();
 	FMovieSceneTracksComponentTypes* Components        = FMovieSceneTracksComponentTypes::Get();
diff --git a/Engine/Source/Runtime/MovieSceneTracks/Private/Systems/MovieSceneEventSystems.cpp b/Engine/Source/Runtime/MovieSceneTracks/Private/Systems/MovieSceneEventSystems.cpp
--- a/Engine/Source/Runtime/MovieSceneTracks/Private/Systems/MovieSceneEventSystems.cpp
+++ b/Engine/Source/Runtime/MovieSceneTracks/Private/Systems/MovieSceneEventSystems.cpp
@@ -383,11 +383,8 @@ void UMovieScenePostEvalEventSystem::OnRun(FSystemTaskPrerequisites& InPrerequis
 {
 	if (HasEvents())
 	{
-		FMovieSceneEntitySystemRunner* Runner = Linker->GetActiveRunner();
-		if (ensure(Runner))
-		{
-			Runner->GetQueuedEventTriggers().AddUObject(this, &UMovieScenePostEvalEventSystem::TriggerAllEvents);
-		}
+		TSharedRef<FMovieSceneEntitySystemRunner> Runner = Linker->GetRunner();
+		Runner->GetQueuedEventTriggers().AddUObject(this, &UMovieScenePostEvalEventSystem::TriggerAllEvents);
 	}
 }
 
diff --git a/Engine/Source/Runtime/MovieSceneTracks/Private/Systems/MovieSceneMotionVectorSimulationSystem.cpp b/Engine/Source/Runtime/MovieSceneTracks/Private/Systems/MovieSceneMotionVectorSimulationSystem.cpp
--- a/Engine/Source/Runtime/MovieSceneTracks/Private/Systems/MovieSceneMotionVectorSimulationSystem.cpp
+++ b/Engine/Source/Runtime/MovieSceneTracks/Private/Systems/MovieSceneMotionVectorSimulationSystem.cpp
@@ -59,7 +59,7 @@ void UMovieSceneMotionVectorSimulationSystem::OnRun(FSystemTaskPrerequisites& In
 {
 	using namespace UE::MovieScene;
 
-	FMovieSceneEntitySystemRunner* Runner = Linker->GetActiveRunner();
+	TSharedRef<FMovieSceneEntitySystemRunner> Runner = Linker->GetRunner();
 	if (Runner->GetCurrentPhase() == ESystemPhase::Finalization && bSimulationEnabled)
 	{
 		Runner->GetQueuedEventTriggers().AddLambda([this] { this->OnPostEvaluation(); });
@@ -135,11 +135,8 @@ void UMovieSceneMotionVectorSimulationSystem::ComputeSimulatedMotion()
 
 	// --------------------------------------------------------------------------------------------------------------------------------------------
 	// Re-execute the evaluation phase to flush the new transforms
-	FMovieSceneEntitySystemRunner* ActiveRunner = Linker->GetActiveRunner();
-	if (ensure(ActiveRunner))
-	{
-		ActiveRunner->FlushSingleEvaluationPhase();
-	}
+	TSharedRef<FMovieSceneEntitySystemRunner> Runner = Linker->GetRunner();
+	Runner->FlushSingleEvaluationPhase();
 
 	// --------------------------------------------------------------------------------------------------------------------------------------------
 	// Harvest results
diff --git a/Engine/Source/Runtime/MovieSceneTracks/Private/Systems/MovieScenePredictionSystem.cpp b/Engine/Source/Runtime/MovieSceneTracks/Private/Systems/MovieScenePredictionSystem.cpp
--- a/Engine/Source/Runtime/MovieSceneTracks/Private/Systems/MovieScenePredictionSystem.cpp
+++ b/Engine/Source/Runtime/MovieSceneTracks/Private/Systems/MovieScenePredictionSystem.cpp
@@ -408,11 +408,7 @@ void UMovieScenePredictionSystem::OnRun(FSystemTaskPrerequisites& InPrerequisite
 {
 	using namespace UE::MovieScene;
 
-	FMovieSceneEntitySystemRunner* Runner = Linker->GetActiveRunner();
-	if (!ensure(Runner))
-	{
-		return;
-	}
+	TSharedRef<FMovieSceneEntitySystemRunner> Runner = Linker->GetRunner();
 
 	FInstanceRegistry* InstanceRegistry = Linker->GetInstanceRegistry();
 
diff --git a/Engine/Source/Runtime/MovieSceneTracks/Private/Systems/MovieSceneSkeletalAnimationSystem.cpp b/Engine/Source/Runtime/MovieSceneTracks/Private/Systems/MovieSceneSkeletalAnimationSystem.cpp
--- a/Engine/Source/Runtime/MovieSceneTracks/Private/Systems/MovieSceneSkeletalAnimationSystem.cpp
+++ b/Engine/Source/Runtime/MovieSceneTracks/Private/Systems/MovieSceneSkeletalAnimationSystem.cpp
@@ -1020,7 +1020,7 @@ void UMovieSceneSkeletalAnimationSystem::OnRun(FSystemTaskPrerequisites& InPrere
 
 	const TStatId GatherStatId = GET_STATID(MovieSceneEval_GatherSkeletalAnimations);
 
-	const FMovieSceneEntitySystemRunner* Runner = Linker->GetActiveRunner();
+	TSharedRef<FMovieSceneEntitySystemRunner> Runner = Linker->GetRunner();
 	if (Runner->GetCurrentPhase() == ESystemPhase::Instantiation)
 	{
 		CleanSystemData();
diff --git a/Engine/Source/Runtime/MovieSceneTracks/Private/Tests/MovieScenePartialEvaluationTests.cpp b/Engine/Source/Runtime/MovieSceneTracks/Private/Tests/MovieScenePartialEvaluationTests.cpp
--- a/Engine/Source/Runtime/MovieSceneTracks/Private/Tests/MovieScenePartialEvaluationTests.cpp
+++ b/Engine/Source/Runtime/MovieSceneTracks/Private/Tests/MovieScenePartialEvaluationTests.cpp
@@ -67,12 +67,10 @@ bool FMovieScenePartialEvaluationTest::RunTest(const FString& Parameters)
 
 	UMovieSceneCompiledDataManager* CompiledDataManager = UMovieSceneCompiledDataManager::GetPrecompiledData();
 	TStrongObjectPtr<UMovieSceneEntitySystemLinker> Linker(NewObject<UMovieSceneEntitySystemLinker>(GetTransientPackage()));
-	TSharedPtr<FMovieSceneEntitySystemRunner> Runner = MakeShared<FMovieSceneEntitySystemRunner>();
+	TSharedPtr<FMovieSceneEntitySystemRunner> Runner = Linker->GetRunner();
 
 	TGuardValue<FEntityManager*> DebugVizGuard(GEntityManagerForDebuggingVisualizers, &Linker->EntityManager);
 
-	Runner->AttachToLinker(Linker.Get());
-
 	struct FLocalPlayer : IMovieScenePlayer
 	{
 		FMovieSceneRootEvaluationTemplateInstance Template;
@@ -91,7 +89,7 @@ bool FMovieScenePartialEvaluationTest::RunTest(const FString& Parameters)
 	Player.TestLinker = Linker.Get();
 
 	FMovieSceneCompiledDataID DataID = CompiledDataManager->Compile(RootSequence.Get());
-	Player.Template.Initialize(*RootSequence, Player, CompiledDataManager, Runner);
+	Player.Template.Initialize(*RootSequence, Player, CompiledDataManager);
 
 	bool bIsActive = true;
 
diff --git a/Engine/Source/Runtime/MovieSceneTracks/Public/Systems/MovieSceneMaterialSystem.h b/Engine/Source/Runtime/MovieSceneTracks/Public/Systems/MovieSceneMaterialSystem.h
--- a/Engine/Source/Runtime/MovieSceneTracks/Public/Systems/MovieSceneMaterialSystem.h
+++ b/Engine/Source/Runtime/MovieSceneTracks/Public/Systems/MovieSceneMaterialSystem.h
@@ -340,13 +340,9 @@ void TMovieSceneMaterialSystem<AccessorType, RequiredComponents...>::OnRun(UMovi
 	FBuiltInComponentTypes*          BuiltInComponents = FBuiltInComponentTypes::Get();
 	FMovieSceneTracksComponentTypes* TracksComponents  = FMovieSceneTracksComponentTypes::Get();
 
-	FMovieSceneEntitySystemRunner* ActiveRunner = Linker->GetActiveRunner();
-	if (!ensure(ActiveRunner))
-	{
-		return;
-	}
+	TSharedRef<FMovieSceneEntitySystemRunner> Runner = Linker->GetRunner();
 
-	ESystemPhase CurrentPhase = ActiveRunner->GetCurrentPhase();
+	ESystemPhase CurrentPhase = Runner->GetCurrentPhase();
 	if (CurrentPhase == ESystemPhase::Instantiation)
 	{
 		// --------------------------------------------------------------------------------------
diff --git a/Engine/Source/Runtime/UMG/Private/Animation/UMGSequencePlayer.cpp b/Engine/Source/Runtime/UMG/Private/Animation/UMGSequencePlayer.cpp
--- a/Engine/Source/Runtime/UMG/Private/Animation/UMGSequencePlayer.cpp
+++ b/Engine/Source/Runtime/UMG/Private/Animation/UMGSequencePlayer.cpp
@@ -234,14 +234,7 @@ void UUMGSequencePlayer::PlayInternal(double StartAtTime, double EndAtTime, int3
 		CSV_EVENT_GLOBAL(TEXT("Play Animation [%s::%s]"), *Widget->GetName(), *Animation->GetName());
 	}
 
-	TSharedPtr<FMovieSceneEntitySystemRunner> RunnerToUse = TickManager ? TickManager->GetRunner() : nullptr;
-	if (EnumHasAnyFlags(Animation->GetFlags(), EMovieSceneSequenceFlags::BlockingEvaluation))
-	{
-		SynchronousRunner = MakeShared<FMovieSceneEntitySystemRunner>();
-		RunnerToUse = SynchronousRunner;
-	}
-
-	RootTemplateInstance.Initialize(*Animation, *this, nullptr, RunnerToUse);
+	RootTemplateInstance.Initialize(*Animation, *this, nullptr);
 
 	if (bInRestoreState)
 	{
@@ -285,7 +278,8 @@ void UUMGSequencePlayer::PlayInternal(double StartAtTime, double EndAtTime, int3
 	PlayerStatus = EMovieScenePlayerStatus::Playing;
 
 	// Playback assumes the start frame has already been evaulated, so we also want to evaluate any events on the start frame here.
-	if (RunnerToUse)
+	TSharedPtr<FMovieSceneEntitySystemRunner> Runner = RootTemplateInstance.GetRunner();
+	if (Runner)
 	{
 		const FMovieSceneContext Context(FMovieSceneEvaluationRange(AbsolutePlaybackStart + TimeCursorPosition, AbsolutePlaybackStart + TimeCursorPosition, AnimationResolution), PlayerStatus);
 
@@ -296,11 +290,11 @@ void UUMGSequencePlayer::PlayInternal(double StartAtTime, double EndAtTime, int3
 
 		// We queue an update instead of immediately flushing the entire linker so that we don't incur a cascade of flushes on frames when multiple animations are played
 		// In rare cases where the linker must be flushed immediately PreTick, the queue should be manually flushed 
-		RunnerToUse->QueueUpdate(Context, RootTemplateInstance.GetRootInstanceHandle(), FSimpleDelegate::CreateWeakLambda(this, OnBegunPlay), UE::MovieScene::ERunnerUpdateFlags::Flush);
+		Runner->QueueUpdate(Context, RootTemplateInstance.GetRootInstanceHandle(), FSimpleDelegate::CreateWeakLambda(this, OnBegunPlay), UE::MovieScene::ERunnerUpdateFlags::Flush);
 
-		if (RunnerToUse == SynchronousRunner || !UE::UMG::GAsyncAnimationControlFlow)
+		if (Runner == SynchronousRunner || !UE::UMG::GAsyncAnimationControlFlow)
 		{
-			RunnerToUse->Flush();
+			Runner->Flush();
 		}
 	}
 }
diff --git a/Engine/Source/Runtime/UMG/Private/Animation/UMGSequenceTickManager.cpp b/Engine/Source/Runtime/UMG/Private/Animation/UMGSequenceTickManager.cpp
--- a/Engine/Source/Runtime/UMG/Private/Animation/UMGSequenceTickManager.cpp
+++ b/Engine/Source/Runtime/UMG/Private/Animation/UMGSequenceTickManager.cpp
@@ -46,10 +46,22 @@ UUMGSequenceTickManager::UUMGSequenceTickManager(const FObjectInitializer& Init)
 	: Super(Init)
 	, bIsTicking(false)
 {
-	if (!HasAnyFlags(RF_ClassDefaultObject))
-	{
-		Runner = MakeShared<FMovieSceneEntitySystemRunner>();
-	}
+}
+
+void UUMGSequenceTickManager::Initialize(UObject* Owner)
+{
+	Linker = UMovieSceneEntitySystemLinker::FindOrCreateLinker(Owner, UE::MovieScene::EEntitySystemLinkerRole::UMG, TEXT("UMGAnimationEntitySystemLinker"));
+	check(Linker);
+	Runner = Linker->GetRunner();
+
+	FSlateApplication& SlateApp = FSlateApplication::Get();
+	FDelegateHandle PreTickHandle = SlateApp.OnPreTick().AddUObject(this, &UUMGSequenceTickManager::TickWidgetAnimations);
+	check(PreTickHandle.IsValid());
+	SlateApplicationPreTickHandle = PreTickHandle;
+
+	FDelegateHandle PostTickHandle = SlateApp.OnPostTick().AddUObject(this, &UUMGSequenceTickManager::HandleSlatePostTick);
+	check(PostTickHandle.IsValid());
+	SlateApplicationPostTickHandle = PostTickHandle;
 }
 
 void UUMGSequenceTickManager::AddWidget(UUserWidget* InWidget)
@@ -238,11 +250,8 @@ void UUMGSequenceTickManager::TickWidgetAnimations(float DeltaSeconds)
 
 void UUMGSequenceTickManager::ForceFlush()
 {
-	if (Runner->IsAttachedToLinker())
-	{
-		Runner->Flush(UE::UMG::GAnimationBudgetMs);
-		RunLatentActions();
-	}
+	Runner->Flush(UE::UMG::GAnimationBudgetMs);
+	RunLatentActions();
 }
 
 void UUMGSequenceTickManager::HandleSlatePostTick(float DeltaSeconds)
@@ -254,7 +263,7 @@ void UUMGSequenceTickManager::HandleSlatePostTick(float DeltaSeconds)
 	}
 
 	// Only tick widgets at the end of the frame if our runner has completely finished, and we still have updates
-	if (UE::UMG::GFlushUMGAnimationsAtEndOfFrame && Runner->IsAttachedToLinker() && Runner->HasQueuedUpdates() && !Runner->IsCurrentlyEvaluating())
+	if (UE::UMG::GFlushUMGAnimationsAtEndOfFrame && Runner->HasQueuedUpdates() && !Runner->IsCurrentlyEvaluating())
 	{
 		SCOPE_CYCLE_COUNTER(MovieSceneEval_FlushEndOfFrameAnimations);
 
@@ -316,19 +325,7 @@ UUMGSequenceTickManager* UUMGSequenceTickManager::Get(UObject* PlaybackContext)
 	if (!TickManager)
 	{
 		TickManager = NewObject<UUMGSequenceTickManager>(Owner, TickManagerName);
-
-		TickManager->Linker = UMovieSceneEntitySystemLinker::FindOrCreateLinker(Owner, UE::MovieScene::EEntitySystemLinkerRole::UMG, TEXT("UMGAnimationEntitySystemLinker"));
-		check(TickManager->Linker);
-		TickManager->Runner->AttachToLinker(TickManager->Linker);
-
-		FSlateApplication& SlateApp = FSlateApplication::Get();
-		FDelegateHandle PreTickHandle = SlateApp.OnPreTick().AddUObject(TickManager, &UUMGSequenceTickManager::TickWidgetAnimations);
-		check(PreTickHandle.IsValid());
-		TickManager->SlateApplicationPreTickHandle = PreTickHandle;
-
-		FDelegateHandle PostTickHandle = SlateApp.OnPostTick().AddUObject(TickManager, &UUMGSequenceTickManager::HandleSlatePostTick);
-		check(PostTickHandle.IsValid());
-		TickManager->SlateApplicationPostTickHandle = PostTickHandle;
+		TickManager->Initialize(Owner);
 	}
 	return TickManager;
 }
diff --git a/Engine/Source/Runtime/UMG/Public/Animation/UMGSequenceTickManager.h b/Engine/Source/Runtime/UMG/Public/Animation/UMGSequenceTickManager.h
--- a/Engine/Source/Runtime/UMG/Public/Animation/UMGSequenceTickManager.h
+++ b/Engine/Source/Runtime/UMG/Public/Animation/UMGSequenceTickManager.h
@@ -37,6 +37,8 @@ public:
 
 	UMG_API UUMGSequenceTickManager(const FObjectInitializer& Init);
 
+	UMG_API void Initialize(UObject* Owner);
+
 	UMovieSceneEntitySystemLinker* GetLinker() { return Linker; }
 	TSharedPtr<FMovieSceneEntitySystemRunner> GetRunner() { return Runner; }
 
