Sequencer: fix potential leaking of shared pointers and playback capabilities

diff --git a/Engine/Source/Runtime/MovieScene/Private/EntitySystem/MovieSceneSequenceInstance.cpp b/Engine/Source/Runtime/MovieScene/Private/EntitySystem/MovieSceneSequenceInstance.cpp
--- a/Engine/Source/Runtime/MovieScene/Private/EntitySystem/MovieSceneSequenceInstance.cpp
+++ b/Engine/Source/Runtime/MovieScene/Private/EntitySystem/MovieSceneSequenceInstance.cpp
@@ -15,6 +15,7 @@
 #include "Evaluation/PreAnimatedState/MovieScenePreAnimatedStateExtension.h"
 
 #include "IMovieScenePlayer.h"
+#include "Misc/AssertionMacros.h"
 #include "MovieSceneSequence.h"
 #include "MovieSceneSequencePlayer.h"
 #include "MovieSceneTimeHelpers.h"
@@ -117,6 +118,9 @@ void FSequenceInstance::Initialize()
 
 FSequenceInstance::~FSequenceInstance()
 {
+	ensureAlwaysMsgf(
+			RootInstanceHandle != InstanceHandle || SharedPlaybackState.IsUnique(),
+			TEXT("References to SharedPlaybackState should not be held past the lifetime of its root sequence instance"));
 }
 
 FSequenceInstance::FSequenceInstance(FSequenceInstance&&) = default;
diff --git a/Engine/Source/Runtime/MovieScene/Private/Evaluation/MovieSceneEvaluationTemplateInstance.cpp b/Engine/Source/Runtime/MovieScene/Private/Evaluation/MovieSceneEvaluationTemplateInstance.cpp
--- a/Engine/Source/Runtime/MovieScene/Private/Evaluation/MovieSceneEvaluationTemplateInstance.cpp
+++ b/Engine/Source/Runtime/MovieScene/Private/Evaluation/MovieSceneEvaluationTemplateInstance.cpp
@@ -35,18 +35,23 @@ void FMovieSceneRootEvaluationTemplateInstance::TearDown()
 {
 	using namespace UE::MovieScene;
 
+	// Let's clear our shared pointer before destroying the sequence instance. This is because the sequence
+	// instance checks that no one is holding onto it, making sure it's destroyed along with it and that nobody
+	// is leaking it.
+	FSharedPlaybackState* SharedPlaybackStatePtr = SharedPlaybackState.Get();
+	SharedPlaybackState.Reset();
+
 	// Avoid redundant work if the linker is being destroyed anyway
-	if (SharedPlaybackState.IsValid() && 
-			SharedPlaybackState->GetRootInstanceHandle().IsValid() &&
+	if (SharedPlaybackStatePtr &&
+			SharedPlaybackStatePtr->GetRootInstanceHandle().IsValid() &&
 			EntitySystemLinker && 
 			IsValidChecked(EntitySystemLinker) && 
 			!EntitySystemLinker->IsUnreachable() && 
 			!EntitySystemLinker->HasAnyFlags(RF_BeginDestroyed))
 	{
-		EntitySystemLinker->DestroyInstanceImmediately(SharedPlaybackState->GetRootInstanceHandle());
+		EntitySystemLinker->DestroyInstanceImmediately(SharedPlaybackStatePtr->GetRootInstanceHandle());
 	}
 
-	SharedPlaybackState.Reset();
 	EntitySystemLinker = nullptr;
 }
 
@@ -454,6 +459,11 @@ void FMovieSceneRootEvaluationTemplateInstance::PlaybackContextChanged(IMovieSce
 		{
 			Player.RestorePreAnimatedState();
 		}
+
+		// Forget our previous shared playback state before we destroy our root instance, because the instance
+		// wants to make sure the state will be destroyed with it.
+		SharedPlaybackState.Reset();
+
 		EntitySystemLinker->GetInstanceRegistry()->DestroyInstance(PreviousRootInstanceHandle);
 	}
 
diff --git a/Engine/Source/Runtime/MovieScene/Private/Evaluation/MovieScenePlaybackCapabilities.cpp b/Engine/Source/Runtime/MovieScene/Private/Evaluation/MovieScenePlaybackCapabilities.cpp
--- a/Engine/Source/Runtime/MovieScene/Private/Evaluation/MovieScenePlaybackCapabilities.cpp
+++ b/Engine/Source/Runtime/MovieScene/Private/Evaluation/MovieScenePlaybackCapabilities.cpp
@@ -18,16 +18,18 @@ FPlaybackCapabilities::FPlaybackCapabilities(FPlaybackCapabilities&& RHS)
 
 FPlaybackCapabilities& FPlaybackCapabilities::operator=(FPlaybackCapabilities&& RHS)
 {
-	Destroy();
-
-	Memory = RHS.Memory;
-	Alignment = RHS.Alignment;
-	Capacity = RHS.Capacity;
-	Num = RHS.Num;
-	AllCapabilities = RHS.AllCapabilities;
+	if (!ensure(this != &RHS))
+	{
+		Destroy();
 
-	new (&RHS) FPlaybackCapabilities();
+		Memory = RHS.Memory;
+		Alignment = RHS.Alignment;
+		Capacity = RHS.Capacity;
+		Num = RHS.Num;
+		AllCapabilities = RHS.AllCapabilities;
 
+		new (&RHS) FPlaybackCapabilities();
+	}
 	return *this;
 }
 
diff --git a/Engine/Source/Runtime/MovieScene/Public/Evaluation/MovieScenePlaybackCapabilities.h b/Engine/Source/Runtime/MovieScene/Public/Evaluation/MovieScenePlaybackCapabilities.h
--- a/Engine/Source/Runtime/MovieScene/Public/Evaluation/MovieScenePlaybackCapabilities.h
+++ b/Engine/Source/Runtime/MovieScene/Public/Evaluation/MovieScenePlaybackCapabilities.h
@@ -493,11 +493,9 @@ protected:
 		// the actual pointer to the actual stored capability.
 		void* CapabilityPtr = Header.Capability.Resolve(Memory);
 
-		// If we have inline storage, call the destructor on the previous capability.
-		if (Header.StorageMode == EPlaybackCapabilityStorageMode::Inline)
-		{
-			Helpers[Index].Destructor(CapabilityPtr);
-		}
+		// Call the destructor on the previous capability, which is important if it was stored inline or as
+		// a shared pointer.
+		Helpers[Index].Destructor(CapabilityPtr);
 
 		// Allocate the new capability.
 		StorageType* TypedCapabilityPtr = reinterpret_cast<StorageType*>(CapabilityPtr);
