[Backout] - CL31660265
Original CL Desc
-----------------------------------------------------------------
Sequencer: fix potential leaking of shared pointers and playback capabilities

diff --git a/Engine/Source/Runtime/MovieScene/Private/EntitySystem/MovieSceneSequenceInstance.cpp b/Engine/Source/Runtime/MovieScene/Private/EntitySystem/MovieSceneSequenceInstance.cpp
--- a/Engine/Source/Runtime/MovieScene/Private/EntitySystem/MovieSceneSequenceInstance.cpp
+++ b/Engine/Source/Runtime/MovieScene/Private/EntitySystem/MovieSceneSequenceInstance.cpp
@@ -15,7 +15,6 @@
 #include "Evaluation/PreAnimatedState/MovieScenePreAnimatedStateExtension.h"
 
 #include "IMovieScenePlayer.h"
-#include "Misc/AssertionMacros.h"
 #include "MovieSceneSequence.h"
 #include "MovieSceneSequencePlayer.h"
 #include "MovieSceneTimeHelpers.h"
@@ -118,9 +117,6 @@ void FSequenceInstance::Initialize()
 
 FSequenceInstance::~FSequenceInstance()
 {
-	ensureAlwaysMsgf(
-			RootInstanceHandle != InstanceHandle || SharedPlaybackState.IsUnique(),
-			TEXT("References to SharedPlaybackState should not be held past the lifetime of its root sequence instance"));
 }
 
 FSequenceInstance::FSequenceInstance(FSequenceInstance&&) = default;
diff --git a/Engine/Source/Runtime/MovieScene/Private/Evaluation/MovieScenePlaybackCapabilities.cpp b/Engine/Source/Runtime/MovieScene/Private/Evaluation/MovieScenePlaybackCapabilities.cpp
--- a/Engine/Source/Runtime/MovieScene/Private/Evaluation/MovieScenePlaybackCapabilities.cpp
+++ b/Engine/Source/Runtime/MovieScene/Private/Evaluation/MovieScenePlaybackCapabilities.cpp
@@ -18,18 +18,16 @@ FPlaybackCapabilities::FPlaybackCapabilities(FPlaybackCapabilities&& RHS)
 
 FPlaybackCapabilities& FPlaybackCapabilities::operator=(FPlaybackCapabilities&& RHS)
 {
-	if (!ensure(this != &RHS))
-	{
-		Destroy();
+	Destroy();
+
+	Memory = RHS.Memory;
+	Alignment = RHS.Alignment;
+	Capacity = RHS.Capacity;
+	Num = RHS.Num;
+	AllCapabilities = RHS.AllCapabilities;
 
-		Memory = RHS.Memory;
-		Alignment = RHS.Alignment;
-		Capacity = RHS.Capacity;
-		Num = RHS.Num;
-		AllCapabilities = RHS.AllCapabilities;
+	new (&RHS) FPlaybackCapabilities();
 
-		new (&RHS) FPlaybackCapabilities();
-	}
 	return *this;
 }
 
diff --git a/Engine/Source/Runtime/MovieScene/Public/Evaluation/MovieScenePlaybackCapabilities.h b/Engine/Source/Runtime/MovieScene/Public/Evaluation/MovieScenePlaybackCapabilities.h
--- a/Engine/Source/Runtime/MovieScene/Public/Evaluation/MovieScenePlaybackCapabilities.h
+++ b/Engine/Source/Runtime/MovieScene/Public/Evaluation/MovieScenePlaybackCapabilities.h
@@ -493,9 +493,11 @@ protected:
 		// the actual pointer to the actual stored capability.
 		void* CapabilityPtr = Header.Capability.Resolve(Memory);
 
-		// Call the destructor on the previous capability, which is important if it was stored inline or as
-		// a shared pointer.
-		Helpers[Index].Destructor(CapabilityPtr);
+		// If we have inline storage, call the destructor on the previous capability.
+		if (Header.StorageMode == EPlaybackCapabilityStorageMode::Inline)
+		{
+			Helpers[Index].Destructor(CapabilityPtr);
+		}
 
 		// Allocate the new capability.
 		StorageType* TypedCapabilityPtr = reinterpret_cast<StorageType*>(CapabilityPtr);
