[Backout] - CL31652683
Original CL Desc
-----------------------------------------------------------------
Sequencer: linker/runner refactor

The original design for runners vs linkers, which was to have a central linker with many runners (one for each player) isn't valid anymore now that runners are also global (owned by the tick manager). This change therefore makes this official by making runners created and owned by the linkers themselves.

diff --git a/Engine/Plugins/Experimental/Avalanche/Source/AvalancheSequence/Private/AvaSequence.cpp b/Engine/Plugins/Experimental/Avalanche/Source/AvalancheSequence/Private/AvaSequence.cpp
--- a/Engine/Plugins/Experimental/Avalanche/Source/AvalancheSequence/Private/AvaSequence.cpp
+++ b/Engine/Plugins/Experimental/Avalanche/Source/AvalancheSequence/Private/AvaSequence.cpp
@@ -36,8 +36,9 @@ namespace UE::AvaSequence::Private
 	{
 	public:
 		FTransientPlayer(UAvaSequence& InSequence)
+			: Runner(MakeShared<FMovieSceneEntitySystemRunner>())
 		{
-			Template.Initialize(InSequence, *this, nullptr);
+			Template.Initialize(InSequence, *this, nullptr, Runner);
 			State.AssignSequence(MovieSceneSequenceID::Root, InSequence, *this);
 		}
 
@@ -50,6 +51,8 @@ namespace UE::AvaSequence::Private
 		//~ End IMovieScenePlayer
 
 	private:
+		TSharedRef<FMovieSceneEntitySystemRunner> Runner;
+
 		FMovieSceneRootEvaluationTemplateInstance Template;
 	};
 }
diff --git a/Engine/Plugins/Experimental/PropertyAnimator/Source/PropertyAnimator/Private/MovieScene/Easing/PropertyAnimatorEasingDoubleChannelEvaluatorSystem.cpp b/Engine/Plugins/Experimental/PropertyAnimator/Source/PropertyAnimator/Private/MovieScene/Easing/PropertyAnimatorEasingDoubleChannelEvaluatorSystem.cpp
--- a/Engine/Plugins/Experimental/PropertyAnimator/Source/PropertyAnimator/Private/MovieScene/Easing/PropertyAnimatorEasingDoubleChannelEvaluatorSystem.cpp
+++ b/Engine/Plugins/Experimental/PropertyAnimator/Source/PropertyAnimator/Private/MovieScene/Easing/PropertyAnimatorEasingDoubleChannelEvaluatorSystem.cpp
@@ -78,7 +78,11 @@ void UPropertyAnimatorEasingDoubleChannelEvaluatorSystem::OnRun(FSystemTaskPrere
 	const FBuiltInComponentTypes* BuiltInComponents = FBuiltInComponentTypes::Get();
 	const FPropertyAnimatorComponentTypes* PropertyAnimatorComponents = FPropertyAnimatorComponentTypes::Get();
 
-	TSharedRef<FMovieSceneEntitySystemRunner> Runner = Linker->GetRunner();
+	FMovieSceneEntitySystemRunner* Runner = Linker->GetActiveRunner();
+	if (!Runner)
+	{
+		return;
+	}
 
 	if (Runner->GetCurrentPhase() == ESystemPhase::Instantiation)
 	{
diff --git a/Engine/Plugins/Experimental/PropertyAnimator/Source/PropertyAnimator/Private/MovieScene/Wave/PropertyAnimatorWaveDoubleChannelEvaluatorSystem.cpp b/Engine/Plugins/Experimental/PropertyAnimator/Source/PropertyAnimator/Private/MovieScene/Wave/PropertyAnimatorWaveDoubleChannelEvaluatorSystem.cpp
--- a/Engine/Plugins/Experimental/PropertyAnimator/Source/PropertyAnimator/Private/MovieScene/Wave/PropertyAnimatorWaveDoubleChannelEvaluatorSystem.cpp
+++ b/Engine/Plugins/Experimental/PropertyAnimator/Source/PropertyAnimator/Private/MovieScene/Wave/PropertyAnimatorWaveDoubleChannelEvaluatorSystem.cpp
@@ -77,7 +77,11 @@ void UPropertyAnimatorWaveDoubleChannelEvaluatorSystem::OnRun(FSystemTaskPrerequ
 	const FBuiltInComponentTypes* BuiltInComponents = FBuiltInComponentTypes::Get();
 	const FPropertyAnimatorComponentTypes* PropertyAnimatorComponents = FPropertyAnimatorComponentTypes::Get();
 
-	TSharedRef<FMovieSceneEntitySystemRunner> Runner = Linker->GetRunner();
+	FMovieSceneEntitySystemRunner* Runner = Linker->GetActiveRunner();
+	if (!Runner)
+	{
+		return;
+	}
 
 	if (Runner->GetCurrentPhase() == ESystemPhase::Instantiation)
 	{
diff --git a/Engine/Plugins/MovieScene/ReplayTracks/Source/ReplayTracks/Private/Systems/MovieSceneReplaySystem.cpp b/Engine/Plugins/MovieScene/ReplayTracks/Source/ReplayTracks/Private/Systems/MovieSceneReplaySystem.cpp
--- a/Engine/Plugins/MovieScene/ReplayTracks/Source/ReplayTracks/Private/Systems/MovieSceneReplaySystem.cpp
+++ b/Engine/Plugins/MovieScene/ReplayTracks/Source/ReplayTracks/Private/Systems/MovieSceneReplaySystem.cpp
@@ -84,8 +84,8 @@ void UMovieSceneReplaySystem::OnRun(FSystemTaskPrerequisites& InPrerequisites, F
 
 	using namespace UE::MovieScene;
 
-	TSharedRef<FMovieSceneEntitySystemRunner> Runner = Linker->GetRunner();
-	ESystemPhase CurrentPhase = Runner->GetCurrentPhase();
+	FMovieSceneEntitySystemRunner* ActiveRunner = Linker->GetActiveRunner();
+	ESystemPhase CurrentPhase = ActiveRunner->GetCurrentPhase();
 
 	if (CurrentPhase == ESystemPhase::Instantiation)
 	{
diff --git a/Engine/Plugins/MovieScene/TemplateSequence/Source/TemplateSequence/Private/CameraAnimationSequencePlayer.cpp b/Engine/Plugins/MovieScene/TemplateSequence/Source/TemplateSequence/Private/CameraAnimationSequencePlayer.cpp
--- a/Engine/Plugins/MovieScene/TemplateSequence/Source/TemplateSequence/Private/CameraAnimationSequencePlayer.cpp
+++ b/Engine/Plugins/MovieScene/TemplateSequence/Source/TemplateSequence/Private/CameraAnimationSequencePlayer.cpp
@@ -374,7 +374,9 @@ void UCameraAnimationSequencePlayer::Initialize(UMovieSceneSequence* InSequence,
 
 	UCameraAnimationSequenceSubsystem* Subsystem = UCameraAnimationSequenceSubsystem::GetCameraAnimationSequenceSubsystem(GetWorld());
 	ensureMsgf(Subsystem, TEXT("Unable to locate a valid camera animation sub-system. Camera anim sequences will not play."));
-	RootTemplateInstance.Initialize(*Sequence, *this, nullptr);
+
+	TSharedPtr<FMovieSceneEntitySystemRunner> Runner = Subsystem ? Subsystem->GetRunner() : nullptr;
+	RootTemplateInstance.Initialize(*Sequence, *this, nullptr, Runner);
 }
 
 void UCameraAnimationSequencePlayer::Play(bool bLoop, bool bRandomStartTime)
diff --git a/Engine/Plugins/MovieScene/TemplateSequence/Source/TemplateSequence/Private/CameraAnimationSequenceSubsystem.cpp b/Engine/Plugins/MovieScene/TemplateSequence/Source/TemplateSequence/Private/CameraAnimationSequenceSubsystem.cpp
--- a/Engine/Plugins/MovieScene/TemplateSequence/Source/TemplateSequence/Private/CameraAnimationSequenceSubsystem.cpp
+++ b/Engine/Plugins/MovieScene/TemplateSequence/Source/TemplateSequence/Private/CameraAnimationSequenceSubsystem.cpp
@@ -211,11 +211,18 @@ UCameraAnimationSequenceSubsystem::~UCameraAnimationSequenceSubsystem()
 
 void UCameraAnimationSequenceSubsystem::Initialize(FSubsystemCollectionBase& Collection)
 {
+	Runner = MakeShared<FMovieSceneEntitySystemRunner>();
 }
 
 void UCameraAnimationSequenceSubsystem::Deinitialize()
 {
-	Runner = nullptr;
+	// We check if the runner still has a valid pointer on the linker because the linker could
+	// have been GC'ed just now, which would make DetachFromLinker complain.
+	if (Runner->GetLinker())
+	{
+		Runner->DetachFromLinker();
+		Runner = nullptr;
+	}
 	Linker = nullptr;
 
 	Super::Deinitialize();
@@ -226,7 +233,7 @@ UMovieSceneEntitySystemLinker* UCameraAnimationSequenceSubsystem::GetLinker(bool
 	if (!Linker && bAutoCreate)
 	{
 		Linker = CreateLinker(this, TEXT("CameraAnimationSequenceSubsystemLinker"));
-		Runner = Linker->GetRunner();
+		Runner->AttachToLinker(Linker);
 	}
 	return Linker;
 }
diff --git a/Engine/Source/Editor/Sequencer/Private/SSequencer.cpp b/Engine/Source/Editor/Sequencer/Private/SSequencer.cpp
--- a/Engine/Source/Editor/Sequencer/Private/SSequencer.cpp
+++ b/Engine/Source/Editor/Sequencer/Private/SSequencer.cpp
@@ -2055,7 +2055,7 @@ void SSequencer::FillAdvancedMenu(FMenuBuilder& MenuBuilder)
 				// Since sequencer owns its own compiled data manager, it's ok to override the mask here and reset everything
 				Template.GetCompiledDataManager()->SetEmulatedNetworkMask(InMode);
 				// Reinitialize the template again
-				Template.Initialize(*RootSequence, *SequencerPin, Template.GetCompiledDataManager());
+				Template.Initialize(*RootSequence, *SequencerPin, Template.GetCompiledDataManager(), Template.GetRunner());
 			}
 		};
 		auto IsNetworkModeChecked = [WeakSequencer = SequencerPtr](EMovieSceneServerClientMask InMode)
diff --git a/Engine/Source/Editor/Sequencer/Private/Sequencer.cpp b/Engine/Source/Editor/Sequencer/Private/Sequencer.cpp
--- a/Engine/Source/Editor/Sequencer/Private/Sequencer.cpp
+++ b/Engine/Source/Editor/Sequencer/Private/Sequencer.cpp
@@ -485,8 +485,8 @@ void FSequencer::InitSequencer(const FSequencerInitParams& InitParams, const TSh
 	ActiveTemplateIDs.Add(MovieSceneSequenceID::Root);
 	ActiveTemplateStates.Add(true);
 
-	RootTemplateInstance.Initialize(*InitParams.RootSequence, *this, CompiledDataManager);
-	Runner = RootTemplateInstance.GetRunner();
+	Runner = MakeShared<FMovieSceneEntitySystemRunner>();
+	RootTemplateInstance.Initialize(*InitParams.RootSequence, *this, CompiledDataManager, Runner);
 
 	RootTemplateInstance.EnableGlobalPreAnimatedStateCapture();
 
@@ -716,7 +716,7 @@ FSequencer::FSequencer()
 
 FSequencer::~FSequencer()
 {
-	if (Runner)
+	if (Runner->IsAttachedToLinker())
 	{
 		Runner->QueueFinalUpdate(RootTemplateInstance.GetRootInstanceHandle());
 		Runner->Flush();
@@ -759,7 +759,7 @@ void FSequencer::Close()
 		OldMaxTickRate.Reset();
 	}
 
-	if (Runner)
+	if (Runner->IsAttachedToLinker())
 	{
 		Runner->QueueFinalUpdate(RootTemplateInstance.GetRootInstanceHandle());
 		Runner->Flush();
@@ -1110,7 +1110,7 @@ void FSequencer::ResetToNewRootSequence(UMovieSceneSequence& NewSequence)
 		}
 	}
 
-	if (Runner)
+	if (Runner->IsAttachedToLinker())
 	{
 		Runner->QueueFinalUpdate(RootTemplateInstance.GetRootInstanceHandle());
 		Runner->Flush();
@@ -1121,7 +1121,7 @@ void FSequencer::ResetToNewRootSequence(UMovieSceneSequence& NewSequence)
 	ActiveTemplateStates.Reset();
 	ActiveTemplateStates.Add(true);
 
-	RootTemplateInstance.Initialize(NewSequence, *this, CompiledDataManager);
+	RootTemplateInstance.Initialize(NewSequence, *this, CompiledDataManager, Runner);
 	RootTemplateInstance.EnableGlobalPreAnimatedStateCapture();
 
 	RootToLocalTransform = FMovieSceneSequenceTransform();
diff --git a/Engine/Source/Runtime/MovieScene/Private/EntitySystem/MovieSceneEntitySystemLinker.cpp b/Engine/Source/Runtime/MovieScene/Private/EntitySystem/MovieSceneEntitySystemLinker.cpp
--- a/Engine/Source/Runtime/MovieScene/Private/EntitySystem/MovieSceneEntitySystemLinker.cpp
+++ b/Engine/Source/Runtime/MovieScene/Private/EntitySystem/MovieSceneEntitySystemLinker.cpp
@@ -163,8 +163,6 @@ UMovieSceneEntitySystemLinker::UMovieSceneEntitySystemLinker(const FObjectInitia
 
 		InstanceRegistry.Reset(new FInstanceRegistry(this));
 
-		Runner = MakeShared<FMovieSceneEntitySystemRunner>(this);
-
 #if WITH_EDITOR
 		FCoreUObjectDelegates::OnObjectsReplaced.AddUObject(this, &UMovieSceneEntitySystemLinker::OnObjectsReplaced);
 #endif
@@ -316,13 +314,15 @@ bool UMovieSceneEntitySystemLinker::HasStructureChangedSinceLastRun() const
 	return EntityManager.HasStructureChangedSince(LastInstantiationVersion);
 }
 
-bool UMovieSceneEntitySystemLinker::StartEvaluation()
+bool UMovieSceneEntitySystemLinker::StartEvaluation(FMovieSceneEntitySystemRunner& InRunner)
 {
-	if (RunnerReentrancyFlags.Num() == 0 || RunnerReentrancyFlags[RunnerReentrancyFlags.Num() - 1])
+
+	if (ActiveRunners.Num() == 0 || ActiveRunnerReentrancyFlags[ActiveRunners.Num()-1])
 	{
 		// Default to re-entrancy being forbidden. The runner will allow re-entrancy at specific spots
 		// in the evaluation loop, via a "re-entrancy window".
-		RunnerReentrancyFlags.Add(false);
+		ActiveRunners.Emplace(&InRunner);
+		ActiveRunnerReentrancyFlags.Add(false);
 		return true;
 	}
 		
@@ -330,25 +330,33 @@ bool UMovieSceneEntitySystemLinker::StartEvaluation()
 	return false;
 }
 
-TSharedRef<FMovieSceneEntitySystemRunner> UMovieSceneEntitySystemLinker::GetRunner() const
+FMovieSceneEntitySystemRunner* UMovieSceneEntitySystemLinker::GetActiveRunner() const
 {
-	// The runner might lose us if we are GC'ed while someone else is keeping the runner's ref-count up,
-	// but we can never lose our runner, since we have at least one reference.
-	return Runner.ToSharedRef();
+	if (ActiveRunners.Num() > 0)
+	{
+		return ActiveRunners.Last();
+	}
+	return nullptr;
 }
 
-void UMovieSceneEntitySystemLinker::PostInstantation()
+void UMovieSceneEntitySystemLinker::PostInstantation(FMovieSceneEntitySystemRunner& InRunner)
 {
 	LastInstantiationVersion = EntityManager.GetSystemSerial();
 
 	GetInstanceRegistry()->PostInstantation();
 }
 
-void UMovieSceneEntitySystemLinker::EndEvaluation()
+void UMovieSceneEntitySystemLinker::EndEvaluation(FMovieSceneEntitySystemRunner& InRunner)
 {
-	const int32 LastIndex = RunnerReentrancyFlags.Num()-1;
-	ensureAlways(RunnerReentrancyFlags[LastIndex] == false);
-	RunnerReentrancyFlags.RemoveAt(LastIndex);
+	if (ensureMsgf((ActiveRunners.Num() > 0 && ActiveRunners.Last() == &InRunner),
+				TEXT("Trying end the evaluation of a runner that's not the latest one to run.")))
+	{
+		const int32 LastIndex = ActiveRunners.Num()-1;
+		ensureAlways(ActiveRunnerReentrancyFlags[LastIndex] == false);
+
+		ActiveRunners.Pop();
+		ActiveRunnerReentrancyFlags.RemoveAt(LastIndex);
+	}
 }
 
 void UMovieSceneEntitySystemLinker::HandlePreGarbageCollection()
@@ -450,13 +458,16 @@ void UMovieSceneEntitySystemLinker::CleanGarbage()
 	InstanceRegistry->CleanupLinkerEntities(FreedEntities);
 
 	// If we have any runners part-way through an evaluation, we need to reset them so that they re-evaluate from the start
-	ResetRunner();
+	ResetActiveRunners();
 }
 
-void UMovieSceneEntitySystemLinker::ResetRunner()
+void UMovieSceneEntitySystemLinker::ResetActiveRunners()
 {
 	// If we have any runners part-way through an evaluation, we need to reset them so that they re-evaluate from the start
-	Runner->ResetFlushState();
+	for (int32 Index = ActiveRunners.Num()-1; Index >= 0; --Index)
+	{
+		ActiveRunners[Index]->ResetFlushState();
+	}
 }
 
 void UMovieSceneEntitySystemLinker::DestroyInstanceImmediately(UE::MovieScene::FRootInstanceHandle Instance)
@@ -467,7 +478,7 @@ void UMovieSceneEntitySystemLinker::DestroyInstanceImmediately(UE::MovieScene::F
 	// Destroy the instance and any sub sequences. Any pre-existing NeedsLink entities will be forcibly made NeedsUnlink and cleaned as garbage
 	GetInstanceRegistry()->DestroyInstance(Instance);
 	CleanGarbage();
-	ResetRunner();
+	ResetActiveRunners();
 }
 
 void UMovieSceneEntitySystemLinker::OnObjectsReplaced(const TMap<UObject*, UObject*>& ReplacementMap)
diff --git a/Engine/Source/Runtime/MovieScene/Private/EntitySystem/MovieSceneEntitySystemRunner.cpp b/Engine/Source/Runtime/MovieScene/Private/EntitySystem/MovieSceneEntitySystemRunner.cpp
--- a/Engine/Source/Runtime/MovieScene/Private/EntitySystem/MovieSceneEntitySystemRunner.cpp
+++ b/Engine/Source/Runtime/MovieScene/Private/EntitySystem/MovieSceneEntitySystemRunner.cpp
@@ -13,7 +13,6 @@
 #include "Algo/Reverse.h"
 #include "Algo/Sort.h"
 #include "ProfilingDebugging/CountersTrace.h"
-#include "UObject/ObjectMacros.h"
 
 DECLARE_CYCLE_STAT(TEXT("Runner Flush"), 				MovieSceneEval_RunnerFlush, 				STATGROUP_MovieSceneEval);
 
@@ -51,21 +50,21 @@ struct FMovieSceneEntitySystemEvaluationReentrancyWindow
 {
 	FMovieSceneEntitySystemRunner* Runner;
 	UMovieSceneEntitySystemLinker* Linker;
-	int32 ReentrancyLevel;
+	int32 RunnerIndex;
 	UE::MovieScene::ERunnerFlushState CachedFlushState;
 	UE::MovieScene::ERunnerFlushState CachedCurrentFlushState;
 
 	FMovieSceneEntitySystemEvaluationReentrancyWindow(FMovieSceneEntitySystemRunner* InRunner, UMovieSceneEntitySystemLinker* InLinker)
 		: Runner(InRunner)
 		, Linker(InLinker)
-		, ReentrancyLevel(Linker->RunnerReentrancyFlags.Num() - 1)
+		, RunnerIndex(Linker->ActiveRunners.Num() - 1)
 		, CachedFlushState(Runner->FlushState)
 		, CachedCurrentFlushState(Runner->CurrentFlushState)
 	{
-		check(ReentrancyLevel >= 0);
-		checkf(Linker->RunnerReentrancyFlags[ReentrancyLevel] == false, TEXT("Nested FMovieSceneEntitySystemEvaluationReentrancyWindows are not supported for the same active runner"));
+		check(RunnerIndex >= 0);
+		checkf(Linker->ActiveRunnerReentrancyFlags[RunnerIndex] == false, TEXT("Nested FMovieSceneEntitySystemEvaluationReentrancyWindows are not supported for the same active runner"));
 
-		Linker->RunnerReentrancyFlags[ReentrancyLevel] = true;
+		Linker->ActiveRunnerReentrancyFlags[RunnerIndex] = true;
 
 		// Clear the current flush state to prevent inner scopes from flushing
 		//      our pending states and set the flag that allows us to be re-entrant
@@ -83,13 +82,14 @@ struct FMovieSceneEntitySystemEvaluationReentrancyWindow
 
 		Runner->FlushState = CachedFlushState;
 		Runner->CurrentFlushState = CachedCurrentFlushState;
-		if (ensure(Linker->RunnerReentrancyFlags.IsValidIndex(ReentrancyLevel)))
+		if (ensure(Linker->ActiveRunnerReentrancyFlags.IsValidIndex(RunnerIndex)))
 		{
-			Linker->RunnerReentrancyFlags[ReentrancyLevel] = false;
+			Linker->ActiveRunnerReentrancyFlags[RunnerIndex] = false;
 		}
 	}
 };
 
+
 FMovieSceneEntitySystemRunner::FMovieSceneEntitySystemRunner()
 	: GameThread(ENamedThreads::GameThread_Local)
 	, CurrentPhase(UE::MovieScene::ESystemPhase::None)
@@ -100,25 +100,57 @@ FMovieSceneEntitySystemRunner::FMovieSceneEntitySystemRunner()
 {
 }
 
-FMovieSceneEntitySystemRunner::FMovieSceneEntitySystemRunner(UMovieSceneEntitySystemLinker* InLinker)
-	: WeakLinker(InLinker)
-	, GameThread(ENamedThreads::GameThread_Local)
-	, CurrentPhase(UE::MovieScene::ESystemPhase::None)
-	, FlushState(UE::MovieScene::ERunnerFlushState::None)
-	, CurrentFlushState(UE::MovieScene::ERunnerFlushState::None)
-	, bRequireFullFlush(false)
-	, bIsUpdatingSequence(false)
+FMovieSceneEntitySystemRunner::~FMovieSceneEntitySystemRunner()
 {
-	// We need to handle the case of our linker being GC'ed, in case someone else is holding
-	// our ref-count up in addition to the linker.
+	if (IsAttachedToLinker())
+	{
+		DetachFromLinker();
+	}
+}
+
+void FMovieSceneEntitySystemRunner::AttachToLinker(UMovieSceneEntitySystemLinker* InLinker)
+{
+	if (!ensureMsgf(InLinker, TEXT("Can't attach to a null linker!")))
+	{
+		return;
+	}
+	if (!ensureMsgf(WeakLinker.IsExplicitlyNull(), TEXT("This runner is already attached to a linker")))
+	{
+		if (ensureMsgf(WeakLinker.IsValid(), TEXT("Our previous linker isn't valid anymore! We will permit attaching to a new one.")))
+		{
+			return;
+		}
+	}
+
+	WeakLinker = InLinker;
 	InLinker->Events.AbandonLinker.AddRaw(this, &FMovieSceneEntitySystemRunner::OnLinkerAbandon);
 }
 
-FMovieSceneEntitySystemRunner::~FMovieSceneEntitySystemRunner()
+bool FMovieSceneEntitySystemRunner::IsAttachedToLinker() const
 {
-	// We don't need to unregister AbandonLinker here since we are either being destroyed along with
-	// our linker, or someone kept us alive longer than the linker and we already unregistered the
-	// event inisde OnLinkerAbandon.
+	return !WeakLinker.IsExplicitlyNull();
+}
+
+void FMovieSceneEntitySystemRunner::DetachFromLinker()
+{
+	if (ensureMsgf(!WeakLinker.IsExplicitlyNull(), TEXT("This runner is not attached to any linker")))
+	{
+		// Abandon our previous linker. We need to do so even for a linker that is pending kill, otherwise
+		// we will later get a OnLinkerAbandon call, which _could_ come _after_ we've been re-attached to a
+		// new valid linker. This would in turn trip the ensure in OnLinkerAbandon that checks that we are
+		// abandoning the linker we have, instead of an unrelated linker.
+		if (UMovieSceneEntitySystemLinker* Linker = WeakLinker.Get(true))
+		{
+			OnLinkerAbandon(Linker);
+		}
+	}
+
+	WeakLinker.Reset();
+}
+
+UMovieSceneEntitySystemLinker* FMovieSceneEntitySystemRunner::GetLinker() const
+{
+	return WeakLinker.Get();
 }
 
 UE::MovieScene::FEntityManager* FMovieSceneEntitySystemRunner::GetEntityManager() const
@@ -289,7 +321,7 @@ UE::MovieScene::ERunnerFlushResult FMovieSceneEntitySystemRunner::StartEvaluatio
 		return ERunnerFlushResult::Break;
 	}
 
-	if (!Linker->StartEvaluation())
+	if (!Linker->StartEvaluation(*this))
 	{
 		return ERunnerFlushResult::Break;
 	}
@@ -310,7 +342,7 @@ UE::MovieScene::ERunnerFlushResult FMovieSceneEntitySystemRunner::StartEvaluatio
 void FMovieSceneEntitySystemRunner::EndEvaluation(UMovieSceneEntitySystemLinker* Linker)
 {
 	using namespace UE::MovieScene;
-	Linker->EndEvaluation();
+	Linker->EndEvaluation(*this);
 }
 
 UE::MovieScene::ERunnerFlushResult FMovieSceneEntitySystemRunner::FlushNext(UMovieSceneEntitySystemLinker* Linker)
@@ -454,9 +486,10 @@ void FMovieSceneEntitySystemRunner::FlushOutstanding(double BudgetMs, UE::MovieS
 		return;
 	}
 
-	// If this runner's linker has been destroyed, early return
-	UMovieSceneEntitySystemLinker* Linker = WeakLinker.Get();
-	if (!Linker)
+	UMovieSceneEntitySystemLinker* Linker = GetLinker();
+
+	// Check that we are attached to a linker that allows starting a new evaluation.
+	if (!ensureMsgf(Linker, TEXT("Runner isn't attached to a valid linker")))
 	{
 		return;
 	}
@@ -706,9 +739,9 @@ UE::MovieScene::ERunnerFlushResult FMovieSceneEntitySystemRunner::GameThread_Upd
 			FSequenceInstance& Instance = InstanceRegistry->MutateInstance(Request.Params.InstanceHandle);
 			if (!Instance.IsRootSequence())
 			{
-				TSharedRef<const FSharedPlaybackState> SharedPlaybackState = Instance.GetSharedPlaybackState();
-				UMovieSceneSequence* RootSequence = SharedPlaybackState->GetRootSequence();
-				UMovieSceneSequence* SubSequence  = SharedPlaybackState->GetSequence(Instance.GetSequenceID());
+				FMovieSceneRootEvaluationTemplateInstance& Template = Instance.GetPlayer()->GetEvaluationTemplate();
+				UMovieSceneSequence* RootSequence = Template.GetRootSequence();
+				UMovieSceneSequence* SubSequence  = Template.GetSequence(Instance.GetSequenceID());
 
 				ensureMsgf(Instance.IsRootSequence(), TEXT("Update request received for a non-root sequence ID 0x%08X (%s) in root-sequence %s. This is not supported."),
 					Instance.GetSequenceID().GetInternalValue(),
@@ -890,7 +923,7 @@ UE::MovieScene::ERunnerFlushResult FMovieSceneEntitySystemRunner::GameThread_Rei
 
 	TGuardValue<bool> IsUpdatingSequenceGuard(bIsUpdatingSequence, true);
 
-	FInstanceRegistry* InstanceRegistry = Linker->GetInstanceRegistry();
+	FInstanceRegistry* InstanceRegistry = WeakLinker.Get()->GetInstanceRegistry();
 
 	// NOTE: IncrementSystemSerial must be called before any instance updates are made
 	//       to ensure that up-to-date versions are used inside FEntityManager::OnStructureChanged
@@ -1034,7 +1067,7 @@ UE::MovieScene::ERunnerFlushResult FMovieSceneEntitySystemRunner::GameThread_Pos
 
 	check(GameThread == ENamedThreads::GameThread || GameThread == ENamedThreads::GameThread_Local);
 
-	Linker->PostInstantation();
+	Linker->PostInstantation(*this);
 
 	FEntityManager& EntityManager = Linker->EntityManager;
 	FBuiltInComponentTypes* BuiltInComponentTypes = FBuiltInComponentTypes::Get();
@@ -1325,8 +1358,8 @@ bool FMovieSceneEntitySystemRunner::FlushSingleEvaluationPhase()
 		return false;
 	}
 
-	UMovieSceneEntitySystemLinker* Linker = WeakLinker.Get();
-	if (!ensureMsgf(Linker, TEXT("Runner doesn't have a valid linker")))
+	UMovieSceneEntitySystemLinker* Linker = GetLinker();
+	if (!ensureMsgf(Linker, TEXT("Runner isn't attached to a valid linker")))
 	{
 		return false;
 	}
diff --git a/Engine/Source/Runtime/MovieScene/Private/EntitySystem/MovieSceneEvaluationHookSystem.cpp b/Engine/Source/Runtime/MovieScene/Private/EntitySystem/MovieSceneEvaluationHookSystem.cpp
--- a/Engine/Source/Runtime/MovieScene/Private/EntitySystem/MovieSceneEvaluationHookSystem.cpp
+++ b/Engine/Source/Runtime/MovieScene/Private/EntitySystem/MovieSceneEvaluationHookSystem.cpp
@@ -107,7 +107,11 @@ void UMovieSceneEvaluationHookSystem::OnRun(FSystemTaskPrerequisites& InPrerequi
 {
 	using namespace UE::MovieScene;
 
-	TSharedRef<FMovieSceneEntitySystemRunner> Runner = Linker->GetRunner();
+	FMovieSceneEntitySystemRunner* Runner = Linker->GetActiveRunner();
+	if (!ensure(Runner))
+	{
+		return;
+	}
 
 	ESystemPhase CurrentPhase = Runner->GetCurrentPhase();
 	if (CurrentPhase == ESystemPhase::Instantiation)
diff --git a/Engine/Source/Runtime/MovieScene/Private/EntitySystem/MovieSceneInstanceRegistry.cpp b/Engine/Source/Runtime/MovieScene/Private/EntitySystem/MovieSceneInstanceRegistry.cpp
--- a/Engine/Source/Runtime/MovieScene/Private/EntitySystem/MovieSceneInstanceRegistry.cpp
+++ b/Engine/Source/Runtime/MovieScene/Private/EntitySystem/MovieSceneInstanceRegistry.cpp
@@ -60,6 +60,7 @@ FInstanceHandle FInstanceRegistry::FindRelatedInstanceHandle(FInstanceHandle Ins
 FRootInstanceHandle FInstanceRegistry::AllocateRootInstance(
 		UMovieSceneSequence& InRootSequence,
 		UObject* InPlaybackContext,
+		TSharedPtr<FMovieSceneEntitySystemRunner> InRunner,
 		UMovieSceneCompiledDataManager* InCompiledDataManager)
 {
 	check(Instances.Num() < 65535);
@@ -69,6 +70,11 @@ FRootInstanceHandle FInstanceRegistry::AllocateRootInstance(
 	FSparseArrayAllocationInfo NewAllocation = Instances.AddUninitialized();
 	FRootInstanceHandle InstanceHandle { (uint16)NewAllocation.Index, InstanceSerial };
 
+	if (!InRunner)
+	{
+		FMovieSceneEntitySystemRunner* ActiveRunner = Linker->GetActiveRunner();
+		InRunner = ActiveRunner ? ActiveRunner->AsShared() : TSharedPtr<FMovieSceneEntitySystemRunner>();
+	}
 	if (!InCompiledDataManager)
 	{
 		InCompiledDataManager = UMovieSceneCompiledDataManager::GetPrecompiledData();
@@ -77,7 +83,7 @@ FRootInstanceHandle FInstanceRegistry::AllocateRootInstance(
 	FSharedPlaybackStateCreateParams PlaybackStateCreateParams;
 	PlaybackStateCreateParams.PlaybackContext = InPlaybackContext;
 	PlaybackStateCreateParams.RootInstanceHandle = InstanceHandle;
-	PlaybackStateCreateParams.Linker = Linker;
+	PlaybackStateCreateParams.Runner = InRunner;
 	PlaybackStateCreateParams.CompiledDataManager = InCompiledDataManager;
 
 	TSharedRef<FSharedPlaybackState> NewPlaybackState = MakeShared<FSharedPlaybackState>(InRootSequence, PlaybackStateCreateParams);
diff --git a/Engine/Source/Runtime/MovieScene/Private/EntitySystem/MovieSceneSequenceUpdaters.cpp b/Engine/Source/Runtime/MovieScene/Private/EntitySystem/MovieSceneSequenceUpdaters.cpp
--- a/Engine/Source/Runtime/MovieScene/Private/EntitySystem/MovieSceneSequenceUpdaters.cpp
+++ b/Engine/Source/Runtime/MovieScene/Private/EntitySystem/MovieSceneSequenceUpdaters.cpp
@@ -752,7 +752,8 @@ void FSequenceUpdater_Hierarchical::Update(TSharedRef<const FSharedPlaybackState
 		}
 	}
 
-	TSharedRef<FMovieSceneEntitySystemRunner> Runner = Linker->GetRunner();
+	FMovieSceneEntitySystemRunner* Runner = Linker->GetActiveRunner();
+	check(Runner);
 
 	for (auto InstanceIt = SequenceInstances.CreateIterator(); InstanceIt; ++InstanceIt)
 	{
diff --git a/Engine/Source/Runtime/MovieScene/Private/EntitySystem/MovieSceneSharedPlaybackState.cpp b/Engine/Source/Runtime/MovieScene/Private/EntitySystem/MovieSceneSharedPlaybackState.cpp
--- a/Engine/Source/Runtime/MovieScene/Private/EntitySystem/MovieSceneSharedPlaybackState.cpp
+++ b/Engine/Source/Runtime/MovieScene/Private/EntitySystem/MovieSceneSharedPlaybackState.cpp
@@ -18,7 +18,7 @@ FSharedPlaybackState::FSharedPlaybackState(
 		const FSharedPlaybackStateCreateParams& CreateParams)
 	: WeakRootSequence(&InRootSequence)
 	, WeakPlaybackContext(CreateParams.PlaybackContext)
-	, WeakLinker(CreateParams.Linker)
+	, WeakRunner(CreateParams.Runner)
 	, CompiledDataManager(CreateParams.CompiledDataManager)
 	, RootInstanceHandle(CreateParams.RootInstanceHandle)
 {
@@ -28,11 +28,11 @@ FSharedPlaybackState::FSharedPlaybackState(
 	}
 }
 
-TSharedPtr<FMovieSceneEntitySystemRunner> FSharedPlaybackState::GetRunner() const
+UMovieSceneEntitySystemLinker* FSharedPlaybackState::GetLinker() const
 {
-	if (UMovieSceneEntitySystemLinker* Linker = WeakLinker.Get())
+	if (TSharedPtr<FMovieSceneEntitySystemRunner> Runner = WeakRunner.Pin())
 	{
-		return Linker->GetRunner();
+		return Runner->GetLinker();
 	}
 	return nullptr;
 }
diff --git a/Engine/Source/Runtime/MovieScene/Private/Evaluation/MovieSceneEvaluationTemplateInstance.cpp b/Engine/Source/Runtime/MovieScene/Private/Evaluation/MovieSceneEvaluationTemplateInstance.cpp
--- a/Engine/Source/Runtime/MovieScene/Private/Evaluation/MovieSceneEvaluationTemplateInstance.cpp
+++ b/Engine/Source/Runtime/MovieScene/Private/Evaluation/MovieSceneEvaluationTemplateInstance.cpp
@@ -72,7 +72,7 @@ UMovieSceneEntitySystemLinker* FMovieSceneRootEvaluationTemplateInstance::Constr
 	return UMovieSceneEntitySystemLinker::FindOrCreateLinker(PlaybackContext, UE::MovieScene::EEntitySystemLinkerRole::Standalone, TEXT("DefaultEntitySystemLinker"));
 }
 
-void FMovieSceneRootEvaluationTemplateInstance::Initialize(UMovieSceneSequence& InRootSequence, IMovieScenePlayer& Player, UMovieSceneCompiledDataManager* InCompiledDataManager)
+void FMovieSceneRootEvaluationTemplateInstance::Initialize(UMovieSceneSequence& InRootSequence, IMovieScenePlayer& Player, UMovieSceneCompiledDataManager* InCompiledDataManager, TSharedPtr<FMovieSceneEntitySystemRunner> InRunner)
 {
 	using namespace UE::MovieScene;
 
@@ -116,7 +116,9 @@ void FMovieSceneRootEvaluationTemplateInstance::Initialize(UMovieSceneSequence&
 	bReinitialize |= (PreviousCompiledDataManager != InCompiledDataManager);
 
 	// Reinitialize if the runner has changed.
-	TSharedPtr<FMovieSceneEntitySystemRunner> PreviousRunner = SharedPlaybackState ? SharedPlaybackState->GetRunner() : nullptr;
+	TSharedPtr<FMovieSceneEntitySystemRunner> PreviousRunner = SharedPlaybackState ?
+		SharedPlaybackState->GetRunner() : nullptr;
+	bReinitialize |= (PreviousRunner != InRunner);
 
 	// Reinitialize if the root sequence has changed.
 	UMovieSceneSequence* PreviousRootSequence = SharedPlaybackState ? 
@@ -151,13 +153,26 @@ void FMovieSceneRootEvaluationTemplateInstance::Initialize(UMovieSceneSequence&
 		// when initializing the new root instance and shared playback state.
 		EntitySystemLinker = ConstructEntityLinker(Player);
 
+		if (ensure(InRunner) && EntitySystemLinker)
+		{
+			if (InRunner->IsAttachedToLinker() && InRunner->GetLinker() != EntitySystemLinker)
+			{
+				InRunner->DetachFromLinker();
+			}
+
+			if (!InRunner->IsAttachedToLinker())
+			{
+				InRunner->AttachToLinker(EntitySystemLinker);
+			}
+		}
+
 		// Create the new root instance and save its new shared playback state.
 		FRootInstanceHandle RootInstanceHandle;
 		if (EntitySystemLinker != nullptr && EntitySystemLinker->GetInstanceRegistry())
 		{
 			UObject* PlaybackContext = Player.GetPlaybackContext();
 			FInstanceRegistry* InstanceRegistry = EntitySystemLinker->GetInstanceRegistry();
-			RootInstanceHandle = InstanceRegistry->AllocateRootInstance(InRootSequence, PlaybackContext, InCompiledDataManager);
+			RootInstanceHandle = InstanceRegistry->AllocateRootInstance(InRootSequence, PlaybackContext, InRunner, InCompiledDataManager);
 			SharedPlaybackState = InstanceRegistry->GetInstance(RootInstanceHandle).GetSharedPlaybackState();
 			Player.InitializeRootInstance(SharedPlaybackState.ToSharedRef());
 		}
@@ -433,6 +448,7 @@ void FMovieSceneRootEvaluationTemplateInstance::PlaybackContextChanged(IMovieSce
 
 	// Only the playback context changed, so we keep the same sequence, runner, and compiled data manager.
 	UMovieSceneSequence* RootSequence = SharedPlaybackState->GetRootSequence();
+	TSharedPtr<FMovieSceneEntitySystemRunner> Runner = SharedPlaybackState->GetRunner();
 	UMovieSceneCompiledDataManager* CompiledDataManager = SharedPlaybackState->GetCompiledDataManager();
 
 	const bool bGlobalCapture = Player.PreAnimatedState.IsCapturingGlobalPreAnimatedState();
@@ -446,7 +462,6 @@ void FMovieSceneRootEvaluationTemplateInstance::PlaybackContextChanged(IMovieSce
 	{
 		EntitySystemLinker->CleanupInvalidBoundObjects();
 
-		TSharedPtr<FMovieSceneEntitySystemRunner> Runner = SharedPlaybackState->GetRunner();
 		if (Runner)
 		{
 			if (Runner->QueueFinalUpdate(PreviousRootInstanceHandle))
@@ -468,11 +483,20 @@ void FMovieSceneRootEvaluationTemplateInstance::PlaybackContextChanged(IMovieSce
 	}
 
 	EntitySystemLinker = ConstructEntityLinker(Player);
+	if (Runner)
+	{
+		if (Runner->IsAttachedToLinker())
+		{
+			Runner->DetachFromLinker();
+		}
+
+		Runner->AttachToLinker(EntitySystemLinker);
+	}
 
 	UObject* PlaybackContext = Player.GetPlaybackContext();
 	FInstanceRegistry* InstanceRegistry = EntitySystemLinker->GetInstanceRegistry();
 	const FRootInstanceHandle RootInstanceHandle = InstanceRegistry->AllocateRootInstance(
-			*RootSequence, PlaybackContext, CompiledDataManager);
+			*RootSequence, PlaybackContext, Runner, CompiledDataManager);
 	FSequenceInstance& RootInstance = InstanceRegistry->MutateInstance(RootInstanceHandle);
 	SharedPlaybackState = RootInstance.GetSharedPlaybackState();
 	Player.InitializeRootInstance(SharedPlaybackState.ToSharedRef());
diff --git a/Engine/Source/Runtime/MovieScene/Private/MovieSceneSequencePlayer.cpp b/Engine/Source/Runtime/MovieScene/Private/MovieSceneSequencePlayer.cpp
--- a/Engine/Source/Runtime/MovieScene/Private/MovieSceneSequencePlayer.cpp
+++ b/Engine/Source/Runtime/MovieScene/Private/MovieSceneSequencePlayer.cpp
@@ -975,7 +975,15 @@ void UMovieSceneSequencePlayer::Initialize(UMovieSceneSequence* InSequence)
 
 	RegisteredTickInterval = TickInterval;
 
-	RootTemplateInstance.Initialize(*Sequence, *this, nullptr);
+	TSharedPtr<FMovieSceneEntitySystemRunner> RunnerToUse = TickManager->GetRunner(RegisteredTickInterval.GetValue());
+	if (EnumHasAnyFlags(Sequence->GetFlags(), EMovieSceneSequenceFlags::BlockingEvaluation))
+	{
+		SynchronousRunner = MakeShared<FMovieSceneEntitySystemRunner>();
+		RunnerToUse = SynchronousRunner;
+	}
+
+	check(RunnerToUse);
+	RootTemplateInstance.Initialize(*Sequence, *this, nullptr, RunnerToUse);
 
 	if (!PlaybackSettings.bDynamicWeighting)
 	{
diff --git a/Engine/Source/Runtime/MovieScene/Private/MovieSceneSequenceTickManager.cpp b/Engine/Source/Runtime/MovieScene/Private/MovieSceneSequenceTickManager.cpp
--- a/Engine/Source/Runtime/MovieScene/Private/MovieSceneSequenceTickManager.cpp
+++ b/Engine/Source/Runtime/MovieScene/Private/MovieSceneSequenceTickManager.cpp
@@ -172,7 +172,8 @@ void UMovieSceneSequenceTickManager::RegisterTickClient(const FMovieSceneSequenc
 
 		FLinkerGroup& NewGroup = LinkerGroups[LinkerIndex];
 		NewGroup.Linker = Linker;
-		NewGroup.Runner = Linker->GetRunner();
+		NewGroup.Runner = MakeShared<FMovieSceneEntitySystemRunner>();
+		NewGroup.Runner->AttachToLinker(Linker);
 		NewGroup.RoundedTickIntervalMs = DesiredTickIntervalMs;
 		NewGroup.FrameBudgetMs = DesiredBudgetMs;
 		NewGroup.NumClients = 1;
diff --git a/Engine/Source/Runtime/MovieScene/Private/Tests/MovieScenePreAnimatedStateTests.cpp b/Engine/Source/Runtime/MovieScene/Private/Tests/MovieScenePreAnimatedStateTests.cpp
--- a/Engine/Source/Runtime/MovieScene/Private/Tests/MovieScenePreAnimatedStateTests.cpp
+++ b/Engine/Source/Runtime/MovieScene/Private/Tests/MovieScenePreAnimatedStateTests.cpp
@@ -393,8 +393,10 @@ bool FMovieScenePreAnimatedStateTrackTypesTest::RunTest(const FString& Parameter
 
 	UMovieSceneCompiledDataManager* CompiledDataManager = UMovieSceneCompiledDataManager::GetPrecompiledData();
 
+	TSharedPtr<FMovieSceneEntitySystemRunner> Runner = MakeShared<FMovieSceneEntitySystemRunner>();
+
 	FMovieSceneCompiledDataID DataID = CompiledDataManager->Compile(Sequence);
-	TestPlayer.Template.Initialize(*Sequence, TestPlayer, CompiledDataManager);
+	TestPlayer.Template.Initialize(*Sequence, TestPlayer, CompiledDataManager, Runner);
 	TestPlayer.Template.EnableGlobalPreAnimatedStateCapture();
 
 	// Test the keep state section
@@ -473,8 +475,10 @@ bool FMovieScenePreAnimatedStateContextChangedTest::RunTest(const FString& Param
 
 	UMovieSceneCompiledDataManager* CompiledDataManager = UMovieSceneCompiledDataManager::GetPrecompiledData();
 
+	TSharedPtr<FMovieSceneEntitySystemRunner> Runner = MakeShared<FMovieSceneEntitySystemRunner>();
+
 	FMovieSceneCompiledDataID DataID = CompiledDataManager->Compile(Sequence);
-	TestPlayer.Template.Initialize(*Sequence, TestPlayer, CompiledDataManager);
+	TestPlayer.Template.Initialize(*Sequence, TestPlayer, CompiledDataManager, Runner);
 	TestPlayer.Template.EnableGlobalPreAnimatedStateCapture();
 
 	// ---------------------------------------------
diff --git a/Engine/Source/Runtime/MovieScene/Public/EntitySystem/MovieSceneEntitySystemLinker.h b/Engine/Source/Runtime/MovieScene/Public/EntitySystem/MovieSceneEntitySystemLinker.h
--- a/Engine/Source/Runtime/MovieScene/Public/EntitySystem/MovieSceneEntitySystemLinker.h
+++ b/Engine/Source/Runtime/MovieScene/Public/EntitySystem/MovieSceneEntitySystemLinker.h
@@ -337,18 +337,15 @@ public:
 	MOVIESCENE_API void InvalidateObjectBinding(const FGuid& ObjectBindingID, FInstanceHandle InstanceHandle);
 	MOVIESCENE_API void CleanupInvalidBoundObjects();
 
-	MOVIESCENE_API bool StartEvaluation();
-	MOVIESCENE_API TSharedRef<FMovieSceneEntitySystemRunner> GetRunner() const;
-	MOVIESCENE_API void PostInstantation();
-	MOVIESCENE_API void EndEvaluation();
+	MOVIESCENE_API bool StartEvaluation(FMovieSceneEntitySystemRunner& InRunner);
+	MOVIESCENE_API FMovieSceneEntitySystemRunner* GetActiveRunner() const;
+	MOVIESCENE_API void PostInstantation(FMovieSceneEntitySystemRunner& InRunner);
+	MOVIESCENE_API void EndEvaluation(FMovieSceneEntitySystemRunner& InRunner);
 
-	MOVIESCENE_API void ResetRunner();
+	MOVIESCENE_API void ResetActiveRunners();
 
 	MOVIESCENE_API void DestroyInstanceImmediately(UE::MovieScene::FRootInstanceHandle Instance);
 
-	UE_DEPRECATED(5.5, "Please use GetRunner()")
-	FMovieSceneEntitySystemRunner* GetActiveRunner() const { return &GetRunner().Get(); }
-
 private:
 
 	MOVIESCENE_API UMovieSceneEntitySystem* LinkSystemImpl(TSubclassOf<UMovieSceneEntitySystem> InClassType);
@@ -373,8 +370,8 @@ private:
 	TSparseArray<UMovieSceneEntitySystem*> EntitySystemsByGlobalGraphID;
 	TMap<TObjectPtr<UClass>, TObjectPtr<UMovieSceneEntitySystem>> EntitySystemsRecyclingPool;
 
-	TSharedPtr<FMovieSceneEntitySystemRunner> Runner;
-	TBitArray<> RunnerReentrancyFlags;
+	TArray<FMovieSceneEntitySystemRunner*> ActiveRunners;
+	TBitArray<> ActiveRunnerReentrancyFlags;
 
 	TSparseArray<void*> ExtensionsByID;
 
diff --git a/Engine/Source/Runtime/MovieScene/Public/EntitySystem/MovieSceneEntitySystemRunner.h b/Engine/Source/Runtime/MovieScene/Public/EntitySystem/MovieSceneEntitySystemRunner.h
--- a/Engine/Source/Runtime/MovieScene/Public/EntitySystem/MovieSceneEntitySystemRunner.h
+++ b/Engine/Source/Runtime/MovieScene/Public/EntitySystem/MovieSceneEntitySystemRunner.h
@@ -79,11 +79,18 @@ public:
 	using FInstanceRegistry = UE::MovieScene::FInstanceRegistry;
 
 public:
-
+	/** Creates an unbound runner */
+	MOVIESCENE_API FMovieSceneEntitySystemRunner();
 	/** Destructor */
 	MOVIESCENE_API ~FMovieSceneEntitySystemRunner();
 
-	/** Gets how many instances are in the update queue. */
+	/** Attach this runner to a linker */
+	MOVIESCENE_API void AttachToLinker(UMovieSceneEntitySystemLinker* InLinker);
+	/** Returns whether this runner is attached to a linker */
+	MOVIESCENE_API bool IsAttachedToLinker() const;
+	/** Detaches this runner from a linker */
+	MOVIESCENE_API void DetachFromLinker();
+
 	MOVIESCENE_API int32 GetQueuedUpdateCount() const;
 	/** Returns whether this runner has any outstanding updates. */
 	MOVIESCENE_API bool HasQueuedUpdates() const;
@@ -161,33 +168,17 @@ public:
 
 public:
 
-	UMovieSceneEntitySystemLinker* GetLinker() const { return WeakLinker.Get(); }
-	bool IsValid() const { return WeakLinker.IsValid(); }
-
+	MOVIESCENE_API UMovieSceneEntitySystemLinker* GetLinker() const;
 	MOVIESCENE_API FEntityManager* GetEntityManager() const;
 	MOVIESCENE_API FInstanceRegistry* GetInstanceRegistry() const;
 
-public:
-
-	UE_DEPRECATED(5.5, "Runners are now owned by linkers, and always attached")
-	void AttachToLinker(UMovieSceneEntitySystemLinker* InLinker) {}
-	UE_DEPRECATED(5.5, "Runners are now owned by linkers, and always attached")
-	bool IsAttachedToLinker() const { return IsValid(); }
-	UE_DEPRECATED(5.5, "Runners are now owned by linkers, and always attached")
-	void DetachFromLinker() {}
-
 public:
 	
 	// Internal API
 
-	void MarkForUpdate(FInstanceHandle InInstanceHandle, UE::MovieScene::ERunnerUpdateFlags UpdateFlags);
+	MOVIESCENE_API void MarkForUpdate(FInstanceHandle InInstanceHandle, UE::MovieScene::ERunnerUpdateFlags UpdateFlags);
 	MOVIESCENE_API FMovieSceneEntitySystemEventTriggers& GetQueuedEventTriggers();
 
-	/** Creates an unbound runner */
-	FMovieSceneEntitySystemRunner();
-	/** Create a bound runner */
-	FMovieSceneEntitySystemRunner(UMovieSceneEntitySystemLinker* InLinker);
-
 private:
 
 	MOVIESCENE_API void OnLinkerAbandon(UMovieSceneEntitySystemLinker* Linker);
@@ -260,7 +251,7 @@ private:
 		UE::MovieScene::ERunnerUpdateFlags UpdateFlags = UE::MovieScene::ERunnerUpdateFlags::None;
 	};
 
-	/** Owning linker */
+	/** Owner linker */
 	TWeakObjectPtr<UMovieSceneEntitySystemLinker> WeakLinker;
 
 	struct FQueuedUpdateParams
diff --git a/Engine/Source/Runtime/MovieScene/Public/EntitySystem/MovieSceneInstanceRegistry.h b/Engine/Source/Runtime/MovieScene/Public/EntitySystem/MovieSceneInstanceRegistry.h
--- a/Engine/Source/Runtime/MovieScene/Public/EntitySystem/MovieSceneInstanceRegistry.h
+++ b/Engine/Source/Runtime/MovieScene/Public/EntitySystem/MovieSceneInstanceRegistry.h
@@ -12,6 +12,7 @@
 #include "MovieSceneSequenceID.h"
 
 
+class FMovieSceneEntitySystemRunner;
 class UMovieSceneCompiledDataManager;
 class UMovieSceneEntitySystemLinker;
 class UMovieSceneCompiledDataManager;
@@ -86,6 +87,7 @@ struct FInstanceRegistry
 	MOVIESCENE_API FRootInstanceHandle AllocateRootInstance(
 			UMovieSceneSequence& InRootSequence,
 			UObject* InPlaybackContext = nullptr,
+			TSharedPtr<FMovieSceneEntitySystemRunner> InRunner = nullptr,
 			UMovieSceneCompiledDataManager* InCompiledDataManager = nullptr);
 
 	MOVIESCENE_API FInstanceHandle AllocateSubInstance(
diff --git a/Engine/Source/Runtime/MovieScene/Public/EntitySystem/MovieSceneSharedPlaybackState.h b/Engine/Source/Runtime/MovieScene/Public/EntitySystem/MovieSceneSharedPlaybackState.h
--- a/Engine/Source/Runtime/MovieScene/Public/EntitySystem/MovieSceneSharedPlaybackState.h
+++ b/Engine/Source/Runtime/MovieScene/Public/EntitySystem/MovieSceneSharedPlaybackState.h
@@ -11,7 +11,6 @@
 
 class FMovieSceneEntitySystemRunner;
 class UMovieSceneCompiledDataManager;
-class UMovieSceneEntitySystemLinker;
 struct FMovieSceneObjectCache;
 struct FMovieSceneSequenceHierarchy;
 
@@ -26,7 +25,7 @@ struct FSharedPlaybackStateCreateParams
 	/**
 	 * The playback context in which the root sequence will be evaluated.
 	 *
-	 * Requires that RootInstanceHandle and Linker are also set.
+	 * Requires that RootInstanceHandle and Runner are also set.
 	 */
 	UObject* PlaybackContext = nullptr;
 
@@ -35,17 +34,17 @@ struct FSharedPlaybackStateCreateParams
 	 * is meant to relate to an instance that has also been created inside
 	 * a runner/linker's instance registry.
 	 *
-	 * Requires that PlaybackContext and Linker are also set.
+	 * Requires that PlaybackContext and Runner are also set.
 	 */
 	FRootInstanceHandle RootInstanceHandle;
 
 	/**
-	 * The linker that will be evaluating the sequence that the created playback
+	 * The runner that will be evaluating the sequence that the created playback
 	 * state relates to.
 	 *
 	 * Requires that PlaybackContext and RootInstanceHandle are also set.
 	 */
-	TObjectPtr<UMovieSceneEntitySystemLinker> Linker;
+	TSharedPtr<FMovieSceneEntitySystemRunner> Runner;
 
 	/**
 	 * The compiled data manager with which the root sequence was compiled, or
@@ -74,8 +73,8 @@ public:
 	/** Gets the root sequence */
 	UMovieSceneSequence* GetRootSequence() const { return WeakRootSequence.Get(); }
 
-	/** Gets the linker evaluating this root sequence */
-	UMovieSceneEntitySystemLinker* GetLinker() const { return WeakLinker.Get(); }
+	/** Gets the runner evaluating this root sequence */
+	TSharedPtr<FMovieSceneEntitySystemRunner> GetRunner() const { return WeakRunner.Pin(); }
 
 	/** Gets the compiled data manager that contains the data for the root sequence */
 	TObjectPtr<UMovieSceneCompiledDataManager> GetCompiledDataManager() const { return CompiledDataManager; }
@@ -90,12 +89,9 @@ public:
 
 	// General utility methods
 
-	/** Gets the runner evaluating this root sequence */
-	TSharedPtr<FMovieSceneEntitySystemRunner> GetRunner() const;
+	UMovieSceneEntitySystemLinker* GetLinker() const;
 
-	/** Gets the hierarchy (if any) for this root sequence */
 	const FMovieSceneSequenceHierarchy* GetHierarchy() const;
-	/** Gets a sub-sequence given an ID */
 	UMovieSceneSequence* GetSequence(FMovieSceneSequenceIDRef SequenceID) const;
 
 public:
@@ -300,8 +296,8 @@ private:
 	/** The playback context */
 	TWeakObjectPtr<UObject> WeakPlaybackContext;
 
-	/** The linker evaluating this root sequence */
-	TWeakObjectPtr<UMovieSceneEntitySystemLinker> WeakLinker;
+	/** The runner evaluating this root sequence */
+	TWeakPtr<FMovieSceneEntitySystemRunner> WeakRunner;
 
 	/** The compiled data manager that contains the data for the root sequence */
 	TObjectPtr<UMovieSceneCompiledDataManager> CompiledDataManager;
diff --git a/Engine/Source/Runtime/MovieScene/Public/Evaluation/MovieSceneEvaluationTemplateInstance.h b/Engine/Source/Runtime/MovieScene/Public/Evaluation/MovieSceneEvaluationTemplateInstance.h
--- a/Engine/Source/Runtime/MovieScene/Public/Evaluation/MovieSceneEvaluationTemplateInstance.h
+++ b/Engine/Source/Runtime/MovieScene/Public/Evaluation/MovieSceneEvaluationTemplateInstance.h
@@ -50,7 +50,7 @@ public:
 	 * @param Player					The player responsible for playback
 	 * @param TemplateStore				Template store responsible for supplying templates for a given sequence
 	 */
-	MOVIESCENE_API void Initialize(UMovieSceneSequence& RootSequence, IMovieScenePlayer& Player, UMovieSceneCompiledDataManager* InCompiledDataManager);
+	MOVIESCENE_API void Initialize(UMovieSceneSequence& RootSequence, IMovieScenePlayer& Player, UMovieSceneCompiledDataManager* InCompiledDataManager, TSharedPtr<FMovieSceneEntitySystemRunner> InRunner);
 
 	/**
 	 * Evaluate this sequence in a synchronous way.
diff --git a/Engine/Source/Runtime/MovieSceneTracks/Private/EntitySystem/Interrogation/MovieSceneInterrogationLinker.cpp b/Engine/Source/Runtime/MovieSceneTracks/Private/EntitySystem/Interrogation/MovieSceneInterrogationLinker.cpp
--- a/Engine/Source/Runtime/MovieSceneTracks/Private/EntitySystem/Interrogation/MovieSceneInterrogationLinker.cpp
+++ b/Engine/Source/Runtime/MovieSceneTracks/Private/EntitySystem/Interrogation/MovieSceneInterrogationLinker.cpp
@@ -844,8 +844,9 @@ void FSystemInterrogator::Update()
 			}
 		});
 
-	TSharedPtr<FMovieSceneEntitySystemRunner> Runner = Linker->GetRunner();
-	Runner->Flush();
+	FMovieSceneEntitySystemRunner Runner;
+	Runner.AttachToLinker(Linker);
+	Runner.Flush();
 
 	Linker->EntityManager.IncrementSystemSerial();
 }
diff --git a/Engine/Source/Runtime/MovieSceneTracks/Private/Systems/DoubleChannelEvaluatorSystem.cpp b/Engine/Source/Runtime/MovieSceneTracks/Private/Systems/DoubleChannelEvaluatorSystem.cpp
--- a/Engine/Source/Runtime/MovieSceneTracks/Private/Systems/DoubleChannelEvaluatorSystem.cpp
+++ b/Engine/Source/Runtime/MovieSceneTracks/Private/Systems/DoubleChannelEvaluatorSystem.cpp
@@ -205,7 +205,11 @@ void UDoubleChannelEvaluatorSystem::OnRun(FSystemTaskPrerequisites& InPrerequisi
 {
 	using namespace UE::MovieScene;
 
-	TSharedRef<FMovieSceneEntitySystemRunner> Runner = Linker->GetRunner();
+	FMovieSceneEntitySystemRunner* Runner = Linker->GetActiveRunner();
+	if (!Runner)
+	{
+		return;
+	}
 
 	FBuiltInComponentTypes* BuiltInComponents = FBuiltInComponentTypes::Get();
 
diff --git a/Engine/Source/Runtime/MovieSceneTracks/Private/Systems/DoublePerlinNoiseChannelEvaluatorSystem.cpp b/Engine/Source/Runtime/MovieSceneTracks/Private/Systems/DoublePerlinNoiseChannelEvaluatorSystem.cpp
--- a/Engine/Source/Runtime/MovieSceneTracks/Private/Systems/DoublePerlinNoiseChannelEvaluatorSystem.cpp
+++ b/Engine/Source/Runtime/MovieSceneTracks/Private/Systems/DoublePerlinNoiseChannelEvaluatorSystem.cpp
@@ -83,7 +83,11 @@ void UDoublePerlinNoiseChannelEvaluatorSystem::OnRun(FSystemTaskPrerequisites& I
 	const FBuiltInComponentTypes* BuiltInComponents = FBuiltInComponentTypes::Get();
 	const FMovieSceneTracksComponentTypes* TrackComponents = FMovieSceneTracksComponentTypes::Get();
 
-	TSharedRef<FMovieSceneEntitySystemRunner> Runner = Linker->GetRunner();
+	FMovieSceneEntitySystemRunner* Runner = Linker->GetActiveRunner();
+	if (!Runner)
+	{
+		return;
+	}
 
 	if (Runner->GetCurrentPhase() == ESystemPhase::Instantiation)
 	{
diff --git a/Engine/Source/Runtime/MovieSceneTracks/Private/Systems/FloatChannelEvaluatorSystem.cpp b/Engine/Source/Runtime/MovieSceneTracks/Private/Systems/FloatChannelEvaluatorSystem.cpp
--- a/Engine/Source/Runtime/MovieSceneTracks/Private/Systems/FloatChannelEvaluatorSystem.cpp
+++ b/Engine/Source/Runtime/MovieSceneTracks/Private/Systems/FloatChannelEvaluatorSystem.cpp
@@ -218,7 +218,11 @@ void UFloatChannelEvaluatorSystem::OnRun(FSystemTaskPrerequisites& InPrerequisit
 {
 	using namespace UE::MovieScene;
 
-	TSharedRef<FMovieSceneEntitySystemRunner> Runner = Linker->GetRunner();
+	FMovieSceneEntitySystemRunner* Runner = Linker->GetActiveRunner();
+	if (!Runner)
+	{
+		return;
+	}
 
 	FBuiltInComponentTypes* BuiltInComponents = FBuiltInComponentTypes::Get();
 
diff --git a/Engine/Source/Runtime/MovieSceneTracks/Private/Systems/FloatPerlinNoiseChannelEvaluatorSystem.cpp b/Engine/Source/Runtime/MovieSceneTracks/Private/Systems/FloatPerlinNoiseChannelEvaluatorSystem.cpp
--- a/Engine/Source/Runtime/MovieSceneTracks/Private/Systems/FloatPerlinNoiseChannelEvaluatorSystem.cpp
+++ b/Engine/Source/Runtime/MovieSceneTracks/Private/Systems/FloatPerlinNoiseChannelEvaluatorSystem.cpp
@@ -84,7 +84,11 @@ void UFloatPerlinNoiseChannelEvaluatorSystem::OnRun(FSystemTaskPrerequisites& In
 	const FBuiltInComponentTypes* BuiltInComponents = FBuiltInComponentTypes::Get();
 	const FMovieSceneTracksComponentTypes* TrackComponents = FMovieSceneTracksComponentTypes::Get();
 
-	TSharedRef<FMovieSceneEntitySystemRunner> Runner = Linker->GetRunner();
+	FMovieSceneEntitySystemRunner* Runner = Linker->GetActiveRunner();
+	if (!Runner)
+	{
+		return;
+	}
 
 	if (Runner->GetCurrentPhase() == ESystemPhase::Instantiation)
 	{
diff --git a/Engine/Source/Runtime/MovieSceneTracks/Private/Systems/MovieSceneConstraintSystem.cpp b/Engine/Source/Runtime/MovieSceneTracks/Private/Systems/MovieSceneConstraintSystem.cpp
--- a/Engine/Source/Runtime/MovieSceneTracks/Private/Systems/MovieSceneConstraintSystem.cpp
+++ b/Engine/Source/Runtime/MovieSceneTracks/Private/Systems/MovieSceneConstraintSystem.cpp
@@ -254,8 +254,8 @@ void UMovieSceneConstraintSystem::OnRun(FSystemTaskPrerequisites& InPrerequisite
 {
 	using namespace UE::MovieScene;
 
-	TSharedRef<FMovieSceneEntitySystemRunner> Runner = Linker->GetRunner();
-	ESystemPhase CurrentPhase = Runner->GetCurrentPhase();
+	FMovieSceneEntitySystemRunner* ActiveRunner = Linker->GetActiveRunner();
+	ESystemPhase CurrentPhase = ActiveRunner->GetCurrentPhase();
 
 	if (CurrentPhase == ESystemPhase::Instantiation)
 	{
diff --git a/Engine/Source/Runtime/MovieSceneTracks/Private/Systems/MovieSceneCustomPrimitiveDataSystem.cpp b/Engine/Source/Runtime/MovieSceneTracks/Private/Systems/MovieSceneCustomPrimitiveDataSystem.cpp
--- a/Engine/Source/Runtime/MovieSceneTracks/Private/Systems/MovieSceneCustomPrimitiveDataSystem.cpp
+++ b/Engine/Source/Runtime/MovieSceneTracks/Private/Systems/MovieSceneCustomPrimitiveDataSystem.cpp
@@ -431,9 +431,13 @@ void UMovieSceneCustomPrimitiveDataSystem::OnRun(FSystemTaskPrerequisites& InPre
 {
 	using namespace UE::MovieScene;
 
-	TSharedRef<FMovieSceneEntitySystemRunner> Runner = Linker->GetRunner();
+	FMovieSceneEntitySystemRunner* ActiveRunner = Linker->GetActiveRunner();
+	if (!ensure(ActiveRunner))
+	{
+		return;
+	}
 
-	ESystemPhase CurrentPhase = Runner->GetCurrentPhase();
+	ESystemPhase CurrentPhase = ActiveRunner->GetCurrentPhase();
 	if (CurrentPhase == ESystemPhase::Instantiation)
 	{
 		OnInstantiation();
diff --git a/Engine/Source/Runtime/MovieSceneTracks/Private/Systems/MovieSceneDeferredComponentMovementSystem.cpp b/Engine/Source/Runtime/MovieSceneTracks/Private/Systems/MovieSceneDeferredComponentMovementSystem.cpp
--- a/Engine/Source/Runtime/MovieSceneTracks/Private/Systems/MovieSceneDeferredComponentMovementSystem.cpp
+++ b/Engine/Source/Runtime/MovieSceneTracks/Private/Systems/MovieSceneDeferredComponentMovementSystem.cpp
@@ -145,8 +145,8 @@ void UMovieSceneDeferredComponentMovementSystem::OnRun(FSystemTaskPrerequisites&
 {
 	using namespace UE::MovieScene;
 
-	TSharedRef<FMovieSceneEntitySystemRunner> Runner = Linker->GetRunner();
-	ESystemPhase CurrentPhase = Runner->GetCurrentPhase();
+	FMovieSceneEntitySystemRunner* ActiveRunner = Linker->GetActiveRunner();
+	ESystemPhase CurrentPhase = ActiveRunner->GetCurrentPhase();
 
 	FBuiltInComponentTypes*          BuiltInComponents = FBuiltInComponentTypes::Get();
 	FMovieSceneTracksComponentTypes* Components        = FMovieSceneTracksComponentTypes::Get();
diff --git a/Engine/Source/Runtime/MovieSceneTracks/Private/Systems/MovieSceneEventSystems.cpp b/Engine/Source/Runtime/MovieSceneTracks/Private/Systems/MovieSceneEventSystems.cpp
--- a/Engine/Source/Runtime/MovieSceneTracks/Private/Systems/MovieSceneEventSystems.cpp
+++ b/Engine/Source/Runtime/MovieSceneTracks/Private/Systems/MovieSceneEventSystems.cpp
@@ -383,8 +383,11 @@ void UMovieScenePostEvalEventSystem::OnRun(FSystemTaskPrerequisites& InPrerequis
 {
 	if (HasEvents())
 	{
-		TSharedRef<FMovieSceneEntitySystemRunner> Runner = Linker->GetRunner();
-		Runner->GetQueuedEventTriggers().AddUObject(this, &UMovieScenePostEvalEventSystem::TriggerAllEvents);
+		FMovieSceneEntitySystemRunner* Runner = Linker->GetActiveRunner();
+		if (ensure(Runner))
+		{
+			Runner->GetQueuedEventTriggers().AddUObject(this, &UMovieScenePostEvalEventSystem::TriggerAllEvents);
+		}
 	}
 }
 
diff --git a/Engine/Source/Runtime/MovieSceneTracks/Private/Systems/MovieSceneMotionVectorSimulationSystem.cpp b/Engine/Source/Runtime/MovieSceneTracks/Private/Systems/MovieSceneMotionVectorSimulationSystem.cpp
--- a/Engine/Source/Runtime/MovieSceneTracks/Private/Systems/MovieSceneMotionVectorSimulationSystem.cpp
+++ b/Engine/Source/Runtime/MovieSceneTracks/Private/Systems/MovieSceneMotionVectorSimulationSystem.cpp
@@ -59,7 +59,7 @@ void UMovieSceneMotionVectorSimulationSystem::OnRun(FSystemTaskPrerequisites& In
 {
 	using namespace UE::MovieScene;
 
-	TSharedRef<FMovieSceneEntitySystemRunner> Runner = Linker->GetRunner();
+	FMovieSceneEntitySystemRunner* Runner = Linker->GetActiveRunner();
 	if (Runner->GetCurrentPhase() == ESystemPhase::Finalization && bSimulationEnabled)
 	{
 		Runner->GetQueuedEventTriggers().AddLambda([this] { this->OnPostEvaluation(); });
@@ -135,8 +135,11 @@ void UMovieSceneMotionVectorSimulationSystem::ComputeSimulatedMotion()
 
 	// --------------------------------------------------------------------------------------------------------------------------------------------
 	// Re-execute the evaluation phase to flush the new transforms
-	TSharedRef<FMovieSceneEntitySystemRunner> Runner = Linker->GetRunner();
-	Runner->FlushSingleEvaluationPhase();
+	FMovieSceneEntitySystemRunner* ActiveRunner = Linker->GetActiveRunner();
+	if (ensure(ActiveRunner))
+	{
+		ActiveRunner->FlushSingleEvaluationPhase();
+	}
 
 	// --------------------------------------------------------------------------------------------------------------------------------------------
 	// Harvest results
diff --git a/Engine/Source/Runtime/MovieSceneTracks/Private/Systems/MovieScenePredictionSystem.cpp b/Engine/Source/Runtime/MovieSceneTracks/Private/Systems/MovieScenePredictionSystem.cpp
--- a/Engine/Source/Runtime/MovieSceneTracks/Private/Systems/MovieScenePredictionSystem.cpp
+++ b/Engine/Source/Runtime/MovieSceneTracks/Private/Systems/MovieScenePredictionSystem.cpp
@@ -408,7 +408,11 @@ void UMovieScenePredictionSystem::OnRun(FSystemTaskPrerequisites& InPrerequisite
 {
 	using namespace UE::MovieScene;
 
-	TSharedRef<FMovieSceneEntitySystemRunner> Runner = Linker->GetRunner();
+	FMovieSceneEntitySystemRunner* Runner = Linker->GetActiveRunner();
+	if (!ensure(Runner))
+	{
+		return;
+	}
 
 	FInstanceRegistry* InstanceRegistry = Linker->GetInstanceRegistry();
 
diff --git a/Engine/Source/Runtime/MovieSceneTracks/Private/Systems/MovieSceneSkeletalAnimationSystem.cpp b/Engine/Source/Runtime/MovieSceneTracks/Private/Systems/MovieSceneSkeletalAnimationSystem.cpp
--- a/Engine/Source/Runtime/MovieSceneTracks/Private/Systems/MovieSceneSkeletalAnimationSystem.cpp
+++ b/Engine/Source/Runtime/MovieSceneTracks/Private/Systems/MovieSceneSkeletalAnimationSystem.cpp
@@ -1020,7 +1020,7 @@ void UMovieSceneSkeletalAnimationSystem::OnRun(FSystemTaskPrerequisites& InPrere
 
 	const TStatId GatherStatId = GET_STATID(MovieSceneEval_GatherSkeletalAnimations);
 
-	TSharedRef<FMovieSceneEntitySystemRunner> Runner = Linker->GetRunner();
+	const FMovieSceneEntitySystemRunner* Runner = Linker->GetActiveRunner();
 	if (Runner->GetCurrentPhase() == ESystemPhase::Instantiation)
 	{
 		CleanSystemData();
diff --git a/Engine/Source/Runtime/MovieSceneTracks/Private/Tests/MovieScenePartialEvaluationTests.cpp b/Engine/Source/Runtime/MovieSceneTracks/Private/Tests/MovieScenePartialEvaluationTests.cpp
--- a/Engine/Source/Runtime/MovieSceneTracks/Private/Tests/MovieScenePartialEvaluationTests.cpp
+++ b/Engine/Source/Runtime/MovieSceneTracks/Private/Tests/MovieScenePartialEvaluationTests.cpp
@@ -67,10 +67,12 @@ bool FMovieScenePartialEvaluationTest::RunTest(const FString& Parameters)
 
 	UMovieSceneCompiledDataManager* CompiledDataManager = UMovieSceneCompiledDataManager::GetPrecompiledData();
 	TStrongObjectPtr<UMovieSceneEntitySystemLinker> Linker(NewObject<UMovieSceneEntitySystemLinker>(GetTransientPackage()));
-	TSharedPtr<FMovieSceneEntitySystemRunner> Runner = Linker->GetRunner();
+	TSharedPtr<FMovieSceneEntitySystemRunner> Runner = MakeShared<FMovieSceneEntitySystemRunner>();
 
 	TGuardValue<FEntityManager*> DebugVizGuard(GEntityManagerForDebuggingVisualizers, &Linker->EntityManager);
 
+	Runner->AttachToLinker(Linker.Get());
+
 	struct FLocalPlayer : IMovieScenePlayer
 	{
 		FMovieSceneRootEvaluationTemplateInstance Template;
@@ -89,7 +91,7 @@ bool FMovieScenePartialEvaluationTest::RunTest(const FString& Parameters)
 	Player.TestLinker = Linker.Get();
 
 	FMovieSceneCompiledDataID DataID = CompiledDataManager->Compile(RootSequence.Get());
-	Player.Template.Initialize(*RootSequence, Player, CompiledDataManager);
+	Player.Template.Initialize(*RootSequence, Player, CompiledDataManager, Runner);
 
 	bool bIsActive = true;
 
diff --git a/Engine/Source/Runtime/MovieSceneTracks/Public/Systems/MovieSceneMaterialSystem.h b/Engine/Source/Runtime/MovieSceneTracks/Public/Systems/MovieSceneMaterialSystem.h
--- a/Engine/Source/Runtime/MovieSceneTracks/Public/Systems/MovieSceneMaterialSystem.h
+++ b/Engine/Source/Runtime/MovieSceneTracks/Public/Systems/MovieSceneMaterialSystem.h
@@ -344,9 +344,13 @@ void TMovieSceneMaterialSystem<AccessorType, RequiredComponents...>::OnRun(UMovi
 	FBuiltInComponentTypes*          BuiltInComponents = FBuiltInComponentTypes::Get();
 	FMovieSceneTracksComponentTypes* TracksComponents  = FMovieSceneTracksComponentTypes::Get();
 
-	TSharedRef<FMovieSceneEntitySystemRunner> Runner = Linker->GetRunner();
+	FMovieSceneEntitySystemRunner* ActiveRunner = Linker->GetActiveRunner();
+	if (!ensure(ActiveRunner))
+	{
+		return;
+	}
 
-	ESystemPhase CurrentPhase = Runner->GetCurrentPhase();
+	ESystemPhase CurrentPhase = ActiveRunner->GetCurrentPhase();
 	if (CurrentPhase == ESystemPhase::Instantiation)
 	{
 		// --------------------------------------------------------------------------------------
diff --git a/Engine/Source/Runtime/UMG/Private/Animation/UMGSequencePlayer.cpp b/Engine/Source/Runtime/UMG/Private/Animation/UMGSequencePlayer.cpp
--- a/Engine/Source/Runtime/UMG/Private/Animation/UMGSequencePlayer.cpp
+++ b/Engine/Source/Runtime/UMG/Private/Animation/UMGSequencePlayer.cpp
@@ -234,7 +234,14 @@ void UUMGSequencePlayer::PlayInternal(double StartAtTime, double EndAtTime, int3
 		CSV_EVENT_GLOBAL(TEXT("Play Animation [%s::%s]"), *Widget->GetName(), *Animation->GetName());
 	}
 
-	RootTemplateInstance.Initialize(*Animation, *this, nullptr);
+	TSharedPtr<FMovieSceneEntitySystemRunner> RunnerToUse = TickManager ? TickManager->GetRunner() : nullptr;
+	if (EnumHasAnyFlags(Animation->GetFlags(), EMovieSceneSequenceFlags::BlockingEvaluation))
+	{
+		SynchronousRunner = MakeShared<FMovieSceneEntitySystemRunner>();
+		RunnerToUse = SynchronousRunner;
+	}
+
+	RootTemplateInstance.Initialize(*Animation, *this, nullptr, RunnerToUse);
 
 	if (bInRestoreState)
 	{
@@ -278,8 +285,7 @@ void UUMGSequencePlayer::PlayInternal(double StartAtTime, double EndAtTime, int3
 	PlayerStatus = EMovieScenePlayerStatus::Playing;
 
 	// Playback assumes the start frame has already been evaulated, so we also want to evaluate any events on the start frame here.
-	TSharedPtr<FMovieSceneEntitySystemRunner> Runner = RootTemplateInstance.GetRunner();
-	if (Runner)
+	if (RunnerToUse)
 	{
 		const FMovieSceneContext Context(FMovieSceneEvaluationRange(AbsolutePlaybackStart + TimeCursorPosition, AbsolutePlaybackStart + TimeCursorPosition, AnimationResolution), PlayerStatus);
 
@@ -290,11 +296,11 @@ void UUMGSequencePlayer::PlayInternal(double StartAtTime, double EndAtTime, int3
 
 		// We queue an update instead of immediately flushing the entire linker so that we don't incur a cascade of flushes on frames when multiple animations are played
 		// In rare cases where the linker must be flushed immediately PreTick, the queue should be manually flushed 
-		Runner->QueueUpdate(Context, RootTemplateInstance.GetRootInstanceHandle(), FSimpleDelegate::CreateWeakLambda(this, OnBegunPlay), UE::MovieScene::ERunnerUpdateFlags::Flush);
+		RunnerToUse->QueueUpdate(Context, RootTemplateInstance.GetRootInstanceHandle(), FSimpleDelegate::CreateWeakLambda(this, OnBegunPlay), UE::MovieScene::ERunnerUpdateFlags::Flush);
 
-		if (Runner == SynchronousRunner || !UE::UMG::GAsyncAnimationControlFlow)
+		if (RunnerToUse == SynchronousRunner || !UE::UMG::GAsyncAnimationControlFlow)
 		{
-			Runner->Flush();
+			RunnerToUse->Flush();
 		}
 	}
 }
diff --git a/Engine/Source/Runtime/UMG/Private/Animation/UMGSequenceTickManager.cpp b/Engine/Source/Runtime/UMG/Private/Animation/UMGSequenceTickManager.cpp
--- a/Engine/Source/Runtime/UMG/Private/Animation/UMGSequenceTickManager.cpp
+++ b/Engine/Source/Runtime/UMG/Private/Animation/UMGSequenceTickManager.cpp
@@ -46,22 +46,10 @@ UUMGSequenceTickManager::UUMGSequenceTickManager(const FObjectInitializer& Init)
 	: Super(Init)
 	, bIsTicking(false)
 {
-}
-
-void UUMGSequenceTickManager::Initialize(UObject* Owner)
-{
-	Linker = UMovieSceneEntitySystemLinker::FindOrCreateLinker(Owner, UE::MovieScene::EEntitySystemLinkerRole::UMG, TEXT("UMGAnimationEntitySystemLinker"));
-	check(Linker);
-	Runner = Linker->GetRunner();
-
-	FSlateApplication& SlateApp = FSlateApplication::Get();
-	FDelegateHandle PreTickHandle = SlateApp.OnPreTick().AddUObject(this, &UUMGSequenceTickManager::TickWidgetAnimations);
-	check(PreTickHandle.IsValid());
-	SlateApplicationPreTickHandle = PreTickHandle;
-
-	FDelegateHandle PostTickHandle = SlateApp.OnPostTick().AddUObject(this, &UUMGSequenceTickManager::HandleSlatePostTick);
-	check(PostTickHandle.IsValid());
-	SlateApplicationPostTickHandle = PostTickHandle;
+	if (!HasAnyFlags(RF_ClassDefaultObject))
+	{
+		Runner = MakeShared<FMovieSceneEntitySystemRunner>();
+	}
 }
 
 void UUMGSequenceTickManager::AddWidget(UUserWidget* InWidget)
@@ -250,8 +238,11 @@ void UUMGSequenceTickManager::TickWidgetAnimations(float DeltaSeconds)
 
 void UUMGSequenceTickManager::ForceFlush()
 {
-	Runner->Flush(UE::UMG::GAnimationBudgetMs);
-	RunLatentActions();
+	if (Runner->IsAttachedToLinker())
+	{
+		Runner->Flush(UE::UMG::GAnimationBudgetMs);
+		RunLatentActions();
+	}
 }
 
 void UUMGSequenceTickManager::HandleSlatePostTick(float DeltaSeconds)
@@ -263,7 +254,7 @@ void UUMGSequenceTickManager::HandleSlatePostTick(float DeltaSeconds)
 	}
 
 	// Only tick widgets at the end of the frame if our runner has completely finished, and we still have updates
-	if (UE::UMG::GFlushUMGAnimationsAtEndOfFrame && Runner->HasQueuedUpdates() && !Runner->IsCurrentlyEvaluating())
+	if (UE::UMG::GFlushUMGAnimationsAtEndOfFrame && Runner->IsAttachedToLinker() && Runner->HasQueuedUpdates() && !Runner->IsCurrentlyEvaluating())
 	{
 		SCOPE_CYCLE_COUNTER(MovieSceneEval_FlushEndOfFrameAnimations);
 
@@ -325,7 +316,19 @@ UUMGSequenceTickManager* UUMGSequenceTickManager::Get(UObject* PlaybackContext)
 	if (!TickManager)
 	{
 		TickManager = NewObject<UUMGSequenceTickManager>(Owner, TickManagerName);
-		TickManager->Initialize(Owner);
+
+		TickManager->Linker = UMovieSceneEntitySystemLinker::FindOrCreateLinker(Owner, UE::MovieScene::EEntitySystemLinkerRole::UMG, TEXT("UMGAnimationEntitySystemLinker"));
+		check(TickManager->Linker);
+		TickManager->Runner->AttachToLinker(TickManager->Linker);
+
+		FSlateApplication& SlateApp = FSlateApplication::Get();
+		FDelegateHandle PreTickHandle = SlateApp.OnPreTick().AddUObject(TickManager, &UUMGSequenceTickManager::TickWidgetAnimations);
+		check(PreTickHandle.IsValid());
+		TickManager->SlateApplicationPreTickHandle = PreTickHandle;
+
+		FDelegateHandle PostTickHandle = SlateApp.OnPostTick().AddUObject(TickManager, &UUMGSequenceTickManager::HandleSlatePostTick);
+		check(PostTickHandle.IsValid());
+		TickManager->SlateApplicationPostTickHandle = PostTickHandle;
 	}
 	return TickManager;
 }
diff --git a/Engine/Source/Runtime/UMG/Public/Animation/UMGSequenceTickManager.h b/Engine/Source/Runtime/UMG/Public/Animation/UMGSequenceTickManager.h
--- a/Engine/Source/Runtime/UMG/Public/Animation/UMGSequenceTickManager.h
+++ b/Engine/Source/Runtime/UMG/Public/Animation/UMGSequenceTickManager.h
@@ -37,8 +37,6 @@ public:
 
 	UMG_API UUMGSequenceTickManager(const FObjectInitializer& Init);
 
-	UMG_API void Initialize(UObject* Owner);
-
 	UMovieSceneEntitySystemLinker* GetLinker() { return Linker; }
 	TSharedPtr<FMovieSceneEntitySystemRunner> GetRunner() { return Runner; }
 
