Sequencer: fix self-assignment test

diff --git a/Engine/Source/Runtime/MovieScene/Private/Evaluation/MovieScenePlaybackCapabilities.cpp b/Engine/Source/Runtime/MovieScene/Private/Evaluation/MovieScenePlaybackCapabilities.cpp
--- a/Engine/Source/Runtime/MovieScene/Private/Evaluation/MovieScenePlaybackCapabilities.cpp
+++ b/Engine/Source/Runtime/MovieScene/Private/Evaluation/MovieScenePlaybackCapabilities.cpp
@@ -18,7 +18,7 @@ FPlaybackCapabilities::FPlaybackCapabilities(FPlaybackCapabilities&& RHS)
 
 FPlaybackCapabilities& FPlaybackCapabilities::operator=(FPlaybackCapabilities&& RHS)
 {
-	if (!ensure(this != &RHS))
+	if (ensure(this != &RHS))
 	{
 		Destroy();
 
