Sequencer: fix reinitialization of root evaluation template instance

It was tripping the assert that was recently put in to detect leaking playback capabilities.

diff --git a/Engine/Source/Runtime/MovieScene/Private/Evaluation/MovieSceneEvaluationTemplateInstance.cpp b/Engine/Source/Runtime/MovieScene/Private/Evaluation/MovieSceneEvaluationTemplateInstance.cpp
--- a/Engine/Source/Runtime/MovieScene/Private/Evaluation/MovieSceneEvaluationTemplateInstance.cpp
+++ b/Engine/Source/Runtime/MovieScene/Private/Evaluation/MovieSceneEvaluationTemplateInstance.cpp
@@ -38,18 +38,18 @@ void FMovieSceneRootEvaluationTemplateInstance::TearDown()
 	// Let's clear our shared pointer before destroying the sequence instance. This is because the sequence
 	// instance checks that no one is holding onto it, making sure it's destroyed along with it and that nobody
 	// is leaking it.
-	FSharedPlaybackState* SharedPlaybackStatePtr = SharedPlaybackState.Get();
+	FRootInstanceHandle RootInstanceHandle = SharedPlaybackState ?
+		SharedPlaybackState->GetRootInstanceHandle() : FRootInstanceHandle();
 	SharedPlaybackState.Reset();
 
 	// Avoid redundant work if the linker is being destroyed anyway
-	if (SharedPlaybackStatePtr &&
-			SharedPlaybackStatePtr->GetRootInstanceHandle().IsValid() &&
+	if (RootInstanceHandle.IsValid() &&
 			EntitySystemLinker && 
 			IsValidChecked(EntitySystemLinker) && 
 			!EntitySystemLinker->IsUnreachable() && 
 			!EntitySystemLinker->HasAnyFlags(RF_BeginDestroyed))
 	{
-		EntitySystemLinker->DestroyInstanceImmediately(SharedPlaybackStatePtr->GetRootInstanceHandle());
+		EntitySystemLinker->DestroyInstanceImmediately(RootInstanceHandle);
 	}
 
 	EntitySystemLinker = nullptr;
@@ -128,15 +128,19 @@ void FMovieSceneRootEvaluationTemplateInstance::Initialize(UMovieSceneSequence&
 	if (bReinitialize)
 	{
 		// Tear down previous sequence first.
-		if (SharedPlaybackState.IsValid() && SharedPlaybackState->GetRootInstanceHandle().IsValid())
+		FRootInstanceHandle OldRootInstanceHandle = SharedPlaybackState ?
+			SharedPlaybackState->GetRootInstanceHandle() : FRootInstanceHandle();
+		SharedPlaybackState.Reset();
+
+		if (OldRootInstanceHandle.IsValid())
 		{
 			if (PreviousRunner)
 			{
-				PreviousRunner->AbandonAndDestroyInstance(SharedPlaybackState->GetRootInstanceHandle());
+				PreviousRunner->AbandonAndDestroyInstance(OldRootInstanceHandle);
 			}
 			else if (EntitySystemLinker)
 			{
-				EntitySystemLinker->GetInstanceRegistry()->DestroyInstance(SharedPlaybackState->GetRootInstanceHandle());
+				EntitySystemLinker->GetInstanceRegistry()->DestroyInstance(OldRootInstanceHandle);
 			}
 			else
 			{
@@ -152,23 +156,19 @@ void FMovieSceneRootEvaluationTemplateInstance::Initialize(UMovieSceneSequence&
 		EntitySystemLinker = ConstructEntityLinker(Player);
 
 		// Create the new root instance and save its new shared playback state.
-		FRootInstanceHandle RootInstanceHandle;
+		FRootInstanceHandle NewRootInstanceHandle;
 		if (EntitySystemLinker != nullptr && EntitySystemLinker->GetInstanceRegistry())
 		{
 			UObject* PlaybackContext = Player.GetPlaybackContext();
 			FInstanceRegistry* InstanceRegistry = EntitySystemLinker->GetInstanceRegistry();
-			RootInstanceHandle = InstanceRegistry->AllocateRootInstance(InRootSequence, PlaybackContext, InCompiledDataManager);
-			SharedPlaybackState = InstanceRegistry->GetInstance(RootInstanceHandle).GetSharedPlaybackState();
+			NewRootInstanceHandle = InstanceRegistry->AllocateRootInstance(InRootSequence, PlaybackContext, InCompiledDataManager);
+			SharedPlaybackState = InstanceRegistry->GetInstance(NewRootInstanceHandle).GetSharedPlaybackState();
 			Player.InitializeRootInstance(SharedPlaybackState.ToSharedRef());
 		}
-		else
-		{
-			SharedPlaybackState.Reset();
-		}
 
-		if (RootInstanceHandle.IsValid())
+		if (NewRootInstanceHandle.IsValid())
 		{
-			Player.PreAnimatedState.Initialize(EntitySystemLinker, RootInstanceHandle);
+			Player.PreAnimatedState.Initialize(EntitySystemLinker, NewRootInstanceHandle);
 		}
 	}
 }
