PR #11547: [MovieScene] Fixed blend range check

diff --git a/Engine/Source/Runtime/MovieScene/Private/Channels/MovieSceneCurveChannelImpl.cpp b/Engine/Source/Runtime/MovieScene/Private/Channels/MovieSceneCurveChannelImpl.cpp
--- a/Engine/Source/Runtime/MovieScene/Private/Channels/MovieSceneCurveChannelImpl.cpp
+++ b/Engine/Source/Runtime/MovieScene/Private/Channels/MovieSceneCurveChannelImpl.cpp
@@ -950,7 +950,7 @@ float TMovieSceneCurveChannelImpl<ChannelType>::CalcSmartTangent(ChannelType* In
 			0.0 : NextKey.Tangent.ArriveTangent;
 		float PrevTangent = (PrevKey.InterpMode == RCIM_Cubic && (PrevKey.TangentMode == RCTM_Auto || PrevKey.TangentMode == RCTM_SmartAuto)) ?
 			0.0 : PrevKey.Tangent.LeaveTangent;
-		if (BlendToNextRange >= 0.0f || BlendToNextRange <= 1.0f)
+		if (BlendToNextRange >= 0.0f && BlendToNextRange <= 1.0f)
 		{
 			NextTangent = PrevTangent = 0.0;
 			if (PercDiff > BlendToNextRange)
