Sequencer: Warn once on zero duration playback.

Note: If the duration later becomes valid, allow the warning to occur again if the duration becomes zero again.

diff --git a/Engine/Source/Runtime/MovieScene/Private/MovieSceneSequencePlayer.cpp b/Engine/Source/Runtime/MovieScene/Private/MovieSceneSequencePlayer.cpp
--- a/Engine/Source/Runtime/MovieScene/Private/MovieSceneSequencePlayer.cpp
+++ b/Engine/Source/Runtime/MovieScene/Private/MovieSceneSequencePlayer.cpp
@@ -170,6 +170,7 @@ UMovieSceneSequencePlayer::UMovieSceneSequencePlayer(const FObjectInitializer& I
 	, bIsAsyncUpdate(false)
 	, bSkipNextUpdate(false)
 	, bUpdateNetSync(false)
+	, bWarnZeroDuration(true)
 	, Sequence(nullptr)
 	, StartTime(0)
 	, DurationFrames(0)
@@ -1100,11 +1101,17 @@ void UMovieSceneSequencePlayer::UpdateTimeCursorPosition_Internal(FFrameTime New
 	EMovieScenePlayerStatus::Type StatusOverride = UpdateMethodToStatus(Method);
 
 	const int32 Duration = DurationFrames;
-	if (Duration == 0)
+	if (Duration == 0 && DurationSubFrames == 0.f)
 	{
-		UE_LOG(LogMovieScene, Warning, TEXT("Attempting to play back a sequence with zero duration"));
+		if (bWarnZeroDuration)
+		{
+			bWarnZeroDuration = false;
+			const FString SequenceName = GetSequenceName(true);
+			UE_LOG(LogMovieScene, Warning, TEXT("Attempting to play back sequence %s with zero duration"), *SequenceName);
+		}
 		return;
 	}
+	bWarnZeroDuration = true;
 	
 	if (bPendingOnStartedPlaying)
 	{
diff --git a/Engine/Source/Runtime/MovieScene/Public/MovieSceneSequencePlayer.h b/Engine/Source/Runtime/MovieScene/Public/MovieSceneSequencePlayer.h
--- a/Engine/Source/Runtime/MovieScene/Public/MovieSceneSequencePlayer.h
+++ b/Engine/Source/Runtime/MovieScene/Public/MovieSceneSequencePlayer.h
@@ -663,6 +663,9 @@ protected:
 	/** Flag that notifies the player to check network synchronization on next update */
 	uint32 bUpdateNetSync : 1;
 
+	/** Flag that indicates whether to warn on zero duration playback */
+	uint32 bWarnZeroDuration : 1;
+
 	/** The sequence to play back */
 	UPROPERTY(transient)
 	TObjectPtr<UMovieSceneSequence> Sequence;
