Sequencer: Warn on zero/invalid frame rate

diff --git a/Engine/Source/Runtime/MovieScene/Private/MovieSceneSequencePlayer.cpp b/Engine/Source/Runtime/MovieScene/Private/MovieSceneSequencePlayer.cpp
--- a/Engine/Source/Runtime/MovieScene/Private/MovieSceneSequencePlayer.cpp
+++ b/Engine/Source/Runtime/MovieScene/Private/MovieSceneSequencePlayer.cpp
@@ -601,6 +601,13 @@ int32 UMovieSceneSequencePlayer::GetFrameDuration() const
 
 void UMovieSceneSequencePlayer::SetFrameRate(FFrameRate FrameRate)
 {
+	if (!FrameRate.IsValid() || FrameRate.Numerator <= 0)
+	{
+		const FString SequenceName = GetSequenceName(true);
+		UE_LOG(LogMovieScene, Error, TEXT("Attempting to set sequence %s with an invalid frame rate: %s"), *SequenceName, *FrameRate.ToPrettyText().ToString());
+		return;
+	}
+
 	UMovieScene* MovieScene = Sequence ? Sequence->GetMovieScene() : nullptr;
 	if (MovieScene)
 	{
@@ -894,7 +901,23 @@ void UMovieSceneSequencePlayer::Initialize(UMovieSceneSequence* InSequence)
 		FFrameRate                TickResolution    = MovieScene->GetTickResolution();
 		FFrameRate                DisplayRate       = MovieScene->GetDisplayRate();
 
-		UE_LOG(LogMovieScene, Verbose, TEXT("Initialize - MovieSceneSequence: %s, TickResolution: %d, DisplayRate: %d"), *InSequence->GetName(), TickResolution.Numerator, DisplayRate.Numerator);
+		UE_LOG(LogMovieScene, Verbose, TEXT("Initialize - MovieSceneSequence: %s, TickResolution: %s, DisplayRate: %s"), *InSequence->GetName(), *TickResolution.ToPrettyText().ToString(), *DisplayRate.ToPrettyText().ToString());
+
+		if (!TickResolution.IsValid() || TickResolution.Numerator <= 0)
+		{
+			const FString SequenceName = GetSequenceName(true);
+			const FFrameRate DefaultTickResolution(60000, 1);
+			UE_LOG(LogMovieScene, Error, TEXT("Attempting to set sequence %s with an invalid tick resolution: %s, defaulting to: %s"), *SequenceName, *TickResolution.ToPrettyText().ToString(), *DefaultTickResolution.ToPrettyText().ToString());
+			TickResolution = DefaultTickResolution;
+		}
+
+		if (!DisplayRate.IsValid() || DisplayRate.Numerator <= 0)
+		{
+			const FString SequenceName = GetSequenceName(true);
+			const FFrameRate DefaultDisplayRate(30, 1);
+			UE_LOG(LogMovieScene, Error, TEXT("Attempting to set sequence %s with an invalid display rate: %s, defaulting to: %s"), *SequenceName, *DisplayRate.ToPrettyText().ToString(), *DefaultDisplayRate.ToPrettyText().ToString());
+			DisplayRate = DefaultDisplayRate;
+		}
 
 		// We set the play position in terms of the display rate,
 		// but want evaluation ranges in the moviescene's tick resolution
