Sequencer- Refactor Bindings to support inherited Custom Bindings.

diff --git a/Engine/Source/Runtime/MovieScene/Private/EntitySystem/MovieSceneBindingLifetimeSystem.cpp b/Engine/Source/Runtime/MovieScene/Private/EntitySystem/MovieSceneBindingLifetimeSystem.cpp
--- a/Engine/Source/Runtime/MovieScene/Private/EntitySystem/MovieSceneBindingLifetimeSystem.cpp
+++ b/Engine/Source/Runtime/MovieScene/Private/EntitySystem/MovieSceneBindingLifetimeSystem.cpp
@@ -25,50 +25,6 @@
 #include UE_INLINE_GENERATED_CPP_BY_NAME(MovieSceneBindingLifetimeSystem)
 #define LOCTEXT_NAMESPACE "MovieSceneBindingLifetimeSystem"
 
-namespace UE
-{
-	namespace MovieScene
-	{
-		const FMovieSceneAnimTypeID BindingLifetimeAnimTypeID = FMovieSceneAnimTypeID::Unique();
-
-		struct FBindingLifetimeTrackPreAnimatedTokenProducer : IMovieScenePreAnimatedTokenProducer
-		{
-			FMovieSceneEvaluationOperand Operand;
-			FBindingLifetimeTrackPreAnimatedTokenProducer(FMovieSceneEvaluationOperand InOperand) : Operand(InOperand) {}
-
-			virtual IMovieScenePreAnimatedTokenPtr CacheExistingState(UObject& Object) const
-			{
-				struct FToken : IMovieScenePreAnimatedToken
-				{
-					FMovieSceneEvaluationOperand OperandToDestroy;
-					FToken(FMovieSceneEvaluationOperand InOperand) : OperandToDestroy(InOperand) {}
-
-					virtual void RestoreState(UObject& Object, const UE::MovieScene::FRestoreStateParams& Params) override
-					{
-						TSharedPtr<const FSharedPlaybackState> PlaybackState = Params.GetTerminalPlaybackState();
-						if (!ensure(PlaybackState))
-						{
-							return;
-						}
-
-						// Destroy any loaded bindings
-						if (FMovieSceneEvaluationState* EvaluationState = PlaybackState->FindCapability<FMovieSceneEvaluationState>())
-						{
-							if (FMovieSceneObjectCache* Cache = EvaluationState->FindObjectCache(OperandToDestroy.SequenceID))
-							{
-								Cache->UnloadBinding(OperandToDestroy.ObjectBindingID, PlaybackState.ToSharedRef());
-							}
-						}
-					}
-				};
-
-				return FToken(Operand);
-			}
-		};
-
-	} // namespace MovieScene
-} // namespace UE
-
 
 UMovieSceneBindingLifetimeSystem::UMovieSceneBindingLifetimeSystem(const FObjectInitializer& ObjInit)
 	: Super(ObjInit)
@@ -85,11 +41,6 @@ UMovieSceneBindingLifetimeSystem::UMovieSceneBindingLifetimeSystem(const FObject
 	}
 }
 
-FMovieSceneAnimTypeID UMovieSceneBindingLifetimeSystem::GetAnimTypeID()
-{
-	return UE::MovieScene::BindingLifetimeAnimTypeID;
-}
-
 void UMovieSceneBindingLifetimeSystem::OnRun(FSystemTaskPrerequisites& InPrerequisites, FSystemSubsequentTasks& Subsequents)
 {
 	using namespace UE::MovieScene;
@@ -137,19 +88,6 @@ void UMovieSceneBindingLifetimeSystem::OnRun(FSystemTaskPrerequisites& InPrerequ
 								}
 							}
 						}
-
-						if (bLink)
-						{
-							Player->PreAnimatedState.SavePreAnimatedState(*BoundObject, BindingLifetimeAnimTypeID, FBindingLifetimeTrackPreAnimatedTokenProducer({SequenceID, BindingLifetime.BindingGuid}));
-						}
-					}
-				}
-				if (!bLink)
-				{
-					// In addition, on unlink of the active range, unload any loaded objects with that binding id
-					if (FMovieSceneObjectCache* Cache = Player->State.FindObjectCache(SequenceID))
-					{
-						Cache->UnloadBinding(BindingLifetime.BindingGuid, Player->GetSharedPlaybackState());
 					}
 				}
 			}
diff --git a/Engine/Source/Runtime/MovieScene/Private/EntitySystem/MovieSceneSpawnablesSystem.cpp b/Engine/Source/Runtime/MovieScene/Private/EntitySystem/MovieSceneSpawnablesSystem.cpp
--- a/Engine/Source/Runtime/MovieScene/Private/EntitySystem/MovieSceneSpawnablesSystem.cpp
+++ b/Engine/Source/Runtime/MovieScene/Private/EntitySystem/MovieSceneSpawnablesSystem.cpp
@@ -17,6 +17,8 @@
 
 #include "IMovieScenePlayer.h"
 #include "IMovieScenePlaybackClient.h"
+#include "MovieSceneBindingReferences.h"
+#include "Bindings/MovieSceneSpawnableBinding.h"
 
 #include UE_INLINE_GENERATED_CPP_BY_NAME(MovieSceneSpawnablesSystem)
 
@@ -32,14 +34,16 @@ DECLARE_CYCLE_STAT(TEXT("Spawnables: Destroy"), MovieSceneEval_DestroySpawnables
 struct FSpawnTrackPreAnimatedTokenProducer : IMovieScenePreAnimatedTokenProducer
 {
 	FMovieSceneEvaluationOperand Operand;
-	FSpawnTrackPreAnimatedTokenProducer(FMovieSceneEvaluationOperand InOperand) : Operand(InOperand) {}
+	int32 BindingIndex;
+	FSpawnTrackPreAnimatedTokenProducer(FMovieSceneEvaluationOperand InOperand, int32 InBindingIndex) : Operand(InOperand), BindingIndex(InBindingIndex) {}
 
 	virtual IMovieScenePreAnimatedTokenPtr CacheExistingState(UObject& Object) const
 	{
 		struct FToken : IMovieScenePreAnimatedToken
 		{
 			FMovieSceneEvaluationOperand OperandToDestroy;
-			FToken(FMovieSceneEvaluationOperand InOperand) : OperandToDestroy(InOperand) {}
+			int32 BindingIndexToDestroy;
+			FToken(FMovieSceneEvaluationOperand InOperand, int32 InBindingIndex) : OperandToDestroy(InOperand), BindingIndexToDestroy(InBindingIndex) {}
 
 			virtual void RestoreState(UObject& Object, const UE::MovieScene::FRestoreStateParams& Params) override
 			{
@@ -50,7 +54,7 @@ struct FSpawnTrackPreAnimatedTokenProducer : IMovieScenePreAnimatedTokenProducer
 				}
 
 				FMovieSceneSpawnRegister* SpawnRegister = PlaybackState->FindCapability<FMovieSceneSpawnRegister>();
-				if (SpawnRegister && !SpawnRegister->DestroySpawnedObject(OperandToDestroy.ObjectBindingID, OperandToDestroy.SequenceID, PlaybackState.ToSharedRef()))
+				if (SpawnRegister && !SpawnRegister->DestroySpawnedObject(OperandToDestroy.ObjectBindingID, OperandToDestroy.SequenceID, PlaybackState.ToSharedRef(), BindingIndexToDestroy))
 				{
 					// This branch should only be taken for Externally owned spawnables that have been 'forgotten',
 					// but still had RestoreState tokens generated for them (ie, in FSequencer, or if bRestoreState is enabled)
@@ -60,7 +64,7 @@ struct FSpawnTrackPreAnimatedTokenProducer : IMovieScenePreAnimatedTokenProducer
 			}
 		};
 		
-		return FToken(Operand);
+		return FToken(Operand, BindingIndex);
 	}
 };
 
@@ -123,8 +127,8 @@ void UMovieSceneSpawnablesSystem::OnRun(FSystemTaskPrerequisites& InPrerequisite
 	}
 
 	// Used below.
-	TArray<TTuple<FGuid, FMovieSceneSequenceID, FInstanceHandle>> DestroyedObjects;
-	auto DestroyOldSpawnables = [&DestroyedObjects, InstanceRegistry](FInstanceHandle InstanceHandle, const FGuid& SpawnableObjectID)
+	TArray<TTuple<FGuid, int32, FMovieSceneSequenceID, FInstanceHandle>> DestroyedObjects;
+	auto DestroyOldSpawnables = [this, &DestroyedObjects, InstanceRegistry](FMovieSceneEntityID EntityID, FInstanceHandle InstanceHandle, const FGuid& SpawnableObjectID)
 	{
 		SCOPE_CYCLE_COUNTER(MovieSceneEval_DestroySpawnables)
 
@@ -133,25 +137,80 @@ void UMovieSceneSpawnablesSystem::OnRun(FSystemTaskPrerequisites& InPrerequisite
 			const FSequenceInstance& Instance = InstanceRegistry->GetInstance(InstanceHandle);
 			TSharedRef<FSharedPlaybackState> SharedPlaybackState = Instance.GetSharedPlaybackState();
 
-			// If the sequence instance has finished and it is a sub sequence, we do not destroy the spawnable
-			// if it is owned by the root sequence or externally. These will get destroyed or forgotten by the player when it ends
-			if (Instance.HasFinished() && Instance.IsSubSequence())
+			const UMovieSceneSequence* Sequence = SharedPlaybackState->GetSequence(Instance.GetSequenceID());
+			if (Sequence)
 			{
-				const UMovieSceneSequence* Sequence = SharedPlaybackState->GetSequence(Instance.GetSequenceID());
-				FMovieSceneSpawnable* Spawnable = Sequence ? Sequence->GetMovieScene()->FindSpawnable(SpawnableObjectID) : nullptr;
-				if (!Spawnable || Spawnable->GetSpawnOwnership() != ESpawnOwnership::InnerSequence)
+				if (Linker->EntityManager.HasComponent(EntityID, FBuiltInComponentTypes::Get()->Tags.OldStyleSpawnable))
 				{
-					return;
+					if (FMovieSceneSpawnable* Spawnable = Sequence->GetMovieScene()->FindSpawnable(SpawnableObjectID))
+					{
+						// If the sequence instance has finished and it is a sub sequence, we do not destroy the spawnable
+						// if it is owned by the root sequence or externally. These will get destroyed or forgotten by the player when it ends
+						if (Instance.HasFinished() && Instance.IsSubSequence())
+						{
+							if (Spawnable->GetSpawnOwnership() != ESpawnOwnership::InnerSequence)
+							{
+								return;
+							}
+						}
+
+						DestroyedObjects.Emplace(SpawnableObjectID, 0, Instance.GetSequenceID(), InstanceHandle);
+					}
+				}
+				else if (const FMovieSceneBindingReferences* BindingReferences = Sequence->GetBindingReferences())
+				{
+					for (int32 Index = 0; Index < BindingReferences->GetReferences(SpawnableObjectID).Num(); ++Index)
+					{
+						if (const UMovieSceneCustomBinding* CustomBinding = BindingReferences->GetCustomBinding(SpawnableObjectID, Index))
+						{
+							if (const UMovieSceneSpawnableBindingBase* SpawnableBinding = CustomBinding->AsSpawnable(SharedPlaybackState))
+							{
+								// If the sequence instance has finished and it is a sub sequence, we do not destroy the spawnable
+								// if it is owned by the root sequence or externally. These will get destroyed or forgotten by the player when it ends
+								if (Instance.HasFinished() && Instance.IsSubSequence())
+								{
+									if (SpawnableBinding->SpawnOwnership != ESpawnOwnership::InnerSequence)
+									{
+										return;
+									}
+								}
+								DestroyedObjects.Emplace(SpawnableObjectID, Index, Instance.GetSequenceID(), InstanceHandle);
+							}
+						}
+					}
 				}
 			}
-
-			DestroyedObjects.Emplace(SpawnableObjectID, Instance.GetSequenceID(), InstanceHandle);
 		}
 	};
 
 	// ----------------------------------------------------------------------------------------------------------------------------------------
-	// Step 1 - iterate all pending spawnables and spawn their objects if necessary
-	auto SpawnNewObjects = [&DestroyOldSpawnables, InstanceRegistry](FInstanceHandle InstanceHandle, const FGuid& SpawnableBindingID)
+	// Step 1 - destroy any spawnable objects that are no longer relevant
+	//          NOTE: We gather the objects into an array because destroying an object can potentially cause a garbage collection to run (ie
+	//                if the spawnable is a level instance), and that could assert because we are currently iterating the ECS
+
+	FEntityTaskBuilder()
+		.ReadEntityIDs()
+		.Read(BuiltInComponents->InstanceHandle)
+		.Read(BuiltInComponents->SpawnableBinding)
+		.FilterAll({ BuiltInComponents->Tags.NeedsUnlink })
+		.Iterate_PerEntity(&Linker->EntityManager, DestroyOldSpawnables);
+
+	for (TTuple<FGuid, int32, FMovieSceneSequenceID, FInstanceHandle> Tuple : DestroyedObjects)
+	{
+		// Have to check whether the player is still valid because there is a possibility it got cleaned up
+		if (InstanceRegistry->IsHandleValid(Tuple.Get<3>()))
+		{
+			TSharedRef<const FSharedPlaybackState> SharedPlaybackState = InstanceRegistry->GetInstance(Tuple.Get<3>()).GetSharedPlaybackState();
+			if (FMovieSceneSpawnRegister* SpawnRegister = SharedPlaybackState->FindCapability<FMovieSceneSpawnRegister>())
+			{
+				SpawnRegister->DestroySpawnedObject(Tuple.Get<0>(), Tuple.Get<2>(), SharedPlaybackState, Tuple.Get<1>());
+			}
+		}
+	}
+
+	// ----------------------------------------------------------------------------------------------------------------------------------------
+	// Step 2 - iterate all pending spawnables and spawn their objects if necessary
+	auto SpawnNewObjects = [this, &DestroyOldSpawnables, InstanceRegistry](FMovieSceneEntityID EntityID, FInstanceHandle InstanceHandle, const FGuid& SpawnableBindingID)
 	{
 		SCOPE_CYCLE_COUNTER(MovieSceneEval_SpawnSpawnables)
 
@@ -181,11 +240,8 @@ void UMovieSceneSpawnablesSystem::OnRun(FSystemTaskPrerequisites& InPrerequisite
 		{
 			// If the binding is currently inactive, don't spawn the object.
 
-			// If we have an existing spawned object, then we need to destroy the spawned object here.
-			if (ExistingSpawnedObject)
-			{
-				DestroyOldSpawnables(InstanceHandle, SpawnableBindingID);
-			}
+			// We may have existing spawned objects, in which case we need to destory them
+			DestroyOldSpawnables(EntityID, InstanceHandle, SpawnableBindingID);
 			return;
 		}
 
@@ -196,70 +252,55 @@ void UMovieSceneSpawnablesSystem::OnRun(FSystemTaskPrerequisites& InPrerequisite
 			bool bUseDefaultBinding = DynamicOverrides->RetrieveBindingOverrides(SpawnableBindingID, SequenceID, FoundObjects);
 			if (!bUseDefaultBinding)
 			{
-				// If the binding has been overridden but we have an existing spawned object, then the binding is new and we need to destroy the spawned object.
-				if (ExistingSpawnedObject)
-				{
-					DestroyOldSpawnables(InstanceHandle, SpawnableBindingID);
-				}
+				// If the binding has been overridden, then the binding is new and we need to destroy any spawned objects that may exist
+				DestroyOldSpawnables(EntityID, InstanceHandle, SpawnableBindingID);
 
 				// This spawnable is overridden so don't try and spawn anything
 				return;
 			}
 		}
 
-		// Check if we already have a spawned object in the spawn register - if we have we use that
-		if (ExistingSpawnedObject)
-		{
-			return;
-		}
-
-		// At this point we've decided that we need to spawn a whole new object
+		// At this point we've decided that we should have a spawned object
 		const UMovieSceneSequence* Sequence = SharedPlaybackState->GetSequence(SequenceID);
 		if (!Sequence)
 		{
 			return;
 		}
 
-		UObject* SpawnedObject = SpawnRegister->SpawnObject(SpawnableBindingID, *Sequence->GetMovieScene(), SequenceID, SharedPlaybackState);
-		IMovieScenePlayer* Player = FPlayerIndexPlaybackCapability::GetPlayer(SharedPlaybackState);
-		if (SpawnedObject && Player)
+		auto SpawnObject = [&SpawnRegister, &SpawnableBindingID, &Sequence, &SequenceID, &SharedPlaybackState](int BindingIndex)
 		{
-			FMovieSceneEvaluationOperand Operand(SequenceID, SpawnableBindingID);
-			Player->OnObjectSpawned(SpawnedObject, Operand);
+			UObject* SpawnedObject = SpawnRegister->SpawnObject(SpawnableBindingID, *Sequence->GetMovieScene(), SequenceID, SharedPlaybackState, BindingIndex);
+			IMovieScenePlayer* Player = FPlayerIndexPlaybackCapability::GetPlayer(SharedPlaybackState);
+			if (SpawnedObject && Player)
+			{
+				FMovieSceneEvaluationOperand Operand(SequenceID, SpawnableBindingID);
+				Player->OnObjectSpawned(SpawnedObject, Operand);
 
-			Player->SavePreAnimatedState(*SpawnedObject, SpawnableAnimTypeID, FSpawnTrackPreAnimatedTokenProducer(Operand));
+				Player->SavePreAnimatedState(*SpawnedObject, SpawnableAnimTypeID, FSpawnTrackPreAnimatedTokenProducer(Operand, BindingIndex));
+			}
+			return SpawnedObject;
+		};
+
+		// If we have an old-style spawnable, we can just spawn a single object, otherwise we need to check our binding references to see if we need to spawn multiple
+		if (Linker->EntityManager.HasComponent(EntityID, FBuiltInComponentTypes::Get()->Tags.OldStyleSpawnable))
+		{
+			SpawnObject(0);
+		}
+		else if (const FMovieSceneBindingReferences* BindingReferences = Sequence->GetBindingReferences())
+		{
+			for (int32 Index = 0; Index < BindingReferences->GetReferences(SpawnableBindingID).Num(); ++Index)
+			{
+				SpawnObject(Index);
+			}
 		}
 
 		InstanceRegistry->InvalidateObjectBinding(SpawnableBindingID, InstanceHandle);
 	};
 	FEntityTaskBuilder()
+	.ReadEntityIDs()
 	.Read(BuiltInComponents->InstanceHandle)
 	.Read(BuiltInComponents->SpawnableBinding)
 	.FilterAll({ BuiltInComponents->Tags.NeedsLink })
 	.Iterate_PerEntity(&Linker->EntityManager, SpawnNewObjects);
-
-	// ----------------------------------------------------------------------------------------------------------------------------------------
-	// Step 2 - destroy any spawnable objects that are no longer relevant
-	//          NOTE: We gather the objects into an array because destroying an object can potentially cause a garbage collection to run (ie
-	//                if the spawnable is a level instance), and that could assert because we are currently iterating the ECS
-
-	FEntityTaskBuilder()
-	.Read(BuiltInComponents->InstanceHandle)
-	.Read(BuiltInComponents->SpawnableBinding)
-	.FilterAll({ BuiltInComponents->Tags.NeedsUnlink })
-	.Iterate_PerEntity(&Linker->EntityManager, DestroyOldSpawnables);
-
-	for (TTuple<FGuid, FMovieSceneSequenceID, FInstanceHandle> Tuple : DestroyedObjects)
-	{
-		// Have to check whether the player is still valid because there is a possibility it got cleaned up
-		if (InstanceRegistry->IsHandleValid(Tuple.Get<2>()))
-		{
-			TSharedRef<const FSharedPlaybackState> SharedPlaybackState = InstanceRegistry->GetInstance(Tuple.Get<2>()).GetSharedPlaybackState();
-			if (FMovieSceneSpawnRegister* SpawnRegister = SharedPlaybackState->FindCapability<FMovieSceneSpawnRegister>())
-			{
-				SpawnRegister->DestroySpawnedObject(Tuple.Get<0>(), Tuple.Get<1>(), SharedPlaybackState);
-			}
-		}
-	}
 }
 
diff --git a/Engine/Source/Runtime/MovieScene/Private/Evaluation/MovieSceneEvaluationState.cpp b/Engine/Source/Runtime/MovieScene/Private/Evaluation/MovieSceneEvaluationState.cpp
--- a/Engine/Source/Runtime/MovieScene/Private/Evaluation/MovieSceneEvaluationState.cpp
+++ b/Engine/Source/Runtime/MovieScene/Private/Evaluation/MovieSceneEvaluationState.cpp
@@ -15,7 +15,6 @@
 #include "MovieSceneObjectBindingID.h"
 #include "MovieSceneSequence.h"
 #include "UniversalObjectLocatorResolveParams.h"
-#include "UniversalObjectLocatorResolveParameterBuffer.inl"
 
 DECLARE_CYCLE_STAT(TEXT("Find Bound Objects"), MovieSceneEval_FindBoundObjects, STATGROUP_MovieSceneEval);
 DECLARE_CYCLE_STAT(TEXT("Iterate Bound Objects"), MovieSceneEval_IterateBoundObjects, STATGROUP_MovieSceneEval);
@@ -314,43 +313,6 @@ bool FMovieSceneObjectCache::InvalidateInternal(const FGuid& InGuid)
 	return true;
 }
 
-void FMovieSceneObjectCache::UnloadBinding(const FGuid& Guid, TSharedRef<const FSharedPlaybackState> SharedPlaybackState)
-{
-	// Invalidate binding, forcing it to be reloaded
-	InvalidateInternal(Guid);
-
-	TArray<FMovieSceneLocatorSpawnedCacheKey, TInlineAllocator<1>> LoadedCacheKeys;
-	Algo::TransformIf(LoadedBindingIds,
-		LoadedCacheKeys,
-		[Guid](const TPair<FMovieSceneLocatorSpawnedCacheKey, TWeakObjectPtr<>>  Pair) { return Pair.Key.BindingID == Guid; },
-		[Guid](const TPair<FMovieSceneLocatorSpawnedCacheKey, TWeakObjectPtr<>> Pair) { return Pair.Key; });
-	
-	for (const FMovieSceneLocatorSpawnedCacheKey& CacheKey : LoadedCacheKeys)
-	{
-		UnloadBindingInternal(CacheKey, SharedPlaybackState);
-	}
-}
-	
-
-void FMovieSceneObjectCache::UnloadBindingInternal(const FMovieSceneLocatorSpawnedCacheKey& CacheKey, TSharedRef<const FSharedPlaybackState> SharedPlaybackState)
-{
-	if (LoadedBindingIds.Contains(CacheKey))
-	{
-		UMovieSceneSequence* Sequence = WeakSequence.Get();
-		if (!Sequence)
-		{
-			return;
-		}
-		UE::UniversalObjectLocator::TResolveParamsWithBuffer<128> ResolveParams(SharedPlaybackState->GetPlaybackContext(), ELocatorResolveFlags::Unload);
-		ResolveParams.AddParameter(FLocatorSpawnedCacheResolveParameter::ParameterType, this);
-
-		SetResolvingBindingCacheKey(CacheKey);
-		Sequence->UnloadBoundObject(ResolveParams, CacheKey.BindingID, CacheKey.BindingIndex);
-		ClearResolvingBindingCacheKey();
-		LoadedBindingIds.Remove(CacheKey);
-	}
-}
-
 void FMovieSceneObjectCache::Clear(TSharedRef<const FSharedPlaybackState> SharedPlaybackState)
 {
 	using namespace UE::MovieScene;
@@ -387,8 +349,6 @@ void FMovieSceneObjectCache::UpdateBindings(const FGuid& InGuid, TSharedRef<cons
 	TGuardValue<bool> ReentrancyGuard(bReentrantUpdate, true);
 
 	// Invalidate existing bindings, we're going to rebuild them.
-
-	// Unload any objects that have been loaded during original resolution
 	FBoundObjects* Bindings = &BoundObjects.FindOrAdd(InGuid);
 	Bindings->Objects.Reset();
 
@@ -416,12 +376,6 @@ void FMovieSceneObjectCache::UpdateBindings(const FGuid& InGuid, TSharedRef<cons
 		return;
 	}
 
-	if (CurrentlyResolvingCacheKey.IsValid())
-	{
-		// We're getting called recursively to resolve something. Return to avoid looping.
-		return;
-	}
-
 	// If we have overrides for this binding, ask the player to find it for us (most probably in a different cache
 	// for a different sequence).
 	// TODO-lchabant: we could technically end up in a circular override that creates an infinite loop...
@@ -493,9 +447,7 @@ void FMovieSceneObjectCache::UpdateBindings(const FGuid& InGuid, TSharedRef<cons
 						}
 						else
 						{
-							UE::UniversalObjectLocator::TResolveParamsWithBuffer<128> ResolveParams(ResolutionContext);
-							ResolveParams.AddParameter(FLocatorSpawnedCacheResolveParameter::ParameterType, this);
-							SetResolvingBindingCacheKey({ InGuid, 0 });
+							UE::UniversalObjectLocator::FResolveParams ResolveParams(ResolutionContext);
 							if (Player)
 							{
 								Player->ResolveBoundObjects(ResolveParams, InGuid, SequenceID, *Sequence, FoundObjects);
@@ -547,9 +499,7 @@ void FMovieSceneObjectCache::UpdateBindings(const FGuid& InGuid, TSharedRef<cons
 					}
 					else
 					{
-						UE::UniversalObjectLocator::TResolveParamsWithBuffer<128> ResolveParams(ResolutionContext);
-						ResolveParams.AddParameter(FLocatorSpawnedCacheResolveParameter::ParameterType, this);
-						SetResolvingBindingCacheKey({ InGuid, 0 });
+						UE::UniversalObjectLocator::FResolveParams ResolveParams(ResolutionContext);
 						if (Player)
 						{
 							Player->ResolveBoundObjects(ResolveParams, InGuid, SequenceID, *Sequence, FoundObjects);
@@ -570,7 +520,7 @@ void FMovieSceneObjectCache::UpdateBindings(const FGuid& InGuid, TSharedRef<cons
 		}
 		else
 		{
-			// Probably a spawnable then (or an phantom)
+			// Probably a FMovieSceneSpawnable then (or an phantom)
 			bool bUseDefault = true;
 
 			// Allow external overrides for spawnables
@@ -589,7 +539,7 @@ void FMovieSceneObjectCache::UpdateBindings(const FGuid& InGuid, TSharedRef<cons
 			if (bUseDefault)
 			{
 				const FMovieSceneSpawnRegister* SpawnRegister = SharedPlaybackState->FindCapability<FMovieSceneSpawnRegister>();
-				UObject* SpawnedObject = SpawnRegister ? SpawnRegister->FindSpawnedObject(InGuid, SequenceID).Get() : nullptr;
+				UObject* SpawnedObject = SpawnRegister ? SpawnRegister->FindSpawnedObject(InGuid, SequenceID, 0).Get() : nullptr;
 				if (SpawnedObject)
 				{
 					Bindings->Objects.Add(SpawnedObject);
@@ -635,8 +585,6 @@ void FMovieSceneObjectCache::UpdateBindings(const FGuid& InGuid, TSharedRef<cons
 		}
 		ChildBindings.Remove(InGuid);
 	}
-
-	ClearResolvingBindingCacheKey();
 }
 
 void FMovieSceneObjectCache::UpdateSerialNumber()
@@ -675,59 +623,6 @@ void FMovieSceneObjectCache::FilterObjectBindings(UObject* PredicateObject, IMov
 	FilterObjectBindings(PredicateObject, Player.GetSharedPlaybackState(), OutBindings);
 }
 
-UObject* FMovieSceneObjectCache::FindExistingObject()
-{
-	if (CurrentlyResolvingCacheKey.IsValid())
-	{
-		if (TWeakObjectPtr<>* LoadedObjectPtr = LoadedBindingIds.Find(CurrentlyResolvingCacheKey))
-		{
-			return LoadedObjectPtr->Get();
-		}
-	}
-	return nullptr;
-}
-
-
-FName FMovieSceneObjectCache::GetRequestedObjectName()
-{
-#if WITH_EDITOR
-	// Find the object binding display name. If it's currently set to "Empty Binding" or some variant, return nothing, otherwise return the name
-	if (CurrentlyResolvingCacheKey.IsValid())
-	{
-		UMovieSceneSequence* Sequence = WeakSequence.Get();
-		UMovieScene* MovieScene = Sequence ? Sequence->GetMovieScene() : nullptr;
-		if (MovieScene)
-		{
-			if (FMovieScenePossessable* Possessable = MovieScene->FindPossessable(CurrentlyResolvingCacheKey.BindingID))
-			{
-				FString DisplayName = Possessable->GetName();
-				if (!DisplayName.StartsWith(TEXT("Empty Binding")))
-				{
-					return *DisplayName;
-				}
-			}
-		}
-	}
-#endif
-	return FName();
-}
-
-void FMovieSceneObjectCache::ReportSpawnedObject(UObject* Object)
-{
-	if (CurrentlyResolvingCacheKey.IsValid())
-	{
-		LoadedBindingIds.Add(CurrentlyResolvingCacheKey, Object);
-	}
-}
-
-void FMovieSceneObjectCache::SpawnedObjectDestroyed()
-{
-	if (CurrentlyResolvingCacheKey.IsValid())
-	{
-		LoadedBindingIds.Remove(CurrentlyResolvingCacheKey);
-	}
-}
-
 UE::MovieScene::TPlaybackCapabilityID<FMovieSceneEvaluationState> FMovieSceneEvaluationState::ID = UE::MovieScene::TPlaybackCapabilityID<FMovieSceneEvaluationState>::Register();
 
 void FMovieSceneEvaluationState::InvalidateExpiredObjects()
@@ -854,11 +749,6 @@ void FMovieSceneEvaluationState::RegisterObjectCacheEvents(UMovieSceneEntitySyst
 			Linker, &UMovieSceneEntitySystemLinker::InvalidateObjectBinding, InstanceHandle);
 }
 
-bool FMovieSceneEvaluationState::IsResolvingObject() const
-{
-	return Algo::AnyOf(ObjectCaches, [](TPair<FMovieSceneSequenceID, FVersionedObjectCache> Pair) { return Pair.Value.ObjectCache.GetResolvingBindingCacheKey().IsValid(); });
-}
-
 void FMovieSceneEvaluationState::AssignSequence(FMovieSceneSequenceIDRef InSequenceID, UMovieSceneSequence& InSequence, IMovieScenePlayer& Player)
 {
 	AssignSequence(InSequenceID, InSequence, Player.GetSharedPlaybackState());
diff --git a/Engine/Source/Runtime/MovieScene/Private/MovieSceneBindingReferences.cpp b/Engine/Source/Runtime/MovieScene/Private/MovieSceneBindingReferences.cpp
--- a/Engine/Source/Runtime/MovieScene/Private/MovieSceneBindingReferences.cpp
+++ b/Engine/Source/Runtime/MovieScene/Private/MovieSceneBindingReferences.cpp
@@ -2,12 +2,12 @@
 
 #include "MovieSceneBindingReferences.h"
 #include "IMovieSceneBoundObjectProxy.h"
-#include "ILocatorSpawnedCache.h"
 #include "UniversalObjectLocatorResolveParams.h"
 #include "Engine/World.h"
 #include "Evaluation/MovieSceneEvaluationState.h"
 #include "UObject/Package.h"
 #include "UnrealEngine.h"
+#include "Bindings/MovieSceneCustomBinding.h"
 
 #include UE_INLINE_GENERATED_CPP_BY_NAME(MovieSceneBindingReferences)
 
@@ -40,6 +40,11 @@ TArrayView<const FMovieSceneBindingReference> FMovieSceneBindingReferences::GetA
 	return SortedReferences;
 }
 
+TArrayView<FMovieSceneBindingReference> FMovieSceneBindingReferences::GetAllReferences() 
+{
+	return SortedReferences;
+}
+
 TArrayView<const FMovieSceneBindingReference> FMovieSceneBindingReferences::GetReferences(const FGuid& ObjectId) const
 {
 	const int32 Num   = SortedReferences.Num();
@@ -55,12 +60,37 @@ TArrayView<const FMovieSceneBindingReference> FMovieSceneBindingReferences::GetR
 	return TArrayView<const FMovieSceneBindingReference>(SortedReferences.GetData() + Index, MatchNum);
 }
 
+const FMovieSceneBindingReference* FMovieSceneBindingReferences::GetReference(const FGuid& ObjectId, int32 BindingIndex) const
+{
+	const int32 Index = Algo::LowerBoundBy(SortedReferences, ObjectId, &FMovieSceneBindingReference::ID) + BindingIndex;
+	if (SortedReferences.IsValidIndex(Index) && SortedReferences[Index].ID == ObjectId)
+	{
+		return &SortedReferences[Index];
+	}
+	return nullptr;
+}
+
 bool FMovieSceneBindingReferences::HasBinding(const FGuid& ObjectId) const
 {
 	const int32 Index = Algo::LowerBoundBy(SortedReferences, ObjectId, &FMovieSceneBindingReference::ID);
 	return SortedReferences.IsValidIndex(Index) && SortedReferences[Index].ID == ObjectId;
 }
 
+UMovieSceneCustomBinding* FMovieSceneBindingReferences::GetCustomBinding(const FGuid& ObjectId, int32 BindingIndex)
+{
+	return const_cast<UMovieSceneCustomBinding*>(const_cast<const FMovieSceneBindingReferences*>(this)->GetCustomBinding(ObjectId, BindingIndex));
+}
+
+const UMovieSceneCustomBinding* FMovieSceneBindingReferences::GetCustomBinding(const FGuid& ObjectId, int32 BindingIndex) const
+{
+	const int32 Index = Algo::LowerBoundBy(SortedReferences, ObjectId, &FMovieSceneBindingReference::ID) + BindingIndex;
+	if (SortedReferences.IsValidIndex(Index) && SortedReferences[Index].ID == ObjectId)
+	{
+		return SortedReferences[Index].CustomBinding;
+	}
+	return nullptr;
+}
+
 const FMovieSceneBindingReference* FMovieSceneBindingReferences::AddBinding(const FGuid& ObjectId, FUniversalObjectLocator&& NewLocator)
 {
 	const int32 Index = Algo::UpperBoundBy(SortedReferences, ObjectId, &FMovieSceneBindingReference::ID);
@@ -70,16 +100,52 @@ const FMovieSceneBindingReference* FMovieSceneBindingReferences::AddBinding(cons
 	return &NewBinding;
 }
 
-const FMovieSceneBindingReference* FMovieSceneBindingReferences::AddBinding(const FGuid& ObjectId, FUniversalObjectLocator&& NewLocator, ELocatorResolveFlags InResolveFlags)
+const FMovieSceneBindingReference* FMovieSceneBindingReferences::AddBinding(const FGuid& ObjectId, FUniversalObjectLocator&& NewLocator, ELocatorResolveFlags InResolveFlags, UMovieSceneCustomBinding* CustomBinding/*=nullptr*/)
 {
 	const int32 Index = Algo::UpperBoundBy(SortedReferences, ObjectId, &FMovieSceneBindingReference::ID);
 
-	FMovieSceneBindingReference& NewBinding = SortedReferences.Insert_GetRef(FMovieSceneBindingReference{ ObjectId, MoveTemp(NewLocator) }, Index);
-	NewBinding.ResolveFlags = InResolveFlags;
+	FMovieSceneBindingReference& NewBinding = SortedReferences.Insert_GetRef(FMovieSceneBindingReference{ ObjectId, MoveTemp(NewLocator), InResolveFlags, CustomBinding }, Index);
 
 	return &NewBinding;
 }
 
+const FMovieSceneBindingReference* FMovieSceneBindingReferences::AddOrReplaceBinding(const FGuid& ObjectId, UMovieSceneCustomBinding* NewCustomBinding, int32 BindingIndex)
+{
+	const int32 Num = SortedReferences.Num();
+	const int32 Index = Algo::LowerBoundBy(SortedReferences, ObjectId, &FMovieSceneBindingReference::ID);
+
+	// Could also use a binary search here, but typically we are only dealing with a single binding
+	if (BindingIndex >= 0 && Index + BindingIndex < Num && SortedReferences[Index + BindingIndex].ID == ObjectId)
+	{
+		// Replace the current binding
+		SortedReferences[Index + BindingIndex] = FMovieSceneBindingReference{ ObjectId, FUniversalObjectLocator(), ELocatorResolveFlags::None, NewCustomBinding};
+		return &SortedReferences[Index + BindingIndex];
+	}
+	else
+	{
+		// Add a new binding instead
+		return AddBinding(ObjectId, NewCustomBinding);
+	}
+}
+
+const FMovieSceneBindingReference* FMovieSceneBindingReferences::AddOrReplaceBinding(const FGuid& ObjectId, FUniversalObjectLocator&& NewLocator, int32 BindingIndex)
+{
+	const int32 Num = SortedReferences.Num();
+	const int32 Index = Algo::LowerBoundBy(SortedReferences, ObjectId, &FMovieSceneBindingReference::ID);
+
+	// Could also use a binary search here, but typically we are only dealing with a single binding
+	if (BindingIndex >= 0 && Index + BindingIndex < Num && SortedReferences[Index + BindingIndex].ID == ObjectId)
+	{
+		// Replace the current binding
+		SortedReferences[Index + BindingIndex] = FMovieSceneBindingReference{ ObjectId, MoveTemp(NewLocator), ELocatorResolveFlags::None, nullptr };
+		return &SortedReferences[Index + BindingIndex];
+	}
+	else
+	{
+		// Add a new binding instead
+		return AddBinding(ObjectId, MoveTemp(NewLocator));
+	}
+}
 
 void FMovieSceneBindingReferences::RemoveBinding(const FGuid& ObjectId)
 {
@@ -92,6 +158,25 @@ void FMovieSceneBindingReferences::RemoveBinding(const FGuid& ObjectId)
 	}
 }
 
+void FMovieSceneBindingReferences::ResolveBindingFromLocator(int32 Index, const UE::UniversalObjectLocator::FResolveParams& ResolveParams, TArray<UObject*, TInlineAllocator<1>>& OutObjects) const
+{
+	// Add our resolve param flags
+	if (ResolveParams.Context)
+	{
+		if (UWorld* World = ResolveParams.Context->GetWorld())
+		{
+			EnumAddFlags(const_cast<UE::UniversalObjectLocator::FResolveParams&>(ResolveParams).Flags, SortedReferences[Index].ResolveFlags);
+		}
+	}
+
+	UObject* ResolvedObject = SortedReferences[Index].Locator.Resolve(ResolveParams).SyncGet().Object;
+	ResolvedObject = UE::MovieScene::FindBoundObjectProxy(ResolvedObject);
+	if (ResolvedObject)
+	{
+		OutObjects.Add(ResolvedObject);
+	}
+}
+
 void FMovieSceneBindingReferences::ResolveBinding(const FGuid& ObjectId, const UE::UniversalObjectLocator::FResolveParams& ResolveParams, TArray<UObject*, TInlineAllocator<1>>& OutObjects) const
 {
 #if WITH_EDITORONLY_DATA
@@ -108,34 +193,31 @@ void FMovieSceneBindingReferences::ResolveBinding(const FGuid& ObjectId, const U
 
 	for (int32 Index = StartIndex; Index < Num && SortedReferences[Index].ID == ObjectId; ++Index)
 	{
-		// If we have cache params in the ResolveParams, update the index
-		if (const FLocatorSpawnedCacheResolveParameter* CacheParameter = ResolveParams.FindParameter<FLocatorSpawnedCacheResolveParameter>())
-		{
-			if (FMovieSceneObjectCache* Cache = static_cast<FMovieSceneObjectCache*>(CacheParameter->Cache))
-			{
-				FMovieSceneLocatorSpawnedCacheKey CacheKey = Cache->GetResolvingBindingCacheKey();
-				if (CacheKey.BindingID == ObjectId && CacheKey.BindingIndex != Index)
-				{
-					CacheKey.BindingIndex = Index;
-					Cache->SetResolvingBindingCacheKey(CacheKey);
-				}
-			}
-		}
+		ResolveBindingFromLocator(Index, ResolveParams, OutObjects);
+	}
+}
 
-		// Add our resolve param flags
-		if (ResolveParams.Context)
+void FMovieSceneBindingReferences::ResolveBinding(const FMovieSceneBindingResolveParams& BindingResolveParams, const UE::UniversalObjectLocator::FResolveParams& LocatorResolveParams, TSharedPtr<const UE::MovieScene::FSharedPlaybackState> SharedPlaybackState, TArray<UObject*, TInlineAllocator<1>>& OutObjects) const
+{
+	const int32 StartIndex = Algo::LowerBoundBy(SortedReferences, BindingResolveParams.ObjectBindingID, &FMovieSceneBindingReference::ID);
+	const int32 Num = SortedReferences.Num();
+
+	for (int32 Index = StartIndex; Index < Num && SortedReferences[Index].ID == BindingResolveParams.ObjectBindingID; ++Index)
+	{
+		// If a custom binding is present and we have valid shared playback state, resolve the custom binding
+		if (SortedReferences[Index].CustomBinding && SharedPlaybackState.IsValid())
 		{
-			if (UWorld* World = ResolveParams.Context->GetWorld())
+			UObject* ResolvedObject = SortedReferences[Index].CustomBinding->ResolveBinding(BindingResolveParams, Index - StartIndex, SharedPlaybackState.ToSharedRef()).Object;
+			ResolvedObject = UE::MovieScene::FindBoundObjectProxy(ResolvedObject);
+			if (ResolvedObject)
 			{
-				EnumAddFlags(const_cast<UE::UniversalObjectLocator::FResolveParams&>(ResolveParams).Flags, SortedReferences[Index].ResolveFlags);
+				OutObjects.Add(ResolvedObject);
 			}
 		}
-
-		UObject* ResolvedObject = SortedReferences[Index].Locator.Resolve(ResolveParams).SyncGet().Object;
-		ResolvedObject = UE::MovieScene::FindBoundObjectProxy(ResolvedObject);
-		if (ResolvedObject)
+		else
 		{
-			OutObjects.Add(ResolvedObject);
+			// Otherwise, attempt to resolve via the locator
+			ResolveBindingFromLocator(Index, LocatorResolveParams, OutObjects);
 		}
 	}
 }
@@ -212,19 +294,6 @@ void FMovieSceneBindingReferences::RemoveInvalidBindings(const TSet<FGuid>& Vali
 	}
 }
 
-void FMovieSceneBindingReferences::UnloadBoundObject(const UE::UniversalObjectLocator::FResolveParams& ResolveParams, const FGuid& ObjectId, int32 BindingIndex)
-{
-	const int32 StartIndex = Algo::LowerBoundBy(SortedReferences, ObjectId, &FMovieSceneBindingReference::ID);
-
-	if (SortedReferences.IsValidIndex(StartIndex) && SortedReferences[StartIndex].ID == ObjectId)
-	{
-		const int32 EndIndex = Algo::UpperBoundBy(SortedReferences, ObjectId, &FMovieSceneBindingReference::ID);
-		ensure(StartIndex + BindingIndex <= EndIndex);
-		ensure(EnumHasAllFlags(ResolveParams.Flags, ELocatorResolveFlags::Unload));
-		SortedReferences[StartIndex + BindingIndex].Locator.Resolve(ResolveParams);
-	}
-}
-
 void FMovieSceneBindingReference::InitializeLocatorResolveFlags()
 {
 	using namespace UE::UniversalObjectLocator;
diff --git a/Engine/Source/Runtime/MovieScene/Private/MovieSceneCommonHelpers.cpp b/Engine/Source/Runtime/MovieScene/Private/MovieSceneCommonHelpers.cpp
--- a/Engine/Source/Runtime/MovieScene/Private/MovieSceneCommonHelpers.cpp
+++ b/Engine/Source/Runtime/MovieScene/Private/MovieSceneCommonHelpers.cpp
@@ -1,5 +1,4 @@
 // Copyright Epic Games, Inc. All Rights Reserved.
-
 #include "MovieSceneCommonHelpers.h"
 #include "Components/SceneComponent.h"
 #include "GameFramework/Actor.h"
@@ -17,6 +16,8 @@
 #include "MovieSceneTrack.h"
 #include "Engine/Engine.h"
 #include "UObject/Package.h"
+#include "MovieSceneBindingReferences.h"
+#include "Bindings/MovieSceneSpawnableBinding.h"
 
 bool MovieSceneHelpers::IsSectionKeyable(const UMovieSceneSection* Section)
 {
@@ -578,6 +579,31 @@ float MovieSceneHelpers::CalculateWeightForBlending(UMovieSceneSection* SectionT
 	return Weight;
 }
 
+FString MovieSceneHelpers::MakeUniqueBindingName(UMovieScene* MovieScene, const FString& InName)
+{
+	FString NewName = InName;
+
+	auto DuplNameSpawnable = [&](FMovieSceneSpawnable& InSpawnable)
+	{
+		return InSpawnable.GetName() == NewName;
+	};
+
+	auto DuplNamePossessable = [&](FMovieScenePossessable& InPossessable)
+	{
+		return InPossessable.GetName() == NewName;
+	};
+
+	int32 Index = 2;
+	FString UniqueString;
+	while (MovieScene->FindSpawnable(DuplNameSpawnable) || MovieScene->FindPossessable(DuplNamePossessable))
+	{
+		NewName.RemoveFromEnd(UniqueString);
+		UniqueString = FString::Printf(TEXT(" (%d)"), Index++);
+		NewName += UniqueString;
+	}
+	return NewName;
+}
+
 FString MovieSceneHelpers::MakeUniqueSpawnableName(UMovieScene* MovieScene, const FString& InName)
 {
 	FString NewName = InName;
@@ -609,6 +635,13 @@ UObject* MovieSceneHelpers::MakeSpawnableTemplateFromInstance(UObject& InSourceO
 	UEngine::CopyPropertiesForUnrelatedObjects(&InSourceObject, NewInstance, CopyParams);
 
 	AActor* Actor = CastChecked<AActor>(NewInstance);
+	
+	// Remove tags that may have gotten stuck on- for spawnables/replaceables these tags will be added after spawning
+	static const FName SequencerActorTag(TEXT("SequencerActor"));
+	static const FName SequencerPreviewActorTag(TEXT("SequencerPreviewActor"));
+	Actor->Tags.Remove(SequencerActorTag);
+	Actor->Tags.Remove(SequencerPreviewActorTag);
+
 	if (Actor->GetAttachParentActor() != nullptr)
 	{
 		// We don't support spawnables and attachments right now
@@ -630,6 +663,206 @@ UObject* MovieSceneHelpers::MakeSpawnableTemplateFromInstance(UObject& InSourceO
 }
 
 
+bool MovieSceneHelpers::IsBoundToAnySpawnable(UMovieSceneSequence* Sequence, const FGuid& ObjectId, TSharedRef<const UE::MovieScene::FSharedPlaybackState> SharedPlaybackState)
+{
+	if (Sequence)
+	{
+		if (UMovieScene* MovieScene = Sequence->GetMovieScene())
+		{
+			if (MovieScene->FindSpawnable(ObjectId))
+			{
+				return true;
+			}
+		}
+
+		const FMovieSceneBindingReferences* Refs = Sequence->GetBindingReferences();
+		if (Refs)
+		{
+			return Algo::AnyOf(Refs->GetReferences(ObjectId), [&SharedPlaybackState](const FMovieSceneBindingReference& BindingReference) {
+				return BindingReference.CustomBinding && BindingReference.CustomBinding->WillSpawnObject(SharedPlaybackState);
+			});
+		}
+	}
+	return false;
+}
+
+bool MovieSceneHelpers::IsBoundToSpawnable(UMovieSceneSequence* Sequence, const FGuid& ObjectId, TSharedRef<const UE::MovieScene::FSharedPlaybackState> SharedPlaybackState, int32 BindingIndex)
+{
+	if (Sequence)
+	{
+		if (UMovieScene* MovieScene = Sequence->GetMovieScene())
+		{
+			if (MovieScene->FindSpawnable(ObjectId))
+			{
+				return true;
+			}
+		}
+
+		const FMovieSceneBindingReferences* Refs = Sequence->GetBindingReferences();
+		if (Refs)
+		{
+			if (const FMovieSceneBindingReference* Ref = Refs->GetReference(ObjectId, BindingIndex))
+			{
+				return Ref->CustomBinding && Ref->CustomBinding->WillSpawnObject(SharedPlaybackState);
+			}
+		}
+	}
+	return false;
+}
+
+UObject* MovieSceneHelpers::GetObjectTemplate(UMovieSceneSequence* Sequence, const FGuid& ObjectId, TSharedRef<const UE::MovieScene::FSharedPlaybackState> SharedPlaybackState, int32 BindingIndex)
+{
+	if (Sequence)
+	{
+		UMovieScene* MovieScene = Sequence->GetMovieScene();
+		if (!MovieScene)
+		{
+			return nullptr;
+		}
+		if (FMovieSceneSpawnable* Spawnable = MovieScene->FindSpawnable(ObjectId))
+		{
+			return Spawnable->GetObjectTemplate();
+		}
+		else if (FMovieSceneBindingReferences* BindingReferences = Sequence->GetBindingReferences())
+		{
+			if (UMovieSceneCustomBinding* CustomBinding = BindingReferences->GetCustomBinding(ObjectId, BindingIndex))
+			{
+				if (UMovieSceneSpawnableBindingBase* SpawnableBinding = CustomBinding->AsSpawnable(SharedPlaybackState))
+				{
+					return SpawnableBinding->GetObjectTemplate();
+				}
+			}
+		}
+	}
+	return nullptr;
+}
+
+bool MovieSceneHelpers::SetObjectTemplate(UMovieSceneSequence* Sequence, const FGuid& ObjectId, UObject* InSourceObject, TSharedRef<const UE::MovieScene::FSharedPlaybackState> SharedPlaybackState, int32 BindingIndex)
+{
+	if (Sequence)
+	{
+		UMovieScene* MovieScene = Sequence->GetMovieScene();
+		if (!MovieScene)
+		{
+			return false;
+		}
+		if (FMovieSceneSpawnable* Spawnable = MovieScene->FindSpawnable(ObjectId))
+		{
+			Spawnable->SetObjectTemplate(InSourceObject);
+			return true;
+		}
+		else if (FMovieSceneBindingReferences* BindingReferences = Sequence->GetBindingReferences())
+		{
+			if (UMovieSceneCustomBinding* CustomBinding = BindingReferences->GetCustomBinding(ObjectId, BindingIndex))
+			{
+				if (UMovieSceneSpawnableBindingBase* SpawnableBinding = CustomBinding->AsSpawnable(SharedPlaybackState))
+				{
+					SpawnableBinding->SetObjectTemplate(InSourceObject);
+					return true;
+				}
+			}
+		}
+	}
+	return false;
+}
+
+bool MovieSceneHelpers::SupportsObjectTemplate(UMovieSceneSequence* Sequence, const FGuid& ObjectId, TSharedRef<const UE::MovieScene::FSharedPlaybackState> SharedPlaybackState, int32 BindingIndex)
+{
+	if (Sequence)
+	{
+		UMovieScene* MovieScene = Sequence->GetMovieScene();
+		if (!MovieScene)
+		{
+			return false;
+		}
+		if (FMovieSceneSpawnable* Spawnable = MovieScene->FindSpawnable(ObjectId))
+		{
+			return true;
+		}
+		else if (const FMovieSceneBindingReferences* BindingReferences = Sequence->GetBindingReferences())
+		{
+			if (const UMovieSceneCustomBinding* CustomBinding = BindingReferences->GetCustomBinding(ObjectId, BindingIndex))
+			{
+				if (const UMovieSceneSpawnableBindingBase* SpawnableBinding = CustomBinding->AsSpawnable(SharedPlaybackState))
+				{
+					if (SpawnableBinding && SpawnableBinding->SupportsObjectTemplates())
+					{
+						return true;
+					}
+				}
+			}
+		}
+	}
+	return false;
+}
+
+bool MovieSceneHelpers::CopyObjectTemplate(UMovieSceneSequence* Sequence, const FGuid& ObjectId, UObject* InSourceObject, TSharedRef<const UE::MovieScene::FSharedPlaybackState> SharedPlaybackState, int32 BindingIndex )
+{
+	if (Sequence && InSourceObject)
+	{
+		UMovieScene* MovieScene = Sequence->GetMovieScene();
+		if (!MovieScene)
+		{
+			return false;
+		}
+		if (FMovieSceneSpawnable* Spawnable = MovieScene->FindSpawnable(ObjectId))
+		{
+			Spawnable->CopyObjectTemplate(*InSourceObject, *Sequence);
+			return true;
+		}
+		else if (FMovieSceneBindingReferences* BindingReferences = Sequence->GetBindingReferences())
+		{
+			if (UMovieSceneCustomBinding* CustomBinding = BindingReferences->GetCustomBinding(ObjectId, BindingIndex))
+			{
+				if (UMovieSceneSpawnableBindingBase* SpawnableBinding = CustomBinding->AsSpawnable(SharedPlaybackState))
+				{
+					SpawnableBinding->CopyObjectTemplate(InSourceObject, *Sequence);
+					return true;
+				}
+			}
+		}
+	}
+	return false;
+}
+
+#if WITH_EDITORONLY_DATA
+
+const UClass* MovieSceneHelpers::GetBoundObjectClass(UMovieSceneSequence* Sequence, const FGuid& ObjectId, int32 BindingIndex)
+{
+	if (Sequence)
+	{
+		UMovieScene* MovieScene = Sequence->GetMovieScene();
+		if (!MovieScene)
+		{
+			return nullptr;
+		}
+		if (FMovieSceneSpawnable* Spawnable = MovieScene->FindSpawnable(ObjectId))
+		{
+			if (UObject* ObjectTemplate = Spawnable->GetObjectTemplate())
+			{
+				return ObjectTemplate->GetClass();
+			}
+		}
+		else if (FMovieSceneBindingReferences* BindingReferences = Sequence->GetBindingReferences())
+		{
+			if (UMovieSceneCustomBinding* CustomBinding = BindingReferences->GetCustomBinding(ObjectId, BindingIndex))
+			{
+				return CustomBinding->GetBoundObjectClass();
+			}
+		}
+
+		if (FMovieScenePossessable* Possessable = MovieScene->FindPossessable(ObjectId))
+		{
+			return Possessable->GetPossessedObjectClass();
+		}
+	}
+	return nullptr;
+}
+
+#endif
+
+
+
 MovieSceneHelpers::FMovieSceneScopedPackageDirtyGuard::FMovieSceneScopedPackageDirtyGuard(USceneComponent* InComponent)
 {
 	Component = InComponent;
diff --git a/Engine/Source/Runtime/MovieScene/Private/MovieScenePlayer.cpp b/Engine/Source/Runtime/MovieScene/Private/MovieScenePlayer.cpp
--- a/Engine/Source/Runtime/MovieScene/Private/MovieScenePlayer.cpp
+++ b/Engine/Source/Runtime/MovieScene/Private/MovieScenePlayer.cpp
@@ -14,6 +14,7 @@
 #include "MovieSceneSequence.h"
 #include "MovieSceneSequenceID.h"
 #include "UniversalObjectLocatorResolveParameterBuffer.inl"
+#include "MovieSceneBindingReferences.h"
 
 namespace UE
 {
@@ -124,9 +125,20 @@ void IMovieScenePlayer::ResolveBoundObjects(const FGuid& InBindingId, FMovieScen
 	}
 }
 
-void IMovieScenePlayer::ResolveBoundObjects(UE::UniversalObjectLocator::FResolveParams& ResolveParams, const FGuid& InBindingId, FMovieSceneSequenceID SequenceID, UMovieSceneSequence& Sequence, TArray<UObject*, TInlineAllocator<1>>& OutObjects) const
+void IMovieScenePlayer::ResolveBoundObjects(UE::UniversalObjectLocator::FResolveParams& LocatorResolveParams, const FGuid& InBindingId, FMovieSceneSequenceID SequenceID, UMovieSceneSequence& InSequence, TArray<UObject*, TInlineAllocator<1>>& OutObjects) const
 {
-	Sequence.LocateBoundObjects(InBindingId, ResolveParams, OutObjects);
+	using namespace UE::UniversalObjectLocator;
+	using namespace UE::MovieScene;
+
+	if (const FMovieSceneBindingReferences* BindingReferences = InSequence.GetBindingReferences())
+	{
+		FMovieSceneBindingResolveParams BindingResolveParams{ &InSequence, InBindingId, SequenceID };
+		BindingReferences->ResolveBinding(BindingResolveParams, LocatorResolveParams, ConstCastSharedPtr<const UE::MovieScene::FSharedPlaybackState>(const_cast<IMovieScenePlayer*>(this)->FindSharedPlaybackState()), OutObjects);
+	}
+	else
+	{
+		InSequence.LocateBoundObjects(InBindingId, LocatorResolveParams, OutObjects);
+	}
 }
 
 TArrayView<TWeakObjectPtr<>> IMovieScenePlayer::FindBoundObjects(const FGuid& ObjectBindingID, FMovieSceneSequenceIDRef SequenceID)
@@ -224,6 +236,16 @@ bool IMovieScenePlayer::IsDisablingEventTriggers(FFrameTime& DisabledUntilTime)
 	return false;
 }
 
+
+FGuid IMovieScenePlayer::CreateBinding(UMovieSceneSequence* InSequence, UObject* InObject)
+{
+	if (InSequence && InObject)
+	{
+		return InSequence->CreatePossessable(InObject);
+	}
+	return FGuid();
+}
+
 void IMovieScenePlayer::InitializeRootInstance(TSharedRef<UE::MovieScene::FSharedPlaybackState> NewSharedPlaybackState)
 {
 	using namespace UE::MovieScene;
diff --git a/Engine/Source/Runtime/MovieScene/Private/MovieSceneSequence.cpp b/Engine/Source/Runtime/MovieScene/Private/MovieSceneSequence.cpp
--- a/Engine/Source/Runtime/MovieScene/Private/MovieSceneSequence.cpp
+++ b/Engine/Source/Runtime/MovieScene/Private/MovieSceneSequence.cpp
@@ -19,6 +19,7 @@
 #include "EntitySystem/IMovieSceneEntityProvider.h"
 #include "Compilation/MovieSceneCompiledDataManager.h"
 #include "UniversalObjectLocator.h"
+#include "Bindings/MovieSceneSpawnableBinding.h"
 
 #include UE_INLINE_GENERATED_CPP_BY_NAME(MovieSceneSequence)
 
@@ -65,15 +66,6 @@ FMovieSceneBindingReferences* UMovieSceneSequence::GetBindingReferences()
 	return const_cast<FMovieSceneBindingReferences*>(Result);
 }
 
-void UMovieSceneSequence::UnloadBoundObject(const UE::UniversalObjectLocator::FResolveParams& ResolveParams, const FGuid& ObjectId, int32 BindingIndex)
-{
-	FMovieSceneBindingReferences* Refs = GetBindingReferences();
-	if (Refs)
-	{
-		Refs->UnloadBoundObject(ResolveParams, ObjectId, BindingIndex);
-	}
-}
-
 void UMovieSceneSequence::LocateBoundObjects(const FGuid& ObjectId, const UE::UniversalObjectLocator::FResolveParams& ResolveParams, TArray<UObject*, TInlineAllocator<1>>& OutObjects) const
 {
 	const FMovieSceneBindingReferences* Refs = GetBindingReferences();
diff --git a/Engine/Source/Runtime/MovieScene/Private/MovieSceneSpawnRegister.cpp b/Engine/Source/Runtime/MovieScene/Private/MovieSceneSpawnRegister.cpp
--- a/Engine/Source/Runtime/MovieScene/Private/MovieSceneSpawnRegister.cpp
+++ b/Engine/Source/Runtime/MovieScene/Private/MovieSceneSpawnRegister.cpp
@@ -6,20 +6,22 @@
 #include "Evaluation/MovieSceneEvaluationTemplateInstance.h"
 #include "MovieSceneDynamicBindingInvoker.h"
 #include "MovieSceneSpawnableAnnotation.h"
+#include "MovieSceneBindingReferences.h"
+#include "Bindings/MovieSceneSpawnableBinding.h"
 
 UE::MovieScene::TPlaybackCapabilityID<FMovieSceneSpawnRegister> FMovieSceneSpawnRegister::ID = UE::MovieScene::TPlaybackCapabilityID<FMovieSceneSpawnRegister>::Register();
 
-TWeakObjectPtr<> FMovieSceneSpawnRegister::FindSpawnedObject(const FGuid& BindingId, FMovieSceneSequenceIDRef TemplateID) const
+TWeakObjectPtr<> FMovieSceneSpawnRegister::FindSpawnedObject(const FGuid& BindingId, FMovieSceneSequenceIDRef TemplateID, int BindingIndex/* = 0*/) const
 {
-	FMovieSceneSpawnRegisterKey Key(TemplateID, BindingId);
+	FMovieSceneSpawnRegisterKey Key(TemplateID, BindingId, BindingIndex);
 
 	const FSpawnedObject* Existing = Register.Find(Key);
 	return Existing ? Existing->Object : TWeakObjectPtr<>();
 }
 
-UObject* FMovieSceneSpawnRegister::SpawnObject(const FGuid& BindingId, UMovieScene& MovieScene, FMovieSceneSequenceIDRef TemplateID, TSharedRef<const FSharedPlaybackState> SharedPlaybackState)
+UObject* FMovieSceneSpawnRegister::SpawnObject(const FGuid& BindingId, UMovieScene& MovieScene, FMovieSceneSequenceIDRef TemplateID, TSharedRef<const FSharedPlaybackState> SharedPlaybackState, int BindingIndex/* = 0*/)
 {
-	TWeakObjectPtr<> WeakObjectInstance = FindSpawnedObject(BindingId, TemplateID);
+	TWeakObjectPtr<> WeakObjectInstance = FindSpawnedObject(BindingId, TemplateID, BindingIndex);
 	UObject*         ObjectInstance     = WeakObjectInstance.Get();
 
 	if (ObjectInstance)
@@ -28,16 +30,9 @@ UObject* FMovieSceneSpawnRegister::SpawnObject(const FGuid& BindingId, UMovieSce
 	}
 
 	// Find the spawnable definition
-	FMovieSceneSpawnable* Spawnable = MovieScene.FindSpawnable(BindingId);
-	if (!Spawnable)
-	{
-		return nullptr;
-	}
 
-	if (WeakObjectInstance.IsStale() && !Spawnable->bContinuouslyRespawn)
-	{
-		return nullptr;
-	}
+	UObject* SpawnedActor = nullptr; 
+	ESpawnOwnership SpawnOwnership = ESpawnOwnership::InnerSequence;
 
 	UMovieSceneSequence* Sequence = SharedPlaybackState->GetSequence(TemplateID);
 	if (!ensure(Sequence))
@@ -45,30 +40,60 @@ UObject* FMovieSceneSpawnRegister::SpawnObject(const FGuid& BindingId, UMovieSce
 		return nullptr;
 	}
 
-	UObject* SpawnedActor = nullptr;
-	ESpawnOwnership SpawnOwnership = Spawnable->GetSpawnOwnership();
-	IMovieScenePlayer* Player = UE::MovieScene::FPlayerIndexPlaybackCapability::GetPlayer(SharedPlaybackState);
-
-	// See if there is some dynamic binding logic to invoke, otherwise spawn the actor
-	FMovieSceneDynamicBindingResolveResult ResolveResult = FMovieSceneDynamicBindingInvoker::ResolveDynamicBinding(SharedPlaybackState, Sequence, TemplateID, *Spawnable);
-	if (ResolveResult.Object)
+	// First check if we're using the old-style FMovieSceneSpawnable
+	FMovieSceneSpawnable* Spawnable = MovieScene.FindSpawnable(BindingId);
+	if (Spawnable)
 	{
-		SpawnedActor = ResolveResult.Object;
-		if (ResolveResult.bIsPossessedObject)
+		if (WeakObjectInstance.IsStale() && !Spawnable->bContinuouslyRespawn)
+		{
+			return nullptr;
+		}
+
+		SpawnOwnership = Spawnable->GetSpawnOwnership();
+		IMovieScenePlayer* Player = UE::MovieScene::FPlayerIndexPlaybackCapability::GetPlayer(SharedPlaybackState);
+
+		// See if there is some dynamic binding logic to invoke, otherwise spawn the actor
+			// See if there is some dynamic binding logic to invoke, otherwise spawn the actor
+	    FMovieSceneDynamicBindingResolveResult ResolveResult = FMovieSceneDynamicBindingInvoker::ResolveDynamicBinding(SharedPlaybackState, Sequence, TemplateID, *Spawnable);
+	    if (ResolveResult.Object)
+	    {
+		    SpawnedActor = ResolveResult.Object;
+		    if (ResolveResult.bIsPossessedObject)
+            {
+                SpawnOwnership = ESpawnOwnership::External;
+            }
+        }
+		if (!SpawnedActor)
 		{
-			SpawnOwnership = ESpawnOwnership::External;
+	        // Call through to the list of spawners to see who can spawn something from this FMovieSceneSpawnable
+			SpawnedActor = SpawnObject(*Spawnable, TemplateID, SharedPlaybackState);
 		}
 	}
-	if (!SpawnedActor)
+	else if (UMovieSceneSequence* MovieSceneSequence = MovieScene.GetTypedOuter<UMovieSceneSequence>())
 	{
-		SpawnedActor = SpawnObject(*Spawnable, TemplateID, SharedPlaybackState);
+		if (FMovieSceneBindingReferences* BindingReferences = MovieSceneSequence->GetBindingReferences())
+		{
+			if (UMovieSceneCustomBinding* CustomBinding = BindingReferences->GetCustomBinding(BindingId, BindingIndex))
+			{
+				if (UMovieSceneSpawnableBindingBase* SpawnableBinding = CustomBinding->AsSpawnable(SharedPlaybackState))
+				{
+					if (WeakObjectInstance.IsStale() && !SpawnableBinding->bContinuouslyRespawn)
+					{
+						return nullptr;
+					}
+
+					// Call the Spawnable binding itself to spawn the object
+					SpawnedActor = SpawnableBinding->SpawnObject(BindingId, BindingIndex, MovieScene, TemplateID, SharedPlaybackState);
+				}
+			}
+		}
 	}
 	
 	if (SpawnedActor)
 	{
 		FMovieSceneSpawnableAnnotation::Add(SpawnedActor, BindingId, TemplateID, Sequence);
 
-		FMovieSceneSpawnRegisterKey Key(TemplateID, BindingId);
+		FMovieSceneSpawnRegisterKey Key(TemplateID, BindingId, BindingIndex);
 		Register.Add(Key, FSpawnedObject(BindingId, *SpawnedActor, SpawnOwnership));
 
 		if (FMovieSceneEvaluationState* State = SharedPlaybackState->FindCapability<FMovieSceneEvaluationState>())
@@ -80,18 +105,48 @@ UObject* FMovieSceneSpawnRegister::SpawnObject(const FGuid& BindingId, UMovieSce
 	return SpawnedActor;
 }
 
-bool FMovieSceneSpawnRegister::DestroySpawnedObject(const FGuid& BindingId, FMovieSceneSequenceIDRef TemplateID, TSharedRef<const FSharedPlaybackState> SharedPlaybackState)
+void FMovieSceneSpawnRegister::PreDestroyObject(UObject & Object, const FGuid & BindingId, FMovieSceneSequenceIDRef TemplateID)
+{
+	PreDestroyObject(Object, BindingId, 0, TemplateID);
+}
+
+bool FMovieSceneSpawnRegister::DestroySpawnedObject(const FGuid& BindingId, FMovieSceneSequenceIDRef TemplateID, TSharedRef<const FSharedPlaybackState> SharedPlaybackState, int BindingIndex/* = 0*/)
 {
 	TGuardValue<bool> CleaningUp(bCleaningUp, true);
 
-	FMovieSceneSpawnRegisterKey Key(TemplateID, BindingId);
+	FMovieSceneSpawnRegisterKey Key(TemplateID, BindingId, BindingIndex);
 	
 	FSpawnedObject* Existing = Register.Find(Key);
 	UObject* SpawnedObject = Existing ? Existing->Object.Get() : nullptr;
 	if (SpawnedObject)
 	{
-		PreDestroyObject(*SpawnedObject, BindingId, TemplateID);
-		DestroySpawnedObject(*SpawnedObject);
+		PreDestroyObject(*SpawnedObject, BindingId, BindingIndex, TemplateID);
+
+		bool bCustomBinding = false;
+		// If we have a custom binding, it will handle object destruction
+		if (FMovieSceneEvaluationState* State = SharedPlaybackState->FindCapability<FMovieSceneEvaluationState>())
+		{
+			if (UMovieSceneSequence* MovieSceneSequence = State->FindSequence(TemplateID))
+			{
+				if (FMovieSceneBindingReferences* BindingReferences = MovieSceneSequence->GetBindingReferences())
+				{
+					if (UMovieSceneCustomBinding* CustomBinding = BindingReferences->GetCustomBinding(BindingId, BindingIndex))
+					{
+						if (UMovieSceneSpawnableBindingBase* SpawnableBinding = CustomBinding->AsSpawnable(SharedPlaybackState))
+						{
+							bCustomBinding = true;
+							SpawnableBinding->PreDestroyObject(SpawnedObject, BindingId, BindingIndex, TemplateID);
+							DestroySpawnedObject(*SpawnedObject, SpawnableBinding);
+						}
+					}
+				}
+			}
+		}
+
+		if (!bCustomBinding)
+		{
+			DestroySpawnedObject(*SpawnedObject, nullptr);
+		}
 	}
 
 	Register.Remove(Key);
@@ -105,16 +160,47 @@ bool FMovieSceneSpawnRegister::DestroySpawnedObject(const FGuid& BindingId, FMov
 }
 
 void FMovieSceneSpawnRegister::DestroyObjectsByPredicate(TSharedRef<const FSharedPlaybackState> SharedPlaybackState, const TFunctionRef<bool(const FGuid&, ESpawnOwnership, FMovieSceneSequenceIDRef)>& Predicate)
+{
+	DestroyObjectsByPredicate(SharedPlaybackState, [&Predicate](const FGuid& Guid, ESpawnOwnership SpawnOwnership, FMovieSceneSequenceIDRef SequenceID, int32 BindingIndex) {
+		return Predicate(Guid, SpawnOwnership, SequenceID);
+		});
+}
+
+void FMovieSceneSpawnRegister::DestroyObjectsByPredicate(TSharedRef<const FSharedPlaybackState> SharedPlaybackState, const TFunctionRef<bool(const FGuid&, ESpawnOwnership, FMovieSceneSequenceIDRef, int32)>& Predicate)
 {
 	for (auto It = Register.CreateIterator(); It; ++It)
 	{
-		if (Predicate(It.Value().Guid, It.Value().Ownership, It.Key().TemplateID))
+		if (Predicate(It.Value().Guid, It.Value().Ownership, It.Key().TemplateID, It.Key().BindingIndex))
 		{
 			UObject* SpawnedObject = It.Value().Object.Get();
 			if (SpawnedObject)
 			{
-				PreDestroyObject(*SpawnedObject, It.Key().BindingId, It.Key().TemplateID);
-				DestroySpawnedObject(*SpawnedObject);
+				bool bCustomBinding = false;
+				// If we have a custom binding, it will handle object destruction
+				if (FMovieSceneEvaluationState* State = SharedPlaybackState->FindCapability<FMovieSceneEvaluationState>())
+				{
+					if (UMovieSceneSequence* MovieSceneSequence = State->FindSequence(It.Key().TemplateID))
+					{
+						if (FMovieSceneBindingReferences* BindingReferences = MovieSceneSequence->GetBindingReferences())
+						{
+							if (UMovieSceneCustomBinding* CustomBinding = BindingReferences->GetCustomBinding(It.Key().BindingId, It.Key().BindingIndex))
+							{
+								if (UMovieSceneSpawnableBindingBase* SpawnableBinding = CustomBinding->AsSpawnable(SharedPlaybackState))
+								{
+									bCustomBinding = true;
+									SpawnableBinding->PreDestroyObject(SpawnedObject, It.Key().BindingId, It.Key().BindingIndex, It.Key().TemplateID);
+									DestroySpawnedObject(*SpawnedObject, SpawnableBinding);
+								}
+							}
+						}
+					}
+				}
+
+				if (!bCustomBinding)
+				{
+					PreDestroyObject(*SpawnedObject, It.Key().BindingId, 0, It.Key().TemplateID);
+					DestroySpawnedObject(*SpawnedObject, nullptr);
+				}
 			}
 
 			It.RemoveCurrent();
@@ -202,7 +288,12 @@ void FMovieSceneSpawnRegister::OnSequenceExpired(FMovieSceneSequenceIDRef Templa
 
 void FMovieSceneSpawnRegister::SaveDefaultSpawnableState(FMovieSceneSpawnable& Spawnable, FMovieSceneSequenceIDRef TemplateID, IMovieScenePlayer& Player)
 {
-	SaveDefaultSpawnableState(Spawnable, TemplateID, Player.GetSharedPlaybackState());
+	SaveDefaultSpawnableState(Spawnable.GetGuid(), TemplateID, Player.GetSharedPlaybackState());
+}
+
+void FMovieSceneSpawnRegister::SaveDefaultSpawnableState(FMovieSceneSpawnable& Spawnable, FMovieSceneSequenceIDRef TemplateID, TSharedRef<const FSharedPlaybackState> SharedPlaybackState)
+{
+	SaveDefaultSpawnableState(Spawnable.GetGuid(), TemplateID, SharedPlaybackState);
 }
 
 void FMovieSceneSpawnRegister::HandleConvertPossessableToSpawnable(UObject* OldObject, IMovieScenePlayer& Player, TOptional<FTransformData>& OutTransformData)
@@ -210,10 +301,37 @@ void FMovieSceneSpawnRegister::HandleConvertPossessableToSpawnable(UObject* OldO
 	HandleConvertPossessableToSpawnable(OldObject, Player.GetSharedPlaybackState(), OutTransformData);
 }
 
+UObject* FMovieSceneSpawnRegister::SpawnObject(FMovieSceneSpawnable& Spawnable, FMovieSceneSequenceIDRef TemplateID, TSharedRef<const FSharedPlaybackState> SharedPlaybackState)
+{
+	return SpawnObject(Spawnable.GetGuid(), *SharedPlaybackState->GetSequence(TemplateID)->GetMovieScene(), TemplateID, SharedPlaybackState);
+}
+
 UObject* FMovieSceneSpawnRegister::SpawnObject(FMovieSceneSpawnable& Spawnable, FMovieSceneSequenceIDRef TemplateID, IMovieScenePlayer& Player)
 {
 	return SpawnObject(Spawnable, TemplateID, Player.GetSharedPlaybackState());
 }
 
+bool FMovieSceneSpawnRegister::CanConvertToPossessable(const FGuid& Guid, FMovieSceneSequenceIDRef TemplateID, TSharedRef<const FSharedPlaybackState> SharedPlaybackState, int32 BindingIndex) const
+{
+	if (UMovieSceneSequence* Sequence = SharedPlaybackState->GetSequence(TemplateID))
+	{
+		if (UMovieScene* MovieScene = Sequence->GetMovieScene())
+		{
+			if (FMovieSceneSpawnable* Spawnable = MovieScene->FindSpawnable(Guid))
+			{
+				return true;
+			}
+			else if (const FMovieSceneBindingReferences* BindingReferences = Sequence->GetBindingReferences())
+			{
+				if (const UMovieSceneCustomBinding* CustomBinding = BindingReferences->GetCustomBinding(Guid, BindingIndex))
+				{
+					return CustomBinding->CanConvertToPossessable(Guid, TemplateID, SharedPlaybackState);
+				}
+			}
+		}
+	}
+	return false;
+}
+
 #endif
 
diff --git a/Engine/Source/Runtime/MovieScene/Private/Sections/MovieSceneBindingLifetimeSection.cpp b/Engine/Source/Runtime/MovieScene/Private/Sections/MovieSceneBindingLifetimeSection.cpp
--- a/Engine/Source/Runtime/MovieScene/Private/Sections/MovieSceneBindingLifetimeSection.cpp
+++ b/Engine/Source/Runtime/MovieScene/Private/Sections/MovieSceneBindingLifetimeSection.cpp
@@ -9,6 +9,9 @@
 #include "EntitySystem/MovieSceneEntitySystemLinker.h"
 #include "Evaluation/MovieSceneEvaluationField.h"
 #include "MovieScene.h"
+#include "EntitySystem/MovieSceneInstanceRegistry.h"
+#include "MovieSceneSequence.h"
+#include "MovieSceneCommonHelpers.h"
 
 #include UE_INLINE_GENERATED_CPP_BY_NAME(MovieSceneBindingLifetimeSection)
 
@@ -24,10 +27,24 @@ void UMovieSceneBindingLifetimeSection::ImportEntityImpl(UMovieSceneEntitySystem
 
 	FBuiltInComponentTypes* BuiltInComponentTypes = FBuiltInComponentTypes::Get();
 
+	// Conditionally add this as a spawnable binding if our binding says so.
+	// Also conditionally add the binding lifetime component to manage lifetime events.
+
+	UMovieSceneSequence* Sequence = GetTypedOuter<UMovieSceneSequence>();
+	ensure(Sequence);
+
+	const FSequenceInstance& RootInstance = EntityLinker->GetInstanceRegistry()->GetInstance(Params.Sequence.RootInstanceHandle);
+	TSharedRef<const FSharedPlaybackState> SharedPlaybackState = RootInstance.GetSharedPlaybackState();
+
+	bool bSpawnable = MovieSceneHelpers::IsBoundToAnySpawnable(Sequence, Params.GetObjectBindingID(), SharedPlaybackState);
+
 	OutImportedEntity->AddBuilder(
 		FEntityBuilder()
+		.AddConditional(BuiltInComponentTypes->SpawnableBinding, Params.GetObjectBindingID(), Params.GetObjectBindingID().IsValid() && bSpawnable)
 		.AddConditional(BuiltInComponentTypes->BindingLifetime, FMovieSceneBindingLifetimeComponentData{ Params.GetObjectBindingID(), EMovieSceneBindingLifetimeState::Active }, Params.GetObjectBindingID().IsValid())
 	);
+
+
 }
 
 bool UMovieSceneBindingLifetimeSection::PopulateEvaluationFieldImpl(const TRange<FFrameNumber>& EffectiveRange, const FMovieSceneEvaluationFieldEntityMetaData& InMetaData, FMovieSceneEntityComponentFieldBuilder* OutFieldBuilder)
diff --git a/Engine/Source/Runtime/MovieScene/Private/Sections/MovieSceneSpawnSection.cpp b/Engine/Source/Runtime/MovieScene/Private/Sections/MovieSceneSpawnSection.cpp
--- a/Engine/Source/Runtime/MovieScene/Private/Sections/MovieSceneSpawnSection.cpp
+++ b/Engine/Source/Runtime/MovieScene/Private/Sections/MovieSceneSpawnSection.cpp
@@ -33,6 +33,7 @@ void UMovieSceneSpawnSection::ImportEntityImpl(UMovieSceneEntitySystemLinker* En
 	OutImportedEntity->AddBuilder(
 		FEntityBuilder()
 		.Add(FBuiltInComponentTypes::Get()->SpawnableBinding, Params.GetObjectBindingID())
+		.AddTag(FBuiltInComponentTypes::Get()->Tags.OldStyleSpawnable)
 	);
 }
 
diff --git a/Engine/Source/Runtime/MovieScene/Public/Bindings/MovieSceneReplaceableBinding.h b/Engine/Source/Runtime/MovieScene/Public/Bindings/MovieSceneReplaceableBinding.h
new file mode 100644
--- /dev/null
+++ b/Engine/Source/Runtime/MovieScene/Public/Bindings/MovieSceneReplaceableBinding.h
@@ -0,0 +1,93 @@
+// Copyright Epic Games, Inc. All Rights Reserved.
+
+#pragma once
+
+#include "Bindings/MovieSceneCustomBinding.h"
+#include "MovieSceneSequenceID.h"
+#include "MovieSceneSpawnableBinding.h"
+#include "Templates/SubclassOf.h"
+#include "MovieSceneReplaceableBinding.generated.h"
+
+namespace UE
+{
+	namespace MovieScene
+	{
+		struct FSharedPlaybackState;
+	}
+}
+struct FGuid;
+class UMovieScene;
+struct FSlateBrush;
+class UMovieSceneSpawnableBindingBase;
+
+/**
+ * The base class for custom replaceable bindings. A replaceable binding uses an internal custom spawnable at editor time to produce a preview object,
+ * while in editor will use some other mechanism to dynamically bind an object to the track. Different replaceable types can choose different combinations
+ * of how to create a spawnable for preview vs. how to dynamically bind an object at runtime.
+ * UMovieSceneReplaceableActorBinding as an example is the simplest type of replaceable binding and provides no method for binding at runtime and relies on the LevelSequenceActor's Binding Override 
+ * mechanism to bind an actor at runtime.
+ */
+UCLASS(Abstract, BlueprintType)
+class MOVIESCENE_API UMovieSceneReplaceableBindingBase
+	: public UMovieSceneCustomBinding
+{
+public:
+
+	GENERATED_BODY()
+
+public:
+
+	/* UMovieSceneCustomBinding overrides */
+	virtual bool SupportsBindingCreationFromObject(const UObject* SourceObject) const override;
+#if WITH_EDITOR
+	virtual	void SetupDefaults(UObject* SpawnedObject, FGuid ObjectBindingId, UMovieScene& OwnerMovieScene, TSharedRef<const UE::MovieScene::FSharedPlaybackState> SharedPlaybackState) override;
+	virtual const FSlateBrush* GetBindingTrackCustomIconOverlay() const override;
+	virtual FText GetBindingTrackIconTooltip() const override;
+	virtual UClass* GetBoundObjectClass() const override;
+
+	virtual bool SupportsConversionFromBinding(const FMovieSceneBindingReference& BindingReference, const UObject* SourceObject) const override;
+	virtual UMovieSceneCustomBinding* CreateCustomBindingFromBinding(const FMovieSceneBindingReference& BindingReference, UObject* SourceObject, UMovieScene& OwnerMovieScene) override;
+#endif
+
+protected:
+
+	/*
+	* Must be implemented. Called during non-editor/runtime to resolve the binding dynamically. In editor worlds/Sequencer will instead use the PreviewSpawnable binding to spawn a preview object.
+	* If no object is returned, Sequencer's BindingOverrides can still be used to dynamically bind the object. See UMovieSceneReplaceableActorBinding for an example.
+	*/
+	virtual FMovieSceneBindingResolveResult ResolveRuntimeBindingInternal(const FMovieSceneBindingResolveParams& ResolveParams, int32 BindingIndex, TSharedRef<const UE::MovieScene::FSharedPlaybackState> SharedPlaybackState) const PURE_VIRTUAL(UMovieSceneReplaceableBindingBase::ResolveRuntimeBindingInternal, return FMovieSceneBindingResolveResult(););
+
+	/* Called from CreateNewCustomBinding to create the inner spawnable used for sequencer preview. By default will just call GetInnerSpawnableClass and call CreateNewCustomBinding on that. */
+	virtual UMovieSceneSpawnableBindingBase* CreateInnerSpawnable(UObject* SourceObject, UMovieScene& OwnerMovieScene);
+
+	/*
+	* Must be implemented and return a non abstract spawnable binding class inheriting from UMovieSceneSpawnableBindingBase to use for the preview for this replaceable binding.
+	*/
+	virtual TSubclassOf<UMovieSceneSpawnableBindingBase> GetInnerSpawnableClass() const PURE_VIRTUAL(UMovieSceneReplaceableBindingBase::GetInnerSpawnableClass, return nullptr;);
+
+	/* Must be implemented. Called from CreateNewCustomBinding to allow the replaceable to initialize any data members from the source object. */
+	virtual void InitReplaceableBinding(UObject* SourceObject, UMovieScene & OwnerMovieScene) PURE_VIRTUAL(UMovieSceneReplaceableBindingBase::InitReplaceableBinding, return;);
+
+#if WITH_EDITORONLY_DATA
+	// Editor-only preview object
+	UPROPERTY(Instanced, VisibleAnywhere, Category="Editor")
+	TObjectPtr<UMovieSceneSpawnableBindingBase> PreviewSpawnable = nullptr;
+#endif
+
+	UPROPERTY()
+	FString BindingName;
+
+private:
+
+	/* UMovieSceneCustomBinding overrides*/
+	bool WillSpawnObject(TSharedRef<const UE::MovieScene::FSharedPlaybackState> SharedPlaybackState) const override final;
+	FMovieSceneBindingResolveResult ResolveBinding(const FMovieSceneBindingResolveParams& ResolveParams, int32 BindingIndex, TSharedRef<const UE::MovieScene::FSharedPlaybackState> SharedPlaybackState) const override final;
+	const UMovieSceneSpawnableBindingBase* AsSpawnable(TSharedRef<const UE::MovieScene::FSharedPlaybackState> SharedPlaybackState) const override final;
+	FString GetDesiredBindingName() const override { return BindingName; }
+
+	/* 
+	*  Note that we choose to implement CreateCustomBinding here rather than in subclasses.
+	*  Instead we rely on subclasses to implement CreateInnerSpawnable and InitReplaceableBinding which we call here.
+	*/
+	UMovieSceneCustomBinding* CreateNewCustomBinding(UObject* SourceObject, UMovieScene& OwnerMovieScene) override final;
+};
\ No newline at end of file
diff --git a/Engine/Source/Runtime/MovieScene/Public/EntitySystem/BuiltInComponentTypes.h b/Engine/Source/Runtime/MovieScene/Public/EntitySystem/BuiltInComponentTypes.h
--- a/Engine/Source/Runtime/MovieScene/Public/EntitySystem/BuiltInComponentTypes.h
+++ b/Engine/Source/Runtime/MovieScene/Public/EntitySystem/BuiltInComponentTypes.h
@@ -769,6 +769,8 @@ public:
 
 		FComponentTypeID RemoveHierarchicalBlendTarget;
 
+		FComponentTypeID OldStyleSpawnable;
+
 	} Tags;
 
 	struct
diff --git a/Engine/Source/Runtime/MovieScene/Public/Evaluation/MovieSceneEvaluationState.h b/Engine/Source/Runtime/MovieScene/Public/Evaluation/MovieSceneEvaluationState.h
--- a/Engine/Source/Runtime/MovieScene/Public/Evaluation/MovieSceneEvaluationState.h
+++ b/Engine/Source/Runtime/MovieScene/Public/Evaluation/MovieSceneEvaluationState.h
@@ -22,7 +22,6 @@
 #include "Templates/UnrealTypeTraits.h"
 #include "UObject/WeakObjectPtr.h"
 #include "UObject/WeakObjectPtrTemplates.h"
-#include "ILocatorSpawnedCache.h"
 
 class IMovieScenePlayer;
 class UMovieSceneSequence;
@@ -73,34 +72,10 @@ namespace UE::MovieScene
 
 }  // namespace UE::MovieScene
 
-struct FMovieSceneLocatorSpawnedCacheKey
-{
-	bool IsValid() const { return BindingID.IsValid() && BindingIndex != INDEX_NONE; }
-	void Reset() { BindingID = FGuid(); BindingIndex = INDEX_NONE; }
-
-	friend uint32 GetTypeHash(const FMovieSceneLocatorSpawnedCacheKey& A)
-	{
-		return GetTypeHash(A.BindingID) ^ GetTypeHash(A.BindingIndex);
-	}
-
-	friend bool operator==(const FMovieSceneLocatorSpawnedCacheKey& A, const FMovieSceneLocatorSpawnedCacheKey& B)
-	{
-		return A.BindingID == B.BindingID && A.BindingIndex == B.BindingIndex;
-	}
-
-	friend bool operator!=(const FMovieSceneLocatorSpawnedCacheKey& A, const FMovieSceneLocatorSpawnedCacheKey& B)
-	{
-		return A.BindingID != B.BindingID || A.BindingIndex != B.BindingIndex;
-	}
-
-	FGuid BindingID = FGuid();
-	int32 BindingIndex = INDEX_NONE;
-};
-
 /**
  * Object cache that looks up, resolves, and caches object bindings for a specific sequence
  */
-struct FMovieSceneObjectCache : public UE::UniversalObjectLocator::ILocatorSpawnedCache
+struct FMovieSceneObjectCache
 {
 	using FSharedPlaybackState = UE::MovieScene::FSharedPlaybackState;
 
@@ -219,18 +194,6 @@ struct FMovieSceneObjectCache : public UE::UniversalObjectLocator::ILocatorSpawn
 	 * @param OutBindings			(mandatory) Array to populate with bindings that relate to the object
 	 */
 	void FilterObjectBindings(UObject* PredicateObject, TSharedRef<const FSharedPlaybackState> SharedPlaybackState, TArray<FMovieSceneObjectBindingID>* OutBindings);
-	
-	/* ILocatorSpawnedCache Implementation */
-	MOVIESCENE_API UObject* FindExistingObject() override;
-	MOVIESCENE_API FName GetRequestedObjectName() override;
-	MOVIESCENE_API void ReportSpawnedObject(UObject* Object) override;
-	MOVIESCENE_API void SpawnedObjectDestroyed() override;
-	
-	FMovieSceneLocatorSpawnedCacheKey GetResolvingBindingCacheKey() const { return CurrentlyResolvingCacheKey; }
-	void SetResolvingBindingCacheKey(const FMovieSceneLocatorSpawnedCacheKey& InCacheKey) { CurrentlyResolvingCacheKey = InCacheKey; }
-	void ClearResolvingBindingCacheKey() { CurrentlyResolvingCacheKey.Reset(); }
-
-	MOVIESCENE_API void UnloadBinding(const FGuid& InBindingID, TSharedRef<const FSharedPlaybackState> SharedPlaybackState);
 
 public:
 
@@ -277,9 +240,6 @@ private:
 	 */
 	void UpdateSerialNumber();
 
-	/* Unloads the binding in question in cases where resolving the binding loaded a locator.*/
-	void UnloadBindingInternal(const FMovieSceneLocatorSpawnedCacheKey& CacheKey, TSharedRef<const FSharedPlaybackState> SharedPlaybackState);
-
 	struct FBoundObjects
 	{
 		bool bUpToDate;
@@ -329,17 +289,10 @@ private:
 	/* A set of inactive binding ids based on Binding Lifetime track. While inactive, these will be prevented from resolving.*/
 	TSet<FGuid> InactiveBindingIds;
 
-	/* A map of binding ids and binding indices that have been loaded by locators. We keep these here and reference them if necessary when resolving objects.*/
-	TMap<FMovieSceneLocatorSpawnedCacheKey, TWeakObjectPtr<>> LoadedBindingIds;
-
 	/** Serial number for this cache */
 	uint32 SerialNumber = 0;
 
 	bool bReentrantUpdate = false;
-
-	// Temporary variable set during a resolve for use by the ILocatorSpawnedCache functions
-	FMovieSceneLocatorSpawnedCacheKey CurrentlyResolvingCacheKey;
-
 };
 
 /**
@@ -515,11 +468,6 @@ struct FMovieSceneEvaluationState : public UE::MovieScene::IPlaybackCapability
 	 */
 	MOVIESCENE_API uint32 GetSerialNumber();
 
-	/** 
-	* Returns whether we are currently resolving a locator
-	*/
-	MOVIESCENE_API bool IsResolvingObject() const;
-
 	/** A map of persistent evaluation data mapped by movie scene evaluation entity (i.e, a given track or section) */
 	TMap<FMovieSceneEvaluationKey, TUniquePtr<IPersistentEvaluationData>> PersistentEntityData;
 
diff --git a/Engine/Source/Runtime/MovieScene/Public/IMovieScenePlayer.h b/Engine/Source/Runtime/MovieScene/Public/IMovieScenePlayer.h
--- a/Engine/Source/Runtime/MovieScene/Public/IMovieScenePlayer.h
+++ b/Engine/Source/Runtime/MovieScene/Public/IMovieScenePlayer.h
@@ -275,6 +275,11 @@ public:
 	*/
 	virtual TScriptInterface<IMovieSceneSequencePlayerObserver> GetObserver() { return nullptr; }
 
+	/*
+	* Attempts to create a binding for the given object in the given sequence.
+	*/
+	MOVIESCENE_API virtual FGuid CreateBinding(UMovieSceneSequence* InSequence, UObject* InObject);
+
 public:
 
 	/**
diff --git a/Engine/Source/Runtime/MovieScene/Public/MovieSceneBindingReferences.h b/Engine/Source/Runtime/MovieScene/Public/MovieSceneBindingReferences.h
--- a/Engine/Source/Runtime/MovieScene/Public/MovieSceneBindingReferences.h
+++ b/Engine/Source/Runtime/MovieScene/Public/MovieSceneBindingReferences.h
@@ -5,10 +5,18 @@
 #include "CoreTypes.h"
 #include "UniversalObjectLocator.h"
 #include "UniversalObjectLocatorResolveParams.h"
+#include "MovieSceneSequenceID.h"
+#include "Bindings/MovieSceneCustomBinding.h"
 #include "MovieSceneBindingReferences.generated.h"
 
 class UWorld;
 struct FWorldPartitionResolveData;
+class UMovieSceneSequence;
+
+namespace UE::MovieScene
+{
+	struct FSharedPlaybackState;
+}
 
 
 /**
@@ -28,9 +36,29 @@ struct FMovieSceneBindingReference
 	UPROPERTY()
 	ELocatorResolveFlags ResolveFlags = ELocatorResolveFlags::None;
 
+	UPROPERTY(Instanced)
+	TObjectPtr<UMovieSceneCustomBinding> CustomBinding;
+
 	void InitializeLocatorResolveFlags();
 };
 
+USTRUCT()
+struct FMovieSceneBindingResolveParams
+{
+	GENERATED_BODY()
+
+	/** The sequence that contains the object binding being resolved */
+	UPROPERTY()
+	TObjectPtr<UMovieSceneSequence> Sequence;
+
+	/** The ID of the object binding being resolved */
+	UPROPERTY()
+	FGuid ObjectBindingID;
+
+	/** The sequence ID of the object binding being resolved */
+	UPROPERTY()
+	FMovieSceneSequenceID SequenceID;
+};
 
 /**
  * Structure that stores a one to many mapping from object binding ID, to object references that pertain to that ID.
@@ -42,14 +70,28 @@ struct FMovieSceneBindingReferences
 
 	MOVIESCENE_API TArrayView<const FMovieSceneBindingReference> GetAllReferences() const;
 
+	MOVIESCENE_API TArrayView<FMovieSceneBindingReference> GetAllReferences();
+
 	MOVIESCENE_API TArrayView<const FMovieSceneBindingReference> GetReferences(const FGuid& ObjectId) const;
 
+	MOVIESCENE_API const FMovieSceneBindingReference* GetReference(const FGuid& ObjectId, int32 BindingIndex) const;
+
 	/**
 	 * Check whether this map has a binding for the specified object id
 	 * @return true if this map contains a binding for the id, false otherwise
 	 */
 	MOVIESCENE_API bool HasBinding(const FGuid& ObjectId) const;
 
+	/*
+	* @return If a custom binding exists at the given id and index, returns it.
+	*/
+	MOVIESCENE_API UMovieSceneCustomBinding* GetCustomBinding(const FGuid& ObjectId, int32 BindingIndex);
+
+	/*
+	* @return If a custom binding exists at the given id and index, returns it.
+	*/
+	MOVIESCENE_API const UMovieSceneCustomBinding* GetCustomBinding(const FGuid& ObjectId, int32 BindingIndex) const;
+
 	/**
 	 * Remove a binding for the specified ID
 	 *
@@ -82,22 +124,62 @@ struct FMovieSceneBindingReferences
 	 */
 	MOVIESCENE_API const FMovieSceneBindingReference* AddBinding(const FGuid& ObjectId, FUniversalObjectLocator&& NewLocator);
 
+	/**
+	 * Add a custom binding for the specified ID, no locator necessary.
+	 *
+	 * @param ObjectId	The ID to associate the object with
+	 * @param CustomBinding A created custom binding for the object
+	 */
+	MOVIESCENE_API const FMovieSceneBindingReference* AddBinding(const FGuid& ObjectId, UMovieSceneCustomBinding* CustomBinding) { return AddBinding(ObjectId, FUniversalObjectLocator(), ELocatorResolveFlags::None, CustomBinding); }
+
 	/**
 	 * Add a binding for the specified ID
 	 *
 	 * @param ObjectId	The ID to associate the object with
 	 * @param InContext	A context in which InObject resides (either a UWorld, or an AActor)
 	 */
-	MOVIESCENE_API const FMovieSceneBindingReference* AddBinding(const FGuid& ObjectId, FUniversalObjectLocator&& NewLocator, ELocatorResolveFlags InResolveFlags);
+	MOVIESCENE_API const FMovieSceneBindingReference* AddBinding(const FGuid& ObjectId, FUniversalObjectLocator&& NewLocator, ELocatorResolveFlags InResolveFlags, UMovieSceneCustomBinding* CustomBinding=nullptr);
 
 	/**
-	 * Resolve a binding for the specified ID using a given context
+	 * Replace the binding associated with the ObjectId at the given BindingIndex with a new custom binding.
+	 * If no such binding exists, then one will be created at the nearest new BindingIndex.
+	 *
+	 * @param ObjectId	The ID to associate the object with
+	 * @param CustomBinding A created custom binding for the object
+	 * @param BindingIndex The binding index of the binding to replace.
+	 */
+	MOVIESCENE_API const FMovieSceneBindingReference* AddOrReplaceBinding(const FGuid& ObjectId, UMovieSceneCustomBinding* NewCustomBinding, int32 BindingIndex);
+
+	/**
+ * Replace the binding associated with the ObjectId at the given BindingIndex with a new possessable locator binding.
+	 * If no such binding exists, then one will be created at the nearest new BindingIndex.
+	 *
+	 * @param ObjectId	The ID to associate the object with
+	 * @param NewLocator A created locator for the object
+	 * @param BindingIndex The binding index of the binding to replace.
+	 */
+	MOVIESCENE_API const FMovieSceneBindingReference* AddOrReplaceBinding(const FGuid& ObjectId, FUniversalObjectLocator&& NewLocator, int32 BindingIndex);
+
+	/**
+	 * Resolve a binding for the specified ID using a given context.
+	 * Calling this version will not resolve custom bindings/spawnables- to resolve those, please call the override with BindingResolveParams and SharedPlaybackState
 	 *
 	 * @param ObjectId					The ID to associate the object with
 	 * @param Params					Resolve parameters specifying the context and fragment-specific parameters
 	 * @param OutObjects				Array to populate with resolved object bindings
 	 */
-	MOVIESCENE_API void ResolveBinding(const FGuid& ObjectId, const UE::UniversalObjectLocator::FResolveParams& ResolveParams, TArray<UObject*, TInlineAllocator<1>>& OutObjects) const;
+	MOVIESCENE_API void ResolveBinding(const FGuid& ObjectId, const UE::UniversalObjectLocator::FResolveParams& LocatorResolveParams, TArray<UObject*, TInlineAllocator<1>>& OutObjects) const;
+
+	/**
+	 * Resolve a binding for the specified ID using a given context
+	 *
+	 * @param ObjectId					The ID to associate the object with
+	 * @param BindingResolveParams		Resolve parameters
+	 * @param LocatorResolveParams      Locator-specific resolve parameters
+	 * @param SharedPlaybackState       Shared Playback State from the Sequencer
+	 * @param OutObjects				Array to populate with resolved object bindings
+	 */
+	MOVIESCENE_API void ResolveBinding(const FMovieSceneBindingResolveParams& BindingResolveParams, const UE::UniversalObjectLocator::FResolveParams& LocatorResolveParams, TSharedPtr<const UE::MovieScene::FSharedPlaybackState> SharedPlaybackState, TArray<UObject*, TInlineAllocator<1>>& OutObjects) const;
 
 	/**
 	 * Resolve a binding for the specified ID using a given context
@@ -115,16 +197,20 @@ struct FMovieSceneBindingReferences
 	 * @param ValidBindingIDs A set of GUIDs that are considered valid. Anything references not matching these will be removed.
 	 */
 	MOVIESCENE_API void RemoveInvalidBindings(const TSet<FGuid>& ValidBindingIDs);
-	
+
 	/**
 	 * Unloads an object that has been loaded via a locator.
 	 *  @param ObjectId	The ID of the binding to unload
 	 *  @param BindingIndex	The index of the binding to unload
 	 */
-	MOVIESCENE_API void UnloadBoundObject(const UE::UniversalObjectLocator::FResolveParams& ResolveParams, const FGuid& ObjectId, int32 BindingIndex);
+
+	UE_DEPRECATED(5.5, "UnloadBoundObject no longer supported")
+	MOVIESCENE_API void UnloadBoundObject(const UE::UniversalObjectLocator::FResolveParams& ResolveParams, const FGuid& ObjectId, int32 BindingIndex) {}
 
 private:
 
+	void ResolveBindingFromLocator(int32 Index, const UE::UniversalObjectLocator::FResolveParams& LocatorResolveParams, TArray<UObject*, TInlineAllocator<1>>& OutObjects) const;
+
 	/** The map from object binding ID to an array of references that pertain to that ID */
 	UPROPERTY()
 	TArray<FMovieSceneBindingReference> SortedReferences;
diff --git a/Engine/Source/Runtime/MovieScene/Public/MovieSceneCommonHelpers.h b/Engine/Source/Runtime/MovieScene/Public/MovieSceneCommonHelpers.h
--- a/Engine/Source/Runtime/MovieScene/Public/MovieSceneCommonHelpers.h
+++ b/Engine/Source/Runtime/MovieScene/Public/MovieSceneCommonHelpers.h
@@ -21,6 +21,11 @@ class USoundBase;
 struct FRichCurve;
 enum class EMovieSceneKeyInterpolation : uint8;
 
+namespace UE::MovieScene
+{
+	struct FSharedPlaybackState;
+}
+
 class MovieSceneHelpers
 {
 public:
@@ -178,6 +183,14 @@ public:
 	*/
 	static MOVIESCENE_API float CalculateWeightForBlending(UMovieSceneSection* SectionToKey, FFrameNumber Time);
 
+	/*
+	 * Return a name unique to the binding names in the given movie scene
+	 * @param InMovieScene The movie scene to look for existing possessables.
+	 * @param InName The requested name to make unique.
+	 * @return The unique name
+	 */
+	static MOVIESCENE_API FString MakeUniqueBindingName(UMovieScene* InMovieScene, const FString& InName);
+
 	/*
 	 * Return a name unique to the spawnable names in the given movie scene
 	 * @param InMovieScene The movie scene to look for existing spawnables.
@@ -194,6 +207,51 @@ public:
 	 * @return The spawnable template
 	 */
 	static MOVIESCENE_API UObject* MakeSpawnableTemplateFromInstance(UObject& InSourceObject, UMovieScene* InMovieScene, FName InName);
+
+	/*
+	* Returns whether the given ObjectId is valid and is currently bound to at least 1 spawnable give the current context.
+	* More specifically, if a FMovieSceneSpawnable exists with this ObjectId, true will be returned.
+	* If a Level Sequence binding reference exists with a Custom Binding implementing MovieSceneSpawnableBindingBase, true will be returned.
+	* If a Level Sequence binding reference exists with a Custom Binding implementing MovieSceneReplaceableBindingBase and the Context is an editor world, then true will be returned.
+	* Otherwise, false will be returned.
+	*/
+	static MOVIESCENE_API bool IsBoundToAnySpawnable(UMovieSceneSequence* Sequence, const FGuid& ObjectId, TSharedRef<const UE::MovieScene::FSharedPlaybackState> SharedPlaybackState);
+
+	/*
+	* Returns whether the given ObjectId is valid and is the given bindingindex is currently bound to a spawnable give the current context.
+	* More specifically, if a FMovieSceneSpawnable exists with this ObjectId, true will be returned.
+	* If a Level Sequence binding reference for this guid with the given BindingIndex exists with a Custom Binding implementing MovieSceneSpawnableBindingBase, true will be returned.
+	* If a Level Sequence binding reference for this guid with the given BindingIndex exists with a Custom Binding implementing MovieSceneReplaceableBindingBase and the Context is an editor world, then true will be returned.
+	* Otherwise, false will be returned.
+	*/
+	static MOVIESCENE_API bool IsBoundToSpawnable(UMovieSceneSequence* Sequence, const FGuid& ObjectId, TSharedRef<const UE::MovieScene::FSharedPlaybackState> SharedPlaybackState, int32 BindingIndex = 0);
+
+	/*
+	* If the binding for the given ObjectId supports object templates, returns the template, otherwise returns nullptr
+	*/
+	static MOVIESCENE_API UObject* GetObjectTemplate(UMovieSceneSequence* Sequence, const FGuid& ObjectId, TSharedRef<const UE::MovieScene::FSharedPlaybackState> SharedPlaybackState, int32 BindingIndex = 0);
+
+	/*
+	* If the binding for the given ObjectId supports object templates, sets the template and returns true, otherwise returns false
+	*/
+	static MOVIESCENE_API bool SetObjectTemplate(UMovieSceneSequence* Sequence, const FGuid& ObjectId, UObject* InSourceObject, TSharedRef<const UE::MovieScene::FSharedPlaybackState> SharedPlaybackState, int32 BindingIndex = 0);
+
+	/*
+	* Returns whether the binding for the given ObjectId supports object templates
+	*/
+	static MOVIESCENE_API bool SupportsObjectTemplate(UMovieSceneSequence* Sequence, const FGuid& ObjectId, TSharedRef<const UE::MovieScene::FSharedPlaybackState> SharedPlaybackState, int32 BindingIndex = 0);
+
+	/*
+	* If the binding for the given ObjectId supports object templates, copies the object template into the binding and returns true, otherwise returns false
+	*/
+	static MOVIESCENE_API bool CopyObjectTemplate(UMovieSceneSequence* Sequence, const FGuid& ObjectId, UObject* InSourceObject, TSharedRef<const UE::MovieScene::FSharedPlaybackState> SharedPlaybackState, int32 BindingIndex = 0);
+#if WITH_EDITORONLY_DATA
+	/*
+	* Returns the bound object class for the binding for the given ObjectId.
+	*/
+	static MOVIESCENE_API const UClass* GetBoundObjectClass(UMovieSceneSequence* Sequence, const FGuid& ObjectId, int32 BindingIndex = 0);
+#endif
+
 };
 
 /**
diff --git a/Engine/Source/Runtime/MovieScene/Public/MovieSceneSequence.h b/Engine/Source/Runtime/MovieScene/Public/MovieSceneSequence.h
--- a/Engine/Source/Runtime/MovieScene/Public/MovieSceneSequence.h
+++ b/Engine/Source/Runtime/MovieScene/Public/MovieSceneSequence.h
@@ -98,7 +98,9 @@ public:
 	/**
 	 * Unloads an object that has been loaded via a locator.
 	 */
-	MOVIESCENE_API void UnloadBoundObject(const UE::UniversalObjectLocator::FResolveParams& ResolveParams, const FGuid& ObjectId, int32 BindingIndex);
+	UE_DEPRECATED(5.5, "UnloadBoundObject no longer supported")
+	MOVIESCENE_API void UnloadBoundObject(const UE::UniversalObjectLocator::FResolveParams& ResolveParams, const FGuid& ObjectId, int32 BindingIndex) {}
+
 
 	/**
 	 * Called when Sequencer has created an object binding for a possessable object
@@ -130,7 +132,8 @@ public:
 	virtual void LocateBoundObjects(const FGuid& ObjectId, UObject* Context, TArray<UObject*, TInlineAllocator<1>>& OutObjects) const {}
 
 	/**
-	 * Locate all the objects that correspond to the specified object ID, using the specified parameters
+	 * Locate all the objects that correspond to the specified object ID, using the specified parameters.
+	 * Calling this directly instead of IMovieScenePlayer::ResolveBoundObjects means that you won't be able to locate Spawnable objects or objects from Custom Bindings.
 	 *
 	 * @param ObjectId				The unique identifier of the object.
 	 * @param Params				Resolve parameters specifying the context and fragment-specific parameters
@@ -140,6 +143,7 @@ public:
 
 	/**
 	 * Locate all the objects that correspond to the specified object ID, using the specified context
+	 * Calling this directly instead of IMovieScenePlayer::ResolveBoundObjects means that you won't be able to locate Spawnable objects or objects from Custom Bindings.
 	 *
 	 * @param ObjectId				The unique identifier of the object.
 	 * @param Context				Optional context to use to find the required object (for instance, a parent spawnable object)
@@ -203,6 +207,13 @@ public:
 	 */
 	virtual bool AllowsSpawnableObjects() const { return false; }
 
+	/**
+	 * Whether the sequence type supports custom bindings.
+	 *
+	 * @return true if custom bindings are supported, false if they are not
+	 */
+	virtual bool AllowsCustomBindings() const { return false; }
+
 	/**
 	 * Unbinds all possessable objects from the provided GUID.
 	 *
