Sequencer- Refactor Bindings to support inherited Custom Bindings.

diff --git a/Engine/Source/Runtime/MovieScene/Public/MovieSceneSpawnRegister.h b/Engine/Source/Runtime/MovieScene/Public/MovieSceneSpawnRegister.h
--- a/Engine/Source/Runtime/MovieScene/Public/MovieSceneSpawnRegister.h
+++ b/Engine/Source/Runtime/MovieScene/Public/MovieSceneSpawnRegister.h
@@ -35,6 +35,7 @@ class USequencerSettings;
 struct FMovieSceneEvaluationState;
 struct FMovieSceneSpawnable;
 struct FTransformData;
+class UMovieSceneSpawnableBindingBase;
 
 namespace UE::MovieScene
 {
@@ -63,9 +64,11 @@ public:
 	 * Attempt to find a previously spawned object represented by the specified object and template IDs
 	 * @param BindingId		ID of the object to find
 	 * @param TemplateID	Unique ID of the template to look within
+	 * @param BindingIndex 	For level sequences using custom spawnable bindings, the index of the binding reference.
 	 * @return The spawned object if found; nullptr otherwise.
+	 * The version without the binding index is to support old clients of the previous API
 	 */
-	MOVIESCENE_API TWeakObjectPtr<> FindSpawnedObject(const FGuid& BindingId, FMovieSceneSequenceIDRef TemplateID) const;
+	MOVIESCENE_API TWeakObjectPtr<> FindSpawnedObject(const FGuid& BindingId, FMovieSceneSequenceIDRef TemplateID, int BindingIndex) const;
 
 	/**
 	 * Spawn an object for the specified GUID, from the specified sequence instance.
@@ -73,9 +76,11 @@ public:
 	 * @param BindingId 	ID of the object to spawn
 	 * @param TemplateID 	Identifier for the current template we're evaluating
 	 * @param Player 		Movie scene player that is ultimately responsible for spawning the object
+	 * @param BindingIndex 	For level sequences using custom spawnable bindings, the index of the binding reference.
 	 * @return the spawned object, or nullptr on failure
+	 * The version without the binding index is to support old clients of the previous API
 	 */
-	MOVIESCENE_API UObject* SpawnObject(const FGuid& BindingId, UMovieScene& MovieScene, FMovieSceneSequenceIDRef Template, TSharedRef<const FSharedPlaybackState> SharedPlaybackState);
+	MOVIESCENE_API virtual UObject* SpawnObject(const FGuid& BindingId, UMovieScene& MovieScene, FMovieSceneSequenceIDRef Template, TSharedRef<const FSharedPlaybackState> SharedPlaybackState, int32 BindingIndex);
 
 	/**
 	 * Destroy a specific previously spawned object
@@ -83,10 +88,13 @@ public:
 	 * @param BindingId		ID of the object to destroy
 	 * @param TemplateID 	Identifier for the current template we're evaluating
 	 * @param Player 		Movie scene player that is ultimately responsible for destroying the object
+	 * @param BindingIndex 	For level sequences using custom spawnable bindings, the index of the binding reference.
 	 *
 	 * @return True if an object was destroyed, false otherwise
+	 * The version without the binding index is to support old clients of the previous API
 	 */
-	MOVIESCENE_API bool DestroySpawnedObject(const FGuid& BindingId, FMovieSceneSequenceIDRef TemplateID, TSharedRef<const FSharedPlaybackState> SharedPlaybackState);
+
+	MOVIESCENE_API bool DestroySpawnedObject(const FGuid& BindingId, FMovieSceneSequenceIDRef TemplateID, TSharedRef<const FSharedPlaybackState> SharedPlaybackState, int32 BindingIndex);
 
 	/**
 	 * Destroy a specific previously spawned object, where its binding ID and sequence ID is not known.
@@ -94,7 +102,7 @@ public:
 	 *
 	 * @param InObject		the object to destroy
 	 */
-	void DestroyObjectDirectly(UObject& InObject) { DestroySpawnedObject(InObject); }
+	void DestroyObjectDirectly(UObject& InObject) { DestroySpawnedObject(InObject, nullptr); }
 
 	/**
 	 * Destroy spawned objects using a custom predicate
@@ -102,7 +110,7 @@ public:
 	 * @param Player 		Movie scene player that is ultimately responsible for destroying the objects
 	 * @param Predicate		Predicate used for testing whether an object should be destroyed. Returns true for destruction, false to skip.
 	 */
-	MOVIESCENE_API void DestroyObjectsByPredicate(TSharedRef<const FSharedPlaybackState> SharedPlaybackState, const TFunctionRef<bool(const FGuid&, ESpawnOwnership, FMovieSceneSequenceIDRef)>& Predicate);
+	MOVIESCENE_API void DestroyObjectsByPredicate(TSharedRef<const FSharedPlaybackState> SharedPlaybackState, const TFunctionRef<bool(const FGuid&, ESpawnOwnership, FMovieSceneSequenceIDRef, int32)>& Predicate);
 
 	/**
 	 * Purge any memory of any objects that are considered externally owned
@@ -132,8 +140,12 @@ public:
 
 	// Backwards compatible API, to be deprecated later
 
+	MOVIESCENE_API TWeakObjectPtr<> FindSpawnedObject(const FGuid& BindingId, FMovieSceneSequenceIDRef TemplateID) const { return FindSpawnedObject(BindingId, TemplateID, 0); }
+	MOVIESCENE_API virtual UObject* SpawnObject(const FGuid& BindingId, UMovieScene& MovieScene, FMovieSceneSequenceIDRef Template, TSharedRef<const FSharedPlaybackState> SharedPlaybackState) { return SpawnObject(BindingId, MovieScene, Template, SharedPlaybackState, 0); }
+	MOVIESCENE_API bool DestroySpawnedObject(const FGuid& BindingId, FMovieSceneSequenceIDRef TemplateID, TSharedRef<const FSharedPlaybackState> SharedPlaybackState) { return DestroySpawnedObject(BindingId, TemplateID, SharedPlaybackState, 0); }
 	MOVIESCENE_API UObject* SpawnObject(const FGuid& BindingId, UMovieScene& MovieScene, FMovieSceneSequenceIDRef Template, IMovieScenePlayer& Player);
 	MOVIESCENE_API bool DestroySpawnedObject(const FGuid& BindingId, FMovieSceneSequenceIDRef TemplateID, IMovieScenePlayer& Player);
+	MOVIESCENE_API void DestroyObjectsByPredicate(TSharedRef<const FSharedPlaybackState> SharedPlaybackState, const TFunctionRef<bool(const FGuid&, ESpawnOwnership, FMovieSceneSequenceIDRef)>& Predicate);
 	MOVIESCENE_API void DestroyObjectsByPredicate(IMovieScenePlayer& Player, const TFunctionRef<bool(const FGuid&, ESpawnOwnership, FMovieSceneSequenceIDRef)>& Predicate);
 	MOVIESCENE_API void ForgetExternallyOwnedSpawnedObjects(FMovieSceneEvaluationState& State, IMovieScenePlayer& Player);
 	MOVIESCENE_API void CleanUp(IMovieScenePlayer& Player);
@@ -164,7 +176,7 @@ public:
 	/**
 	 * Called to save the default state of the specified spawnable
 	 */
-	virtual void SaveDefaultSpawnableState(FMovieSceneSpawnable& Spawnable, FMovieSceneSequenceIDRef TemplateID, TSharedRef<const FSharedPlaybackState> SharedPlaybackState) {}
+	virtual void SaveDefaultSpawnableState(const FGuid& Guid, FMovieSceneSequenceIDRef TemplateID, TSharedRef<const FSharedPlaybackState> SharedPlaybackState) {}
 
 	/**
 	 * Setup a new spawnable object with some default tracks and keys
@@ -187,19 +199,23 @@ public:
 	virtual void HandleConvertPossessableToSpawnable(UObject* OldObject, TSharedRef<const FSharedPlaybackState> SharedPlaybackState, TOptional<FTransformData>& OutTransformData) {}
 
 	/**
-	 * Check whether the specified Spawnable can become a Possessable.
-	 * @param	Spawnable	The spawnable to check
+	 * Check whether the specified binding can become a Possessable.
+	 * @param BindingId		ID of the object to destroy
+	 * @param TemplateID 	Identifier for the current template we're evaluating
 	 * @return whether the conversion from Spawnable to Possessable can occur.
 	 */
-	virtual bool CanConvertSpawnableToPossessable(FMovieSceneSpawnable& Spawnable) const { return true; }
+	MOVIESCENE_API virtual bool CanConvertToPossessable(const FGuid& Guid, FMovieSceneSequenceIDRef TemplateID, TSharedRef<const FSharedPlaybackState> SharedPlaybackState, int32 BindingIndex=0) const;
 
 public:
 
 	// Backwards compatible API, to be deprecated later
-
+#if WITH_EDITOR
+	MOVIESCENE_API virtual void SaveDefaultSpawnableState(FMovieSceneSpawnable& Spawnable, FMovieSceneSequenceIDRef TemplateID, TSharedRef<const FSharedPlaybackState> SharedPlaybackState);
 	MOVIESCENE_API void SaveDefaultSpawnableState(FMovieSceneSpawnable& Spawnable, FMovieSceneSequenceIDRef TemplateID, IMovieScenePlayer& Player);
+#endif
 	MOVIESCENE_API void HandleConvertPossessableToSpawnable(UObject* OldObject, IMovieScenePlayer& Player, TOptional<FTransformData>& OutTransformData);
 	MOVIESCENE_API UObject* SpawnObject(FMovieSceneSpawnable& Spawnable, FMovieSceneSequenceIDRef TemplateID, IMovieScenePlayer& Player);
+	virtual bool CanConvertSpawnableToPossessable(FMovieSceneSpawnable& Spawnable) const { return true; }
 
 #endif
 
@@ -213,19 +229,26 @@ protected:
 	 * @param Player 		Movie scene player that is ultimately responsible for spawning the object
 	 * @return the spawned object, or nullptr on failure
 	 */
-	virtual UObject* SpawnObject(FMovieSceneSpawnable& Spawnable, FMovieSceneSequenceIDRef TemplateID, TSharedRef<const FSharedPlaybackState> SharedPlaybackState) = 0;
+	MOVIESCENE_API virtual UObject* SpawnObject(FMovieSceneSpawnable& Spawnable, FMovieSceneSequenceIDRef TemplateID, TSharedRef<const FSharedPlaybackState> SharedPlaybackState) = 0;
 
 	/**
 	 * Called right before a spawned object with the specified ID and template ID is destroyed
 	 */
-	virtual void PreDestroyObject(UObject& Object, const FGuid& BindingId, FMovieSceneSequenceIDRef TemplateID) {}
+	UE_DEPRECATED(5.5, "Please use the version of this with BindingIndex")
+	MOVIESCENE_API virtual void PreDestroyObject(UObject& Object, const FGuid& BindingId, FMovieSceneSequenceIDRef TemplateID);
+
+	MOVIESCENE_API virtual void PreDestroyObject(UObject& Object, const FGuid& BindingId, int32 BindingIndex, FMovieSceneSequenceIDRef TemplateID) {}
+
+	UE_DEPRECATED(5.5, "Please use the version of this with CustomSpawnableBinding")
+	virtual void DestroySpawnedObject(UObject& Object) {}
 
 	/**
 	 * Destroy a specific previously spawned object
 	 *
-	 * @param Object 		The object to destroy
+	 * @param Object 			The object to destroy
+	 * @param CustomSpawnableBinding Optional custom spawnable binding to handle the destroy.
 	 */
-	virtual void DestroySpawnedObject(UObject& Object) = 0;
+	virtual void DestroySpawnedObject(UObject& Object, UMovieSceneSpawnableBindingBase* CustomSpawnableBinding) = 0;
 
 protected:
 
@@ -253,20 +276,21 @@ protected:
 	 */
 	struct FMovieSceneSpawnRegisterKey
 	{
-		FMovieSceneSpawnRegisterKey(FMovieSceneSequenceIDRef InTemplateID, const FGuid& InBindingId)
+		FMovieSceneSpawnRegisterKey(FMovieSceneSequenceIDRef InTemplateID, const FGuid& InBindingId, int32 InBindingIndex)
 			: BindingId(InBindingId)
 			, TemplateID(InTemplateID)
+			, BindingIndex(InBindingIndex)
 		{
 		}
 
 		bool operator==(const FMovieSceneSpawnRegisterKey& Other) const
 		{
-			return BindingId == Other.BindingId && TemplateID == Other.TemplateID;
+			return BindingId == Other.BindingId && TemplateID == Other.TemplateID && BindingIndex == Other.BindingIndex;
 		}
 		
 		friend uint32 GetTypeHash(const FMovieSceneSpawnRegisterKey& Key)
 		{
-			return HashCombine(GetTypeHash(Key.BindingId), GetTypeHash(Key.TemplateID));
+			return HashCombine(HashCombine(GetTypeHash(Key.BindingId), GetTypeHash(Key.TemplateID)), GetTypeHash(Key.BindingIndex));
 		}
 
 		/** BindingId of the object binding */
@@ -274,6 +298,9 @@ protected:
 
 		/** Movie Scene template identifier that spawned the object */
 		FMovieSceneSequenceID TemplateID;
+
+		/* For level sequences using custom spawnable bindings, the index of the binding reference. */
+		int32 BindingIndex = 0;
 	};
 
 protected:
@@ -289,5 +316,5 @@ class FNullMovieSceneSpawnRegister : public FMovieSceneSpawnRegister
 {
 public:
 	virtual UObject* SpawnObject(FMovieSceneSpawnable&, FMovieSceneSequenceIDRef, TSharedRef<const FSharedPlaybackState>) override { check(false); return nullptr; }
-	virtual void DestroySpawnedObject(UObject&) override { }
+	virtual void DestroySpawnedObject(UObject&, UMovieSceneSpawnableBindingBase* CustomSpawnableBinding) override {}
 };
diff --git a/Engine/Source/Runtime/MovieSceneTracks/Public/Bindings/MovieSceneSpawnableActorBinding.h b/Engine/Source/Runtime/MovieSceneTracks/Public/Bindings/MovieSceneSpawnableActorBinding.h
new file mode 100644
--- /dev/null
+++ b/Engine/Source/Runtime/MovieSceneTracks/Public/Bindings/MovieSceneSpawnableActorBinding.h
@@ -0,0 +1,159 @@
+// Copyright Epic Games, Inc. All Rights Reserved.
+
+#pragma once
+
+#include "Bindings/MovieSceneSpawnableBinding.h"
+#include "Templates/SubclassOf.h"
+#include "GameFramework/Actor.h"
+#include "MovieSceneSpawnRegister.h"
+#include "MovieSceneSpawnableActorBinding.generated.h"
+
+/**
+ * The base class for actor-specific spawnable bindings. Contains a default implementation that can handle spawning an Actor from provided Actor class and optional Actor template.
+ * Can be overridden in C++ or blueprint to provide an Actor class and to add custom PostSpawnObject behavior such as mesh setup based on an asset.
+ * The below UMovieSceneSpawnableActorBinding class implements this base class and replicates the old FMovieSceneSpawnable behavior by using a specified Actor template to spawn an Actor and can be used out of the box.
+ */
+UCLASS(Abstract, Blueprintable, MinimalAPI)
+class UMovieSceneSpawnableActorBindingBase
+	: public UMovieSceneSpawnableBindingBase
+{
+public:
+
+	GENERATED_BODY()
+
+public:
+
+	// Optional template support for spawnables. Subclasses can override this to provide support for saving an object template into a binding.
+	
+	/* Override and return true if the binding type supports object templates.*/
+	MOVIESCENETRACKS_API virtual bool SupportsObjectTemplates() const { return false; }
+
+	/* Override and return the object template if the binding type supports object templates*/
+	MOVIESCENETRACKS_API virtual UObject* GetObjectTemplate() { return nullptr; }
+
+	/**
+	 * Sets the object template to the specified object directly.
+	 * Used for Copy/Paste, typically you should use CopyObjectTemplate.
+	 */
+	MOVIESCENETRACKS_API virtual void SetObjectTemplate(UObject* InObjectTemplate) {}
+
+	/**
+	 * Copy the specified object into this spawnable's template
+	 *
+	 * @param InSourceObject The source object to use. This object will be duplicated into the spawnable.
+	 * @param MovieSceneSequence The movie scene sequence to which this spawnable belongs
+	 */
+	MOVIESCENETRACKS_API virtual void CopyObjectTemplate(UObject* InSourceObject, UMovieSceneSequence& MovieSceneSequence) {}
+
+
+	/**
+	 * Automatically determine a value for bNetAddressableName based on the spawnable type
+	 */
+	MOVIESCENETRACKS_API void AutoSetNetAddressableName();
+
+	/*
+	* Returns the optional level name to spawn the actor in, otherwise the Persistent level is used.
+	*/
+	MOVIESCENETRACKS_API FName GetLevelName() const { return LevelName; }
+	
+	/*
+	* Used to provide an optional level name to spawn the actor in, otherwise the Persistent level is used.
+	*/
+	MOVIESCENETRACKS_API void SetLevelName(FName InLevelName) { LevelName = InLevelName; }
+
+
+public:
+	/** When enabled, the actor will be spawned with a unique name so that it can be addressable between clients and servers. */
+	UPROPERTY(EditAnywhere, Category=Actor)
+	bool bNetAddressableName = false;
+
+	/** Name of level to spawn into */
+	UPROPERTY(EditAnywhere, Category = Actor)
+	FName LevelName;
+
+protected:
+
+	/* Override to provide Actor class to be spawned*/
+	// TODO: Make UFUNCTION
+	virtual TSubclassOf<AActor> GetActorClass() const PURE_VIRTUAL(UMovieSceneSpawnableActorBindingBase::GetActorClass, return AActor::StaticClass(););
+
+	/* Optionally override to provide an Actor template to use during Spawn */
+	virtual AActor* GetActorTemplate() const { return nullptr; }
+
+	/* Returns the transform to spawn the actor at*/
+	virtual FTransform GetSpawnTransform() const;
+
+protected:
+	// UMovieSceneSpawnableBindingBase overrides
+
+	/* Overridden to handle Actor-specific spawning */
+	virtual UObject* SpawnObjectInternal(UWorld* WorldContext, FName SpawnName, const FGuid& BindingId, int32 BindingIndex, UMovieScene& MovieScene, FMovieSceneSequenceIDRef TemplateID, TSharedRef<const UE::MovieScene::FSharedPlaybackState> SharedPlaybackState) override;
+
+	/* Overridden to handle Actor-specific destruction*/
+	virtual void DestroySpawnedObjectInternal(UObject* Object) override;
+
+
+protected:
+};
+
+/*
+* An implementation of UMovieSceneSpawnableActorBindingBase that matches the old FMovieSceneSpawnable spawnable implementation, allowing the spawning 
+* of Actors from a UObject template which is serialized inside the Sequence.
+*/
+UCLASS(BlueprintType, MinimalAPI, EditInlineNew, DefaultToInstanced)
+class UMovieSceneSpawnableActorBinding
+	: public UMovieSceneSpawnableActorBindingBase
+{
+public:
+
+	GENERATED_BODY()
+
+public:
+
+	/* MovieSceneSpawnableActorBindingBase overrides*/
+	bool SupportsObjectTemplates() const override { return true; }
+
+	/* Override and return the object template if the binding type supports object templates*/
+	UObject* GetObjectTemplate() override { return ActorTemplate.Get(); }
+
+	/* MovieSceneSpawnableBindingBase overrides*/
+	void SetObjectTemplate(UObject* InObjectTemplate) override;
+	void CopyObjectTemplate(UObject* InSourceObject, UMovieSceneSequence& MovieSceneSequence) override;
+
+	/* MovieSceneCustomBinding overrides*/
+	bool SupportsBindingCreationFromObject(const UObject* SourceObject) const override;
+	UMovieSceneCustomBinding* CreateNewCustomBinding(UObject* SourceObject, UMovieScene& OwnerMovieScene) override;
+	FString GetDesiredBindingName() const override { return BindingName; }
+
+#if WITH_EDITOR
+	bool SupportsConversionFromBinding(const FMovieSceneBindingReference& BindingReference, const UObject* SourceObject) const override;
+	UMovieSceneCustomBinding* CreateCustomBindingFromBinding(const FMovieSceneBindingReference& BindingReference, UObject* SourceObject, UMovieScene& OwnerMovieScene) override;
+	FText GetBindingTypePrettyName() const override;
+	UClass* GetBoundObjectClass() const override { return GetActorClass(); }
+#endif
+
+protected:
+
+	/* MovieSceneSpawnableBindingBase overrides*/
+	UWorld* GetWorldContext(TSharedRef<const UE::MovieScene::FSharedPlaybackState> SharedPlaybackState) const override;
+	FName GetSpawnName(const FGuid& BindingId, UMovieScene& MovieScene, FMovieSceneSequenceIDRef TemplateID, TSharedRef<const UE::MovieScene::FSharedPlaybackState> SharedPlaybackState) const override;
+
+	/* MovieSceneSpawnableActorBindingBase overrides*/
+	TSubclassOf<AActor> GetActorClass() const override;
+	AActor* GetActorTemplate() const override { return ActorTemplate; }
+
+private:
+
+	FName GetNetAddressableName(TSharedRef<const UE::MovieScene::FSharedPlaybackState> SharedPlaybackState, const FGuid& BindingId, FMovieSceneSequenceID SequenceID, const FString& BaseName) const;
+
+
+private:
+
+	UPROPERTY()
+	TObjectPtr<AActor> ActorTemplate;
+
+	UPROPERTY()
+	FString BindingName;
+
+};
+
diff --git a/Engine/Source/Runtime/UniversalObjectLocator/Private/UniversalObjectLocatorModule.cpp b/Engine/Source/Runtime/UniversalObjectLocator/Private/UniversalObjectLocatorModule.cpp
--- a/Engine/Source/Runtime/UniversalObjectLocator/Private/UniversalObjectLocatorModule.cpp
+++ b/Engine/Source/Runtime/UniversalObjectLocator/Private/UniversalObjectLocatorModule.cpp
@@ -11,7 +11,6 @@
 #include "Misc/DelayedAutoRegister.h"
 
 #include "DirectPathObjectLocator.h"
-#include "ILocatorSpawnedCache.h"
 
 namespace UE::UniversalObjectLocator
 {
@@ -37,8 +36,6 @@ public:
 					FFragmentTypeParameters FragmentTypeParams("uobj", NSLOCTEXT("DirectPathObjectLocator", "Object", "Object"));
 					FDirectPathObjectLocator::FragmentType = this->RegisterFragmentType<FDirectPathObjectLocator>(FragmentTypeParams);
 				}
-
-				FLocatorSpawnedCacheResolveParameter::ParameterType = this->RegisterParameterType<FLocatorSpawnedCacheResolveParameter>();
 			}
 		);
 
