CIS fix

diff --git a/Engine/Source/Runtime/MovieScene/Private/Bindings/MovieSceneReplaceableBinding.cpp b/Engine/Source/Runtime/MovieScene/Private/Bindings/MovieSceneReplaceableBinding.cpp
--- a/Engine/Source/Runtime/MovieScene/Private/Bindings/MovieSceneReplaceableBinding.cpp
+++ b/Engine/Source/Runtime/MovieScene/Private/Bindings/MovieSceneReplaceableBinding.cpp
@@ -16,7 +16,10 @@
 
 #define LOCTEXT_NAMESPACE "FPossessableModel"
 
-static const FName SequencerPreviewActorTag(TEXT("SequencerPreviewActor"));
+namespace UE::MovieScene::ReplaceableBinding
+{
+	static const FName SequencerPreviewActorTag(TEXT("SequencerPreviewActor"));
+} // namespace UE::MovieScene::ReplaceableBinding
 
 #if WITH_EDITOR
 void UMovieSceneReplaceableBindingBase::SetupDefaults(UObject* SpawnedObject, FGuid ObjectBindingId, UMovieScene& OwnerMovieScene, TSharedRef<const UE::MovieScene::FSharedPlaybackState> SharedPlaybackState)
@@ -110,7 +113,7 @@ FMovieSceneBindingResolveResult UMovieSceneReplaceableBindingBase::ResolveBindin
 				if (AActor* Actor = Cast<AActor>(Result.Object.Get()))
 				{
 					// In addition to the spawnable tag (which the spawnable will have added), we add a replaceable tag
-					Actor->Tags.AddUnique(SequencerPreviewActorTag);
+					Actor->Tags.AddUnique(UE::MovieScene::ReplaceableBinding::SequencerPreviewActorTag);
 				}
 				return Result;
 			}
