Sequencer: New Auto/Smart auto with no overshoot. Set's Sequencer.AutoTangentNew="2" as default revert this value to "1" to go to previous behavior.

diff --git a/Engine/Source/Runtime/MovieScene/Private/Channels/MovieSceneCurveChannelImpl.cpp b/Engine/Source/Runtime/MovieScene/Private/Channels/MovieSceneCurveChannelImpl.cpp
--- a/Engine/Source/Runtime/MovieScene/Private/Channels/MovieSceneCurveChannelImpl.cpp
+++ b/Engine/Source/Runtime/MovieScene/Private/Channels/MovieSceneCurveChannelImpl.cpp
@@ -17,10 +17,11 @@ static FAutoConsoleVariableRef CVarSequencerLinearCubicInterpolation(
 	TEXT("If 1 Linear Keys Act As Cubic Interpolation with Linear Tangents, if 0 Linear Key Forces Linear Interpolation to Next Key."),
 	ECVF_Default);
 
-static TAutoConsoleVariable<int32> CVarSequencerAutoTangentInterpolation(
+int32 GCachedSequencerAutoTangentInterpolation = 2;
+static FAutoConsoleVariableRef CVarSequencerAutoTangentInterpolation(
 	TEXT("Sequencer.AutoTangentNew"),
-	1,
-	TEXT("If 1 Auto Tangent will use new algorithm to gradually flatten maximum/minimum keys, if 0 Auto Tangent will average all keys (pre 4.23 behavior)."),
+	GCachedSequencerAutoTangentInterpolation,
+	TEXT("If 2 Sequencer will flatten tangents with no overshoot, if 1 Auto Tangent will use new algorithm to gradually flatten maximum/minimum keys, if 0 Auto Tangent will average all keys (pre 4.23 behavior)."),
 	ECVF_Default);
 
 static TAutoConsoleVariable<float> CVarSequencerSmartAutoBlendLocationPercentage(
@@ -876,6 +877,32 @@ bool TMovieSceneCurveChannelImpl<ChannelType>::EvaluateLegacy(const ChannelType*
 
 	return true;
 }
+
+template <typename ChannelValueType>
+int signNoZero(ChannelValueType val)
+{
+	return (ChannelValueType(0) < val) ? -1 : 1;
+};
+
+template <typename ChannelValueType>
+ChannelValueType ClampTangent(ChannelValueType NewTangent, ChannelValueType PreviousSlope, ChannelValueType NextSlope)
+{
+	if (signNoZero(PreviousSlope) != signNoZero(NextSlope) ||
+		signNoZero(NewTangent) != signNoZero(NextSlope))
+	{
+		NewTangent = 0.0;
+	}
+	else if (NextSlope >= 0)
+	{
+		NewTangent = FMath::Min(FMath::Min(NewTangent, NextSlope), PreviousSlope);
+	}
+	else
+	{
+		NewTangent = FMath::Max(FMath::Max(NewTangent, NextSlope), PreviousSlope);
+	}
+	return NewTangent;
+};
+
 template<typename ChannelType>
 float TMovieSceneCurveChannelImpl<ChannelType>::CalcSmartTangent(ChannelType* InChannel, int32 Index)
 {
@@ -976,7 +1003,6 @@ void TMovieSceneCurveChannelImpl<ChannelType>::AutoSetTangents(ChannelType* InCh
 		return;
 	}
 
-	const int UseNewAutoTangent = CVarSequencerAutoTangentInterpolation->GetInt();
 	{
 		ChannelValueType& FirstValue = InChannel->Values[0];
 		if (FirstValue.InterpMode == RCIM_Linear)
@@ -1011,6 +1037,7 @@ void TMovieSceneCurveChannelImpl<ChannelType>::AutoSetTangents(ChannelType* InCh
 			LastValue.Tangent.TangentWeightMode = RCTWM_WeightedNone;
 		}
 	}
+
 	const int32 MaxSmartAutoIterations = 10; // with smart auto we may need to iterate tangent calculations since it depends 
 											 // upon next tangent value which may be 0 since it is new.
 	int32 SmartAutoCount = 0;
@@ -1046,7 +1073,7 @@ void TMovieSceneCurveChannelImpl<ChannelType>::AutoSetTangents(ChannelType* InCh
 					SmartAutoCount = 0;
 				}
 			}
-			else if (!UseNewAutoTangent)
+			else if (GCachedSequencerAutoTangentInterpolation == 0)
 			{
 				//need to pass in the curve value type since though unfortunately tangents are just always floats (for doubles or floats),
 				//the AutoCalcTangent works with either.
@@ -1057,29 +1084,46 @@ void TMovieSceneCurveChannelImpl<ChannelType>::AutoSetTangents(ChannelType* InCh
 			}
 			else
 			{
-				// if key doesn't lie between we keep it flat(0.0).
-				if ( (ThisKey.Value > PrevKey.Value && ThisKey.Value < NextKey.Value) ||
-					(ThisKey.Value < PrevKey.Value && ThisKey.Value > NextKey.Value))
+				if (GCachedSequencerAutoTangentInterpolation >= 1)
 				{
-					CurveValueType ValueNewTangent = 0.0;
-					AutoCalcTangent(PrevKey.Value, ThisKey.Value, NextKey.Value, Tension, ValueNewTangent);
-					NewTangent = ValueNewTangent;
-					NewTangent /= PrevToNextTimeDiff;
-					//if within 0 to 15% or 85% to 100% range we gradually weight tangent to zero
-					const double AverageToZeroRange = 0.85;
-					const double ValDiff = FMath::Abs<double>(NextKey.Value - PrevKey.Value);
-					const double OurDiff = FMath::Abs<double>(ThisKey.Value - PrevKey.Value);
-					//ValDiff won't be zero ever due to previous check
-					double PercDiff = OurDiff / ValDiff;
-					if (PercDiff > AverageToZeroRange)
-					{
-						PercDiff = (PercDiff - AverageToZeroRange) / (1.0 - AverageToZeroRange);
-						NewTangent = NewTangent * (1.0 - PercDiff);
-					}
-					else if (PercDiff < (1.0 - AverageToZeroRange))
+					// if key doesn't lie between we keep it flat(0.0).
+					if ((ThisKey.Value > PrevKey.Value && ThisKey.Value < NextKey.Value) ||
+						(ThisKey.Value < PrevKey.Value && ThisKey.Value > NextKey.Value))
 					{
-						PercDiff = PercDiff  / (1.0 - AverageToZeroRange);
-						NewTangent = NewTangent * PercDiff;
+						CurveValueType ValueNewTangent = 0.0;
+						AutoCalcTangent(PrevKey.Value, ThisKey.Value, NextKey.Value, Tension, ValueNewTangent);
+						NewTangent = ValueNewTangent;
+						if (GCachedSequencerAutoTangentInterpolation == 2)
+						{
+							NewTangent /= PrevToNextTimeDiff;
+							//if within 0 to 15% or 85% to 100% range we gradually weight tangent to zero
+							const double AverageToZeroRange = 0.85;
+							const double ValDiff = FMath::Abs<double>(NextKey.Value - PrevKey.Value);
+							const double OurDiff = FMath::Abs<double>(ThisKey.Value - PrevKey.Value);
+							//ValDiff won't be zero ever due to previous check
+							double PercDiff = OurDiff / ValDiff;
+							if (PercDiff > AverageToZeroRange)
+							{
+								PercDiff = (PercDiff - AverageToZeroRange) / (1.0 - AverageToZeroRange);
+								NewTangent = NewTangent * (1.0 - PercDiff);
+							}
+							else if (PercDiff < (1.0 - AverageToZeroRange))
+							{
+								PercDiff = PercDiff / (1.0 - AverageToZeroRange);
+								NewTangent = NewTangent * PercDiff;
+							}
+
+						}
+						else if (GCachedSequencerAutoTangentInterpolation == 2) //use flattening, no overshoot
+						{
+							const double TimeToPrevious = FMath::Max<double>(KINDA_SMALL_NUMBER, InChannel->Times[Index].Value - InChannel->Times[Index - 1].Value);
+							const double TimeToNext = FMath::Max<double>(KINDA_SMALL_NUMBER, InChannel->Times[Index + 1].Value - InChannel->Times[Index].Value);
+							const double PreviousSlope = 3.0 * (ThisKey.Value - PrevKey.Value) / (TimeToPrevious);
+							const double NextSlope = 3.0 * (NextKey.Value - ThisKey.Value) / (TimeToNext);
+							NewTangent = ClampTangent<double>(NewTangent, PreviousSlope, NextSlope);
+							NewTangent /= PrevToNextTimeDiff;
+						}
+						
 					}
 				}
 			}
