CIS Fix for Static Analysis from new auto tangent tweak

diff --git a/Engine/Source/Runtime/MovieScene/Private/Channels/MovieSceneCurveChannelImpl.cpp b/Engine/Source/Runtime/MovieScene/Private/Channels/MovieSceneCurveChannelImpl.cpp
--- a/Engine/Source/Runtime/MovieScene/Private/Channels/MovieSceneCurveChannelImpl.cpp
+++ b/Engine/Source/Runtime/MovieScene/Private/Channels/MovieSceneCurveChannelImpl.cpp
@@ -1093,7 +1093,7 @@ void TMovieSceneCurveChannelImpl<ChannelType>::AutoSetTangents(ChannelType* InCh
 						CurveValueType ValueNewTangent = 0.0;
 						AutoCalcTangent(PrevKey.Value, ThisKey.Value, NextKey.Value, Tension, ValueNewTangent);
 						NewTangent = ValueNewTangent;
-						if (GCachedSequencerAutoTangentInterpolation == 2)
+						if (GCachedSequencerAutoTangentInterpolation < 2)
 						{
 							NewTangent /= PrevToNextTimeDiff;
 							//if within 0 to 15% or 85% to 100% range we gradually weight tangent to zero
