Sequencer: Auto tangent real fix for overshoot. Unshevled wrong CL

diff --git a/Engine/Source/Runtime/MovieScene/Private/Channels/MovieSceneCurveChannelImpl.cpp b/Engine/Source/Runtime/MovieScene/Private/Channels/MovieSceneCurveChannelImpl.cpp
--- a/Engine/Source/Runtime/MovieScene/Private/Channels/MovieSceneCurveChannelImpl.cpp
+++ b/Engine/Source/Runtime/MovieScene/Private/Channels/MovieSceneCurveChannelImpl.cpp
@@ -1093,9 +1093,9 @@ void TMovieSceneCurveChannelImpl<ChannelType>::AutoSetTangents(ChannelType* InCh
 						CurveValueType ValueNewTangent = 0.0;
 						AutoCalcTangent(PrevKey.Value, ThisKey.Value, NextKey.Value, Tension, ValueNewTangent);
 						NewTangent = ValueNewTangent;
-						//fix for now if (GCachedSequencerAutoTangentInterpolation < 2)
+						NewTangent /= PrevToNextTimeDiff;
+						if (GCachedSequencerAutoTangentInterpolation < 2)
 						{
-							NewTangent /= PrevToNextTimeDiff;
 							//if within 0 to 15% or 85% to 100% range we gradually weight tangent to zero
 							const double AverageToZeroRange = 0.85;
 							const double ValDiff = FMath::Abs<double>(NextKey.Value - PrevKey.Value);
@@ -1113,18 +1113,18 @@ void TMovieSceneCurveChannelImpl<ChannelType>::AutoSetTangents(ChannelType* InCh
 								NewTangent = NewTangent * PercDiff;
 							}
 
-						}
-						/*
+						}				
 						else if (GCachedSequencerAutoTangentInterpolation == 2) //use flattening, no overshoot
 						{
+							const double TwoThird = 2.0 / 3.0;
 							const double TimeToPrevious = FMath::Max<double>(KINDA_SMALL_NUMBER, InChannel->Times[Index].Value - InChannel->Times[Index - 1].Value);
 							const double TimeToNext = FMath::Max<double>(KINDA_SMALL_NUMBER, InChannel->Times[Index + 1].Value - InChannel->Times[Index].Value);
-							const double PreviousSlope = 3.0 * (ThisKey.Value - PrevKey.Value) / (TimeToPrevious);
-							const double NextSlope = 3.0 * (NextKey.Value - ThisKey.Value) / (TimeToNext);
+							const double PreviousSlope = (ThisKey.Value - PrevKey.Value) / (TwoThird * TimeToPrevious);
+							const double NextSlope = (NextKey.Value - ThisKey.Value) / (TwoThird * TimeToNext);
 							NewTangent = ClampTangent<double>(NewTangent, PreviousSlope, NextSlope);
-							NewTangent /= PrevToNextTimeDiff;
+
 						}
-						*/
+						
 					}
 				}
 			}
