Sequencer: move global pre-animated state capture counting to shared playback state

diff --git a/Engine/Source/Editor/Sequencer/Private/Sequencer.cpp b/Engine/Source/Editor/Sequencer/Private/Sequencer.cpp
--- a/Engine/Source/Editor/Sequencer/Private/Sequencer.cpp
+++ b/Engine/Source/Editor/Sequencer/Private/Sequencer.cpp
@@ -450,10 +450,6 @@ void FSequencer::InitSequencer(const FSequencerInitParams& InitParams, const TSh
 			{
 				State.InvalidateExpiredObjects();
 			}
-
-			// Replace pointers inside our pre-animate state storages.
-			PreAnimatedState.OnObjectsReplaced(ReplacementMap);
-
 		});
 		AcquiredResources.Add([=] { FCoreUObjectDelegates::OnObjectsReplaced.Remove(OnObjectsReplacedHandle); });
 	}
diff --git a/Engine/Source/Runtime/MovieScene/Private/EntitySystem/MovieSceneEntitySystemLinker.cpp b/Engine/Source/Runtime/MovieScene/Private/EntitySystem/MovieSceneEntitySystemLinker.cpp
--- a/Engine/Source/Runtime/MovieScene/Private/EntitySystem/MovieSceneEntitySystemLinker.cpp
+++ b/Engine/Source/Runtime/MovieScene/Private/EntitySystem/MovieSceneEntitySystemLinker.cpp
@@ -487,6 +487,8 @@ void UMovieSceneEntitySystemLinker::OnObjectsReplaced(const TMap<UObject*, UObje
 		}
 	});
 #endif
+
+	PreAnimatedState.OnObjectsReplaced(ReplacementMap);
 }
 
 void UMovieSceneEntitySystemLinker::AddReferencedObjects(UObject* Object, FReferenceCollector& Collector)
diff --git a/Engine/Source/Runtime/MovieScene/Private/EntitySystem/MovieSceneSharedPlaybackState.cpp b/Engine/Source/Runtime/MovieScene/Private/EntitySystem/MovieSceneSharedPlaybackState.cpp
--- a/Engine/Source/Runtime/MovieScene/Private/EntitySystem/MovieSceneSharedPlaybackState.cpp
+++ b/Engine/Source/Runtime/MovieScene/Private/EntitySystem/MovieSceneSharedPlaybackState.cpp
@@ -9,7 +9,9 @@
 namespace UE::MovieScene
 {
 
-FSharedPlaybackState::FSharedPlaybackState()
+FSharedPlaybackState::FSharedPlaybackState(UMovieSceneEntitySystemLinker* InLinker)
+	: WeakLinker(InLinker)
+	, PreAnimatedState(InLinker, FRootInstanceHandle())
 {
 }
 
@@ -21,6 +23,7 @@ FSharedPlaybackState::FSharedPlaybackState(
 	, WeakLinker(CreateParams.Linker)
 	, CompiledDataManager(CreateParams.CompiledDataManager)
 	, RootInstanceHandle(CreateParams.RootInstanceHandle)
+	, PreAnimatedState(CreateParams.Linker, CreateParams.RootInstanceHandle)
 {
 	if (CompiledDataManager)
 	{
diff --git a/Engine/Source/Runtime/MovieScene/Private/Evaluation/MovieSceneEvaluationTemplateInstance.cpp b/Engine/Source/Runtime/MovieScene/Private/Evaluation/MovieSceneEvaluationTemplateInstance.cpp
--- a/Engine/Source/Runtime/MovieScene/Private/Evaluation/MovieSceneEvaluationTemplateInstance.cpp
+++ b/Engine/Source/Runtime/MovieScene/Private/Evaluation/MovieSceneEvaluationTemplateInstance.cpp
@@ -255,11 +255,9 @@ void FMovieSceneRootEvaluationTemplateInstance::EnableGlobalPreAnimatedStateCapt
 {
 	using namespace UE::MovieScene;
 
-	const FRootInstanceHandle RootInstanceHandle = GetRootInstanceHandle();
-	if (ensure(EntitySystemLinker && RootInstanceHandle.IsValid()))
+	if (ensure(SharedPlaybackState))
 	{
-		const FSequenceInstance& Instance = EntitySystemLinker->GetInstanceRegistry()->GetInstance(RootInstanceHandle);
-		Instance.GetPlayer()->PreAnimatedState.EnableGlobalPreAnimatedStateCapture();
+		SharedPlaybackState->GetPreAnimatedState().EnableGlobalPreAnimatedStateCapture();
 	}
 }
 
@@ -435,7 +433,7 @@ void FMovieSceneRootEvaluationTemplateInstance::PlaybackContextChanged(IMovieSce
 	UMovieSceneSequence* RootSequence = SharedPlaybackState->GetRootSequence();
 	UMovieSceneCompiledDataManager* CompiledDataManager = SharedPlaybackState->GetCompiledDataManager();
 
-	const bool bGlobalCapture = Player.PreAnimatedState.IsCapturingGlobalPreAnimatedState();
+	const bool bGlobalCapture = SharedPlaybackState->GetPreAnimatedState().IsCapturingGlobalPreAnimatedState();
 	const FRootInstanceHandle PreviousRootInstanceHandle = SharedPlaybackState->GetRootInstanceHandle();
 
 	if (PreviousRootInstanceHandle.IsValid() &&
@@ -480,7 +478,7 @@ void FMovieSceneRootEvaluationTemplateInstance::PlaybackContextChanged(IMovieSce
 	Player.PreAnimatedState.Initialize(EntitySystemLinker, RootInstanceHandle);
 	if (bGlobalCapture)
 	{
-		Player.PreAnimatedState.EnableGlobalPreAnimatedStateCapture();
+		SharedPlaybackState->GetPreAnimatedState().EnableGlobalPreAnimatedStateCapture();
 	}
 }
 
diff --git a/Engine/Source/Runtime/MovieScene/Private/Evaluation/MovieScenePreAnimatedState.cpp b/Engine/Source/Runtime/MovieScene/Private/Evaluation/MovieScenePreAnimatedState.cpp
--- a/Engine/Source/Runtime/MovieScene/Private/Evaluation/MovieScenePreAnimatedState.cpp
+++ b/Engine/Source/Runtime/MovieScene/Private/Evaluation/MovieScenePreAnimatedState.cpp
@@ -4,6 +4,8 @@
 #include "GameFramework/Actor.h"
 #include "Components/ActorComponent.h"
 #include "EntitySystem/MovieSceneEntitySystemLinker.h"
+#include "EntitySystem/MovieSceneInstanceRegistry.h"
+#include "EntitySystem/MovieSceneSharedPlaybackState.h"
 #include "Evaluation/PreAnimatedState/MovieScenePreAnimatedStateExtension.h"
 #include "Evaluation/PreAnimatedState/MovieSceneRestoreStateParams.h"
 #include "Evaluation/PreAnimatedState/MovieScenePreAnimatedCaptureSources.h"
@@ -12,7 +14,14 @@
 
 DECLARE_CYCLE_STAT(TEXT("Save Pre Animated State"), MovieSceneEval_SavePreAnimatedState, STATGROUP_MovieSceneEval);
 
-FMovieScenePreAnimatedState::~FMovieScenePreAnimatedState()
+FMovieSceneInstancePreAnimatedState::FMovieSceneInstancePreAnimatedState(UMovieSceneEntitySystemLinker* InLinker, UE::MovieScene::FRootInstanceHandle InInstanceHandle)
+	: WeakLinker(InLinker)
+	, InstanceHandle(InInstanceHandle)
+	, bCapturingGlobalPreAnimatedState(false)
+{
+}
+
+FMovieSceneInstancePreAnimatedState::~FMovieSceneInstancePreAnimatedState()
 {
 	// Ensure that the global state capture request is removed
 	UMovieSceneEntitySystemLinker* CurrentLinker = WeakLinker.Get();
@@ -23,35 +32,17 @@ FMovieScenePreAnimatedState::~FMovieScenePreAnimatedState()
 	}
 }
 
-void FMovieScenePreAnimatedState::Initialize(UMovieSceneEntitySystemLinker* Linker, UE::MovieScene::FRootInstanceHandle InInstanceHandle)
-{
-	// If we're re-using a pre-animated state class and it was previously
-	// capturing global state make sure to decrement that request
-	UMovieSceneEntitySystemLinker* PreviousLinker = WeakLinker.Get();
-	if (PreviousLinker && bCapturingGlobalPreAnimatedState)
-	{
-		checkf(PreviousLinker->PreAnimatedState.NumRequestsForGlobalState > 0, TEXT("Increment/Decrement mismatch on FPreAnimatedState::NumRequestsForGlobalState"));
-		--PreviousLinker->PreAnimatedState.NumRequestsForGlobalState;
-
-	}
-
-	bCapturingGlobalPreAnimatedState = false;
-
-	WeakLinker = Linker;
-	InstanceHandle = InInstanceHandle;
-}
-
-UMovieSceneEntitySystemLinker* FMovieScenePreAnimatedState::GetLinker() const
+UMovieSceneEntitySystemLinker* FMovieSceneInstancePreAnimatedState::GetLinker() const
 {
 	return WeakLinker.Get();
 }
 
-bool FMovieScenePreAnimatedState::IsCapturingGlobalPreAnimatedState() const
+bool FMovieSceneInstancePreAnimatedState::IsCapturingGlobalPreAnimatedState() const
 {
 	return bCapturingGlobalPreAnimatedState;
 }
 
-void FMovieScenePreAnimatedState::EnableGlobalPreAnimatedStateCapture()
+void FMovieSceneInstancePreAnimatedState::EnableGlobalPreAnimatedStateCapture()
 {
 	if (bCapturingGlobalPreAnimatedState)
 	{
@@ -67,7 +58,7 @@ void FMovieScenePreAnimatedState::EnableGlobalPreAnimatedStateCapture()
 	}
 }
 
-void FMovieScenePreAnimatedState::SavePreAnimatedState(UObject& InObject, FMovieSceneAnimTypeID InTokenType, const IMovieScenePreAnimatedTokenProducer& Producer)
+void FMovieSceneInstancePreAnimatedState::SavePreAnimatedState(UObject& InObject, FMovieSceneAnimTypeID InTokenType, const IMovieScenePreAnimatedTokenProducer& Producer)
 {
 	using namespace UE::MovieScene;
 
@@ -83,7 +74,7 @@ void FMovieScenePreAnimatedState::SavePreAnimatedState(UObject& InObject, FMovie
 	}
 }
 
-void FMovieScenePreAnimatedState::SavePreAnimatedState(FMovieSceneAnimTypeID InTokenType, const IMovieScenePreAnimatedGlobalTokenProducer& Producer)
+void FMovieSceneInstancePreAnimatedState::SavePreAnimatedState(FMovieSceneAnimTypeID InTokenType, const IMovieScenePreAnimatedGlobalTokenProducer& Producer)
 {
 	using namespace UE::MovieScene;
 
@@ -99,7 +90,7 @@ void FMovieScenePreAnimatedState::SavePreAnimatedState(FMovieSceneAnimTypeID InT
 	}
 }
 
-void FMovieScenePreAnimatedState::RestorePreAnimatedState()
+void FMovieSceneInstancePreAnimatedState::RestorePreAnimatedState()
 {
 	using namespace UE::MovieScene;
 
@@ -110,7 +101,7 @@ void FMovieScenePreAnimatedState::RestorePreAnimatedState()
 	}
 }
 
-void FMovieScenePreAnimatedState::DiscardPreAnimatedState()
+void FMovieSceneInstancePreAnimatedState::DiscardPreAnimatedState()
 {
 	using namespace UE::MovieScene;
 
@@ -121,7 +112,7 @@ void FMovieScenePreAnimatedState::DiscardPreAnimatedState()
 	}
 }
 
-void FMovieScenePreAnimatedState::RestorePreAnimatedState(UObject& Object)
+void FMovieSceneInstancePreAnimatedState::RestorePreAnimatedState(UObject& Object)
 {
 	using namespace UE::MovieScene;
 
@@ -146,7 +137,7 @@ void FMovieScenePreAnimatedState::RestorePreAnimatedState(UObject& Object)
 	Linker->PreAnimatedState.RestoreStateForGroup(Group, FRestoreStateParams{Linker, InstanceHandle});
 }
 
-void FMovieScenePreAnimatedState::RestorePreAnimatedState(UClass* GeneratedClass)
+void FMovieSceneInstancePreAnimatedState::RestorePreAnimatedState(UClass* GeneratedClass)
 {
 	using namespace UE::MovieScene;
 
@@ -171,7 +162,7 @@ void FMovieScenePreAnimatedState::RestorePreAnimatedState(UClass* GeneratedClass
 }
 
 
-void FMovieScenePreAnimatedState::RestorePreAnimatedState(UObject& Object, TFunctionRef<bool(FMovieSceneAnimTypeID)> InFilter)
+void FMovieSceneInstancePreAnimatedState::RestorePreAnimatedState(UObject& Object, TFunctionRef<bool(FMovieSceneAnimTypeID)> InFilter)
 {
 	using namespace UE::MovieScene;
 
@@ -205,7 +196,7 @@ void FMovieScenePreAnimatedState::RestorePreAnimatedState(UObject& Object, TFunc
 	ObjectStorage->SetRestoreMask(nullptr);
 }
 
-void FMovieScenePreAnimatedState::DiscardEntityTokens()
+void FMovieSceneInstancePreAnimatedState::DiscardEntityTokens()
 {
 	using namespace UE::MovieScene;
 
@@ -216,7 +207,7 @@ void FMovieScenePreAnimatedState::DiscardEntityTokens()
 	}
 }
 
-void FMovieScenePreAnimatedState::DiscardAndRemoveEntityTokensForObject(UObject& Object)
+void FMovieSceneInstancePreAnimatedState::DiscardAndRemoveEntityTokensForObject(UObject& Object)
 {
 	using namespace UE::MovieScene;
 
@@ -241,27 +232,133 @@ void FMovieScenePreAnimatedState::DiscardAndRemoveEntityTokensForObject(UObject&
 	Linker->PreAnimatedState.DiscardStateForGroup(Group);
 }
 
-void FMovieScenePreAnimatedState::OnObjectsReplaced(const TMap<UObject*, UObject*>& ReplacementMap)
+bool FMovieSceneInstancePreAnimatedState::ContainsAnyStateForSequence() const
 {
 	using namespace UE::MovieScene;
 
 	UMovieSceneEntitySystemLinker* Linker = WeakLinker.Get();
-	if (!Linker)
+	return Linker && InstanceHandle.IsValid() && Linker->PreAnimatedState.ContainsAnyStateForInstanceHandle(InstanceHandle);
+}
+
+void FMovieScenePreAnimatedState::Initialize(UMovieSceneEntitySystemLinker* Linker, UE::MovieScene::FRootInstanceHandle InInstanceHandle)
+{
+	WeakLinker = Linker;
+	InstanceHandle = InInstanceHandle;
+}
+
+bool FMovieScenePreAnimatedState::IsCapturingGlobalPreAnimatedState() const
+{
+	if (FMovieSceneInstancePreAnimatedState* State = GetState())
 	{
-		return;
+		return State->IsCapturingGlobalPreAnimatedState();
 	}
+	return false;
+}
 
-	TSharedPtr<FPreAnimatedObjectGroupManager> ObjectGroupManager = Linker->PreAnimatedState.FindGroupManager<FPreAnimatedObjectGroupManager>();
-	if (ObjectGroupManager)
+void FMovieScenePreAnimatedState::EnableGlobalPreAnimatedStateCapture()
+{
+	if (FMovieSceneInstancePreAnimatedState* State = GetState())
+	{
+		State->EnableGlobalPreAnimatedStateCapture();
+	}
+}
+
+UMovieSceneEntitySystemLinker* FMovieScenePreAnimatedState::GetLinker() const
+{
+	return WeakLinker.Get();
+}
+
+void FMovieScenePreAnimatedState::SavePreAnimatedState(UObject& InObject, FMovieSceneAnimTypeID InTokenType, const IMovieScenePreAnimatedTokenProducer& Producer)
+{
+	if (FMovieSceneInstancePreAnimatedState* State = GetState())
+	{
+		State->SavePreAnimatedState(InObject, InTokenType, Producer);
+	}
+}
+
+void FMovieScenePreAnimatedState::SavePreAnimatedState(FMovieSceneAnimTypeID InTokenType, const IMovieScenePreAnimatedGlobalTokenProducer& Producer)
+{
+	if (FMovieSceneInstancePreAnimatedState* State = GetState())
+	{
+		State->SavePreAnimatedState(InTokenType, Producer);
+	}
+}
+
+void FMovieScenePreAnimatedState::RestorePreAnimatedState()
+{
+	if (FMovieSceneInstancePreAnimatedState* State = GetState())
+	{
+		State->RestorePreAnimatedState();
+	}
+}
+
+void FMovieScenePreAnimatedState::RestorePreAnimatedState(UObject& Object)
+{
+	if (FMovieSceneInstancePreAnimatedState* State = GetState())
+	{
+		State->RestorePreAnimatedState(Object);
+	}
+}
+
+void FMovieScenePreAnimatedState::RestorePreAnimatedState(UClass* GeneratedClass)
+{
+	if (FMovieSceneInstancePreAnimatedState* State = GetState())
+	{
+		State->RestorePreAnimatedState(GeneratedClass);
+	}
+}
+
+void FMovieScenePreAnimatedState::RestorePreAnimatedState(UObject& Object, TFunctionRef<bool(FMovieSceneAnimTypeID)> InFilter)
+{
+	if (FMovieSceneInstancePreAnimatedState* State = GetState())
+	{
+		State->RestorePreAnimatedState(Object, InFilter);
+	}
+}
+
+void FMovieScenePreAnimatedState::DiscardPreAnimatedState()
+{
+	if (FMovieSceneInstancePreAnimatedState* State = GetState())
+	{
+		State->DiscardPreAnimatedState();
+	}
+}
+
+void FMovieScenePreAnimatedState::DiscardEntityTokens()
+{
+	if (FMovieSceneInstancePreAnimatedState* State = GetState())
+	{
+		State->DiscardEntityTokens();
+	}
+}
+
+void FMovieScenePreAnimatedState::DiscardAndRemoveEntityTokensForObject(UObject& Object)
+{
+	if (FMovieSceneInstancePreAnimatedState* State = GetState())
 	{
-		ObjectGroupManager->OnObjectsReplaced(ReplacementMap);
+		State->DiscardAndRemoveEntityTokensForObject(Object);
 	}
 }
 
 bool FMovieScenePreAnimatedState::ContainsAnyStateForSequence() const
+{
+	if (FMovieSceneInstancePreAnimatedState* State = GetState())
+	{
+		return State->ContainsAnyStateForSequence();
+	}
+	return false;
+}
+
+FMovieSceneInstancePreAnimatedState* FMovieScenePreAnimatedState::GetState() const
 {
 	using namespace UE::MovieScene;
 
 	UMovieSceneEntitySystemLinker* Linker = WeakLinker.Get();
-	return Linker && InstanceHandle.IsValid() && Linker->PreAnimatedState.ContainsAnyStateForInstanceHandle(InstanceHandle);
+	if (Linker && InstanceHandle.IsValid())
+	{
+		const FSequenceInstance& Instance = Linker->GetInstanceRegistry()->GetInstance(InstanceHandle);
+		return &Instance.GetSharedPlaybackState()->GetPreAnimatedState();
+	}
+	return nullptr;
 }
+
diff --git a/Engine/Source/Runtime/MovieScene/Private/Evaluation/PreAnimatedState/MovieScenePreAnimatedStateExtension.cpp b/Engine/Source/Runtime/MovieScene/Private/Evaluation/PreAnimatedState/MovieScenePreAnimatedStateExtension.cpp
--- a/Engine/Source/Runtime/MovieScene/Private/Evaluation/PreAnimatedState/MovieScenePreAnimatedStateExtension.cpp
+++ b/Engine/Source/Runtime/MovieScene/Private/Evaluation/PreAnimatedState/MovieScenePreAnimatedStateExtension.cpp
@@ -933,6 +933,14 @@ void FPreAnimatedStateExtension::AddReferencedObjects(UMovieSceneEntitySystemLin
 	}
 }
 
+void FPreAnimatedStateExtension::OnObjectsReplaced(const TMap<UObject*, UObject*>& ReplacementMap)
+{
+	TSharedPtr<FPreAnimatedObjectGroupManager> ObjectGroupManager = FindGroupManager<FPreAnimatedObjectGroupManager>();
+	if (ObjectGroupManager)
+	{
+		ObjectGroupManager->OnObjectsReplaced(ReplacementMap);
+	}
+}
 
 } // namespace MovieScene
 } // namespace UE
diff --git a/Engine/Source/Runtime/MovieScene/Private/Tests/MovieScenePreAnimatedStateTests.cpp b/Engine/Source/Runtime/MovieScene/Private/Tests/MovieScenePreAnimatedStateTests.cpp
--- a/Engine/Source/Runtime/MovieScene/Private/Tests/MovieScenePreAnimatedStateTests.cpp
+++ b/Engine/Source/Runtime/MovieScene/Private/Tests/MovieScenePreAnimatedStateTests.cpp
@@ -160,8 +160,8 @@ bool FMovieScenePreAnimatedStateGlobalTest::RunTest(const FString& Parameters)
 
 	ResetValues();
 
-	FMovieScenePreAnimatedState State;
-	State.Initialize(Linker, FRootInstanceHandle());
+	TSharedRef<FSharedPlaybackState> PlaybackState = MakeShared<FSharedPlaybackState>(Linker);
+	FMovieSceneInstancePreAnimatedState& State(PlaybackState->GetPreAnimatedState());
 	State.EnableGlobalPreAnimatedStateCapture();
 
 	FPreAnimatedTokenProducer Producer(&TestValue1);
@@ -194,14 +194,14 @@ bool FMovieScenePreAnimatedStateEntityTest::RunTest(const FString& Parameters)
 
 	UMovieSceneEntitySystemLinker* Linker = GetTestLinker();
 
-	FMovieScenePreAnimatedState State;
-	State.Initialize(Linker, FRootInstanceHandle());
+	TSharedRef<FSharedPlaybackState> PlaybackState = MakeShared<FSharedPlaybackState>(Linker);
+	FMovieSceneInstancePreAnimatedState& State(PlaybackState->GetPreAnimatedState());
 	State.EnableGlobalPreAnimatedStateCapture();
 
 	FPreAnimatedTokenProducer Producer(&TestValue1);
 
 	{
-		FScopedPreAnimatedCaptureSource CaptureSource(&State, SectionKey1, true);
+		FScopedPreAnimatedCaptureSource CaptureSource(PlaybackState, SectionKey1, true);
 		State.SavePreAnimatedState(AnimType1, Producer);
 	}
 
@@ -231,8 +231,8 @@ bool FMovieScenePreAnimatedStateOverlappingEntitiesTest::RunTest(const FString&
 
 	UMovieSceneEntitySystemLinker* Linker = GetTestLinker();
 
-	FMovieScenePreAnimatedState State;
-	State.Initialize(Linker, FRootInstanceHandle());
+	TSharedRef<FSharedPlaybackState> PlaybackState = MakeShared<FSharedPlaybackState>(Linker);
+	FMovieSceneInstancePreAnimatedState& State(PlaybackState->GetPreAnimatedState());
 	State.EnableGlobalPreAnimatedStateCapture();
 
 	FPreAnimatedTokenProducer Producer(&TestValue1);
@@ -247,7 +247,7 @@ bool FMovieScenePreAnimatedStateOverlappingEntitiesTest::RunTest(const FString&
 
 	// 2. Save a token for the track's evaluation
 	{
-		FScopedPreAnimatedCaptureSource CaptureSource(&State, TrackKey1, true);
+		FScopedPreAnimatedCaptureSource CaptureSource(PlaybackState, TrackKey1, true);
 		State.SavePreAnimatedState(AnimType1, Producer);
 
 		TestValue1 = 50;
@@ -256,7 +256,7 @@ bool FMovieScenePreAnimatedStateOverlappingEntitiesTest::RunTest(const FString&
 
 	// 3. Save a token for the section's evaluation
 	{
-		FScopedPreAnimatedCaptureSource CaptureSource(&State, SectionKey1, true);
+		FScopedPreAnimatedCaptureSource CaptureSource(PlaybackState, SectionKey1, true);
 		State.SavePreAnimatedState(AnimType1, Producer);
 
 		TestValue1 = 100;
@@ -265,7 +265,7 @@ bool FMovieScenePreAnimatedStateOverlappingEntitiesTest::RunTest(const FString&
 
 	// 4. Save a token for another section's evaluation
 	{
-		FScopedPreAnimatedCaptureSource CaptureSource(&State, SectionKey2, true);
+		FScopedPreAnimatedCaptureSource CaptureSource(PlaybackState, SectionKey2, true);
 		State.SavePreAnimatedState(AnimType1, Producer);
 
 		TestValue1 = 150;
@@ -303,15 +303,15 @@ bool FMovieScenePreAnimatedStateKeepThenRestoreEntityTest::RunTest(const FString
 
 	UMovieSceneEntitySystemLinker* Linker = GetTestLinker();
 
-	FMovieScenePreAnimatedState State;
-	State.Initialize(Linker, FRootInstanceHandle());
+	TSharedRef<FSharedPlaybackState> PlaybackState = MakeShared<FSharedPlaybackState>(Linker);
+	FMovieSceneInstancePreAnimatedState& State(PlaybackState->GetPreAnimatedState());
 	State.EnableGlobalPreAnimatedStateCapture();
 
 	FPreAnimatedTokenProducer Producer(&TestValue1);
 
 	// Indicate that the entity should not capture state
 	{
-		FScopedPreAnimatedCaptureSource CaptureSource(&State, SectionKey1, false);
+		FScopedPreAnimatedCaptureSource CaptureSource(PlaybackState, SectionKey1, false);
 		// Save state - this will only save globally
 		State.SavePreAnimatedState(AnimType1, Producer);
 	}
@@ -327,7 +327,7 @@ bool FMovieScenePreAnimatedStateKeepThenRestoreEntityTest::RunTest(const FString
 
 	{
 		// Indicate that SectionKey2 is now animating, and wants to restore state
-		FScopedPreAnimatedCaptureSource CaptureSource(&State, SectionKey2, true);
+		FScopedPreAnimatedCaptureSource CaptureSource(PlaybackState, SectionKey2, true);
 		State.SavePreAnimatedState(AnimType1, Producer);
 	}
 
@@ -527,14 +527,14 @@ bool FMovieScenePreAnimatedStatePerformanceTest::RunTest(const FString& Paramete
 
 	UMovieSceneEntitySystemLinker* Linker = GetTestLinker();
 
-	FMovieScenePreAnimatedState State;
-	State.Initialize(Linker, FRootInstanceHandle());
+	TSharedRef<FSharedPlaybackState> PlaybackState = MakeShared<FSharedPlaybackState>(Linker);
+	FMovieSceneInstancePreAnimatedState& State(PlaybackState->GetPreAnimatedState());
 	State.EnableGlobalPreAnimatedStateCapture();
 
 	FPreAnimatedTokenProducer Producer(&TestValue1);
 
 	// Indicate that the entity should not capture state
-	FScopedPreAnimatedCaptureSource CaptureSource(&State, SectionKey1, false);
+	FScopedPreAnimatedCaptureSource CaptureSource(PlaybackState, SectionKey1, false);
 
 	for (int32 Iteration = 0; Iteration < 1000000; ++Iteration)
 	{
diff --git a/Engine/Source/Runtime/MovieScene/Public/EntitySystem/MovieSceneSharedPlaybackState.h b/Engine/Source/Runtime/MovieScene/Public/EntitySystem/MovieSceneSharedPlaybackState.h
--- a/Engine/Source/Runtime/MovieScene/Public/EntitySystem/MovieSceneSharedPlaybackState.h
+++ b/Engine/Source/Runtime/MovieScene/Public/EntitySystem/MovieSceneSharedPlaybackState.h
@@ -7,6 +7,7 @@
 #include "EntitySystem/MovieSceneSequenceInstanceHandle.h"
 #include "Evaluation/MovieSceneEvaluationOperand.h"
 #include "Evaluation/MovieScenePlaybackCapabilities.h"
+#include "Evaluation/MovieScenePreAnimatedState.h"
 #include "MovieSceneSequenceID.h"
 
 class FMovieSceneEntitySystemRunner;
@@ -25,8 +26,6 @@ struct FSharedPlaybackStateCreateParams
 {
 	/**
 	 * The playback context in which the root sequence will be evaluated.
-	 *
-	 * Requires that RootInstanceHandle and Linker are also set.
 	 */
 	UObject* PlaybackContext = nullptr;
 
@@ -34,16 +33,12 @@ struct FSharedPlaybackStateCreateParams
 	 * The handle of the root sequence instance, if the created playback state
 	 * is meant to relate to an instance that has also been created inside
 	 * a runner/linker's instance registry.
-	 *
-	 * Requires that PlaybackContext and Linker are also set.
 	 */
 	FRootInstanceHandle RootInstanceHandle;
 
 	/**
 	 * The linker that will be evaluating the sequence that the created playback
 	 * state relates to.
-	 *
-	 * Requires that PlaybackContext and RootInstanceHandle are also set.
 	 */
 	TObjectPtr<UMovieSceneEntitySystemLinker> Linker;
 
@@ -61,7 +56,7 @@ struct MOVIESCENE_API FSharedPlaybackState : TSharedFromThis<FSharedPlaybackStat
 {
 public:
 
-	FSharedPlaybackState();
+	FSharedPlaybackState(UMovieSceneEntitySystemLinker* InLinker = nullptr);
 	FSharedPlaybackState(
 			UMovieSceneSequence& InRootSequence,
 			const FSharedPlaybackStateCreateParams& CreateParams);
@@ -86,6 +81,12 @@ public:
 	/** Gets the compiled data ID for the root sequence */
 	const FMovieSceneCompiledDataID&  GetRootCompiledDataID() const { return RootCompiledDataID; }
 
+	/** Gets the pre-animated state utility for the sequence hierarchy */
+	const FMovieSceneInstancePreAnimatedState& GetPreAnimatedState() const { return PreAnimatedState; }
+
+	/** Gets the pre-animated state utility for the sequence hierarchy */
+	FMovieSceneInstancePreAnimatedState& GetPreAnimatedState() { return PreAnimatedState; }
+
 public:
 
 	// General utility methods
@@ -314,6 +315,9 @@ private:
 
 	/** Playback capabilities for the root sequence */
 	FPlaybackCapabilities Capabilities;
+
+	/** Pre-animated state utility for the sequence hierarchy */
+	FMovieSceneInstancePreAnimatedState PreAnimatedState;
 };
 
 } // namespace UE::MovieScene
diff --git a/Engine/Source/Runtime/MovieScene/Public/Evaluation/MovieScenePreAnimatedState.h b/Engine/Source/Runtime/MovieScene/Public/Evaluation/MovieScenePreAnimatedState.h
--- a/Engine/Source/Runtime/MovieScene/Public/Evaluation/MovieScenePreAnimatedState.h
+++ b/Engine/Source/Runtime/MovieScene/Public/Evaluation/MovieScenePreAnimatedState.h
@@ -9,8 +9,6 @@
 #include "UObject/WeakObjectPtr.h"
 #include "UObject/WeakObjectPtrTemplates.h"
 
-class FMovieScenePreAnimatedState;
-class IMovieScenePlayer;
 class UClass;
 class UMovieSceneEntitySystemLinker;
 class UObject;
@@ -24,18 +22,16 @@ template <typename FuncType> class TFunctionRef;
 /**
  * Class that caches pre-animated state for objects that were manipulated by sequencer
  */
-class FMovieScenePreAnimatedState
+class FMovieSceneInstancePreAnimatedState
 {
 public:
 
-	FMovieScenePreAnimatedState() = default;
-
-	FMovieScenePreAnimatedState(const FMovieScenePreAnimatedState&) = delete;
-	FMovieScenePreAnimatedState& operator=(const FMovieScenePreAnimatedState&) = delete;
+	FMovieSceneInstancePreAnimatedState(UMovieSceneEntitySystemLinker* InLinker, UE::MovieScene::FRootInstanceHandle InInstanceHandle);
 
-	MOVIESCENE_API ~FMovieScenePreAnimatedState();
+	FMovieSceneInstancePreAnimatedState(const FMovieSceneInstancePreAnimatedState&) = delete;
+	FMovieSceneInstancePreAnimatedState& operator=(const FMovieSceneInstancePreAnimatedState&) = delete;
 
-	MOVIESCENE_API void Initialize(UMovieSceneEntitySystemLinker* Linker, UE::MovieScene::FRootInstanceHandle InstanceHandle);
+	MOVIESCENE_API ~FMovieSceneInstancePreAnimatedState();
 
 	/**
 	 * Check whether this sequence instance is capturing any and all changes of state so they can be restored later
@@ -91,11 +87,6 @@ public:
 	 */
 	MOVIESCENE_API void DiscardAndRemoveEntityTokensForObject(UObject& Object);
 
-	/**
-	 * Called when objects have been replaced so that pre animated state can swap out to the new objects
-	 */
-	MOVIESCENE_API void OnObjectsReplaced(const TMap<UObject*, UObject*>& ReplacementMap);
-
 public:
 
 	/**
@@ -110,13 +101,51 @@ private:
 
 private:
 
-	friend struct FScopedPreAnimatedCaptureSource;
-
 	/** Weak pointer to the linker that we're associated with */
 	TWeakObjectPtr<UMovieSceneEntitySystemLinker> WeakLinker;
 
 	/** The instance handle for the root sequence instance */
 	UE::MovieScene::FRootInstanceHandle InstanceHandle;
 
+	/** Whether global pre-animated state capture has been requested on the linker extension */
 	bool bCapturingGlobalPreAnimatedState;
 };
+
+/**
+ * A class similar to FMovieSceneInstancePreAnimatedState (above), which doesn't manage global pre-animated state 
+ * requests, deferring those to the other class. This class is mostly for IMovieScenePlayer, for backwards
+ * compatibility.
+ */
+class FMovieScenePreAnimatedState
+{
+public:
+
+	FMovieScenePreAnimatedState() = default;
+	FMovieScenePreAnimatedState(const FMovieScenePreAnimatedState&) = delete;
+	FMovieScenePreAnimatedState& operator=(const FMovieScenePreAnimatedState&) = delete;
+
+	MOVIESCENE_API void Initialize(UMovieSceneEntitySystemLinker* Linker, UE::MovieScene::FRootInstanceHandle InstanceHandle);
+	MOVIESCENE_API bool IsCapturingGlobalPreAnimatedState() const;
+	MOVIESCENE_API void EnableGlobalPreAnimatedStateCapture();
+	MOVIESCENE_API UMovieSceneEntitySystemLinker* GetLinker() const;
+	MOVIESCENE_API void SavePreAnimatedState(UObject& InObject, FMovieSceneAnimTypeID InTokenType, const IMovieScenePreAnimatedTokenProducer& Producer);
+	MOVIESCENE_API void SavePreAnimatedState(FMovieSceneAnimTypeID InTokenType, const IMovieScenePreAnimatedGlobalTokenProducer& Producer);
+	MOVIESCENE_API void RestorePreAnimatedState();
+	MOVIESCENE_API void RestorePreAnimatedState(UObject& Object);
+	MOVIESCENE_API void RestorePreAnimatedState(UClass* GeneratedClass);
+	MOVIESCENE_API void RestorePreAnimatedState(UObject& Object, TFunctionRef<bool(FMovieSceneAnimTypeID)> InFilter);
+	MOVIESCENE_API void DiscardPreAnimatedState();
+	MOVIESCENE_API void DiscardEntityTokens();
+	MOVIESCENE_API void DiscardAndRemoveEntityTokensForObject(UObject& Object);
+	MOVIESCENE_API bool ContainsAnyStateForSequence() const;
+
+private:
+
+	FMovieSceneInstancePreAnimatedState* GetState() const;
+
+	TWeakObjectPtr<UMovieSceneEntitySystemLinker> WeakLinker;
+	UE::MovieScene::FRootInstanceHandle InstanceHandle;
+
+	friend struct FScopedPreAnimatedCaptureSource;
+};
+
diff --git a/Engine/Source/Runtime/MovieScene/Public/Evaluation/PreAnimatedState/MovieScenePreAnimatedStateExtension.h b/Engine/Source/Runtime/MovieScene/Public/Evaluation/PreAnimatedState/MovieScenePreAnimatedStateExtension.h
--- a/Engine/Source/Runtime/MovieScene/Public/Evaluation/PreAnimatedState/MovieScenePreAnimatedStateExtension.h
+++ b/Engine/Source/Runtime/MovieScene/Public/Evaluation/PreAnimatedState/MovieScenePreAnimatedStateExtension.h
@@ -43,6 +43,7 @@ struct FPreAnimatedTemplateCaptureSources;
 struct FPreAnimatedTrackInstanceCaptureSources;
 struct FPreAnimatedTrackInstanceInputCaptureSources;
 struct FRestoreStateParams;
+struct FSharedPlaybackState;
 struct IPreAnimatedCaptureSource;
 struct IPreAnimatedStorage;
 
@@ -337,6 +338,9 @@ private:
 
 public:
 
+	/** Called to handle replaced objects */
+	void OnObjectsReplaced(const TMap<UObject*, UObject*>& ReplacementMap);
+
 	/** The number of requests that have been made to capture global state - only one should exist per playing sequence */
 	uint32 NumRequestsForGlobalState;
 
