Sequencer- Fix tooltip for 'Continuously Respawn' in Binding Properties on Spawnable Base.

diff --git a/Engine/Source/Runtime/MovieScene/Public/Bindings/MovieSceneSpawnableBinding.h b/Engine/Source/Runtime/MovieScene/Public/Bindings/MovieSceneSpawnableBinding.h
--- a/Engine/Source/Runtime/MovieScene/Public/Bindings/MovieSceneSpawnableBinding.h
+++ b/Engine/Source/Runtime/MovieScene/Public/Bindings/MovieSceneSpawnableBinding.h
@@ -70,7 +70,7 @@ public:
 	UPROPERTY(EditAnywhere, Category=Spawnable)
 	ESpawnOwnership SpawnOwnership = ESpawnOwnership::InnerSequence;
 
-	/** When enabled, this spawnable will always be respawned if it gets destroyed externally. When disabled, this object will only ever be spawned once for each spawn key even if it is destroyed externally. */
+	/** When enabled, this spawnable will always be respawned if it gets destroyed externally. When disabled, this object will only ever be spawned once for each binding lifetime section even if destroyed externally. */
 	UPROPERTY(EditAnywhere, Category= Spawnable)
 	bool bContinuouslyRespawn = false;
 
