Sequencer- Fix regression on name of 'Spawn Ownership' entry "This Sequence" which has turned to "Inner Sequence".

diff --git a/Engine/Source/Runtime/MovieScene/Public/MovieSceneSpawnable.h b/Engine/Source/Runtime/MovieScene/Public/MovieSceneSpawnable.h
--- a/Engine/Source/Runtime/MovieScene/Public/MovieSceneSpawnable.h
+++ b/Engine/Source/Runtime/MovieScene/Public/MovieSceneSpawnable.h
@@ -22,7 +22,8 @@ UENUM()
 enum class ESpawnOwnership : uint8
 {
 	/** The object's lifetime is managed by the sequence that spawned it */
-	InnerSequence,
+	
+	InnerSequence UMETA(DisplayName = "This Sequence"),
 
 	/** The object's lifetime is managed by the outermost sequence */
 	RootSequence,
