Sequencer: Fix SetPlaybackPosition when there is no level sequence asset set on the player

diff --git a/Engine/Source/Runtime/MovieScene/Private/MovieSceneSequencePlayer.cpp b/Engine/Source/Runtime/MovieScene/Private/MovieSceneSequencePlayer.cpp
--- a/Engine/Source/Runtime/MovieScene/Private/MovieSceneSequencePlayer.cpp
+++ b/Engine/Source/Runtime/MovieScene/Private/MovieSceneSequencePlayer.cpp
@@ -690,6 +690,11 @@ void UMovieSceneSequencePlayer::SetPlaybackPosition(FMovieSceneSequencePlaybackP
 		return;
 	}
 
+	if (!Sequence)
+	{
+		return;
+	}
+
 	if (NeedsQueueLatentAction())
 	{
 		QueueLatentAction(FMovieSceneSequenceLatentActionDelegate::CreateUObject(this, &UMovieSceneSequencePlayer::SetPlaybackPosition, InPlaybackParams));
