Sequencer: Fix divide by zero

Fix for when duration (frames) is 0, but subframes is non zero.
Determine looping in tick resolution, rather than in display rate.

diff --git a/Engine/Source/Runtime/MovieScene/Private/MovieSceneSequencePlayer.cpp b/Engine/Source/Runtime/MovieScene/Private/MovieSceneSequencePlayer.cpp
--- a/Engine/Source/Runtime/MovieScene/Private/MovieSceneSequencePlayer.cpp
+++ b/Engine/Source/Runtime/MovieScene/Private/MovieSceneSequencePlayer.cpp
@@ -1157,9 +1157,11 @@ void UMovieSceneSequencePlayer::UpdateTimeCursorPosition_Internal(FFrameTime New
 		// The actual end time taking into account reverse playback
 		FFrameTime EndTimeWithReversed = bReversePlayback ? StartTime : GetLastValidTime().FrameNumber;
 
-		FFrameTime PositionRelativeToStart = NewPosition.FrameNumber - StartTimeWithReversed;
+		// Operate in tick resolution (for subframes)
+		FFrameTime PositionRelativeToStart = ConvertFrameTime(NewPosition.FrameNumber - StartTimeWithReversed, PlayPosition.GetInputRate(), PlayPosition.GetOutputRate());
+		FFrameTime DurationWithSubFrames = ConvertFrameTime(GetDuration().Time, PlayPosition.GetInputRate(), PlayPosition.GetOutputRate());
 
-		const int32 NumTimesLooped    = FMath::Abs(PositionRelativeToStart.FrameNumber.Value / Duration);
+		const int32 NumTimesLooped    = FMath::Abs(PositionRelativeToStart.FrameNumber.Value / DurationWithSubFrames.FrameNumber.Value);
 		const bool  bLoopIndefinitely = PlaybackSettings.LoopCount.Value < 0;
 
 		// loop playback
@@ -1183,7 +1185,7 @@ void UMovieSceneSequencePlayer::UpdateTimeCursorPosition_Internal(FFrameTime New
 				UpdateMovieSceneInstance(Range, StatusOverride);
 			}
 
-			const FFrameTime Overplay = FFrameTime(PositionRelativeToStart.FrameNumber.Value % Duration, PositionRelativeToStart.GetSubFrame());
+			const FFrameTime Overplay = FFrameTime(ConvertFrameTime(PositionRelativeToStart.FrameNumber.Value % DurationWithSubFrames.FrameNumber.Value, PlayPosition.GetOutputRate(), PlayPosition.GetInputRate()).FrameNumber, PositionRelativeToStart.GetSubFrame());
 			FFrameTime NewFrameOffset;
 			
 			if (bReversePlayback)
