Sequencer: Smart Tangents, avoid precision issues when values are the same by automatically giving them zero tangents when both are auto.

diff --git a/Engine/Source/Runtime/MovieScene/Private/Channels/MovieSceneCurveChannelImpl.cpp b/Engine/Source/Runtime/MovieScene/Private/Channels/MovieSceneCurveChannelImpl.cpp
--- a/Engine/Source/Runtime/MovieScene/Private/Channels/MovieSceneCurveChannelImpl.cpp
+++ b/Engine/Source/Runtime/MovieScene/Private/Channels/MovieSceneCurveChannelImpl.cpp
@@ -902,7 +902,6 @@ ChannelValueType ClampTangent(ChannelValueType NewTangent, ChannelValueType Prev
 	}
 	return NewTangent;
 };
-
 template<typename ChannelType>
 float TMovieSceneCurveChannelImpl<ChannelType>::CalcSmartTangent(ChannelType* InChannel, int32 Index)
 {
@@ -914,21 +913,23 @@ float TMovieSceneCurveChannelImpl<ChannelType>::CalcSmartTangent(ChannelType* In
 	ChannelValueType NextKey = InChannel->Values[NextIndex];
 
 	float NewTangent = 0.0f;
-	// if key doesn't lie between we keep it flat(0.0).
-	if ((ThisKey.Value > PrevKey.Value && ThisKey.Value < NextKey.Value) ||
+	// if key doesn't lie between we keep it flat(0.0), except if using auto tangent option 2 since that handles it automatically
+    // and let's us do blending to tangents that have overshoot
+	if ((GCachedSequencerAutoTangentInterpolation == 2)|| (ThisKey.Value > PrevKey.Value && ThisKey.Value < NextKey.Value) ||
 		(ThisKey.Value < PrevKey.Value && ThisKey.Value > NextKey.Value))
 	{
-		
-		while (NextIndex <= (InChannel->Values.Num() - 2)
-			&& FMath::IsNearlyZero(NextKey.Tangent.ArriveTangent)
-			&& NextKey.InterpMode == RCIM_Cubic && (NextKey.TangentMode == RCTM_Auto || NextKey.TangentMode == RCTM_SmartAuto)
-			&& ((NextKey.Value > ThisKey.Value && NextKey.Value < InChannel->Values[NextIndex + 1].Value) ||
-				(NextKey.Value < ThisKey.Value && NextKey.Value > InChannel->Values[NextIndex + 1].Value)))
+		if (GCachedSequencerAutoTangentInterpolation != 2) //for older versions we try to match over a longer period of tie, not needed with improved flattening
 		{
-			++NextIndex;
-			NextKey = InChannel->Values[NextIndex];
+			while (NextIndex <= (InChannel->Values.Num() - 2)
+				&& FMath::IsNearlyZero(NextKey.Tangent.ArriveTangent)
+				&& NextKey.InterpMode == RCIM_Cubic && (NextKey.TangentMode == RCTM_Auto || NextKey.TangentMode == RCTM_SmartAuto)
+				&& ((NextKey.Value > ThisKey.Value && NextKey.Value < InChannel->Values[NextIndex + 1].Value) ||
+					(NextKey.Value < ThisKey.Value && NextKey.Value > InChannel->Values[NextIndex + 1].Value)))
+			{
+				++NextIndex;
+				NextKey = InChannel->Values[NextIndex];
+			}
 		}
-		
 
 		const CurveValueType OneThird = 1.0 / 3.0;
 		const CurveValueType TwoThird = 2.0 / 3.0;
@@ -945,50 +946,45 @@ float TMovieSceneCurveChannelImpl<ChannelType>::CalcSmartTangent(ChannelType* In
 		NewTangent = ((ThisKey.Value - PrevY) + (NextY - ThisKey.Value))
 			/ (TwoThird * TimeToPrevious + TwoThird * TimeToNext);
 
-		const double PreviousSlope = (ThisKey.Value - PrevY) / (TwoThird * TimeToPrevious);
-		const double NextSlope = (NextY - ThisKey.Value) / (TwoThird * TimeToNext);
-		/*
-		* Leaving this in for now but averaging is over agressive.
 		
-		if (NextKey.Value > PrevKey.Value)
+		if (GCachedSequencerAutoTangentInterpolation == 2) //use flattening, no overshoot
 		{
-			CurveValueType NextToPrev = NextKey.Value - PrevKey.Value;
-			CurveValueType ToNext = (NextKey.Value - ThisKey.Value) / NextToPrev;
-			CurveValueType ToPrev = (ThisKey.Value - PrevKey.Value) / NextToPrev;
-
-			//NewTangent = (PreviousSlope * ToNext)+ (ToPrev * NextSlope);
-
+			//if two keys are equivalent in value and both auto tangent is zero
+			if (FMath::IsNearlyEqual(ThisKey.Value, NextKey.Value) && NextKey.InterpMode == RCIM_Cubic && (NextKey.TangentMode == RCTM_Auto || NextKey.TangentMode == RCTM_SmartAuto))
+			{
+				NewTangent = 0.0;
+			}
+			else
+			{
+				const double PreviousSlope = (ThisKey.Value - PrevY) / (TwoThird * TimeToPrevious);
+				const double NextSlope = (NextY - ThisKey.Value) / (TwoThird * TimeToNext);
+				NewTangent = ClampTangent<double>(NewTangent, PreviousSlope, NextSlope);
+			}
 		}
 		else
 		{
-			CurveValueType NextToPrev = PrevKey.Value - NextKey.Value;
-			CurveValueType ToNext = (ThisKey.Value - NextKey.Value) / NextToPrev;
-			CurveValueType ToPrev = (PrevKey.Value  - ThisKey.Value) / NextToPrev;
-
-			//NewTangent = (PreviousSlope * ToNext) + (ToPrev * NextSlope);
-		}
-		*/
-		const float BlendToNextRange = CVarSequencerSmartAutoBlendLocationPercentage->GetFloat();
-		const double ValDiff = FMath::Abs<double>(NextKey.Value - PrevKey.Value);
-		const double OurDiff = FMath::Abs<double>(ThisKey.Value - PrevKey.Value);
-		//ValDiff won't be zero ever due to previous check
-		double PercDiff = OurDiff / ValDiff;
-		float NextTangent = (NextKey.InterpMode == RCIM_Cubic && (NextKey.TangentMode == RCTM_Auto || NextKey.TangentMode == RCTM_SmartAuto)) ?
-			0.0 : NextKey.Tangent.ArriveTangent;
-		float PrevTangent = (PrevKey.InterpMode == RCIM_Cubic && (PrevKey.TangentMode == RCTM_Auto || PrevKey.TangentMode == RCTM_SmartAuto)) ?
-			0.0 : PrevKey.Tangent.LeaveTangent;
-		if (BlendToNextRange >= 0.0f && BlendToNextRange <= 1.0f)
-		{
-			NextTangent = PrevTangent = 0.0;
-			if (PercDiff > BlendToNextRange)
-			{
-				PercDiff = (PercDiff - BlendToNextRange) / (1.0 - BlendToNextRange);
-				NewTangent = NewTangent * (1.0 - PercDiff) + (PercDiff * NextTangent);
-			}
-			else if (PercDiff < (1.0 - BlendToNextRange))
+			const float BlendToNextRange = CVarSequencerSmartAutoBlendLocationPercentage->GetFloat();
+			const double ValDiff = FMath::Abs<double>(NextKey.Value - PrevKey.Value);
+			const double OurDiff = FMath::Abs<double>(ThisKey.Value - PrevKey.Value);
+			//ValDiff won't be zero ever due to previous check
+			double PercDiff = OurDiff / ValDiff;
+			float NextTangent = (NextKey.InterpMode == RCIM_Cubic && (NextKey.TangentMode == RCTM_Auto || NextKey.TangentMode == RCTM_SmartAuto)) ?
+				0.0 : NextKey.Tangent.ArriveTangent;
+			float PrevTangent = (PrevKey.InterpMode == RCIM_Cubic && (PrevKey.TangentMode == RCTM_Auto || PrevKey.TangentMode == RCTM_SmartAuto)) ?
+				0.0 : PrevKey.Tangent.LeaveTangent;
+			if (BlendToNextRange >= 0.0f && BlendToNextRange <= 1.0f)
 			{
-				PercDiff = PercDiff / (1.0 - BlendToNextRange);
-				NewTangent = NewTangent * PercDiff + (1.0 - PercDiff) * PrevTangent;
+				NextTangent = PrevTangent = 0.0;
+				if (PercDiff > BlendToNextRange)
+				{
+					PercDiff = (PercDiff - BlendToNextRange) / (1.0 - BlendToNextRange);
+					NewTangent = NewTangent * (1.0 - PercDiff) + (PercDiff * NextTangent);
+				}
+				else if (PercDiff < (1.0 - BlendToNextRange))
+				{
+					PercDiff = PercDiff / (1.0 - BlendToNextRange);
+					NewTangent = NewTangent * PercDiff + (1.0 - PercDiff) * PrevTangent;
+				}
 			}
 		}
 	}
