Sequencer: Updated bitflags to use new UENUM(Flags) functionality to ensure correct serialization

diff --git a/Engine/Source/Runtime/MovieScene/Private/Compilation/MovieSceneCompiledDataManager.cpp b/Engine/Source/Runtime/MovieScene/Private/Compilation/MovieSceneCompiledDataManager.cpp
--- a/Engine/Source/Runtime/MovieScene/Private/Compilation/MovieSceneCompiledDataManager.cpp
+++ b/Engine/Source/Runtime/MovieScene/Private/Compilation/MovieSceneCompiledDataManager.cpp
@@ -405,25 +405,25 @@ void UMovieSceneCompiledDataManager::CopyCompiledData(UMovieSceneSequence* Seque
 	if (const FMovieSceneSequenceHierarchy* Hierarchy = FindHierarchy(DataID))
 	{
 		CompiledData->Hierarchy = *Hierarchy;
-		CompiledData->AllocatedMask.bHierarchy = true;
+		CompiledData->AllocatedMask |= EMovieSceneSequenceCompilerMask::Hierarchy;
 	}
 	if (const FMovieSceneEvaluationTemplate* TrackTemplate = FindTrackTemplate(DataID))
 	{
 		CompiledData->EvaluationTemplate = *TrackTemplate;
-		CompiledData->AllocatedMask.bEvaluationTemplate = true;
+		CompiledData->AllocatedMask |= EMovieSceneSequenceCompilerMask::EvaluationTemplate;
 	}
 	if (const FMovieSceneEvaluationField* TrackTemplateField = FindTrackTemplateField(DataID))
 	{
 		if (Sequence->IsPlayableDirectly())
 		{
 			CompiledData->TrackTemplateField = *TrackTemplateField;
-			CompiledData->AllocatedMask.bEvaluationTemplateField = true;
+			CompiledData->AllocatedMask  |= EMovieSceneSequenceCompilerMask::EvaluationTemplateField;
 		}
 	}
 	if (const FMovieSceneEntityComponentField* EntityComponentField = FindEntityComponentField(DataID))
 	{
 		CompiledData->EntityComponentField = *EntityComponentField;
-		CompiledData->AllocatedMask.bEntityComponentField = true;
+		CompiledData->AllocatedMask |= EMovieSceneSequenceCompilerMask::EntityComponentField;
 	}
 
 	const FMovieSceneCompiledDataEntry& DataEntry = CompiledDataEntries[DataID.Value];
@@ -451,19 +451,19 @@ void UMovieSceneCompiledDataManager::LoadCompiledData(UMovieSceneSequence* Seque
 			return;
 		}
 
-		if (CompiledData->AllocatedMask.bHierarchy)
+		if (EnumHasAnyFlags(CompiledData->AllocatedMask, EMovieSceneSequenceCompilerMask::Hierarchy))
 		{
 			Hierarchies.Add(DataID.Value, MoveTemp(CompiledData->Hierarchy));
 		}
-		if (CompiledData->AllocatedMask.bEvaluationTemplate)
+		if (EnumHasAnyFlags(CompiledData->AllocatedMask, EMovieSceneSequenceCompilerMask::EvaluationTemplate))
 		{
 			TrackTemplates.Add(DataID.Value, MoveTemp(CompiledData->EvaluationTemplate));
 		}
-		if (CompiledData->AllocatedMask.bEvaluationTemplateField)
+		if (EnumHasAnyFlags(CompiledData->AllocatedMask, EMovieSceneSequenceCompilerMask::EvaluationTemplateField))
 		{
 			TrackTemplateFields.Add(DataID.Value, MoveTemp(CompiledData->TrackTemplateField));
 		}
-		if (CompiledData->AllocatedMask.bEntityComponentField)
+		if (EnumHasAnyFlags(CompiledData->AllocatedMask, EMovieSceneSequenceCompilerMask::EntityComponentField))
 		{
 			EntityComponentFields.Add(DataID.Value, MoveTemp(CompiledData->EntityComponentField));
 		}
@@ -472,7 +472,7 @@ void UMovieSceneCompiledDataManager::LoadCompiledData(UMovieSceneSequence* Seque
 
 		EntryPtr->DeterminismFences = MoveTemp(CompiledData->DeterminismFences);
 		EntryPtr->CompiledSignature = CompiledData->CompiledSignature;
-		EntryPtr->AccumulatedMask = CompiledData->AccumulatedMask.AsEnum();
+		EntryPtr->AccumulatedMask = CompiledData->AccumulatedMask;
 		EntryPtr->AccumulatedFlags = CompiledData->AccumulatedFlags;
 		EntryPtr->CompiledFlags = CompiledData->CompiledFlags;
 
diff --git a/Engine/Source/Runtime/MovieScene/Public/Compilation/MovieSceneCompiledDataManager.h b/Engine/Source/Runtime/MovieScene/Public/Compilation/MovieSceneCompiledDataManager.h
--- a/Engine/Source/Runtime/MovieScene/Public/Compilation/MovieSceneCompiledDataManager.h
+++ b/Engine/Source/Runtime/MovieScene/Public/Compilation/MovieSceneCompiledDataManager.h
@@ -47,6 +47,7 @@ namespace MovieScene
 }
 }
 
+UENUM(Flags)
 enum class EMovieSceneSequenceCompilerMask : uint8
 {
 	Hierarchy               = 1 << 0,
@@ -58,48 +59,6 @@ enum class EMovieSceneSequenceCompilerMask : uint8
 };
 ENUM_CLASS_FLAGS(EMovieSceneSequenceCompilerMask);
 
-/** Flag struct necessary while flag enums are not supported on UPROPERTY */
-USTRUCT()
-struct FMovieSceneSequenceCompilerMaskStruct
-{
-	GENERATED_BODY()
-
-	FMovieSceneSequenceCompilerMaskStruct()
-		: bHierarchy(0)
-		, bEvaluationTemplate(0)
-		, bEvaluationTemplateField(0)
-		, bEntityComponentField(0)
-	{}
-
-	FMovieSceneSequenceCompilerMaskStruct& operator=(EMovieSceneSequenceCompilerMask InMask)
-	{
-		bHierarchy               = EnumHasAnyFlags(InMask, EMovieSceneSequenceCompilerMask::Hierarchy);
-		bEvaluationTemplate      = EnumHasAnyFlags(InMask, EMovieSceneSequenceCompilerMask::EvaluationTemplate);
-		bEvaluationTemplateField = EnumHasAnyFlags(InMask, EMovieSceneSequenceCompilerMask::EvaluationTemplateField);
-		bEntityComponentField    = EnumHasAnyFlags(InMask, EMovieSceneSequenceCompilerMask::EntityComponentField);
-		return *this;
-	}
-
-	EMovieSceneSequenceCompilerMask AsEnum() const
-	{
-		EMovieSceneSequenceCompilerMask Enum = EMovieSceneSequenceCompilerMask::None;
-		Enum |= bHierarchy               ? EMovieSceneSequenceCompilerMask::Hierarchy               : EMovieSceneSequenceCompilerMask::None;
-		Enum |= bEvaluationTemplate      ? EMovieSceneSequenceCompilerMask::EvaluationTemplate      : EMovieSceneSequenceCompilerMask::None;
-		Enum |= bEvaluationTemplateField ? EMovieSceneSequenceCompilerMask::EvaluationTemplateField : EMovieSceneSequenceCompilerMask::None;
-		Enum |= bEntityComponentField    ? EMovieSceneSequenceCompilerMask::EntityComponentField    : EMovieSceneSequenceCompilerMask::None;
-		return Enum;
-	}
-
-	UPROPERTY()
-	uint8 bHierarchy : 1;
-	UPROPERTY()
-	uint8 bEvaluationTemplate : 1;
-	UPROPERTY()
-	uint8 bEvaluationTemplateField : 1;
-	UPROPERTY()
-	uint8 bEntityComponentField : 1;
-};
-
 /** Flags generated at compile time for a given sequence */
 USTRUCT()
 struct FMovieSceneCompiledSequenceFlagStruct
@@ -165,11 +124,11 @@ private:
 
 	/** 1 Byte */
 	UPROPERTY()
-	FMovieSceneSequenceCompilerMaskStruct AccumulatedMask;
+	EMovieSceneSequenceCompilerMask AccumulatedMask;
 
 	/** 1 Byte */
 	UPROPERTY()
-	FMovieSceneSequenceCompilerMaskStruct AllocatedMask;
+	EMovieSceneSequenceCompilerMask AllocatedMask;
 
 	/** 1 Byte */
 	UPROPERTY()
diff --git a/Engine/Source/Runtime/MovieScene/Public/Evaluation/MovieSceneSectionParameters.h b/Engine/Source/Runtime/MovieScene/Public/Evaluation/MovieSceneSectionParameters.h
--- a/Engine/Source/Runtime/MovieScene/Public/Evaluation/MovieSceneSectionParameters.h
+++ b/Engine/Source/Runtime/MovieScene/Public/Evaluation/MovieSceneSectionParameters.h
@@ -11,7 +11,7 @@
  * Flag structure that can be applied to any sub-section allowing control over various
  * behaviors for the nested sub-sequence.
  */
-UENUM(BlueprintType)
+UENUM(Flags, BlueprintType)
 enum class EMovieSceneSubSectionFlags : uint8
 {
 	None = 0,
diff --git a/Engine/Source/Runtime/MovieScene/Public/MovieSceneFwd.h b/Engine/Source/Runtime/MovieScene/Public/MovieSceneFwd.h
--- a/Engine/Source/Runtime/MovieScene/Public/MovieSceneFwd.h
+++ b/Engine/Source/Runtime/MovieScene/Public/MovieSceneFwd.h
@@ -86,7 +86,7 @@ enum class EUpdateClockSource : uint8
 /**
  * Bitfield flags that define special behavior for any UMovieSceneSequence.
  */
-UENUM()
+UENUM(Flags)
 enum class EMovieSceneSequenceFlags : uint8
 {
 	/** Symbolic entry for no flags */
@@ -118,7 +118,7 @@ ENUM_CLASS_FLAGS(EMovieSceneSequenceFlags);
 /**
  * 
  */
-UENUM(meta=(Bitflags, UseEnumValuesAsMaskValuesInEditor = "true"))
+UENUM(Flags, meta=(Bitflags, UseEnumValuesAsMaskValuesInEditor = "true"))
 enum class EMovieSceneServerClientMask : uint8
 {
 	None           = 0 UMETA(Hidden),
