Sequencer: Clamp duration sub frames to 1

Follow up to: Fix divide by zero

Fix for when duration (frames) is 0, but subframes is non zero.
Determine looping in tick resolution, rather than in display rate.

diff --git a/Engine/Source/Runtime/MovieScene/Private/MovieSceneSequencePlayer.cpp b/Engine/Source/Runtime/MovieScene/Private/MovieSceneSequencePlayer.cpp
--- a/Engine/Source/Runtime/MovieScene/Private/MovieSceneSequencePlayer.cpp
+++ b/Engine/Source/Runtime/MovieScene/Private/MovieSceneSequencePlayer.cpp
@@ -1161,7 +1161,8 @@ void UMovieSceneSequencePlayer::UpdateTimeCursorPosition_Internal(FFrameTime New
 		FFrameTime PositionRelativeToStart = ConvertFrameTime(NewPosition.FrameNumber - StartTimeWithReversed, PlayPosition.GetInputRate(), PlayPosition.GetOutputRate());
 		FFrameTime DurationWithSubFrames = ConvertFrameTime(GetDuration().Time, PlayPosition.GetInputRate(), PlayPosition.GetOutputRate());
 
-		const int32 NumTimesLooped    = FMath::Abs(PositionRelativeToStart.FrameNumber.Value / DurationWithSubFrames.FrameNumber.Value);
+		const int32 ClampedDurationWithSubFrames = FMath::Max(1, DurationWithSubFrames.FrameNumber.Value);
+		const int32 NumTimesLooped    = FMath::Abs(PositionRelativeToStart.FrameNumber.Value / ClampedDurationWithSubFrames);
 		const bool  bLoopIndefinitely = PlaybackSettings.LoopCount.Value < 0;
 
 		// loop playback
@@ -1185,7 +1186,7 @@ void UMovieSceneSequencePlayer::UpdateTimeCursorPosition_Internal(FFrameTime New
 				UpdateMovieSceneInstance(Range, StatusOverride);
 			}
 
-			const FFrameTime Overplay = FFrameTime(ConvertFrameTime(PositionRelativeToStart.FrameNumber.Value % DurationWithSubFrames.FrameNumber.Value, PlayPosition.GetOutputRate(), PlayPosition.GetInputRate()).FrameNumber, PositionRelativeToStart.GetSubFrame());
+			const FFrameTime Overplay = FFrameTime(ConvertFrameTime(PositionRelativeToStart.FrameNumber.Value % ClampedDurationWithSubFrames, PlayPosition.GetOutputRate(), PlayPosition.GetInputRate()).FrameNumber, PositionRelativeToStart.GetSubFrame());
 			FFrameTime NewFrameOffset;
 			
 			if (bReversePlayback)
