Sequencer- Fix regression where Binding Overrides specified on a Level Sequence Actor wouldn't preview in editor in Sequencer.

diff --git a/Engine/Source/Runtime/MovieScene/Private/MovieScenePlayer.cpp b/Engine/Source/Runtime/MovieScene/Private/MovieScenePlayer.cpp
--- a/Engine/Source/Runtime/MovieScene/Private/MovieScenePlayer.cpp
+++ b/Engine/Source/Runtime/MovieScene/Private/MovieScenePlayer.cpp
@@ -130,14 +130,21 @@ void IMovieScenePlayer::ResolveBoundObjects(UE::UniversalObjectLocator::FResolve
 	using namespace UE::UniversalObjectLocator;
 	using namespace UE::MovieScene;
 
-	if (const FMovieSceneBindingReferences* BindingReferences = InSequence.GetBindingReferences())
-	{
-		FMovieSceneBindingResolveParams BindingResolveParams{ &InSequence, InBindingId, SequenceID };
-		BindingReferences->ResolveBinding(BindingResolveParams, LocatorResolveParams, ConstCastSharedPtr<const UE::MovieScene::FSharedPlaybackState>(const_cast<IMovieScenePlayer*>(this)->FindSharedPlaybackState()), OutObjects);
-	}
-	else
+	const IMovieScenePlaybackClient* PlaybackClient = GetPlaybackClient();
+
+	bool bAllowDefault = PlaybackClient ? PlaybackClient->RetrieveBindingOverrides(InBindingId, SequenceID, OutObjects) : true;
+
+	if (bAllowDefault)
 	{
-		InSequence.LocateBoundObjects(InBindingId, LocatorResolveParams, OutObjects);
+		if (const FMovieSceneBindingReferences* BindingReferences = InSequence.GetBindingReferences())
+		{
+			FMovieSceneBindingResolveParams BindingResolveParams{ &InSequence, InBindingId, SequenceID };
+			BindingReferences->ResolveBinding(BindingResolveParams, LocatorResolveParams, ConstCastSharedPtr<const UE::MovieScene::FSharedPlaybackState>(const_cast<IMovieScenePlayer*>(this)->FindSharedPlaybackState()), OutObjects);
+		}
+		else
+		{
+			InSequence.LocateBoundObjects(InBindingId, LocatorResolveParams, OutObjects);
+		}
 	}
 }
 
diff --git a/Engine/Source/Runtime/MovieScene/Public/IMovieScenePlayer.h b/Engine/Source/Runtime/MovieScene/Public/IMovieScenePlayer.h
--- a/Engine/Source/Runtime/MovieScene/Public/IMovieScenePlayer.h
+++ b/Engine/Source/Runtime/MovieScene/Public/IMovieScenePlayer.h
@@ -169,6 +169,8 @@ public:
 	 */
 	virtual IMovieScenePlaybackClient* GetPlaybackClient() { return nullptr; }
 
+	const IMovieScenePlaybackClient* GetPlaybackClient() const { return const_cast<IMovieScenePlayer*>(this)->GetPlaybackClient(); }
+
 	/**
 	 * Obtain an object responsible for managing movie scene spawnables
 	 */
