Sequencer: go through entire sub-instance hierarchy when invalidating data after restoring pre-animated state

There was a regression due to an optimization for day-sequences where we didn't want to invalidate everything when data was dynamically changed at runtime. However, when restoring pre-animated state, we do want to invalidate everything.

diff --git a/Engine/Source/Runtime/MovieScene/Private/EntitySystem/MovieSceneSequenceInstance.cpp b/Engine/Source/Runtime/MovieScene/Private/EntitySystem/MovieSceneSequenceInstance.cpp
--- a/Engine/Source/Runtime/MovieScene/Private/EntitySystem/MovieSceneSequenceInstance.cpp
+++ b/Engine/Source/Runtime/MovieScene/Private/EntitySystem/MovieSceneSequenceInstance.cpp
@@ -164,7 +164,7 @@ void FSequenceInstance::InitializeLegacyEvaluator()
 	}
 }
 
-void FSequenceInstance::InvalidateCachedData()
+void FSequenceInstance::InvalidateCachedData(ESequenceInstanceInvalidationType InvalidationType)
 {
 	ensureMsgf(bInitialized, TEXT("Sequence instance hasn't been initialized yet!"));
 
@@ -223,7 +223,7 @@ void FSequenceInstance::InvalidateCachedData()
 
 		ISequenceUpdater::FactoryInstance(SequenceUpdater, CompiledDataManager, RootCompiledDataID);
 
-		SequenceUpdater->InvalidateCachedData(SharedPlaybackState);
+		SequenceUpdater->InvalidateCachedData(SharedPlaybackState, InvalidationType);
 		SequenceUpdater->PopulateUpdateFlags(SharedPlaybackState, UpdateFlags);
 
 		if (LegacyEvaluator)
diff --git a/Engine/Source/Runtime/MovieScene/Private/EntitySystem/MovieSceneSequenceUpdaters.cpp b/Engine/Source/Runtime/MovieScene/Private/EntitySystem/MovieSceneSequenceUpdaters.cpp
--- a/Engine/Source/Runtime/MovieScene/Private/EntitySystem/MovieSceneSequenceUpdaters.cpp
+++ b/Engine/Source/Runtime/MovieScene/Private/EntitySystem/MovieSceneSequenceUpdaters.cpp
@@ -41,7 +41,7 @@ struct FSequenceUpdater_Flat : ISequenceUpdater
 	virtual void Update(TSharedRef<const FSharedPlaybackState> SharedPlaybackState, const FMovieSceneContext& Context) override;
 	virtual bool CanFinishImmediately(TSharedRef<const FSharedPlaybackState> SharedPlaybackState) const override;
 	virtual void Finish(TSharedRef<const FSharedPlaybackState> SharedPlaybackState) override;
-	virtual void InvalidateCachedData(TSharedRef<const FSharedPlaybackState> SharedPlaybackState) override;
+	virtual void InvalidateCachedData(TSharedRef<const FSharedPlaybackState> SharedPlaybackState, ESequenceInstanceInvalidationType InvalidationType) override;
 	virtual void Destroy(TSharedRef<const FSharedPlaybackState> SharedPlaybackState) override;
 	virtual TUniquePtr<ISequenceUpdater> MigrateToHierarchical() override;
 	virtual FInstanceHandle FindSubInstance(FMovieSceneSequenceID SubSequenceID) const override { return FInstanceHandle(); }
@@ -70,7 +70,7 @@ struct FSequenceUpdater_Hierarchical : ISequenceUpdater
 	virtual void Update(TSharedRef<const FSharedPlaybackState> SharedPlaybackState, const FMovieSceneContext& Context) override;
 	virtual bool CanFinishImmediately(TSharedRef<const FSharedPlaybackState> SharedPlaybackState) const override;
 	virtual void Finish(TSharedRef<const FSharedPlaybackState> SharedPlaybackState) override;
-	virtual void InvalidateCachedData(TSharedRef<const FSharedPlaybackState> SharedPlaybackState) override;
+	virtual void InvalidateCachedData(TSharedRef<const FSharedPlaybackState> SharedPlaybackState, ESequenceInstanceInvalidationType InvalidationType) override;
 	virtual void Destroy(TSharedRef<const FSharedPlaybackState> SharedPlaybackState) override;
 	virtual TUniquePtr<ISequenceUpdater> MigrateToHierarchical() override { return nullptr; }
 	virtual FInstanceHandle FindSubInstance(FMovieSceneSequenceID SubSequenceID) const override { return SequenceInstances.FindRef(SubSequenceID).Handle; }
@@ -335,14 +335,14 @@ bool FSequenceUpdater_Flat::CanFinishImmediately(TSharedRef<const FSharedPlaybac
 
 void FSequenceUpdater_Flat::Finish(TSharedRef<const FSharedPlaybackState> SharedPlaybackState)
 {
-	InvalidateCachedData(SharedPlaybackState);
+	InvalidateCachedData(SharedPlaybackState, ESequenceInstanceInvalidationType::All);
 }
 
 void FSequenceUpdater_Flat::Destroy(TSharedRef<const FSharedPlaybackState> SharedPlaybackState)
 {
 }
 
-void FSequenceUpdater_Flat::InvalidateCachedData(TSharedRef<const FSharedPlaybackState> SharedPlaybackState)
+void FSequenceUpdater_Flat::InvalidateCachedData(TSharedRef<const FSharedPlaybackState> SharedPlaybackState, ESequenceInstanceInvalidationType InvalidationType)
 {
 	CachedEntityRange = TRange<FFrameNumber>::Empty();
 	CachedDeterminismFences.Reset();
@@ -520,7 +520,7 @@ void FSequenceUpdater_Hierarchical::OverrideRootSequence(TSharedRef<const FShare
 			InstanceRegistry->MutateInstance(RootInstanceHandle).Ledger.UnlinkEverything(Linker);
 		}
 
-		InvalidateCachedData(SharedPlaybackState);
+		InvalidateCachedData(SharedPlaybackState, ESequenceInstanceInvalidationType::All);
 		RootOverrideSequenceID = NewRootOverrideSequenceID;
 	}
 }
@@ -805,7 +805,7 @@ void FSequenceUpdater_Hierarchical::Finish(TSharedRef<const FSharedPlaybackState
 		InstanceRegistry->MutateInstance(Pair.Value.Handle).Finish();
 	}
 
-	InvalidateCachedData(SharedPlaybackState);
+	InvalidateCachedData(SharedPlaybackState, ESequenceInstanceInvalidationType::All);
 }
 
 void FSequenceUpdater_Hierarchical::Destroy(TSharedRef<const FSharedPlaybackState> SharedPlaybackState)
@@ -819,7 +819,7 @@ void FSequenceUpdater_Hierarchical::Destroy(TSharedRef<const FSharedPlaybackStat
 	}
 }
 
-void FSequenceUpdater_Hierarchical::InvalidateCachedData(TSharedRef<const FSharedPlaybackState> SharedPlaybackState)
+void FSequenceUpdater_Hierarchical::InvalidateCachedData(TSharedRef<const FSharedPlaybackState> SharedPlaybackState, ESequenceInstanceInvalidationType InvalidationType)
 {
 	bDynamicWeighting.Reset();
 	CachedEntityRange = TRange<FFrameNumber>::Empty();
@@ -829,16 +829,30 @@ void FSequenceUpdater_Hierarchical::InvalidateCachedData(TSharedRef<const FShare
 
 	for (TPair<FMovieSceneSequenceID, FSubInstanceData>& Pair : SequenceInstances)
 	{
-		UMovieSceneSequence* Sequence = SharedPlaybackState->GetSequence(Pair.Key);
-		if (!Sequence)
-		{
-			Pair.Value.SequenceSignature = FGuid();
-			InstanceRegistry->MutateInstance(Pair.Value.Handle).Ledger.Invalidate();
-		}
-		else if (Pair.Value.SequenceSignature != Sequence->GetSignature())
+		FSequenceInstance& SubInstance = InstanceRegistry->MutateInstance(Pair.Value.Handle);
+
+		switch (InvalidationType)
 		{
-			Pair.Value.SequenceSignature = Sequence->GetSignature();
-			InstanceRegistry->MutateInstance(Pair.Value.Handle).Ledger.Invalidate();
+			case ESequenceInstanceInvalidationType::All:
+				{
+					SubInstance.Ledger.Invalidate();
+				}
+				break;
+			case ESequenceInstanceInvalidationType::DataChanged:
+				{
+					UMovieSceneSequence* Sequence = SharedPlaybackState->GetSequence(Pair.Key);
+					if (!Sequence)
+					{
+						Pair.Value.SequenceSignature = FGuid();
+						SubInstance.Ledger.Invalidate();
+					}
+					else if (Pair.Value.SequenceSignature != Sequence->GetSignature())
+					{
+						Pair.Value.SequenceSignature = Sequence->GetSignature();
+						SubInstance.Ledger.Invalidate();
+					}
+				}
+				break;
 		}
 	}
 }
@@ -856,5 +870,10 @@ TRange<FFrameNumber> FSequenceUpdater_Hierarchical::UpdateEntitiesForSequence(co
 	return CachedRange;
 }
 
+void ISequenceUpdater::InvalidateCachedData(TSharedRef<const FSharedPlaybackState> SharedPlaybackState)
+{
+	InvalidateCachedData(SharedPlaybackState, ESequenceInstanceInvalidationType::All);
+}
+
 } // namespace MovieScene
 } // namespace UE
diff --git a/Engine/Source/Runtime/MovieScene/Private/Evaluation/PreAnimatedState/MovieScenePreAnimatedStateExtension.cpp b/Engine/Source/Runtime/MovieScene/Private/Evaluation/PreAnimatedState/MovieScenePreAnimatedStateExtension.cpp
--- a/Engine/Source/Runtime/MovieScene/Private/Evaluation/PreAnimatedState/MovieScenePreAnimatedStateExtension.cpp
+++ b/Engine/Source/Runtime/MovieScene/Private/Evaluation/PreAnimatedState/MovieScenePreAnimatedStateExtension.cpp
@@ -377,6 +377,18 @@ void FPreAnimatedStateExtension::RestoreGlobalState(const FRestoreStateParams& P
 			{
 				Storage.RestorePreAnimatedStateStorage(StorageIndex, EPreAnimatedStorageRequirement::Persistent, EPreAnimatedStorageRequirement::None, Params);
 			});
+
+	// Invalidate cached data for any sequence instance that belongs to the terminal instance.
+	// This is because we have just restored pre-animated state, so trying to re-update the sequence will 
+	// need to re-setup everything. It wouldn't do it unless we've invalidated it.
+	if (Params.TerminalInstanceHandle.IsValid() && 
+			ensureMsgf(
+				Linker->GetInstanceRegistry()->IsHandleValid(Params.TerminalInstanceHandle),
+				TEXT("Terminal instance handle is not valid anymore, was the sequence destroyed?")))
+	{
+		FSequenceInstance& TerminalInstance = Linker->GetInstanceRegistry()->MutateInstance(Params.TerminalInstanceHandle);
+		TerminalInstance.InvalidateCachedData(ESequenceInstanceInvalidationType::All);
+	}
 }
 
 void FPreAnimatedStateExtension::DiscardStaleObjectState()
@@ -789,15 +801,6 @@ void FPreAnimatedStateExtension::HandleMetaDataToRemove(const FRestoreStateParam
 
 	GroupMetaData.Shrink();
 
-	// Invalidate cached data for any sequence instance that belongs to the terminal instance
-	if (Params.TerminalInstanceHandle.IsValid() && 
-			ensureMsgf(
-				Linker->GetInstanceRegistry()->IsHandleValid(Params.TerminalInstanceHandle),
-				TEXT("Terminal instance handle is not valid anymore, was the sequence destroyed?")))
-	{
-		Linker->GetInstanceRegistry()->MutateInstance(Params.TerminalInstanceHandle).InvalidateCachedData();
-	}
-
 	bEntriesInvalidated = true;
 }
 
diff --git a/Engine/Source/Runtime/MovieScene/Public/EntitySystem/MovieSceneSequenceInstance.h b/Engine/Source/Runtime/MovieScene/Public/EntitySystem/MovieSceneSequenceInstance.h
--- a/Engine/Source/Runtime/MovieScene/Public/EntitySystem/MovieSceneSequenceInstance.h
+++ b/Engine/Source/Runtime/MovieScene/Public/EntitySystem/MovieSceneSequenceInstance.h
@@ -45,6 +45,22 @@ enum class ESequenceInstanceUpdateFlags : uint8
 };
 ENUM_CLASS_FLAGS(ESequenceInstanceUpdateFlags);
 
+/**
+ * Specifies what cached data should be invalidated in a sequence instance.
+ */
+enum class ESequenceInstanceInvalidationType : uint8
+{
+	/**
+	 * Only invalidate cached data for sequences whose data has been dynamically changed.
+	 */
+	DataChanged,
+
+	/**
+	 * Invalidate all cached data for all sequences.
+	 */
+	All
+};
+
 /**
  * A sequence instance represents a specific instance of a currently playing sequence, either as a top-level sequence in an IMovieScenePlayer, or as a sub sequence.
  * Any given sequence asset may have any number of instances created for it at any given time depending on how many times it is referenced by playing sequences
@@ -297,7 +313,7 @@ public:
 	/**
 	 * Invalidate any cached data that may be being used for evaluation due to a change in the source asset data
 	 */
-	MOVIESCENE_API void InvalidateCachedData();
+	MOVIESCENE_API void InvalidateCachedData(ESequenceInstanceInvalidationType InvalidationType = ESequenceInstanceInvalidationType::All);
 
 	/**
 	 * Destroy this sequence instance immediately - Finish must previously have been called
diff --git a/Engine/Source/Runtime/MovieScene/Public/EntitySystem/MovieSceneSequenceUpdaters.h b/Engine/Source/Runtime/MovieScene/Public/EntitySystem/MovieSceneSequenceUpdaters.h
--- a/Engine/Source/Runtime/MovieScene/Public/EntitySystem/MovieSceneSequenceUpdaters.h
+++ b/Engine/Source/Runtime/MovieScene/Public/EntitySystem/MovieSceneSequenceUpdaters.h
@@ -26,7 +26,7 @@ struct FInstanceHandle;
 struct FSharedPlaybackState;
 
 enum class ESequenceInstanceUpdateFlags : uint8;
-
+enum class ESequenceInstanceInvalidationType : uint8;
 
 /**
  * Interface for an object that will update a sequence based on the current context. It holds several responsibilities:
@@ -119,7 +119,7 @@ public:
 	 * @param InLinker         The linker that is evaluating this sequence
 	 * @param InInstanceHandle The instance handle for the top level sequence instance that this updater belongs to
 	 */
-	virtual void InvalidateCachedData(TSharedRef<const FSharedPlaybackState> SharedPlaybackState) = 0;
+	virtual void InvalidateCachedData(TSharedRef<const FSharedPlaybackState> SharedPlaybackState, ESequenceInstanceInvalidationType InvalidationType) = 0;
 
 
 	/**
@@ -151,8 +151,12 @@ public:
 	 * Find a sub sequence instance from its ID (if it exists)
 	 */
 	virtual FInstanceHandle FindSubInstance(FMovieSceneSequenceID SubSequenceID) const = 0;
-};
 
+public:
+
+	UE_DEPRECATED(5.5, "Please use the method that takes a reason parameter.")
+	MOVIESCENE_API void InvalidateCachedData(TSharedRef<const FSharedPlaybackState> SharedPlaybackState);
+};
 
 } // namespace MovieScene
 } // namespace UE
