Sequencer: Protect against animating arrays

Check NumParams > 0. This prevents a crash with existing sequences that have already managed to animate arrays (like SetTextureData).

diff --git a/Engine/Source/Runtime/MovieScene/Public/MovieSceneCommonHelpers.h b/Engine/Source/Runtime/MovieScene/Public/MovieSceneCommonHelpers.h
--- a/Engine/Source/Runtime/MovieScene/Public/MovieSceneCommonHelpers.h
+++ b/Engine/Source/Runtime/MovieScene/Public/MovieSceneCommonHelpers.h
@@ -555,7 +555,7 @@ void FTrackInstancePropertyBindings::InvokeSetterFunction(UObject* InRuntimeObje
 	uint8* Params = reinterpret_cast<uint8*>(InputParameter);
 
 	check(InRuntimeObject && Setter);
-	if (Setter->ReturnValueOffset != MAX_uint16 || Setter->NumParms > 1)
+	if (Setter->ReturnValueOffset != MAX_uint16 || Setter->NumParms > 0)
 	{
 		// Function has a return value or multiple parameters, we need to initialize memory for the entire parameter pack
 		// We use alloca here (as in UObject::ProcessEvent) to avoid a heap allocation. Alloca memory survives the current function's stack frame.
