Sequencer: Do not tick sequences if the manager is unreachable or being destroyed, or any similarly expiring linkers

diff --git a/Engine/Source/Runtime/MovieScene/Private/MovieSceneSequenceTickManager.cpp b/Engine/Source/Runtime/MovieScene/Private/MovieSceneSequenceTickManager.cpp
--- a/Engine/Source/Runtime/MovieScene/Private/MovieSceneSequenceTickManager.cpp
+++ b/Engine/Source/Runtime/MovieScene/Private/MovieSceneSequenceTickManager.cpp
@@ -312,6 +312,11 @@ void UMovieSceneSequenceTickManager::TickSequenceActors(float DeltaSeconds)
 	CSV_SCOPED_TIMING_STAT_EXCLUSIVE(MovieSceneEval);
 	SCOPE_CYCLE_COUNTER(MovieSceneEval_SequenceTickManager);
 
+	if (IsUnreachable() || HasAnyFlags(RF_BeginDestroyed))
+	{
+		return;
+	}
+
 	// Let all tickable clients update. Some of them won't do anything, others will do synchronous
 	// things (e.g. start/stop, loop, etc.), but in 95% of cases, they will just queue up a normal evaluation
 	// request...
@@ -419,6 +424,11 @@ void UMovieSceneSequenceTickManager::TickSequenceActors(float DeltaSeconds)
 		{
 			FLinkerGroup& Group = LinkerGroups[LinkerIndex.GetIndex()];
 
+			if (Group.Linker == nullptr || Group.Linker->IsUnreachable() || Group.Linker->HasAnyFlags(RF_BeginDestroyed))
+			{
+				continue;
+			}
+
 			check(!UpdatedDeltaTimes.IsValidIndex(LinkerIndex.GetIndex()));
 
 			Group.Runner->Flush(Group.FrameBudgetMs);
@@ -431,6 +441,12 @@ void UMovieSceneSequenceTickManager::TickSequenceActors(float DeltaSeconds)
 			if (UpdatedDeltaTimes.IsAllocated(Index))
 			{
 				FLinkerGroup& Group = LinkerGroups[Index];
+
+				if (Group.Linker == nullptr || Group.Linker->IsUnreachable() || Group.Linker->HasAnyFlags(RF_BeginDestroyed))
+				{
+					continue;
+				}
+
 				// Hitting this check would indicate that the loop above that processes OutstandingLinkers either failed, or some other partial flush happened between then and now.
 				ensureMsgf(!Group.Runner->IsCurrentlyEvaluating(), TEXT("Linker is part-way thorugh a flush when a new flush is being instigated. This is undefined behavior."));
 
@@ -494,6 +510,11 @@ void UMovieSceneSequenceTickManager::FlushRunners()
 		TGuardValue<TArray<FPendingOperation>*> Guard(PendingActorOperations, &CurrentPendingActorOperations);
 		for (FLinkerGroup& LinkerGroup : LinkerGroups)
 		{
+			if (LinkerGroup.Linker == nullptr || LinkerGroup.Linker->IsUnreachable() || LinkerGroup.Linker->HasAnyFlags(RF_BeginDestroyed))
+			{
+				continue;
+			}
+
 			LinkerGroup.Runner->Flush();
 		}
 	}
