Sequencer: Fixed runtime playback of sequences with subframe-durations

diff --git a/Engine/Source/Runtime/MovieScene/Private/MovieSceneSequencePlayer.cpp b/Engine/Source/Runtime/MovieScene/Private/MovieSceneSequencePlayer.cpp
--- a/Engine/Source/Runtime/MovieScene/Private/MovieSceneSequencePlayer.cpp
+++ b/Engine/Source/Runtime/MovieScene/Private/MovieSceneSequencePlayer.cpp
@@ -1158,11 +1158,12 @@ void UMovieSceneSequencePlayer::UpdateTimeCursorPosition_Internal(FFrameTime New
 		FFrameTime EndTimeWithReversed = bReversePlayback ? StartTime : GetLastValidTime().FrameNumber;
 
 		// Operate in tick resolution (for subframes)
-		FFrameTime PositionRelativeToStart = ConvertFrameTime(NewPosition.FrameNumber - StartTimeWithReversed, PlayPosition.GetInputRate(), PlayPosition.GetOutputRate());
-		FFrameTime DurationWithSubFrames = ConvertFrameTime(GetDuration().Time, PlayPosition.GetInputRate(), PlayPosition.GetOutputRate());
+		const double DurationWithSubFrames   = FMath::Max<double>(UE_SMALL_NUMBER, GetDuration().Time.AsDecimal());
+		const double PositionRelativeToStart = (NewPosition - StartTimeWithReversed).AsDecimal();
 
-		const int32 ClampedDurationWithSubFrames = FMath::Max(1, DurationWithSubFrames.FrameNumber.Value);
-		const int32 NumTimesLooped    = FMath::Abs(PositionRelativeToStart.FrameNumber.Value / ClampedDurationWithSubFrames);
+		ensure(PositionRelativeToStart >= DurationWithSubFrames);
+
+		const int32 NumTimesLooped    = FMath::Abs(FMath::TruncToInt32(PositionRelativeToStart / DurationWithSubFrames));
 		const bool  bLoopIndefinitely = PlaybackSettings.LoopCount.Value < 0;
 
 		// loop playback
@@ -1186,7 +1187,7 @@ void UMovieSceneSequencePlayer::UpdateTimeCursorPosition_Internal(FFrameTime New
 				UpdateMovieSceneInstance(Range, StatusOverride);
 			}
 
-			const FFrameTime Overplay = FFrameTime(ConvertFrameTime(PositionRelativeToStart.FrameNumber.Value % ClampedDurationWithSubFrames, PlayPosition.GetOutputRate(), PlayPosition.GetInputRate()).FrameNumber, PositionRelativeToStart.GetSubFrame());
+			const FFrameTime Overplay = FFrameTime::FromDecimal(FMath::Fmod(PositionRelativeToStart, DurationWithSubFrames));
 			FFrameTime NewFrameOffset;
 			
 			if (bReversePlayback)
