Sequencer: Fix move binding contents so that the source child tracks are moved to the destination child binding.

In 5.3, the conversion just calls Move Binding Contents from the source to the destination. The destination binding doesn't have any existing tracks.

In 5.4, there were some changes which result in OnActorAddedToSequencer now being called. That results in default tracks and then when Move Binding Contents is called, the source child tracks aren't copied over the default tracks. This is the issue that this CL fixes.

In Main, Move Binding Contents is never called because the conversion now calls AddOrReplaceBinding, so it's not a problem.

This is a redo of 19871735 and also fixes UE-145266

diff --git a/Engine/Source/Runtime/MovieScene/Private/MovieScene.cpp b/Engine/Source/Runtime/MovieScene/Private/MovieScene.cpp
--- a/Engine/Source/Runtime/MovieScene/Private/MovieScene.cpp
+++ b/Engine/Source/Runtime/MovieScene/Private/MovieScene.cpp
@@ -1588,34 +1588,45 @@ void UMovieScene::MoveBindingContents(const FGuid& SourceBindingId, const FGuid&
 
 	FMovieSceneSpawnable* DestinationSpawnable = FindSpawnable(DestinationBindingId);
 
-	for( auto PossesableIter( Possessables.CreateIterator() ); PossesableIter; ++PossesableIter )
+	TArray<FGuid> BindingsToRemove;
+	for( auto SourcePossessableIter( Possessables.CreateIterator() ); SourcePossessableIter; ++SourcePossessableIter)
 	{
-		FMovieScenePossessable& SourcePossessable = *PossesableIter;
+		FMovieScenePossessable& SourcePossessable = *SourcePossessableIter;
 		
 		// If there is a possessable whose parent is the binding we're moving contents for, 
 		// that possessable needs to be remapped to the new destination parent
 		if (SourcePossessable.GetParent() == SourceBindingId)
 		{
-			// But if there is already a possessable for that destination binding, don't keep the source possessable around.
-			bool bAlreadyExists = false;
-			for (FMovieScenePossessable& DestinationPossessable : Possessables)
+			// Look for an existing destination child possessable
+			bool bUseSourcePossessable = true;
+			for (auto DestinationPossessableIter(Possessables.CreateIterator()); DestinationPossessableIter; ++DestinationPossessableIter)
 			{
+				FMovieScenePossessable& DestinationPossessable = *DestinationPossessableIter;
+
 				if (DestinationPossessable.GetName() == SourcePossessable.GetName() &&
 					DestinationPossessable.GetParent() == DestinationBindingId)
 				{
-					FGuid CurGuid = SourcePossessable.GetGuid();
-
-					// Found it!
-					Possessables.RemoveAt( PossesableIter.GetIndex() );
-
-					RemoveBinding( CurGuid );
-
-					bAlreadyExists = true;
+					// If it's not the same class, we can't use the source possessable because it's going to lead to trouble
+#if WITH_EDITORONLY_DATA
+					if (SourcePossessable.GetPossessedObjectClass() != DestinationPossessable.GetPossessedObjectClass())
+					{
+						bUseSourcePossessable = false;
+					}
+#endif
+					if (bUseSourcePossessable)
+					{
+						BindingsToRemove.Add(DestinationPossessable.GetGuid());
+					}
+					// Otherwise, discard the source possessable since it's a different class and probably going to be trouble
+					else
+					{
+						BindingsToRemove.Add(SourcePossessable.GetGuid());
+					}
 					break;
 				}
 			}
-
-			if (!bAlreadyExists)
+			
+			if (bUseSourcePossessable)
 			{
 				SourcePossessable.SetParent(DestinationBindingId, this);
 				if (DestinationSpawnable)
@@ -1625,6 +1636,22 @@ void UMovieScene::MoveBindingContents(const FGuid& SourceBindingId, const FGuid&
 			}
 		}
 	}
+
+	for (FGuid BindingToRemove : BindingsToRemove)
+	{
+		for (auto PossessableIter(Possessables.CreateIterator()); PossessableIter; ++PossessableIter)
+		{
+			FMovieScenePossessable& Possessable = *PossessableIter;
+			if (Possessable.GetGuid() == BindingToRemove)
+			{
+				Possessables.RemoveAt(PossessableIter.GetIndex());
+				break;
+			}
+		}
+
+		RemoveBinding(BindingToRemove);
+		EventHandlers.Trigger(&UE::MovieScene::ISequenceDataEventHandler::OnBindingRemoved, BindingToRemove);
+	}
 }
 
 FMovieSceneBinding* UMovieScene::FindBinding(const FGuid& ForGuid)
