Take Recorder: Rework FScopedVolatilityManagerSuppression so that it doesn't keep raw pointers to the instance handle and registry

This fixes an issue in Take Recorder where compilation is suppressed with FScopedVolatilityManagerSuppression during recording. If the Sequencer UI is destroyed during recording, it will still try to compile after recording is complete. This now skips that compile if the UI is destroyed.

diff --git a/Engine/Plugins/VirtualProduction/Takes/Source/TakeRecorder/Private/Recorder/TakeRecorder.cpp b/Engine/Plugins/VirtualProduction/Takes/Source/TakeRecorder/Private/Recorder/TakeRecorder.cpp
--- a/Engine/Plugins/VirtualProduction/Takes/Source/TakeRecorder/Private/Recorder/TakeRecorder.cpp
+++ b/Engine/Plugins/VirtualProduction/Takes/Source/TakeRecorder/Private/Recorder/TakeRecorder.cpp
@@ -988,9 +988,7 @@ void UTakeRecorder::Start()
 
 	if (Sequencer.IsValid())
 	{
-		FRootInstanceHandle RootInstanceHandle = Sequencer->GetEvaluationTemplate().GetRootInstanceHandle();
-		FInstanceRegistry* InstanceRegistry = Sequencer->GetEvaluationTemplate().GetEntitySystemLinker()->GetInstanceRegistry();
-		CompileSuppression = MakeUnique<FScopedVolatilityManagerSuppression>(InstanceRegistry, RootInstanceHandle);
+		CompileSuppression = MakeUnique<FScopedVolatilityManagerSuppression>(Sequencer->GetSharedPlaybackState().ToSharedPtr());
 	}
 }
 
diff --git a/Engine/Source/Runtime/MovieScene/Private/EntitySystem/MovieSceneInstanceRegistry.cpp b/Engine/Source/Runtime/MovieScene/Private/EntitySystem/MovieSceneInstanceRegistry.cpp
--- a/Engine/Source/Runtime/MovieScene/Private/EntitySystem/MovieSceneInstanceRegistry.cpp
+++ b/Engine/Source/Runtime/MovieScene/Private/EntitySystem/MovieSceneInstanceRegistry.cpp
@@ -144,16 +144,28 @@ void FInstanceRegistry::CleanupLinkerEntities(const TSet<FMovieSceneEntityID>& E
 	}
 }
 
-FScopedVolatilityManagerSuppression::FScopedVolatilityManagerSuppression(FInstanceRegistry* InInstanceRegistry, FRootInstanceHandle InRootInstanceHandle)
-	: InstanceRegistry(InInstanceRegistry)
-	, RootInstanceHandle(InRootInstanceHandle)
+FScopedVolatilityManagerSuppression::FScopedVolatilityManagerSuppression(TSharedPtr<FSharedPlaybackState> PlaybackState)
+	: WeakPlaybackState(PlaybackState)
 {
+	ensure(PlaybackState.IsValid());
+
+	FRootInstanceHandle RootInstanceHandle = PlaybackState->GetRootInstanceHandle();
+	FInstanceRegistry* InstanceRegistry = PlaybackState->GetLinker()->GetInstanceRegistry();
+
 	FSequenceInstance& Instance = InstanceRegistry->MutateInstance(RootInstanceHandle);
 	PreviousVolatilityManager = MoveTemp(Instance.VolatilityManager);
 }
 
 FScopedVolatilityManagerSuppression::~FScopedVolatilityManagerSuppression()
 {
+	if (!WeakPlaybackState.IsValid())
+	{
+		return;
+	}
+
+	FRootInstanceHandle RootInstanceHandle = WeakPlaybackState.Pin()->GetRootInstanceHandle();
+	FInstanceRegistry* InstanceRegistry = WeakPlaybackState.Pin()->GetLinker()->GetInstanceRegistry();
+
 	FSequenceInstance& Instance = InstanceRegistry->MutateInstance(RootInstanceHandle);
 	Instance.VolatilityManager = MoveTemp(PreviousVolatilityManager);
 	Instance.ConditionalRecompile();
diff --git a/Engine/Source/Runtime/MovieScene/Public/EntitySystem/MovieSceneInstanceRegistry.h b/Engine/Source/Runtime/MovieScene/Public/EntitySystem/MovieSceneInstanceRegistry.h
--- a/Engine/Source/Runtime/MovieScene/Public/EntitySystem/MovieSceneInstanceRegistry.h
+++ b/Engine/Source/Runtime/MovieScene/Public/EntitySystem/MovieSceneInstanceRegistry.h
@@ -9,6 +9,7 @@
 #include "Evaluation/MovieScenePlayback.h"
 #include "EntitySystem/MovieSceneSequenceInstance.h"
 #include "EntitySystem/MovieSceneComponentDebug.h"
+#include "EntitySystem/MovieSceneSharedPlaybackState.h"
 #include "MovieSceneSequenceID.h"
 
 
@@ -134,12 +135,11 @@ private:
  */
 struct MOVIESCENE_API FScopedVolatilityManagerSuppression
 {
-	FScopedVolatilityManagerSuppression(FInstanceRegistry* InInstanceRegistry, FRootInstanceHandle InRootInstanceHandle);
+	FScopedVolatilityManagerSuppression(TSharedPtr<FSharedPlaybackState> PlaybackState);
 	~FScopedVolatilityManagerSuppression();
 
 private:
-	FInstanceRegistry* InstanceRegistry;
-	FRootInstanceHandle RootInstanceHandle;
+	TWeakPtr<FSharedPlaybackState> WeakPlaybackState;
 	TUniquePtr<FCompiledDataVolatilityManager> PreviousVolatilityManager;
 };
 
