MovieSceneSequencePlayer: Warp evaluated range on pause so that there are no discontinuities in playback when pausing/unpausing.

diff --git a/Engine/Source/Runtime/MovieScene/Private/MovieSceneSequencePlayer.cpp b/Engine/Source/Runtime/MovieScene/Private/MovieSceneSequencePlayer.cpp
--- a/Engine/Source/Runtime/MovieScene/Private/MovieSceneSequencePlayer.cpp
+++ b/Engine/Source/Runtime/MovieScene/Private/MovieSceneSequencePlayer.cpp
@@ -412,8 +412,14 @@ void UMovieSceneSequencePlayer::Pause()
 		if (TSharedPtr<FMovieSceneEntitySystemRunner> Runner = RootTemplateInstance.GetRunner())
 		{
 			FMovieSceneEvaluationRange CurrentTimeRange = PlayPosition.GetCurrentPositionAsRange();
-			const FMovieSceneContext Context(CurrentTimeRange, EMovieScenePlayerStatus::Stopped);
 
+			if (PlaybackClient)
+			{
+				PlaybackClient->WarpEvaluationRange(CurrentTimeRange);
+			}
+			
+			const FMovieSceneContext Context(CurrentTimeRange, EMovieScenePlayerStatus::Stopped);
+			
 			Runner->QueueUpdate(Context, RootTemplateInstance.GetRootInstanceHandle(), FSimpleDelegate::CreateWeakLambda(this, FinishPause));
 		}
 		else
