Enable MovieScene Module to be instrumented with AutoRTFM

diff --git a/Engine/Source/Runtime/MovieScene/MovieScene.Build.cs b/Engine/Source/Runtime/MovieScene/MovieScene.Build.cs
--- a/Engine/Source/Runtime/MovieScene/MovieScene.Build.cs
+++ b/Engine/Source/Runtime/MovieScene/MovieScene.Build.cs
@@ -25,5 +25,7 @@ public class MovieScene : ModuleRules
 			}
 		);
 		SetupIrisSupport(Target);
+
+		bAllowAutoRTFMInstrumentation = true;
 	}
 }
