Sequencer: Fix deprecated include warning

#rnx

diff --git a/Engine/Plugins/Cameras/GameplayCameras/Source/GameplayCameras/Legacy/Camera/CameraShake.h b/Engine/Plugins/Cameras/GameplayCameras/Source/GameplayCameras/Legacy/Camera/CameraShake.h
--- a/Engine/Plugins/Cameras/GameplayCameras/Source/GameplayCameras/Legacy/Camera/CameraShake.h
+++ b/Engine/Plugins/Cameras/GameplayCameras/Source/GameplayCameras/Legacy/Camera/CameraShake.h
@@ -6,6 +6,6 @@
 
 #include "CoreTypes.h"
 
-COMPILE_WARNING("UCameraShake has been renamed to ULegacyCameraShake, change to #include \"LegacyCameraShake.h\" and rename references")
 #include "LegacyCameraShake.h"
 
+UE_DEPRECATED_HEADER(5.5, "Use LegacyCameraShake.h instead of CameraShake.h")
diff --git a/Engine/Plugins/Cameras/GameplayCameras/Source/GameplayCameras/Public/Shakes/MatineeCameraShake.h b/Engine/Plugins/Cameras/GameplayCameras/Source/GameplayCameras/Public/Shakes/MatineeCameraShake.h
--- a/Engine/Plugins/Cameras/GameplayCameras/Source/GameplayCameras/Public/Shakes/MatineeCameraShake.h
+++ b/Engine/Plugins/Cameras/GameplayCameras/Source/GameplayCameras/Public/Shakes/MatineeCameraShake.h
@@ -6,5 +6,6 @@
 
 #include "CoreTypes.h"
 
-COMPILE_WARNING("UMatineeCameraShake has been renamed to ULegacyCameraShake, change to #include \"LegacyCameraShake.h\" and rename references")
 #include "LegacyCameraShake.h"
+
+UE_DEPRECATED_HEADER(5.5, "Use LegacyCameraShake.h instead of MatineeCameraShake.h")
diff --git a/Engine/Source/Runtime/MovieScene/Public/Evaluation/PreAnimatedState/MovieScenePreAnimatedEntityCaptureSource.h b/Engine/Source/Runtime/MovieScene/Public/Evaluation/PreAnimatedState/MovieScenePreAnimatedEntityCaptureSource.h
--- a/Engine/Source/Runtime/MovieScene/Public/Evaluation/PreAnimatedState/MovieScenePreAnimatedEntityCaptureSource.h
+++ b/Engine/Source/Runtime/MovieScene/Public/Evaluation/PreAnimatedState/MovieScenePreAnimatedEntityCaptureSource.h
@@ -4,5 +4,4 @@
 
 #include "Evaluation/PreAnimatedState/MovieScenePreAnimatedCaptureSources.h"
 
-COMPILE_WARNING("FPreAnimatedEntityCaptureSource has been moved to MovieScenePreAnimatedCaptureSources.h")
-
+UE_DEPRECATED_HEADER(5.4, "Use MovieScenePreAnimatedCaptureSources.h instead of MovieScenePreAnimatedEntityCaptureSource.h")
