Sequencer: Fix deprecated include warning

#rnx

diff --git a/Engine/Source/Runtime/MovieScene/Public/Evaluation/PreAnimatedState/MovieScenePreAnimatedEntityCaptureSource.h b/Engine/Source/Runtime/MovieScene/Public/Evaluation/PreAnimatedState/MovieScenePreAnimatedEntityCaptureSource.h
--- a/Engine/Source/Runtime/MovieScene/Public/Evaluation/PreAnimatedState/MovieScenePreAnimatedEntityCaptureSource.h
+++ b/Engine/Source/Runtime/MovieScene/Public/Evaluation/PreAnimatedState/MovieScenePreAnimatedEntityCaptureSource.h
@@ -1,8 +1,9 @@
 // Copyright Epic Games, Inc. All Rights Reserved.
 
-#pragma once
+// HEADER_UNIT_SKIP - Deprecated
 
-#include "Evaluation/PreAnimatedState/MovieScenePreAnimatedCaptureSources.h"
+#pragma once
 
-COMPILE_WARNING("FPreAnimatedEntityCaptureSource has been moved to MovieScenePreAnimatedCaptureSources.h")
+UE_DEPRECATED_HEADER(5.4, "Use MovieScenePreAnimatedCaptureSources.h instead of MovieScenePreAnimatedEntityCaptureSource.h")
 
+#include "Evaluation/PreAnimatedState/MovieScenePreAnimatedCaptureSources.h"
