[Backout] - CL33605064
Original CL Desc
-----------------------------------------------------------------
Sequencer: Fix deprecated include warning

#rnx

diff --git a/Engine/Source/Runtime/MovieScene/Public/Evaluation/PreAnimatedState/MovieScenePreAnimatedEntityCaptureSource.h b/Engine/Source/Runtime/MovieScene/Public/Evaluation/PreAnimatedState/MovieScenePreAnimatedEntityCaptureSource.h
--- a/Engine/Source/Runtime/MovieScene/Public/Evaluation/PreAnimatedState/MovieScenePreAnimatedEntityCaptureSource.h
+++ b/Engine/Source/Runtime/MovieScene/Public/Evaluation/PreAnimatedState/MovieScenePreAnimatedEntityCaptureSource.h
@@ -2,10 +2,7 @@
 
 #pragma once
 
-// HEADER_UNIT_SKIP - Deprecated
-
-#include "CoreTypes.h"
+#include "Evaluation/PreAnimatedState/MovieScenePreAnimatedCaptureSources.h"
 
-UE_DEPRECATED_HEADER(5.4, "Use MovieScenePreAnimatedCaptureSources.h instead of MovieScenePreAnimatedEntityCaptureSource.h")
+COMPILE_WARNING("FPreAnimatedEntityCaptureSource has been moved to MovieScenePreAnimatedCaptureSources.h")
 
-#include "Evaluation/PreAnimatedState/MovieScenePreAnimatedCaptureSources.h"
