Fix cyclic action graph
#rnx

diff --git a/Engine/Source/Runtime/MovieScene/MovieScene.Build.cs b/Engine/Source/Runtime/MovieScene/MovieScene.Build.cs
--- a/Engine/Source/Runtime/MovieScene/MovieScene.Build.cs
+++ b/Engine/Source/Runtime/MovieScene/MovieScene.Build.cs
@@ -8,7 +8,8 @@ public class MovieScene : ModuleRules
 	{
         PrivateIncludePathModuleNames.AddRange(
             new string[] {
-                "TargetPlatform",
+				"MovieSceneTracks",
+				"TargetPlatform",
 				"UniversalObjectLocator"
 			}
         );
@@ -19,7 +20,6 @@ public class MovieScene : ModuleRules
 				"CoreUObject",
                 "InputCore",
                 "Engine",
-				"MovieSceneTracks",
 				"SlateCore",
 				"TimeManagement",
 				"UniversalObjectLocator"
