Use header deprecation macro rather than bespoke warnings
#rnx

diff --git a/Engine/Source/Runtime/Engine/Public/PerPlatformProperties.h b/Engine/Source/Runtime/Engine/Public/PerPlatformProperties.h
--- a/Engine/Source/Runtime/Engine/Public/PerPlatformProperties.h
+++ b/Engine/Source/Runtime/Engine/Public/PerPlatformProperties.h
@@ -6,6 +6,7 @@ Redirect PerPlatformProperties.h include to UObject
 
 #pragma once
 
-#pragma message("WARNING: Do not #include PerPlatformProperties.h. Please use #include UObject/PerPlatformProperties.h instead. Deprecated in 5.5")
-
 #include "UObject/PerPlatformProperties.h"
+
+UE_DEPRECATED_HEADER(5.5, "WARNING: Do not #include PerPlatformProperties.h. Please use #include UObject/PerPlatformProperties.h instead.")
+
diff --git a/Engine/Source/Runtime/MovieScene/Public/Evaluation/PreAnimatedState/MovieScenePreAnimatedEntityCaptureSource.h b/Engine/Source/Runtime/MovieScene/Public/Evaluation/PreAnimatedState/MovieScenePreAnimatedEntityCaptureSource.h
--- a/Engine/Source/Runtime/MovieScene/Public/Evaluation/PreAnimatedState/MovieScenePreAnimatedEntityCaptureSource.h
+++ b/Engine/Source/Runtime/MovieScene/Public/Evaluation/PreAnimatedState/MovieScenePreAnimatedEntityCaptureSource.h
@@ -4,5 +4,5 @@
 
 #include "Evaluation/PreAnimatedState/MovieScenePreAnimatedCaptureSources.h"
 
-COMPILE_WARNING("FPreAnimatedEntityCaptureSource has been moved to MovieScenePreAnimatedCaptureSources.h")
+UE_DEPRECATED_HEADER(5.4, "FPreAnimatedEntityCaptureSource has been moved to MovieScenePreAnimatedCaptureSources.h")
 
