[MovieScene]
* Added so UE_SELECT_ANY translates to __attribute__((weak)) for non-microsoft platforms

diff --git a/Engine/Source/Runtime/MovieScene/Public/EntitySystem/MovieSceneEntityManager.h b/Engine/Source/Runtime/MovieScene/Public/EntitySystem/MovieSceneEntityManager.h
--- a/Engine/Source/Runtime/MovieScene/Public/EntitySystem/MovieSceneEntityManager.h
+++ b/Engine/Source/Runtime/MovieScene/Public/EntitySystem/MovieSceneEntityManager.h
@@ -1089,10 +1089,10 @@ private:
 	FEntityManager* EntityManager;
 };
 
-#if PLATFORM_WINDOWS
-#define UE_SELECT_ANY __declspec(selectany)
+#if PLATFORM_COMPILER_CLANG
+#define UE_SELECT_ANY __attribute__((weak))
 #else
-#define UE_SELECT_ANY 
+#define UE_SELECT_ANY __declspec(selectany)
 #endif
 
 extern MOVIESCENE_API FEntityManager* GEntityManagerForDebuggingVisualizers;
