Sequencer: runners shouldn't abandon their linker

Remove the idea of runners abandonning their linkers because:

- We don't need to remove event handlers anymore because the last one was removed in CL 17721867
- We don't need to clear our linker pointer anymore because it's a weak pointer since CL 16912978

So there's no reason anymore to handle a linker going away. The runner's WeakLinker pointer will naturally go stale and we have nothing to do.

This change however fixes one bug... Calling a linker's Reset() method would destroy and reset everything, including telling all systems to abandon their linker. But the runner was also listening to the AbandonLinker event and clearing its WeakLinker pointer for those historical reasons outlined above. This would cause the runner to become effectively inert, unable to run any evaluations after that. So this change keeps the runner valid even after we reset everything.

diff --git a/Engine/Source/Runtime/MovieScene/Private/EntitySystem/MovieSceneEntitySystemRunner.cpp b/Engine/Source/Runtime/MovieScene/Private/EntitySystem/MovieSceneEntitySystemRunner.cpp
--- a/Engine/Source/Runtime/MovieScene/Private/EntitySystem/MovieSceneEntitySystemRunner.cpp
+++ b/Engine/Source/Runtime/MovieScene/Private/EntitySystem/MovieSceneEntitySystemRunner.cpp
@@ -109,16 +109,10 @@ FMovieSceneEntitySystemRunner::FMovieSceneEntitySystemRunner(UMovieSceneEntitySy
 	, bRequireFullFlush(false)
 	, bIsUpdatingSequence(false)
 {
-	// We need to handle the case of our linker being GC'ed, in case someone else is holding
-	// our ref-count up in addition to the linker.
-	InLinker->Events.AbandonLinker.AddRaw(this, &FMovieSceneEntitySystemRunner::OnLinkerAbandon);
 }
 
 FMovieSceneEntitySystemRunner::~FMovieSceneEntitySystemRunner()
 {
-	// We don't need to unregister AbandonLinker here since we are either being destroyed along with
-	// our linker, or someone kept us alive longer than the linker and we already unregistered the
-	// event inisde OnLinkerAbandon.
 }
 
 UE::MovieScene::FEntityManager* FMovieSceneEntitySystemRunner::GetEntityManager() const
@@ -1296,18 +1290,6 @@ void FMovieSceneEntitySystemRunner::MarkForUpdate(FInstanceHandle InInstanceHand
 	CurrentInstances.Add(FQueuedUpdateParams{ InInstanceHandle, UpdateFlags });
 }
 
-void FMovieSceneEntitySystemRunner::OnLinkerAbandon(UMovieSceneEntitySystemLinker* InLinker)
-{
-	// WARNING: this can be called with a linker that is PendingKill
-
-	if (ensure(InLinker))
-	{
-		InLinker->Events.AbandonLinker.RemoveAll(this);
-	}
-
-	WeakLinker.Reset();
-}
-
 FMovieSceneEntitySystemEventTriggers& FMovieSceneEntitySystemRunner::GetQueuedEventTriggers()
 {
 	checkf(bCanQueueEventTriggers, TEXT("Can't queue event triggers at this point in the update loop."));
diff --git a/Engine/Source/Runtime/MovieScene/Public/EntitySystem/MovieSceneEntitySystemRunner.h b/Engine/Source/Runtime/MovieScene/Public/EntitySystem/MovieSceneEntitySystemRunner.h
--- a/Engine/Source/Runtime/MovieScene/Public/EntitySystem/MovieSceneEntitySystemRunner.h
+++ b/Engine/Source/Runtime/MovieScene/Public/EntitySystem/MovieSceneEntitySystemRunner.h
@@ -188,10 +188,6 @@ public:
 	/** Create a bound runner */
 	FMovieSceneEntitySystemRunner(UMovieSceneEntitySystemLinker* InLinker);
 
-private:
-
-	MOVIESCENE_API void OnLinkerAbandon(UMovieSceneEntitySystemLinker* Linker);
-
 private:
 
 	/**
