Support adjusting dynamic resolution frame time budget based on display rate when sequence uses "Lock to Display Rate at Runtime"
- behaviour can be controlled using cvar Sequencer.ApplyDisplayRateToDynamicResolutionFrameTimeBudget (disabled by default)

diff --git a/Engine/Source/Runtime/MovieScene/Private/MovieSceneSequencePlayer.cpp b/Engine/Source/Runtime/MovieScene/Private/MovieSceneSequencePlayer.cpp
--- a/Engine/Source/Runtime/MovieScene/Private/MovieSceneSequencePlayer.cpp
+++ b/Engine/Source/Runtime/MovieScene/Private/MovieSceneSequencePlayer.cpp
@@ -68,6 +68,13 @@ static FAutoConsoleVariableRef CVarSequencerSmoothedNetSyncDeviationThreshold(
 	TEXT("(Default: 200ms. Defines the acceptable deviation for smoothed net sync samples. Samples outside this deviation will be discarded.")
 	);
 
+int32 GSequencerApplyDisplayRateToDynResFrameTimeBudget = 0;
+static FAutoConsoleVariableRef CVarSequencerApplyDisplayRateToDynResFrameTimeBudget(
+	TEXT("Sequencer.ApplyDisplayRateToDynamicResolutionFrameTimeBudget"),
+	GSequencerApplyDisplayRateToDynResFrameTimeBudget,
+	TEXT("(Whether to override r.DynamicRes.FrameTimeBudget based on sequence display rate when using 'Lock to Display Rate at Runtime'.")
+);
+
 bool FMovieSceneSequenceLoopCount::SerializeFromMismatchedTag( const FPropertyTag& Tag, FStructuredArchive::FSlot Slot )
 {
 	if (Tag.Type == NAME_IntProperty)
@@ -191,6 +198,14 @@ UMovieSceneSequencePlayer::~UMovieSceneSequencePlayer()
 	{
 		GEngine->SetMaxFPS(OldMaxTickRate.GetValue());
 	}
+
+	if (bOverridingDynResFrameTimeBudget)
+	{
+		static IConsoleVariable* CVarDynResFrameTimeBudget = IConsoleManager::Get().FindConsoleVariable(TEXT("r.DynamicRes.FrameTimeBudget"));
+		CVarDynResFrameTimeBudget->Unset(ECVF_SetByCode);
+
+		bOverridingDynResFrameTimeBudget = false;
+	}
 }
 
 void UMovieSceneSequencePlayer::UpdateNetworkSyncProperties()
@@ -343,6 +358,16 @@ void UMovieSceneSequencePlayer::PlayInternal()
 			}
 
 			GEngine->SetMaxFPS(1.f / PlayPosition.GetInputRate().AsInterval());
+
+			if (GSequencerApplyDisplayRateToDynResFrameTimeBudget)
+			{
+				const float DyResFrameTimeBudget = PlayPosition.GetInputRate().AsInterval() * 1000.0f;
+
+				static IConsoleVariable* CVarDynResFrameTimeBudget = IConsoleManager::Get().FindConsoleVariable(TEXT("r.DynamicRes.FrameTimeBudget"));
+				CVarDynResFrameTimeBudget->Set(DyResFrameTimeBudget, ECVF_SetByCode);
+
+				bOverridingDynResFrameTimeBudget = true;
+			}
 		}
 
 		if (!PlayPosition.GetLastPlayEvalPostition().IsSet() || PlayPosition.GetLastPlayEvalPostition() != PlayPosition.GetCurrentPosition())
@@ -499,6 +524,14 @@ void UMovieSceneSequencePlayer::StopInternal(FFrameTime TimeToResetTo)
 				this->OldMaxTickRate.Reset();
 			}
 
+			if (bOverridingDynResFrameTimeBudget)
+			{
+				static IConsoleVariable* CVarDynResFrameTimeBudget = IConsoleManager::Get().FindConsoleVariable(TEXT("r.DynamicRes.FrameTimeBudget"));
+				CVarDynResFrameTimeBudget->Unset(ECVF_SetByCode);
+
+				bOverridingDynResFrameTimeBudget = false;
+			}
+
 			this->UpdateNetworkSyncProperties();
 
 			const FString SequenceName = this->GetSequenceName(true);
diff --git a/Engine/Source/Runtime/MovieScene/Public/MovieSceneSequencePlayer.h b/Engine/Source/Runtime/MovieScene/Public/MovieSceneSequencePlayer.h
--- a/Engine/Source/Runtime/MovieScene/Public/MovieSceneSequencePlayer.h
+++ b/Engine/Source/Runtime/MovieScene/Public/MovieSceneSequencePlayer.h
@@ -765,6 +765,9 @@ private:
 	/** The maximum tick rate prior to playing (used for overriding delta time during playback). */
 	TOptional<double> OldMaxTickRate;
 
+	/** Whether dynamic resolution frame time budget is being overridden. */
+	bool bOverridingDynResFrameTimeBudget = false;
+
 	/**
 	* The last world game time at which we were ticked. Game time used is dependent on bTickEvenWhenPaused
 	* Valid only if we've been ticked at least once since having a tick interval
