Enforced non-const source argument in RelocateConstructItems.

diff --git a/Engine/Source/Runtime/Core/Public/Templates/MemoryOps.h b/Engine/Source/Runtime/Core/Public/Templates/MemoryOps.h
--- a/Engine/Source/Runtime/Core/Public/Templates/MemoryOps.h
+++ b/Engine/Source/Runtime/Core/Public/Templates/MemoryOps.h
@@ -218,7 +218,11 @@ FORCEINLINE void RelocateConstructItems(void* Dest, SourceElementType* Source, S
 	{
 		// Should never get here, but this construct should improve the error messages we get when we try to call this function with incomplete types
 	}
-	else if constexpr (UE::Core::Private::MemoryOps::TCanBitwiseRelocate_V<DestinationElementType, std::remove_const_t<SourceElementType>>)
+	else if constexpr (std::is_const_v<SourceElementType>)
+	{
+		static_assert(sizeof(SourceElementType) == 0, "RelocateConstructItems: Source cannot be const");
+	}
+	else if constexpr (UE::Core::Private::MemoryOps::TCanBitwiseRelocate_V<DestinationElementType, SourceElementType>)
 	{
 		/* All existing UE containers seem to assume trivial relocatability (i.e. memcpy'able) of their members,
 		 * so we're going to assume that this is safe here.  However, it's not generally possible to assume this
diff --git a/Engine/Source/Runtime/Core/Tests/Templates/MemoryOpsTest.cpp b/Engine/Source/Runtime/Core/Tests/Templates/MemoryOpsTest.cpp
--- a/Engine/Source/Runtime/Core/Tests/Templates/MemoryOpsTest.cpp
+++ b/Engine/Source/Runtime/Core/Tests/Templates/MemoryOpsTest.cpp
@@ -88,7 +88,7 @@ TEST_CASE_NAMED(FRelocateConstructItemsTest, "UE::Core::RelocateConstructItems",
 	{
 		alignas(FString) char Buffer[sizeof(FString) * 5];
 
-		const FString* TypedBuffer = (const FString*)Buffer;
+		FString* TypedBuffer = (FString*)Buffer;
 
 		new ((void*)(TypedBuffer + 2)) FString(TEXT("String A1"));
 		new ((void*)(TypedBuffer + 3)) FString(TEXT("String B1"));
@@ -109,7 +109,7 @@ TEST_CASE_NAMED(FRelocateConstructItemsTest, "UE::Core::RelocateConstructItems",
 	{
 		alignas(FString) char Buffer[sizeof(FString) * 5];
 
-		const FString* TypedBuffer = (const FString*)Buffer;
+		FString* TypedBuffer = (FString*)Buffer;
 
 		new ((void*)(TypedBuffer + 0)) FString(TEXT("String A2"));
 		new ((void*)(TypedBuffer + 1)) FString(TEXT("String B2"));
diff --git a/Engine/Source/Runtime/MovieScene/Public/Evaluation/MovieSceneEvaluationTree.h b/Engine/Source/Runtime/MovieScene/Public/Evaluation/MovieSceneEvaluationTree.h
--- a/Engine/Source/Runtime/MovieScene/Public/Evaluation/MovieSceneEvaluationTree.h
+++ b/Engine/Source/Runtime/MovieScene/Public/Evaluation/MovieSceneEvaluationTree.h
@@ -921,7 +921,7 @@ void TEvaluationTreeEntryContainer<ElementType>::Insert(FEvaluationTreeEntryHand
 	// Shift anything >= the insert index up one
 	if (Index <= Entry.Size-1)
 	{
-		const ElementType* OldLocation = &Items[Entry.StartIndex + Index];
+		ElementType* OldLocation = &Items[Entry.StartIndex + Index];
 		void* NewLocation = &Items[Entry.StartIndex + Index + 1];
 
 		RelocateConstructItems<ElementType>(NewLocation, OldLocation, Entry.Size - Index);
@@ -972,7 +972,7 @@ void TEvaluationTreeEntryContainer<ElementType>::ReserveEntry(FEvaluationTreeEnt
 	// Reallocate the items to the new entry location
 	Items.SetNum(Items.Num() + Entry.Capacity);
 
-	const ElementType* OldLocation = &Items[Entry.StartIndex];
+	ElementType* OldLocation = &Items[Entry.StartIndex];
 	void* NewLocation = &Items[NewStartIndex];
 	RelocateConstructItems<ElementType>(NewLocation, OldLocation, Entry.Size);
 	DefaultConstructItems<ElementType>((void*)OldLocation, Entry.Size);
