Sequencer: Fixed printing of system dependencies not working for the currently relevant set of systems

Also fixed Scheduling dependencies not showing up as part of the Evaluation phase

diff --git a/Engine/Source/Runtime/MovieScene/Private/EntitySystem/MovieSceneEntitySystem.cpp b/Engine/Source/Runtime/MovieScene/Private/EntitySystem/MovieSceneEntitySystem.cpp
--- a/Engine/Source/Runtime/MovieScene/Private/EntitySystem/MovieSceneEntitySystem.cpp
+++ b/Engine/Source/Runtime/MovieScene/Private/EntitySystem/MovieSceneEntitySystem.cpp
@@ -118,7 +118,36 @@ struct FSystemDependencyGraph
 		}
 	}
 
-	void DebugPrint(bool bUpdateCache = true)
+	void PrintFilteredFlowGraph(TArrayView<const TBitArray<>> GroupFilters, FStringBuilderBase& OutString)
+	{
+		using FDirectionalEdge = FDirectedGraph::FDirectionalEdge;
+
+		FDirectedGraph::FDiscoverCyclicEdges CyclicEdges(&FlowGraph);
+		CyclicEdges.Search();
+
+		TArrayView<const FDirectionalEdge> FlowEdges = FlowGraph.GetEdges();
+		for (int32 EdgeIndex = 0; EdgeIndex < FlowEdges.Num(); ++EdgeIndex)
+		{
+			FDirectionalEdge Edge = FlowEdges[EdgeIndex];
+
+			for (int32 GroupIndex = 0; GroupIndex < GroupFilters.Num(); ++GroupIndex)
+			{
+				if (GroupFilters[GroupIndex].IsValidIndex(Edge.FromNode) && GroupFilters[GroupIndex][Edge.FromNode] && 
+					GroupFilters[GroupIndex].IsValidIndex(Edge.ToNode) && GroupFilters[GroupIndex][Edge.ToNode])
+				{
+					const bool bIsCyclic = CyclicEdges.IsCyclic(EdgeIndex);
+
+					OutString.Appendf(TEXT("\tflow_node%d_%d -> flow_node%d_%d [color=\"%s\"];\n"),
+						(int32)Edge.FromNode, GroupIndex,
+						(int32)Edge.ToNode, GroupIndex,
+						bIsCyclic ? TEXT("#FF0000") : TEXT("#3992ad")
+					);
+				}
+			}
+		}
+	}
+
+	void DebugPrint(bool bUpdateCache)
 	{
 		const TCHAR FormatString[] =
 			TEXT("----------------------------------------------------------------------------------\n")
@@ -134,6 +163,11 @@ struct FSystemDependencyGraph
 		GLog->Log(FString::Printf(FormatString, *ToString()));
 	}
 
+	MOVIESCENE_API void DebugPrint()
+	{
+		DebugPrint(true);
+	}
+
 	FString ToString() const
 	{
 		using FDirectionalEdge = FDirectedGraph::FDirectionalEdge;
@@ -190,7 +224,7 @@ struct FSystemDependencyGraph
 				{
 					FlowStrings[1] += FString::Printf(TEXT("\t\tflow_node%d_1[label=\"%s\"];\n"), NodeID, *SystemClass->GetName());
 				}
-				if (EnumHasAnyFlags(SystemPhase, ESystemPhase::Evaluation))
+				if (EnumHasAnyFlags(SystemPhase, ESystemPhase::Evaluation | ESystemPhase::Scheduling))
 				{
 					FlowStrings[2] += FString::Printf(TEXT("\t\tflow_node%d_2[label=\"%s\"];\n"), NodeID, *SystemClass->GetName());
 				}
@@ -222,7 +256,7 @@ struct FSystemDependencyGraph
 						FlowStrings[1] += FString::Printf(TEXT("\t\tflow_node%d_1[label=\"Component ID=%s\"];\n"), NodeID, Node.WriteComponentType.BitIndex());
 #endif
 					}
-					if (EnumHasAnyFlags(*ConnectedSystemPhase, ESystemPhase::Evaluation))
+					if (EnumHasAnyFlags(*ConnectedSystemPhase, ESystemPhase::Evaluation | ESystemPhase::Scheduling))
 					{
 #if UE_MOVIESCENE_ENTITY_DEBUG
 						FlowStrings[2] += FString::Printf(TEXT("\t\tflow_node%d_2[label=\"%s\"];\n"), NodeID, *WriteComponentInfo.DebugInfo->DebugName);
@@ -281,8 +315,8 @@ struct FSystemDependencyGraph
 				{
 					String += FString::Printf(TEXT("\tflow_node%d_1 -> flow_node%d_1 [color=\"%s\"];\n"), (int32)Edge.FromNode, (int32)Edge.ToNode, *EdgeColor);
 				}
-				if ((FromPhase == ESystemPhase::None || EnumHasAnyFlags(FromPhase, ESystemPhase::Evaluation)) && 
-						(ToPhase == ESystemPhase::None || EnumHasAnyFlags(ToPhase, ESystemPhase::Evaluation)))
+				if ((FromPhase == ESystemPhase::None || EnumHasAnyFlags(FromPhase, ESystemPhase::Evaluation | ESystemPhase::Scheduling)) && 
+						(ToPhase == ESystemPhase::None || EnumHasAnyFlags(ToPhase, ESystemPhase::Evaluation | ESystemPhase::Scheduling)))
 				{
 					String += FString::Printf(TEXT("\tflow_node%d_2 -> flow_node%d_2 [color=\"%s\"];\n"), (int32)Edge.FromNode, (int32)Edge.ToNode, *EdgeColor);
 				}
@@ -347,7 +381,7 @@ private:
 
 		FlowGraph.CleanUpDanglingEdges();
 
-		checkf(!FlowGraph.IsCyclic(), TEXT("Cycle detected in system flow graph!"));
+		checkf(!FlowGraph.IsCyclic(), TEXT("Cycle detected in system flow graph! \n\n %s"), *ToString());
 	}
 
 	void SetupFlowDependencies(int32 NodeIndex)
@@ -613,6 +647,13 @@ void UMovieSceneEntitySystem::GetSubsequentSystems(uint16 FromGlobalNodeID, TArr
 	GlobalDependencyGraph.GetSubsequents(FromGlobalNodeID, OutSubsequentGlobalNodeIDs);
 }
 
+void UMovieSceneEntitySystem::PrintFilteredFlowGraph(TArrayView<const TBitArray<>> GroupFilters, FStringBuilderBase& OutString)
+{
+	using namespace UE::MovieScene;
+
+	GlobalDependencyGraph.PrintFilteredFlowGraph(GroupFilters, OutString);
+}
+
 void UMovieSceneEntitySystem::DebugPrintGlobalDependencyGraph(bool bUpdateCache)
 {
 	using namespace UE::MovieScene;
diff --git a/Engine/Source/Runtime/MovieScene/Private/EntitySystem/MovieSceneEntitySystemGraphs.cpp b/Engine/Source/Runtime/MovieScene/Private/EntitySystem/MovieSceneEntitySystemGraphs.cpp
--- a/Engine/Source/Runtime/MovieScene/Private/EntitySystem/MovieSceneEntitySystemGraphs.cpp
+++ b/Engine/Source/Runtime/MovieScene/Private/EntitySystem/MovieSceneEntitySystemGraphs.cpp
@@ -238,7 +238,8 @@ FString FMovieSceneEntitySystemGraph::ToString() const
 {
 	using namespace UE::MovieScene;
 
-	FString String;
+	TStringBuilder<1024> String;
+
 	String += TEXT("\ndigraph FMovieSceneEntitySystemGraph {\n");
 	String += TEXT("\tnode [shape=record,height=.1];\n");
 
@@ -252,6 +253,9 @@ FString FMovieSceneEntitySystemGraph::ToString() const
 
 	FString ReferenceGraphString = TEXT("\tsubgraph cluster_references { label=\"Explicit Reference Graph (connections imply ownership)\"; color=\"#bfc74c\";\n");
 
+	TBitArray<> FlowFilter[4];
+	TBitArray<> ReferencedSystems;
+
 	for (int32 SystemIndex = 0; SystemIndex < this->Nodes.Array.GetMaxIndex(); ++SystemIndex)
 	{
 		if (Nodes.Array.IsAllocated(SystemIndex))
@@ -259,28 +263,56 @@ FString FMovieSceneEntitySystemGraph::ToString() const
 			UMovieSceneEntitySystem* System = Nodes.Array[SystemIndex].System;
 
 			ESystemPhase SystemPhase = System->GetPhase();
+			const uint16 GlobalIndex = System->GetGlobalDependencyGraphID();
 
 			if (EnumHasAnyFlags(SystemPhase, ESystemPhase::Spawn))
 			{
-				FlowStrings[0] += FString::Printf(TEXT("\t\tflow_node%d_0[label=\"%s\"];\n"), SystemIndex, *System->GetName());
+				FlowStrings[0] += FString::Printf(TEXT("\t\tflow_node%d_0[label=\"%s\"];\n"), GlobalIndex, *System->GetName());
+				FlowFilter[0].PadToNum(GlobalIndex+1, false);
+				FlowFilter[0][GlobalIndex] = true;
 			}
 			if (EnumHasAnyFlags(SystemPhase, ESystemPhase::Instantiation))
 			{
-				FlowStrings[1] += FString::Printf(TEXT("\t\tflow_node%d_1[label=\"%s\"];\n"), SystemIndex, *System->GetName());
+				FlowStrings[1] += FString::Printf(TEXT("\t\tflow_node%d_1[label=\"%s\"];\n"), GlobalIndex, *System->GetName());
+				FlowFilter[1].PadToNum(GlobalIndex+1, false);
+				FlowFilter[1][GlobalIndex] = true;
 			}
-			if (EnumHasAnyFlags(SystemPhase, ESystemPhase::Evaluation))
+			if (EnumHasAnyFlags(SystemPhase, ESystemPhase::Evaluation | ESystemPhase::Scheduling))
 			{
-				FlowStrings[2] += FString::Printf(TEXT("\t\tflow_node%d_2[label=\"%s\"];\n"), SystemIndex, *System->GetName());
+				FlowStrings[2] += FString::Printf(TEXT("\t\tflow_node%d_2[label=\"%s\"];\n"), GlobalIndex, *System->GetName());
+				FlowFilter[2].PadToNum(GlobalIndex+1, false);
+				FlowFilter[2][GlobalIndex] = true;
 			}
 			if (EnumHasAnyFlags(SystemPhase, ESystemPhase::Finalization))
 			{
-				FlowStrings[3] += FString::Printf(TEXT("\t\tflow_node%d_3[label=\"%s\"];\n"), SystemIndex, *System->GetName());
+				FlowStrings[3] += FString::Printf(TEXT("\t\tflow_node%d_3[label=\"%s\"];\n"), GlobalIndex, *System->GetName());
+				FlowFilter[3].PadToNum(GlobalIndex+1, false);
+				FlowFilter[3][GlobalIndex] = true;
 			}
 
-			ReferenceGraphString += FString::Printf(TEXT("\t\treference_node%d[label=\"%s\"];\n"), SystemIndex, *System->GetName());
+			TArrayView<const FDirectionalEdge> ReferenceEdges = ReferenceGraph.GetEdgesFrom(SystemIndex);
+			if (ReferenceEdges.Num() > 0)
+			{
+				ReferencedSystems.PadToNum(SystemIndex+1, false);
+				ReferencedSystems[SystemIndex] = true;
+
+				for (FDirectionalEdge Edge : ReferenceEdges)
+				{
+					ReferencedSystems.PadToNum(Edge.ToNode+1, false);
+					ReferencedSystems[Edge.ToNode] = true;
+				}
+			}
 		}
 	}
 
+	for (TConstSetBitIterator<> ReferencedSystemIt(ReferencedSystems); ReferencedSystemIt; ++ReferencedSystemIt)
+	{
+		UMovieSceneEntitySystem* System = Nodes.Array[ReferencedSystemIt.GetIndex()].System;
+		const uint16 GlobalIndex = System->GetGlobalDependencyGraphID();
+
+		ReferenceGraphString += FString::Printf(TEXT("\t\treference_node%d[label=\"%s\"];\n"), GlobalIndex, *System->GetName());
+	}
+
 	for (FString& FlowString : FlowStrings)
 	{
 		String += FlowString;
@@ -289,6 +321,8 @@ FString FMovieSceneEntitySystemGraph::ToString() const
 	String += ReferenceGraphString;
 	String += TEXT("\t}\n");
 
+	UMovieSceneEntitySystem::PrintFilteredFlowGraph(FlowFilter, String);
+
 	{
 		FDirectedGraph::FDiscoverCyclicEdges CyclicEdges(&ReferenceGraph);
 		CyclicEdges.Search();
@@ -299,12 +333,14 @@ FString FMovieSceneEntitySystemGraph::ToString() const
 			FDirectionalEdge Edge = ReferenceEdges[EdgeIndex];
 			const bool bIsCyclic = CyclicEdges.IsCyclic(EdgeIndex);
 
-			String += FString::Printf(TEXT("\treference_node%d -> reference_node%d [color=\"%s\"];\n"), (int32)Edge.FromNode, (int32)Edge.ToNode, bIsCyclic ? TEXT("#FF0000") : TEXT("#3992ad"));
+			const uint16 FromGlobalIndex = Nodes.Array[Edge.FromNode].System->GetGlobalDependencyGraphID();
+			const uint16 ToGlobalIndex   = Nodes.Array[Edge.ToNode].System->GetGlobalDependencyGraphID();
+			String += FString::Printf(TEXT("\treference_node%d -> reference_node%d [color=\"%s\"];\n"), (int32)FromGlobalIndex, (int32)ToGlobalIndex, bIsCyclic ? TEXT("#FF0000") : TEXT("#3992ad"));
 		}
 	}
 
 	String += TEXT("}");
-	return String;
+	return String.ToString();
 }
 
 TArray<UMovieSceneEntitySystem*> FMovieSceneEntitySystemGraph::GetSystems() const
diff --git a/Engine/Source/Runtime/MovieScene/Public/EntitySystem/MovieSceneEntitySystem.h b/Engine/Source/Runtime/MovieScene/Public/EntitySystem/MovieSceneEntitySystem.h
--- a/Engine/Source/Runtime/MovieScene/Public/EntitySystem/MovieSceneEntitySystem.h
+++ b/Engine/Source/Runtime/MovieScene/Public/EntitySystem/MovieSceneEntitySystem.h
@@ -130,6 +130,14 @@ public:
 	 */
 	static MOVIESCENE_API void GetSubsequentSystems(uint16 FromGlobalNodeID, TArray<uint16>& OutSubsequentGlobalNodeIDs);
 
+	/**
+	 * Print the global system graph in graphviz format filtered by the specified global system IDs, one filter per phase
+	 * *
+	 * @param GroupFilters  Array of bit filters, one per phase, defining which systems to include in the output string
+	 * @param OutString     String builder to add the output graph to
+	 */
+	static MOVIESCENE_API void PrintFilteredFlowGraph(TArrayView<const TBitArray<>> GroupFilters, FStringBuilderBase& OutString);
+
 	/**
 	 * Prints a graphviz markup for the global system dependency graph
 	 */
