Sequencer: Improvements to debugging visualizers

All component types are now supported by the component visualizers by using a polymorphic debug wrapper for component headers.
Statically defined template specializations for each component type are no longer necessary.
This is enabled by setting Sequencer.RichComponentDebugging to 1.

diff --git a/Engine/Source/Runtime/MovieScene/Private/EntitySystem/MovieSceneEntityManager.cpp b/Engine/Source/Runtime/MovieScene/Private/EntitySystem/MovieSceneEntityManager.cpp
--- a/Engine/Source/Runtime/MovieScene/Private/EntitySystem/MovieSceneEntityManager.cpp
+++ b/Engine/Source/Runtime/MovieScene/Private/EntitySystem/MovieSceneEntityManager.cpp
@@ -53,6 +53,15 @@ FAutoConsoleVariableRef CVarThreadedEvaluationEntityThreshold(
 	ECVF_Default
 );
 
+#if UE_MOVIESCENE_ENTITY_DEBUG
+bool GRichComponentDebugging = false;
+FAutoConsoleVariableRef CVarRichComponentDebugging(
+	TEXT("Sequencer.RichComponentDebugging"),
+	GRichComponentDebugging,
+	TEXT("(Default: false. Whether to enable rich component debugging within Sequencer.")
+	);
+#endif // UE_MOVIESCENE_ENTITY_DEBUG
+
 FEntityManager* GEntityManagerForDebuggingVisualizers = nullptr;
 
 static bool IsValidUint16(int32 Test)
@@ -211,11 +220,24 @@ struct FEntityInitializer
 
 			for (FComponentMaskIterator It = EntityComponentMask.Iterate(); It; ++It, ++Header)
 			{
-				new (Header) FComponentHeader();
-
 				FComponentTypeID ComponentTypeID = FComponentTypeID::FromBitIndex(It.GetIndex());
 				const FComponentTypeInfo& TypeInfo = EntityManager.GetComponents()->GetComponentTypeChecked(ComponentTypeID);
 
+
+#if UE_MOVIESCENE_ENTITY_DEBUG
+				if (GRichComponentDebugging)
+				{
+					TypeInfo.DebugInfo->InitializeComponentHeader(Header);
+					Header->Size = &Allocation->Size;
+				}
+				else
+				{
+					new (Header) FComponentHeader();
+				}
+#else
+				new (Header) FComponentHeader();
+#endif
+
 				Header->ComponentType = ComponentTypeID;
 				Header->Sizeof = TypeInfo.Sizeof;
 				Header->PostWriteComponents(WriteContext);
@@ -235,6 +257,13 @@ struct FEntityInitializer
 					check(IsAligned(Header->Components, TypeInfo.Alignment));
 
 					ComponentDataPtr += TypeInfo.Sizeof * InitInfo.InitialCapacity;
+
+#if UE_MOVIESCENE_ENTITY_DEBUG
+					if (GRichComponentDebugging)
+					{
+						TypeInfo.DebugInfo->InitializeDebugComponentData(*Header, InitInfo.InitialCapacity);
+					}
+#endif
 				}
 			}
 		}
diff --git a/Engine/Source/Runtime/MovieScene/Private/EntitySystem/MovieSceneEntitySystemTypes.cpp b/Engine/Source/Runtime/MovieScene/Private/EntitySystem/MovieSceneEntitySystemTypes.cpp
--- a/Engine/Source/Runtime/MovieScene/Private/EntitySystem/MovieSceneEntitySystemTypes.cpp
+++ b/Engine/Source/Runtime/MovieScene/Private/EntitySystem/MovieSceneEntitySystemTypes.cpp
@@ -3,6 +3,7 @@
 #include "EntitySystem/MovieSceneEntitySystemTypes.h"
 #include "EntitySystem/MovieSceneEntityManager.h"
 #include "EntitySystem/MovieSceneComponentAccessors.h"
+#include "EntitySystem/MovieSceneComponentDebug.h"
 
 namespace UE
 {
@@ -155,5 +156,17 @@ const FComponentMask& FEntityAllocationProxy::GetAllocationType() const
 	return Manager->EntityAllocationMasks[AllocationIndex];
 }
 
+#if UE_MOVIESCENE_ENTITY_DEBUG
+FComponentHeader::~FComponentHeader()
+{
+	//~~~~~~~~~~~~~~~~~~~~~~~~~~~~~~~~~~~~~~~~~~~~~~~~~~~~~~~~~~~
+	// WARNING: This destructor is only compiled in when UE_MOVIESCENE_ENTITY_DEBUG is true
+	if (DebugComponents)
+	{
+		delete[] DebugComponents;
+	}
+}
+#endif // UE_MOVIESCENE_ENTITY_DEBUG
+
 } // namespace MovieScene
 } // namespace UE
diff --git a/Engine/Source/Runtime/MovieScene/Public/EntitySystem/BuiltInComponentTypes.h b/Engine/Source/Runtime/MovieScene/Public/EntitySystem/BuiltInComponentTypes.h
--- a/Engine/Source/Runtime/MovieScene/Public/EntitySystem/BuiltInComponentTypes.h
+++ b/Engine/Source/Runtime/MovieScene/Public/EntitySystem/BuiltInComponentTypes.h
@@ -808,9 +808,5 @@ private:
 	TMap<FComponentTypeID, FComponentTypeID> ResultToBase;
 };
 
-#if UE_MOVIESCENE_ENTITY_DEBUG
-template<> struct TComponentDebugType<FEntityGroupID> { static const EComponentDebugType Type = EComponentDebugType::GroupID; };
-#endif
-
 } // namespace MovieScene
 } // namespace UE
diff --git a/Engine/Source/Runtime/MovieScene/Public/EntitySystem/MovieSceneBlenderSystemTypes.h b/Engine/Source/Runtime/MovieScene/Public/EntitySystem/MovieSceneBlenderSystemTypes.h
--- a/Engine/Source/Runtime/MovieScene/Public/EntitySystem/MovieSceneBlenderSystemTypes.h
+++ b/Engine/Source/Runtime/MovieScene/Public/EntitySystem/MovieSceneBlenderSystemTypes.h
@@ -89,10 +89,3 @@ struct FMovieSceneBlendChannelID
 	/** The blend channel ID for the given blender system */
 	uint16 ChannelID;
 };
-
-#if UE_MOVIESCENE_ENTITY_DEBUG
-namespace UE::MovieScene
-{
-	template<> struct TComponentDebugType<FMovieSceneBlendChannelID> { static const EComponentDebugType Type = EComponentDebugType::BlendChannelID; };
-}
-#endif
diff --git a/Engine/Source/Runtime/MovieScene/Public/EntitySystem/MovieSceneComponentDebug.h b/Engine/Source/Runtime/MovieScene/Public/EntitySystem/MovieSceneComponentDebug.h
--- a/Engine/Source/Runtime/MovieScene/Public/EntitySystem/MovieSceneComponentDebug.h
+++ b/Engine/Source/Runtime/MovieScene/Public/EntitySystem/MovieSceneComponentDebug.h
@@ -2,44 +2,44 @@
 
 #pragma once
 
+#include "MovieSceneFwd.h"
+
+#if UE_MOVIESCENE_ENTITY_DEBUG
+
 #include "EntitySystem/MovieSceneEntitySystemTypes.h"
 #include "Math/Vector4.h"
 #include "Misc/FrameTime.h"
 #include "UObject/ObjectKey.h"
-
-#if UE_MOVIESCENE_ENTITY_DEBUG
+#include "Misc/InlineValue.h"
 
 class FName;
 
-namespace UE
+namespace UE::MovieScene
+{
+
+MOVIESCENE_API extern bool GRichComponentDebugging;
+
+struct IComponentDebuggingTypedPtr
 {
-namespace MovieScene
+	virtual ~IComponentDebuggingTypedPtr() {}
+	void* Ptr;
+};
+template<typename T>
+struct TComponentDebuggingTypedPtr : IComponentDebuggingTypedPtr
 {
+	TComponentDebuggingTypedPtr()
+	{
+		static_assert(sizeof(TComponentDebuggingTypedPtr) == sizeof(IComponentDebuggingTypedPtr), "Size must match");
+	}
+};
 
-/** Defines a static type identifier for the natvis visualizer */
-enum class EComponentDebugType
+template<typename T>
+struct TComponentHeader : FComponentHeader
 {
-	Unknown,
-	Bool,
-	Uint8,
-	int16,
-	Uint16,
-	Int32,
-	Float,
-	Double,
-	Vector2,
-	Vector3,
-	Vector4,
-	FrameTime,
-	ObjectKey,
-	Object,
-	Property,
-	InstanceHandle,
-	SequenceID,
-	EntityID,
-	Name,
-	GroupID,
-	BlendChannelID,
+	TComponentHeader()
+	{
+		static_assert(sizeof(TComponentHeader<T>) == sizeof(FComponentHeader), "Size must match!");
+	}
 };
 
 /**
@@ -49,28 +49,39 @@ struct FComponentTypeDebugInfo
 {
 	FString DebugName;
 	const TCHAR* DebugTypeName = nullptr;
-	EComponentDebugType Type = EComponentDebugType::Unknown;
+
+	virtual ~FComponentTypeDebugInfo() {}
+	virtual void InitializeComponentHeader(void* Ptr) const
+	{
+		new (Ptr) FComponentHeader();
+	}
+	virtual void InitializeDebugComponentData(FComponentHeader& Header, uint8 Capacity) const
+	{
+	}
+};
+
+template<typename T>
+struct TComponentTypeDebugInfo : FComponentTypeDebugInfo
+{
+	void InitializeComponentHeader(void* Ptr) const override
+	{
+		new (Ptr) TComponentHeader<T>();
+	}
+
+	virtual void InitializeDebugComponentData(FComponentHeader& Header, uint8 Capacity) const
+	{
+		TComponentDebuggingTypedPtr<T>* TypedComponents = new TComponentDebuggingTypedPtr<T>[Capacity];
+		
+		for (int32 Index = 0; Index < Capacity; ++Index)
+		{
+			TypedComponents[Index].Ptr = static_cast<T*>(Header.GetValuePtr(Index));
+		}
+
+		Header.DebugComponents = TypedComponents;
+	}
 };
 
-template<typename T> struct TComponentDebugType                            { static const EComponentDebugType Type = EComponentDebugType::Unknown;           };
-template<>           struct TComponentDebugType<bool>                      { static const EComponentDebugType Type = EComponentDebugType::Bool;              };
-template<>           struct TComponentDebugType<uint8>                     { static const EComponentDebugType Type = EComponentDebugType::Uint8;             };
-template<>           struct TComponentDebugType<int16>                     { static const EComponentDebugType Type = EComponentDebugType::int16;             };
-template<>           struct TComponentDebugType<uint16>                    { static const EComponentDebugType Type = EComponentDebugType::Uint16;            };
-template<>           struct TComponentDebugType<int32>                     { static const EComponentDebugType Type = EComponentDebugType::Int32;             };
-template<>           struct TComponentDebugType<float>                     { static const EComponentDebugType Type = EComponentDebugType::Float;             };
-template<>           struct TComponentDebugType<double>                    { static const EComponentDebugType Type = EComponentDebugType::Double;            };
-template<>           struct TComponentDebugType<FVector2D>                 { static const EComponentDebugType Type = EComponentDebugType::Vector2;           };
-template<>           struct TComponentDebugType<FVector>                   { static const EComponentDebugType Type = EComponentDebugType::Vector3;           };
-template<>           struct TComponentDebugType<FVector4>                  { static const EComponentDebugType Type = EComponentDebugType::Vector4;           };
-template<>           struct TComponentDebugType<FFrameTime>                { static const EComponentDebugType Type = EComponentDebugType::FrameTime;         };
-template<>           struct TComponentDebugType<FObjectKey>                { static const EComponentDebugType Type = EComponentDebugType::ObjectKey;         };
-template<>           struct TComponentDebugType<UObject*>                  { static const EComponentDebugType Type = EComponentDebugType::Object;            };
-template<>           struct TComponentDebugType<FMovieSceneEntityID>       { static const EComponentDebugType Type = EComponentDebugType::EntityID;          };
-template<>           struct TComponentDebugType<FName>                     { static const EComponentDebugType Type = EComponentDebugType::Name;              };
-
-} // namespace MovieScene
-} // namespace UE
+} // namespace UE::MovieScene
 
 
 #endif // UE_MOVIESCENE_ENTITY_DEBUG
diff --git a/Engine/Source/Runtime/MovieScene/Public/EntitySystem/MovieSceneEntityFactoryTemplates.h b/Engine/Source/Runtime/MovieScene/Public/EntitySystem/MovieSceneEntityFactoryTemplates.h
--- a/Engine/Source/Runtime/MovieScene/Public/EntitySystem/MovieSceneEntityFactoryTemplates.h
+++ b/Engine/Source/Runtime/MovieScene/Public/EntitySystem/MovieSceneEntityFactoryTemplates.h
@@ -230,10 +230,9 @@ FComponentTypeInfo FComponentRegistry::MakeComponentTypeInfoWithoutComponentOps(
 	NewTypeInfo.bHasReferencedObjects = false;
 
 #if UE_MOVIESCENE_ENTITY_DEBUG
-	NewTypeInfo.DebugInfo = MakeUnique<FComponentTypeDebugInfo>();
+	NewTypeInfo.DebugInfo = MakeUnique<TComponentTypeDebugInfo<T>>();
 	NewTypeInfo.DebugInfo->DebugName = DebugName;
 	NewTypeInfo.DebugInfo->DebugTypeName = GetGeneratedTypeName<T>();
-	NewTypeInfo.DebugInfo->Type = TComponentDebugType<T>::Type;
 #endif
 
 	return NewTypeInfo;
diff --git a/Engine/Source/Runtime/MovieScene/Public/EntitySystem/MovieSceneEntitySystemTypes.h b/Engine/Source/Runtime/MovieScene/Public/EntitySystem/MovieSceneEntitySystemTypes.h
--- a/Engine/Source/Runtime/MovieScene/Public/EntitySystem/MovieSceneEntitySystemTypes.h
+++ b/Engine/Source/Runtime/MovieScene/Public/EntitySystem/MovieSceneEntitySystemTypes.h
@@ -6,6 +6,7 @@
 #include "Containers/Array.h"
 #include "Containers/ArrayView.h"
 #include "Containers/BitArray.h"
+#include "MovieSceneFwd.h"
 #include "EntitySystem/MovieSceneEntityIDs.h"
 #include "HAL/CriticalSection.h"
 #include "Math/NumericLimits.h"
@@ -22,10 +23,6 @@
 class UMovieSceneSequence;
 namespace UE { namespace MovieScene { class FEntityManager; } }
 
-#ifndef UE_MOVIESCENE_ENTITY_DEBUG
-	#define UE_MOVIESCENE_ENTITY_DEBUG !UE_BUILD_SHIPPING
-#endif
-
 
 DECLARE_STATS_GROUP(TEXT("Movie Scene Evaluation Systems"), STATGROUP_MovieSceneECS, STATCAT_Advanced)
 
@@ -459,23 +456,52 @@ private:
 	uint64 SystemSerial;
 };
 
-
-
 struct FComponentHeader
 {
+	FComponentHeader() = default;
+
+	FComponentHeader(FComponentHeader&&) = delete;
+	FComponentHeader& operator=(FComponentHeader&&) = delete;
+
+	FComponentHeader(const FComponentHeader&) = delete;
+	FComponentHeader& operator=(const FComponentHeader&) = delete;
+
+#if UE_MOVIESCENE_ENTITY_DEBUG
+
+	// This virtual destructor exists in order to create a vftable for this type
+	//    that is used by the debugger to locate and display the derived TComponentHeader<T> type
+	//    enabling visualization of the concrete type in the debugger. It is completely compiled
+	//    out when UE_MOVIESCENE_ENTITY_DEBUG is disabled ensuring this type is not polymorphic.
+	MOVIESCENE_API virtual ~FComponentHeader();
+
+	/** Pointer to a polymorphic view of the component data used for debugging purposes - only populated when GRichComponentDebugging is enabled */
+	struct IComponentDebuggingTypedPtr* DebugComponents = nullptr;
+
+	/** Pointer to this header's owning FEntityAllocation::Size member. Only required and used for debugging. */
+	const uint16* Size = nullptr;
+
+#endif // UE_MOVIESCENE_ENTITY_DEBUG
+
+	/** Pointer to the first component of this header's type, with array size == FEntityAllocation::Capacity, and stride == FComponentHeader::Sizeof */
 	mutable uint8* Components;
 
+	/** Legacy mutex used for locked threaded evaluation */
 	mutable FRWLock ReadWriteLock;
 
 private:
 
+	/** Serial number that is increased after every time this header is opened for write access. Not currently meaningfully used but may be used for caching in future. */
 	mutable uint64 SerialNumber;
 
 public:
 
+	/** Atomic access count used for verifying write exclusivity at runtime. */
 	mutable std::atomic<int32> ScheduledAccessCount;
 
+	/** sizeof(T) for the component type this header represents (ie: the stride of our component data). 0 if this header represents a tag. */
 	uint8 Sizeof;
+
+	/** Component type identifier */
 	FComponentTypeID ComponentType;
 
 	/**
diff --git a/Engine/Source/Runtime/MovieScene/Public/EntitySystem/MovieSceneInstanceRegistry.h b/Engine/Source/Runtime/MovieScene/Public/EntitySystem/MovieSceneInstanceRegistry.h
--- a/Engine/Source/Runtime/MovieScene/Public/EntitySystem/MovieSceneInstanceRegistry.h
+++ b/Engine/Source/Runtime/MovieScene/Public/EntitySystem/MovieSceneInstanceRegistry.h
@@ -8,6 +8,7 @@
 #include "Containers/SparseArray.h"
 #include "Evaluation/MovieScenePlayback.h"
 #include "EntitySystem/MovieSceneSequenceInstance.h"
+#include "EntitySystem/MovieSceneSequenceInstanceHandle.h"
 #include "EntitySystem/MovieSceneComponentDebug.h"
 #include "EntitySystem/MovieSceneSharedPlaybackState.h"
 #include "MovieSceneSequenceID.h"
diff --git a/Engine/Source/Runtime/MovieScene/Public/EntitySystem/MovieScenePropertyBinding.h b/Engine/Source/Runtime/MovieScene/Public/EntitySystem/MovieScenePropertyBinding.h
--- a/Engine/Source/Runtime/MovieScene/Public/EntitySystem/MovieScenePropertyBinding.h
+++ b/Engine/Source/Runtime/MovieScene/Public/EntitySystem/MovieScenePropertyBinding.h
@@ -60,20 +60,3 @@ struct TStructOpsTypeTraits<FMovieScenePropertyBinding> : public TStructOpsTypeT
 };
 #endif // WITH_EDITORONLY_DATA
 
-
-#if UE_MOVIESCENE_ENTITY_DEBUG
-
-namespace UE
-{
-namespace MovieScene
-{
-
-	template<> struct TComponentDebugType<FMovieScenePropertyBinding>
-	{
-		static const EComponentDebugType Type = EComponentDebugType::Property;
-	};
-
-} // namespace MovieScene
-} // namespace UE
-
-#endif // UE_MOVIESCENE_ENTITY_DEBUG
\ No newline at end of file
diff --git a/Engine/Source/Runtime/MovieScene/Public/EntitySystem/MovieSceneSequenceInstanceHandle.h b/Engine/Source/Runtime/MovieScene/Public/EntitySystem/MovieSceneSequenceInstanceHandle.h
--- a/Engine/Source/Runtime/MovieScene/Public/EntitySystem/MovieSceneSequenceInstanceHandle.h
+++ b/Engine/Source/Runtime/MovieScene/Public/EntitySystem/MovieSceneSequenceInstanceHandle.h
@@ -64,11 +64,6 @@ struct FRootInstanceHandle : FInstanceHandle
 	{}
 };
 
-#if UE_MOVIESCENE_ENTITY_DEBUG
-template<> struct TComponentDebugType<FInstanceHandle>     { static const EComponentDebugType Type = EComponentDebugType::InstanceHandle; };
-template<> struct TComponentDebugType<FRootInstanceHandle> { static const EComponentDebugType Type = EComponentDebugType::InstanceHandle; };
-#endif
-
 
 } // namespace MovieScene
 } // namespace UE
\ No newline at end of file
diff --git a/Engine/Source/Runtime/MovieScene/Public/EntitySystem/MovieSceneSharedPlaybackState.h b/Engine/Source/Runtime/MovieScene/Public/EntitySystem/MovieSceneSharedPlaybackState.h
--- a/Engine/Source/Runtime/MovieScene/Public/EntitySystem/MovieSceneSharedPlaybackState.h
+++ b/Engine/Source/Runtime/MovieScene/Public/EntitySystem/MovieSceneSharedPlaybackState.h
@@ -10,6 +10,7 @@
 #include "Evaluation/MovieScenePreAnimatedState.h"
 #include "MovieSceneSequenceID.h"
 
+class UMovieSceneSequence;
 class FMovieSceneEntitySystemRunner;
 class UMovieSceneCompiledDataManager;
 class UMovieSceneEntitySystemLinker;
diff --git a/Engine/Source/Runtime/MovieScene/Public/MovieSceneFwd.h b/Engine/Source/Runtime/MovieScene/Public/MovieSceneFwd.h
--- a/Engine/Source/Runtime/MovieScene/Public/MovieSceneFwd.h
+++ b/Engine/Source/Runtime/MovieScene/Public/MovieSceneFwd.h
@@ -4,9 +4,7 @@
 
 #include "CoreTypes.h"
 #include "Stats/Stats2.h"
-#include "Misc/FrameRate.h"
 #include "Misc/EnumClassFlags.h"
-#include "MovieSceneSequenceID.h"
 #include "HAL/PreprocessorHelpers.h"
 
 #include "MovieSceneFwd.generated.h"
@@ -19,15 +17,28 @@
 	#define UE_MOVIESCENE_EXPENSIVE_CONSISTENCY_CHECKS 0
 #endif
 
-namespace UE
-{
-namespace MovieScene
+#ifndef UE_MOVIESCENE_TIME_WARP_DEPRECATED
+	#define UE_MOVIESCENE_TIME_WARP_DEPRECATED 0
+#endif
+
+#ifndef UE_MOVIESCENE_ENTITY_DEBUG
+	#define UE_MOVIESCENE_ENTITY_DEBUG !UE_BUILD_SHIPPING
+#endif
+
+#ifndef MOVIESCENE_DETAILED_STATS
+	#define MOVIESCENE_DETAILED_STATS 0
+#endif
+
+struct FFrameRate;
+struct FFrameNumber;
+class UObject;
+
+namespace UE::MovieScene
 {
 
-class ISequenceDataEventHandler;
+	class ISequenceDataEventHandler;
 
-} // namespace MovieScene
-} // namespace UE
+} // namespace UE::MovieScene
 
 UENUM()
 namespace EMovieScenePlayerStatus
@@ -138,9 +149,6 @@ MOVIESCENE_API FFrameRate GetLegacyConversionFrameRate();
 MOVIESCENE_API void EmitLegacyOutOfBoundsError(UObject* ErrorContext, FFrameRate InFrameRate, double InTime);
 MOVIESCENE_API FFrameNumber UpgradeLegacyMovieSceneTime(UObject* ErrorContext, FFrameRate InFrameRate, double InTime);
 
-#ifndef MOVIESCENE_DETAILED_STATS
-	#define MOVIESCENE_DETAILED_STATS 0
-#endif
 
 #if MOVIESCENE_DETAILED_STATS
 	#define MOVIESCENE_DETAILED_SCOPE_CYCLE_COUNTER SCOPE_CYCLE_COUNTER
diff --git a/Engine/Source/Runtime/MovieScene/Public/MovieSceneSequenceID.h b/Engine/Source/Runtime/MovieScene/Public/MovieSceneSequenceID.h
--- a/Engine/Source/Runtime/MovieScene/Public/MovieSceneSequenceID.h
+++ b/Engine/Source/Runtime/MovieScene/Public/MovieSceneSequenceID.h
@@ -90,13 +90,6 @@ struct TStructOpsTypeTraits<FMovieSceneSequenceID> : public TStructOpsTypeTraits
 	static constexpr EPropertyObjectReferenceType WithSerializerObjectReferences = EPropertyObjectReferenceType::None;
 };
 
-#if UE_MOVIESCENE_ENTITY_DEBUG
-namespace UE::MovieScene
-{
-	template<> struct TComponentDebugType<FMovieSceneSequenceID> { static const EComponentDebugType Type = EComponentDebugType::SequenceID; };
-}
-#endif
-
 typedef TCallTraits<FMovieSceneSequenceID>::ParamType FMovieSceneSequenceIDRef;
 
 namespace MovieSceneSequenceID
diff --git a/Engine/Source/Runtime/MovieSceneTracks/Public/MovieSceneTracksPropertyTypes.h b/Engine/Source/Runtime/MovieSceneTracks/Public/MovieSceneTracksPropertyTypes.h
--- a/Engine/Source/Runtime/MovieSceneTracks/Public/MovieSceneTracksPropertyTypes.h
+++ b/Engine/Source/Runtime/MovieSceneTracks/Public/MovieSceneTracksPropertyTypes.h
@@ -81,12 +81,6 @@ enum class EColorPropertyType : uint8
 	Color,
 };
 
-#if UE_MOVIESCENE_ENTITY_DEBUG
-
-template<> struct TComponentDebugType<EColorPropertyType> { static const EComponentDebugType Type = EComponentDebugType::Uint8;   };
-
-#endif
-
 /** Intermediate type for the color property system that lets us store what kind of color type we should use */
 struct FIntermediateColor
 {
