CIS Fix

diff --git a/Engine/Source/Runtime/MovieScene/Public/EntitySystem/MovieSceneComponentTypeInfo.h b/Engine/Source/Runtime/MovieScene/Public/EntitySystem/MovieSceneComponentTypeInfo.h
--- a/Engine/Source/Runtime/MovieScene/Public/EntitySystem/MovieSceneComponentTypeInfo.h
+++ b/Engine/Source/Runtime/MovieScene/Public/EntitySystem/MovieSceneComponentTypeInfo.h
@@ -10,6 +10,10 @@
 #include "EntitySystem/MovieSceneComponentDebug.h"
 #include "EntitySystem/MovieSceneEntityIDs.h"
 
+struct FObjectKey;
+template <typename> class TObjectKey;
+template <typename> struct TObjectPtr;
+
 namespace UE
 {
 namespace MovieScene
