Delete this Component if we abort as it will create some new locks

diff --git a/Engine/Source/Runtime/MovieScene/Private/EntitySystem/MovieSceneEntityManager.cpp b/Engine/Source/Runtime/MovieScene/Private/EntitySystem/MovieSceneEntityManager.cpp
--- a/Engine/Source/Runtime/MovieScene/Private/EntitySystem/MovieSceneEntityManager.cpp
+++ b/Engine/Source/Runtime/MovieScene/Private/EntitySystem/MovieSceneEntityManager.cpp
@@ -15,6 +15,8 @@
 
 #include "EntitySystem/EntityAllocationIterator.h"
 
+#include "AutoRTFM/AutoRTFM.h"
+
 UE_SELECT_ANY UE::MovieScene::FEntityManager*& GEntityManagerForDebugging = UE::MovieScene::GEntityManagerForDebuggingVisualizers;
 
 #if DO_GUARD_SLOW
@@ -232,10 +234,26 @@ struct FEntityInitializer
 				}
 				else
 				{
-					new (Header) FComponentHeader();
+					UE_AUTORTFM_OPEN(
+					{
+						new (Header) FComponentHeader();
+					});
+						
+					UE_AUTORTFM_ONABORT(
+					{
+						Header->~FComponentHeader();
+					});
 				}
 #else
-				new (Header) FComponentHeader();
+				UE_AUTORTFM_OPEN(
+				{
+					new (Header) FComponentHeader();
+				});
+				
+				UE_AUTORTFM_ONABORT(
+				{
+					Header->~FComponentHeader();
+				});
 #endif
 
 				Header->ComponentType = ComponentTypeID;
