[Backout] - CL34065553

diff --git a/Engine/Source/Runtime/MovieScene/Private/Evaluation/MovieSceneEvaluationState.cpp b/Engine/Source/Runtime/MovieScene/Private/Evaluation/MovieSceneEvaluationState.cpp
--- a/Engine/Source/Runtime/MovieScene/Private/Evaluation/MovieSceneEvaluationState.cpp
+++ b/Engine/Source/Runtime/MovieScene/Private/Evaluation/MovieSceneEvaluationState.cpp
@@ -11,6 +11,7 @@
 #include "IMovieScenePlaybackClient.h"
 #include "IMovieScenePlayer.h"
 #include "MovieScene.h"
+#include "MovieSceneDynamicBindingInvoker.h"
 #include "MovieSceneObjectBindingID.h"
 #include "MovieSceneSequence.h"
 #include "UniversalObjectLocatorResolveParams.h"
@@ -433,14 +434,28 @@ void FMovieSceneObjectCache::UpdateBindings(const FGuid& InGuid, TSharedRef<cons
 					}
 					else
 					{
-						UE::UniversalObjectLocator::FResolveParams ResolveParams(ResolutionContext);
-						if (Player)
+						FMovieSceneDynamicBindingResolveResult ResolveResult = FMovieSceneDynamicBindingInvoker::ResolveDynamicBinding(SharedPlaybackState, Sequence, SequenceID, *Possessable);
+						if (ResolveResult.Object)
 						{
-							Player->ResolveBoundObjects(ResolveParams, InGuid, SequenceID, *Sequence, FoundObjects);
+							if (!ResolveResult.bIsPossessedObject)
+							{
+								UE_LOG(LogMovieScene, Error, 
+									TEXT("Possessable '%s' (dynamically resolved to '%s') can't have spawnable-type ownership. The user-defined director blueprint endpoint should set bIsPossessedObject to true."),
+									*LexToString(Possessable->GetName()), *ResolveResult.Object->GetName());
+							}
+							FoundObjects.Add(ResolveResult.Object);
 						}
 						else
 						{
-							Sequence->LocateBoundObjects(InGuid, ResolveParams, SharedPlaybackState, FoundObjects);
+							UE::UniversalObjectLocator::FResolveParams ResolveParams(ResolutionContext);
+							if (Player)
+							{
+								Player->ResolveBoundObjects(ResolveParams, InGuid, SequenceID, *Sequence, FoundObjects);
+							}
+							else
+							{
+								Sequence->LocateBoundObjects(InGuid, ResolveParams, FoundObjects);
+							}
 						}
 					}
 					
@@ -471,14 +486,28 @@ void FMovieSceneObjectCache::UpdateBindings(const FGuid& InGuid, TSharedRef<cons
 				}
 				else
 				{
-					UE::UniversalObjectLocator::FResolveParams ResolveParams(ResolutionContext);
-					if (Player)
+					FMovieSceneDynamicBindingResolveResult ResolveResult = FMovieSceneDynamicBindingInvoker::ResolveDynamicBinding(SharedPlaybackState, Sequence, SequenceID, *Possessable);
+					if (ResolveResult.Object)
 					{
-						Player->ResolveBoundObjects(ResolveParams, InGuid, SequenceID, *Sequence, FoundObjects);
+						if (!ResolveResult.bIsPossessedObject)
+						{
+							UE_LOG(LogMovieScene, Error,
+								TEXT("Possessable '%s' (dynamically resolved to '%s') can't have spawnable-type ownership. The user-defined director blueprint endpoint should set bIsPossessedObject to true."),
+								*LexToString(Possessable->GetName()), *ResolveResult.Object->GetName());
+						}
+						FoundObjects.Add(ResolveResult.Object);
 					}
 					else
 					{
-						Sequence->LocateBoundObjects(InGuid, ResolveParams, SharedPlaybackState, FoundObjects);
+						UE::UniversalObjectLocator::FResolveParams ResolveParams(ResolutionContext);
+						if (Player)
+						{
+							Player->ResolveBoundObjects(ResolveParams, InGuid, SequenceID, *Sequence, FoundObjects);
+						}
+						else
+						{
+							Sequence->LocateBoundObjects(InGuid, ResolveParams, FoundObjects);
+						}
 					}
 				}
 				
diff --git a/Engine/Source/Runtime/MovieScene/Private/MovieScene.cpp b/Engine/Source/Runtime/MovieScene/Private/MovieScene.cpp
--- a/Engine/Source/Runtime/MovieScene/Private/MovieScene.cpp
+++ b/Engine/Source/Runtime/MovieScene/Private/MovieScene.cpp
@@ -283,7 +283,7 @@ void UMovieScene::AddSpawnable(const FMovieSceneSpawnable& InNewSpawnable, const
 	FMovieSceneBinding NewBinding = InNewBinding;
 	for (auto Track : NewBinding.GetTracks())
 	{
-		Track->Rename(nullptr, this, REN_DontCreateRedirectors | REN_ForceNoResetLoaders);
+		Track->Rename(nullptr, this);
 	}
 	int32 NewBindingIndex = InsertSorted(ObjectBindings, MoveTemp(NewBinding));
 
@@ -363,7 +363,7 @@ void UMovieScene::AddPossessable(const FMovieScenePossessable& InNewPossessable,
 	FMovieSceneBinding NewBinding = InNewBinding;
 	for (auto Track : NewBinding.GetTracks())
 	{
-		Track->Rename(nullptr, this, REN_DontCreateRedirectors | REN_ForceNoResetLoaders);
+		Track->Rename(nullptr, this);
 	}
 	int32 NewBindingIndex = InsertSorted(ObjectBindings, MoveTemp(NewBinding));
 	check(ObjectBindings.IsValidIndex(NewBindingIndex));
@@ -1108,7 +1108,7 @@ bool UMovieScene::AddGivenTrack(UMovieSceneTrack* InTrack, const FGuid& ObjectGu
 	Modify();
 	if (FMovieSceneBinding* Binding = FindBinding(ObjectGuid))
 	{
-		InTrack->Rename(nullptr, this, REN_DontCreateRedirectors | REN_ForceNoResetLoaders);
+		InTrack->Rename(nullptr, this);
 		Binding->AddTrack(*InTrack, this);
 		return true;
 	}
@@ -1219,7 +1219,7 @@ bool UMovieScene::AddGivenTrack(UMovieSceneTrack* InTrack)
 	{
 		Modify();
 		Tracks.Add(InTrack);
-		InTrack->Rename(nullptr, this, REN_DontCreateRedirectors | REN_ForceNoResetLoaders);
+		InTrack->Rename(nullptr, this);
 
 		EventHandlers.Trigger(&UE::MovieScene::ISequenceDataEventHandler::OnTrackAdded, InTrack);
 
@@ -1290,7 +1290,7 @@ void UMovieScene::SetCameraCutTrack(UMovieSceneTrack* InTrack)
 	}
 
 	Modify();
-	InTrack->Rename(nullptr, this, REN_DontCreateRedirectors | REN_ForceNoResetLoaders);
+	InTrack->Rename(nullptr, this);
 	UMovieSceneTrack* OldCameraCutTrack = CameraCutTrack;
 	CameraCutTrack = InTrack;
 
@@ -1560,7 +1560,7 @@ void UMovieScene::ReplaceBinding(const FGuid& BindingToReplaceGuid, const FMovie
 		// We also need to change the track's owners to be the MovieScene.
 		for (auto Track : Binding->GetTracks())
 		{
-			Track->Rename(nullptr, this, REN_DontCreateRedirectors | REN_ForceNoResetLoaders);
+			Track->Rename(nullptr, this);
 		}
 
 		EventHandlers.Trigger(&UE::MovieScene::ISequenceDataEventHandler::OnBindingRemoved, BindingToReplaceGuid);
diff --git a/Engine/Source/Runtime/MovieScene/Private/MovieSceneCommonHelpers.cpp b/Engine/Source/Runtime/MovieScene/Private/MovieSceneCommonHelpers.cpp
--- a/Engine/Source/Runtime/MovieScene/Private/MovieSceneCommonHelpers.cpp
+++ b/Engine/Source/Runtime/MovieScene/Private/MovieSceneCommonHelpers.cpp
@@ -19,7 +19,6 @@
 #include "MovieSceneBindingReferences.h"
 #include "Bindings/MovieSceneSpawnableBinding.h"
 #include "Evaluation/MovieSceneEvaluationState.h"
-#include "UObject/UObjectIterator.h"
 
 bool MovieSceneHelpers::IsSectionKeyable(const UMovieSceneSection* Section)
 {
@@ -713,64 +712,6 @@ bool MovieSceneHelpers::IsBoundToSpawnable(UMovieSceneSequence* Sequence, const
 }
 
 
-FGuid MovieSceneHelpers::TryCreateCustomSpawnableBinding(UMovieSceneSequence* Sequence, UObject* CustomBindingObject)
-{
-	FGuid NewID;
-	if (!Sequence)
-	{
-		return NewID;
-	}
-	FMovieSceneBindingReferences* BindingReferences = Sequence->GetBindingReferences();
-	if (!BindingReferences)
-	{
-		return NewID;
-	}
-	static TArray<const TSubclassOf<UMovieSceneCustomBinding>> CachedCustomBindingTypes;
-	static bool CustomBindingTypesCached = false;
-	if (!CustomBindingTypesCached)
-	{
-		CustomBindingTypesCached = true;
-		MovieSceneHelpers::GetPrioritySortedCustomBindingTypes(CachedCustomBindingTypes);
-	}
-
-	UMovieSceneCustomBinding* NewCustomBinding = nullptr;
-
-	for (const TSubclassOf<UMovieSceneCustomBinding>& CustomBindingType : CachedCustomBindingTypes)
-	{
-		// We only want to use children of UMovieSceneSpawnableBindingBase
-		if ((CustomBindingType->IsChildOf<UMovieSceneSpawnableBindingBase>()))
-		{
-			if (UMovieSceneCustomBinding* CustomBindingCDO = CustomBindingType ? CustomBindingType->GetDefaultObject<UMovieSceneCustomBinding>() : nullptr)
-			{
-				if (CustomBindingObject && CustomBindingCDO->SupportsBindingCreationFromObject(CustomBindingObject))
-				{
-					// Create a custom binding from this Object
-					NewCustomBinding = CustomBindingCDO->CreateNewCustomBinding(CustomBindingObject, *Sequence->GetMovieScene());
-					if (NewCustomBinding)
-					{
-						break;
-					}
-				}
-			}
-		}
-	}
-
-	if (NewCustomBinding)
-	{
-		FString DesiredBindingName = NewCustomBinding->GetDesiredBindingName();
-		FString CurrentName = DesiredBindingName.IsEmpty() ? FName::NameToDisplayString(CustomBindingObject->GetName(), false) : DesiredBindingName;
-		CurrentName = MovieSceneHelpers::MakeUniqueBindingName(Sequence->GetMovieScene(), CurrentName);
-
-		NewID = Sequence->GetMovieScene()->AddPossessable(CurrentName, NewCustomBinding->GetBoundObjectClass());
-
-		// Add the custom binding
-		Sequence->GetBindingReferences()->AddOrReplaceBinding(NewID, NewCustomBinding, 0);
-	}
-
-	return NewID;
-}
-
-
 UObject* MovieSceneHelpers::GetSingleBoundObject(UMovieSceneSequence* Sequence, const FGuid& ObjectId, TSharedRef<const UE::MovieScene::FSharedPlaybackState> SharedPlaybackState, int32 BindingIndex)
 {
 	if (Sequence)
@@ -789,30 +730,13 @@ UObject* MovieSceneHelpers::GetSingleBoundObject(UMovieSceneSequence* Sequence,
 						return BoundObjects[0].Get();
 					}
 				}
-				else if (const FMovieScenePossessable* Possessable = MovieScene->FindPossessable(ObjectId))
+
+				const FMovieSceneBindingReferences* Refs = Sequence->GetBindingReferences();
+				if (Refs)
 				{
-					const FMovieSceneBindingReferences* Refs = Sequence->GetBindingReferences();
-					if (Refs)
-					{
-						UObject* ResolutionContext = MovieSceneHelpers::GetResolutionContext(Sequence, ObjectId, SequenceID, SharedPlaybackState);
-
-						if (Possessable->GetParent().IsValid() && Sequence->AreParentContextsSignificant())
-						{
-							TArrayView<TWeakObjectPtr<>> ParentBoundObjects = EvaluationState->FindBoundObjects(FMovieSceneEvaluationOperand(SequenceID, Possessable->GetParent()), SharedPlaybackState);
-							for (TWeakObjectPtr<> Parent : ParentBoundObjects)
-							{
-								ResolutionContext = Parent.Get();
-								if (!ResolutionContext)
-								{
-									continue;
-								}
-							}
-						}
-
-						UE::UniversalObjectLocator::FResolveParams LocatorResolveParams(ResolutionContext);
-						FMovieSceneBindingResolveParams BindingResolveParams{ Sequence, ObjectId, SequenceID, ResolutionContext };
-						return Refs->ResolveSingleBinding(BindingResolveParams, BindingIndex, LocatorResolveParams, SharedPlaybackState);
-					}
+					UE::UniversalObjectLocator::FResolveParams LocatorResolveParams(SharedPlaybackState->GetPlaybackContext());
+					FMovieSceneBindingResolveParams BindingResolveParams{ Sequence, ObjectId, SequenceID };
+					return Refs->ResolveSingleBinding(BindingResolveParams, BindingIndex, LocatorResolveParams, SharedPlaybackState);
 				}
 			}
 		}
@@ -972,64 +896,6 @@ const UClass* MovieSceneHelpers::GetBoundObjectClass(UMovieSceneSequence* Sequen
 #endif
 
 
-void MovieSceneHelpers::GetPrioritySortedCustomBindingTypes(TArray<const TSubclassOf<UMovieSceneCustomBinding>>& OutCustomBindingTypes)
-{
-	OutCustomBindingTypes.Empty();
-	for (TObjectIterator<UClass> It; It; ++It)
-	{
-		if (It->IsChildOf(UMovieSceneCustomBinding::StaticClass()) && !It->HasAnyClassFlags(CLASS_Abstract))
-		{
-			OutCustomBindingTypes.Add(*It);
-		}
-	}// Sort by spawner priority to allow disambiguation for similar object types
-	OutCustomBindingTypes.Sort([](const TSubclassOf<UMovieSceneCustomBinding>& A, const TSubclassOf<UMovieSceneCustomBinding>& B) {
-		return A && B && A->GetDefaultObject<UMovieSceneCustomBinding>()->GetCustomBindingPriority() > B->GetDefaultObject<UMovieSceneCustomBinding>()->GetCustomBindingPriority(); });
-}
-
-
-TSharedRef<UE::MovieScene::FSharedPlaybackState> MovieSceneHelpers::CreateTransientSharedPlaybackState(UObject* WorldContext, UMovieSceneSequence* Sequence)
-{
-	verify(WorldContext && Sequence);
-
-	using namespace UE::MovieScene;
-	FSharedPlaybackStateCreateParams CreateParams;
-	CreateParams.PlaybackContext = WorldContext;
-	TSharedRef<FSharedPlaybackState> TransientPlaybackState = MakeShared<FSharedPlaybackState>(*Sequence, CreateParams);
-
-	FMovieSceneEvaluationState State;
-	TransientPlaybackState->AddCapabilityRaw(&State);
-	State.AssignSequence(MovieSceneSequenceID::Root, *Sequence, TransientPlaybackState);
-
-	return TransientPlaybackState;
-}
-
-UObject* MovieSceneHelpers::GetResolutionContext(UMovieSceneSequence* Sequence, const FGuid& ObjectId, const FMovieSceneSequenceID& SequenceID, TSharedRef<const UE::MovieScene::FSharedPlaybackState> SharedPlaybackState)
-{
-	if (!Sequence)
-	{
-		return nullptr;
-	}
-	UObject* ResolutionContext = SharedPlaybackState->GetPlaybackContext();
-	if (FMovieScenePossessable* Possessable = Sequence->GetMovieScene()->FindPossessable(ObjectId))
-	{
-		if (Possessable->GetParent().IsValid() && Sequence->AreParentContextsSignificant())
-		{
-			if (FMovieSceneEvaluationState* EvaluationState = SharedPlaybackState->FindCapability<FMovieSceneEvaluationState>())
-			{
-				TArrayView<TWeakObjectPtr<>> ParentBoundObjects = EvaluationState->FindBoundObjects(Possessable->GetParent(), SequenceID, SharedPlaybackState);
-				for (TWeakObjectPtr<> Parent : ParentBoundObjects)
-				{
-					ResolutionContext = Parent.Get();
-					if (!ResolutionContext)
-					{
-						continue;
-					}
-				}
-			}
-		}
-	}
-	return ResolutionContext;
-}
 
 MovieSceneHelpers::FMovieSceneScopedPackageDirtyGuard::FMovieSceneScopedPackageDirtyGuard(USceneComponent* InComponent)
 {
diff --git a/Engine/Source/Runtime/MovieScene/Private/MovieScenePlayer.cpp b/Engine/Source/Runtime/MovieScene/Private/MovieScenePlayer.cpp
--- a/Engine/Source/Runtime/MovieScene/Private/MovieScenePlayer.cpp
+++ b/Engine/Source/Runtime/MovieScene/Private/MovieScenePlayer.cpp
@@ -138,12 +138,12 @@ void IMovieScenePlayer::ResolveBoundObjects(UE::UniversalObjectLocator::FResolve
 	{
 		if (const FMovieSceneBindingReferences* BindingReferences = InSequence.GetBindingReferences())
 		{
-			FMovieSceneBindingResolveParams BindingResolveParams{ &InSequence, InBindingId, SequenceID, LocatorResolveParams.Context };
-			BindingReferences->ResolveBinding(BindingResolveParams, LocatorResolveParams, FindSharedPlaybackState(), OutObjects);
+			FMovieSceneBindingResolveParams BindingResolveParams{ &InSequence, InBindingId, SequenceID };
+			BindingReferences->ResolveBinding(BindingResolveParams, LocatorResolveParams, ConstCastSharedPtr<const UE::MovieScene::FSharedPlaybackState>(const_cast<IMovieScenePlayer*>(this)->FindSharedPlaybackState()), OutObjects);
 		}
 		else
 		{
-			InSequence.LocateBoundObjects(InBindingId, LocatorResolveParams, FindSharedPlaybackState(), OutObjects);
+			InSequence.LocateBoundObjects(InBindingId, LocatorResolveParams, OutObjects);
 		}
 	}
 }
@@ -176,23 +176,12 @@ TSharedPtr<UE::MovieScene::FSharedPlaybackState> IMovieScenePlayer::FindSharedPl
 	return GetEvaluationTemplate().GetSharedPlaybackState();
 }
 
-TSharedPtr<const UE::MovieScene::FSharedPlaybackState> IMovieScenePlayer::FindSharedPlaybackState() const
-{
-	return ConstCastSharedPtr<const UE::MovieScene::FSharedPlaybackState>(const_cast<IMovieScenePlayer*>(this)->GetEvaluationTemplate().GetSharedPlaybackState());
-}
-
 TSharedRef<UE::MovieScene::FSharedPlaybackState> IMovieScenePlayer::GetSharedPlaybackState()
 {
 	// ToSharedRef will assert if evaluation template isn't initialized
 	return GetEvaluationTemplate().GetSharedPlaybackState().ToSharedRef();
 }
 
-TSharedRef<const UE::MovieScene::FSharedPlaybackState> IMovieScenePlayer::GetSharedPlaybackState() const
-{
-	// ToSharedRef will assert if evaluation template isn't initialized
-	return ConstCastSharedRef<const UE::MovieScene::FSharedPlaybackState>(const_cast<IMovieScenePlayer*>(this)->GetEvaluationTemplate().GetSharedPlaybackState().ToSharedRef());
-}
-
 void IMovieScenePlayer::ResetDirectorInstances()
 {
 	using namespace UE::MovieScene;
diff --git a/Engine/Source/Runtime/MovieScene/Private/MovieScenePossessable.cpp b/Engine/Source/Runtime/MovieScene/Private/MovieScenePossessable.cpp
--- a/Engine/Source/Runtime/MovieScene/Private/MovieScenePossessable.cpp
+++ b/Engine/Source/Runtime/MovieScene/Private/MovieScenePossessable.cpp
@@ -47,7 +47,7 @@ void FMovieScenePossessable::SetParent(const FGuid& InParentGuid, UMovieScene* O
 void FMovieScenePossessable::FixupPossessedObjectClass(UMovieSceneSequence* InSequence, UObject* Context)
 {
 	TArray<UObject*, TInlineAllocator<1>> BoundObjects;
-	InSequence->LocateBoundObjects(Guid, UE::UniversalObjectLocator::FResolveParams(Context), nullptr, BoundObjects);
+	InSequence->LocateBoundObjects(Guid, UE::UniversalObjectLocator::FResolveParams(Context), BoundObjects);
 
 	TArray<UClass*> Classes;
 	for (UObject* BoundObject : BoundObjects)
diff --git a/Engine/Source/Runtime/MovieScene/Private/MovieSceneSequence.cpp b/Engine/Source/Runtime/MovieScene/Private/MovieSceneSequence.cpp
--- a/Engine/Source/Runtime/MovieScene/Private/MovieSceneSequence.cpp
+++ b/Engine/Source/Runtime/MovieScene/Private/MovieSceneSequence.cpp
@@ -20,7 +20,6 @@
 #include "Compilation/MovieSceneCompiledDataManager.h"
 #include "UniversalObjectLocator.h"
 #include "Bindings/MovieSceneSpawnableBinding.h"
-#include "MovieSceneCommonHelpers.h"
 
 #include UE_INLINE_GENERATED_CPP_BY_NAME(MovieSceneSequence)
 
@@ -68,11 +67,6 @@ FMovieSceneBindingReferences* UMovieSceneSequence::GetBindingReferences()
 }
 
 void UMovieSceneSequence::LocateBoundObjects(const FGuid& ObjectId, const UE::UniversalObjectLocator::FResolveParams& ResolveParams, TArray<UObject*, TInlineAllocator<1>>& OutObjects) const
-{
-	LocateBoundObjects(ObjectId, ResolveParams, nullptr, OutObjects);
-}
-
-void UMovieSceneSequence::LocateBoundObjects(const FGuid& ObjectId, const UE::UniversalObjectLocator::FResolveParams& ResolveParams, TSharedPtr<const FSharedPlaybackState> SharedPlaybackState, TArray<UObject*, TInlineAllocator<1>>& OutObjects) const
 {
 	const FMovieSceneBindingReferences* Refs = GetBindingReferences();
 	if (Refs)
@@ -346,14 +340,18 @@ UMovieSceneCompiledData* UMovieSceneSequence::GetOrCreateCompiledData()
 FGuid UMovieSceneSequence::FindPossessableObjectId(UObject& Object, UObject* Context) const
 {
 	using namespace UE::MovieScene;
+
+	FSharedPlaybackStateCreateParams CreateParams;
+	CreateParams.PlaybackContext = Context;
 	UMovieSceneSequence* ThisSequence = const_cast<UMovieSceneSequence*>(this);
-	TSharedRef<UE::MovieScene::FSharedPlaybackState> TransientPlaybackState = MovieSceneHelpers::CreateTransientSharedPlaybackState(Context, ThisSequence);
-	if (FMovieSceneEvaluationState* EvaluationState = TransientPlaybackState->FindCapability<FMovieSceneEvaluationState>())
-	{
-		FGuid ExistingID = EvaluationState->FindObjectId(Object, MovieSceneSequenceID::Root, TransientPlaybackState);
-		return ExistingID;
-	}
-	return FGuid();
+	TSharedRef<FSharedPlaybackState> TransientPlaybackState = MakeShared<FSharedPlaybackState>(*ThisSequence, CreateParams);
+
+	FMovieSceneEvaluationState State;
+	TransientPlaybackState->AddCapabilityRaw(&State);
+	State.AssignSequence(MovieSceneSequenceID::Root, *ThisSequence, TransientPlaybackState);
+
+	FGuid ExistingID = State.FindObjectId(Object, MovieSceneSequenceID::Root, TransientPlaybackState);
+	return ExistingID;
 }
 
 FMovieSceneObjectBindingID UMovieSceneSequence::FindBindingByTag(FName InBindingName) const
diff --git a/Engine/Source/Runtime/MovieScene/Private/MovieSceneSequencePlayer.cpp b/Engine/Source/Runtime/MovieScene/Private/MovieSceneSequencePlayer.cpp
--- a/Engine/Source/Runtime/MovieScene/Private/MovieSceneSequencePlayer.cpp
+++ b/Engine/Source/Runtime/MovieScene/Private/MovieSceneSequencePlayer.cpp
@@ -249,7 +249,7 @@ void UMovieSceneSequencePlayer::ResolveBoundObjects(UE::UniversalObjectLocator::
 
 	if (bAllowDefault)
 	{
-		InSequence.LocateBoundObjects(InBindingId, ResolveParams, FindSharedPlaybackState(), OutObjects);
+		InSequence.LocateBoundObjects(InBindingId, ResolveParams, OutObjects);
 	}
 }
 
diff --git a/Engine/Source/Runtime/MovieScene/Private/MovieSceneSpawnRegister.cpp b/Engine/Source/Runtime/MovieScene/Private/MovieSceneSpawnRegister.cpp
--- a/Engine/Source/Runtime/MovieScene/Private/MovieSceneSpawnRegister.cpp
+++ b/Engine/Source/Runtime/MovieScene/Private/MovieSceneSpawnRegister.cpp
@@ -4,6 +4,7 @@
 #include "MovieScene.h"
 #include "IMovieScenePlayer.h"
 #include "Evaluation/MovieSceneEvaluationTemplateInstance.h"
+#include "MovieSceneDynamicBindingInvoker.h"
 #include "MovieSceneSpawnableAnnotation.h"
 #include "MovieSceneBindingReferences.h"
 #include "Bindings/MovieSceneSpawnableBinding.h"
@@ -51,8 +52,22 @@ UObject* FMovieSceneSpawnRegister::SpawnObject(const FGuid& BindingId, UMovieSce
 		SpawnOwnership = Spawnable->GetSpawnOwnership();
 		IMovieScenePlayer* Player = UE::MovieScene::FPlayerIndexPlaybackCapability::GetPlayer(SharedPlaybackState);
 
-		// Call through to the list of spawners to see who can spawn something from this FMovieSceneSpawnable
-		SpawnedActor = SpawnObject(*Spawnable, TemplateID, SharedPlaybackState);
+		// See if there is some dynamic binding logic to invoke, otherwise spawn the actor
+			// See if there is some dynamic binding logic to invoke, otherwise spawn the actor
+	    FMovieSceneDynamicBindingResolveResult ResolveResult = FMovieSceneDynamicBindingInvoker::ResolveDynamicBinding(SharedPlaybackState, Sequence, TemplateID, *Spawnable);
+	    if (ResolveResult.Object)
+	    {
+		    SpawnedActor = ResolveResult.Object;
+		    if (ResolveResult.bIsPossessedObject)
+            {
+                SpawnOwnership = ESpawnOwnership::External;
+            }
+        }
+		if (!SpawnedActor)
+		{
+	        // Call through to the list of spawners to see who can spawn something from this FMovieSceneSpawnable
+			SpawnedActor = SpawnObject(*Spawnable, TemplateID, SharedPlaybackState);
+		}
 	}
 	else if (UMovieSceneSequence* MovieSceneSequence = MovieScene.GetTypedOuter<UMovieSceneSequence>())
 	{
diff --git a/Engine/Source/Runtime/MovieScene/Public/Bindings/MovieSceneCustomBinding.h b/Engine/Source/Runtime/MovieScene/Public/Bindings/MovieSceneCustomBinding.h
--- a/Engine/Source/Runtime/MovieScene/Public/Bindings/MovieSceneCustomBinding.h
+++ b/Engine/Source/Runtime/MovieScene/Public/Bindings/MovieSceneCustomBinding.h
@@ -88,13 +88,6 @@ public:
 	*/
 	UMovieSceneSpawnableBindingBase* AsSpawnable(TSharedRef<const UE::MovieScene::FSharedPlaybackState> SharedPlaybackState) { return const_cast<UMovieSceneSpawnableBindingBase*>(const_cast<const UMovieSceneCustomBinding*>(this)->AsSpawnable(SharedPlaybackState)); }
 
-	/*
-	* Must be implemented.
-	* Should return the most specific relevant class of the bound object. Used to populate the FMovieScenePossessable bound object class.
-	*/
-	virtual UClass* GetBoundObjectClass() const PURE_VIRTUAL(UMovieSceneCustomBinding::GetBoundObjectClass, return UObject::StaticClass(););
-
-
 #if WITH_EDITOR
 
 	/*
@@ -127,16 +120,17 @@ public:
 	*/
 	virtual FText GetBindingTypePrettyName() const PURE_VIRTUAL(UMovieSceneCustomBinding::GetBindingTypePrettyName, return FText(););
 
-
 	/*
-	* Called by UI code to see if this custom binding supports converting to a possessable.
+	* Must be implemented.
+	* Should return the most specific relevant class of the bound object. Used to populate the FMovieScenePossessable bound object class.
 	*/
-	virtual bool CanConvertToPossessable(const FGuid& Guid, FMovieSceneSequenceIDRef TemplateID, TSharedRef<const UE::MovieScene::FSharedPlaybackState> SharedPlaybackState) const { return true; }
+	virtual UClass* GetBoundObjectClass() const PURE_VIRTUAL(UMovieSceneCustomBinding::GetBoundObjectClass, return UObject::StaticClass(););
+
 
 	/*
-	* Called by UI code when the binding has recently been added or modified in the case anything needs to be initialized or modified based on this. 
+	* Called by UI code to see if this custom binding supports converting to a possessable.
 	*/
-	virtual void OnBindingAddedOrChanged(UMovieScene& OwnerMovieScene) {}
+	virtual bool CanConvertToPossessable(const FGuid& Guid, FMovieSceneSequenceIDRef TemplateID, TSharedRef<const UE::MovieScene::FSharedPlaybackState> SharedPlaybackState) const { return true; }
 
 #endif
 
diff --git a/Engine/Source/Runtime/MovieScene/Public/Bindings/MovieSceneReplaceableBinding.h b/Engine/Source/Runtime/MovieScene/Public/Bindings/MovieSceneReplaceableBinding.h
--- a/Engine/Source/Runtime/MovieScene/Public/Bindings/MovieSceneReplaceableBinding.h
+++ b/Engine/Source/Runtime/MovieScene/Public/Bindings/MovieSceneReplaceableBinding.h
@@ -39,30 +39,16 @@ public:
 
 	/* UMovieSceneCustomBinding overrides */
 	virtual bool SupportsBindingCreationFromObject(const UObject* SourceObject) const override;
-	virtual UClass* GetBoundObjectClass() const override;
 #if WITH_EDITOR
 	virtual	void SetupDefaults(UObject* SpawnedObject, FGuid ObjectBindingId, UMovieScene& OwnerMovieScene, TSharedRef<const UE::MovieScene::FSharedPlaybackState> SharedPlaybackState) override;
 	virtual const FSlateBrush* GetBindingTrackCustomIconOverlay() const override;
 	virtual FText GetBindingTrackIconTooltip() const override;
+	virtual UClass* GetBoundObjectClass() const override;
 
 	virtual bool SupportsConversionFromBinding(const FMovieSceneBindingReference& BindingReference, const UObject* SourceObject) const override;
 	virtual UMovieSceneCustomBinding* CreateCustomBindingFromBinding(const FMovieSceneBindingReference& BindingReference, UObject* SourceObject, UMovieScene& OwnerMovieScene) override;
 #endif
 
-#if WITH_EDITORONLY_DATA
-	// Optional Editor-only preview object
-	UPROPERTY(Instanced, VisibleAnywhere, Category="Editor")
-	TObjectPtr<UMovieSceneSpawnableBindingBase> PreviewSpawnable = nullptr;
-#endif
-
-public:
-
-	/*
-	*  Note that we choose to implement CreateCustomBinding here rather than in subclasses.
-	*  Instead we rely on subclasses to implement CreateInnerSpawnable and InitReplaceableBinding which we call here.
-	*/
-	UMovieSceneCustomBinding* CreateNewCustomBinding(UObject* SourceObject, UMovieScene& OwnerMovieScene) override final;
-
 protected:
 
 	/*
@@ -82,6 +68,12 @@ protected:
 	/* Must be implemented. Called from CreateNewCustomBinding to allow the replaceable to initialize any data members from the source object. */
 	virtual void InitReplaceableBinding(UObject* SourceObject, UMovieScene & OwnerMovieScene) PURE_VIRTUAL(UMovieSceneReplaceableBindingBase::InitReplaceableBinding, return;);
 
+#if WITH_EDITORONLY_DATA
+	// Editor-only preview object
+	UPROPERTY(Instanced, VisibleAnywhere, Category="Editor")
+	TObjectPtr<UMovieSceneSpawnableBindingBase> PreviewSpawnable = nullptr;
+#endif
+
 protected:
 
 	/* UMovieSceneCustomBinding overrides*/
@@ -89,4 +81,9 @@ protected:
 	FMovieSceneBindingResolveResult ResolveBinding(const FMovieSceneBindingResolveParams& ResolveParams, int32 BindingIndex, TSharedRef<const UE::MovieScene::FSharedPlaybackState> SharedPlaybackState) const override final;
 	const UMovieSceneSpawnableBindingBase* AsSpawnable(TSharedRef<const UE::MovieScene::FSharedPlaybackState> SharedPlaybackState) const override final;
 
+	/* 
+	*  Note that we choose to implement CreateCustomBinding here rather than in subclasses.
+	*  Instead we rely on subclasses to implement CreateInnerSpawnable and InitReplaceableBinding which we call here.
+	*/
+	UMovieSceneCustomBinding* CreateNewCustomBinding(UObject* SourceObject, UMovieScene& OwnerMovieScene) override final;
 };
\ No newline at end of file
diff --git a/Engine/Source/Runtime/MovieScene/Public/IMovieScenePlayer.h b/Engine/Source/Runtime/MovieScene/Public/IMovieScenePlayer.h
--- a/Engine/Source/Runtime/MovieScene/Public/IMovieScenePlayer.h
+++ b/Engine/Source/Runtime/MovieScene/Public/IMovieScenePlayer.h
@@ -397,12 +397,10 @@ public:
 	 * Returns the evaluated sequence instance's shared playback state, if any.
 	 */
 	MOVIESCENE_API TSharedPtr<UE::MovieScene::FSharedPlaybackState> FindSharedPlaybackState();
-	MOVIESCENE_API TSharedPtr<const UE::MovieScene::FSharedPlaybackState> FindSharedPlaybackState() const;
 	/**
 	 * Returns the evaluated sequence instance's shared playback state, asserts if there is none.
 	 */
 	MOVIESCENE_API TSharedRef<UE::MovieScene::FSharedPlaybackState> GetSharedPlaybackState();
-	MOVIESCENE_API TSharedRef<const UE::MovieScene::FSharedPlaybackState> GetSharedPlaybackState() const;
 
 	uint16 GetUniqueIndex() const
 	{
diff --git a/Engine/Source/Runtime/MovieScene/Public/MovieSceneBindingReferences.h b/Engine/Source/Runtime/MovieScene/Public/MovieSceneBindingReferences.h
--- a/Engine/Source/Runtime/MovieScene/Public/MovieSceneBindingReferences.h
+++ b/Engine/Source/Runtime/MovieScene/Public/MovieSceneBindingReferences.h
@@ -58,10 +58,6 @@ struct FMovieSceneBindingResolveParams
 	/** The sequence ID of the object binding being resolved */
 	UPROPERTY()
 	FMovieSceneSequenceID SequenceID;
-
-	/* The outer context with which to resolve this binding. May be the world, or may be an outer UObject.*/
-	UPROPERTY()
-	TObjectPtr<UObject> Context = nullptr;
 };
 
 /**
diff --git a/Engine/Source/Runtime/MovieScene/Public/MovieSceneCommonHelpers.h b/Engine/Source/Runtime/MovieScene/Public/MovieSceneCommonHelpers.h
--- a/Engine/Source/Runtime/MovieScene/Public/MovieSceneCommonHelpers.h
+++ b/Engine/Source/Runtime/MovieScene/Public/MovieSceneCommonHelpers.h
@@ -18,11 +18,8 @@ class UMovieSceneSubSection;
 class UMovieSceneSequence;
 class USceneComponent;
 class USoundBase;
-template <class TClass> class TSubclassOf;
-class UMovieSceneCustomBinding;
 struct FRichCurve;
 enum class EMovieSceneKeyInterpolation : uint8;
-struct FMovieSceneSequenceID;
 
 namespace UE::MovieScene
 {
@@ -229,15 +226,6 @@ public:
 	*/
 	static MOVIESCENE_API bool IsBoundToSpawnable(UMovieSceneSequence* Sequence, const FGuid& ObjectId, TSharedRef<const UE::MovieScene::FSharedPlaybackState> SharedPlaybackState, int32 BindingIndex = 0);
 
-	/*
-	* Attempts to create a new custom spawnable binding for the passed in UObject*. 
-	* Where possible, it is preferred to call FSequencerUtilities::CreateOrReplaceBinding as it handles more cases. This should only be called in cases where there is no editor or sequencer context.
-	* FactoryCreatedActor may be passed in as an alternative option for creating the binding in the case an actor factory was able to create an actor from this object.
-	*/
-
-	static MOVIESCENE_API FGuid TryCreateCustomSpawnableBinding(UMovieSceneSequence* Sequence, UObject* CustomBindingObject);
-	
-
 	/*
 	* Returns the single bound object currently bound to the given objectid and binding index (optional).
 	*/
@@ -269,14 +257,6 @@ public:
 	static MOVIESCENE_API const UClass* GetBoundObjectClass(UMovieSceneSequence* Sequence, const FGuid& ObjectId, int32 BindingIndex = 0);
 #endif
 
-	/* Returns a sorted list of all custom binding type classes currently known. Slow, may desire to cache result*/
-	static void MOVIESCENE_API GetPrioritySortedCustomBindingTypes(TArray<const TSubclassOf<UMovieSceneCustomBinding>>& OutCustomBindingTypes);
-
-	/* For cases where the user does not have a IMovieScenePlayer with a shared playback state, creates a transient one. Use sparingly. */
-	static TSharedRef<UE::MovieScene::FSharedPlaybackState> MOVIESCENE_API CreateTransientSharedPlaybackState(UObject* WorldContext, UMovieSceneSequence* Sequence);
-
-	/* Finds the resolution context to use to resolve the given guid. */
-	static MOVIESCENE_API UObject* GetResolutionContext(UMovieSceneSequence* Sequence, const FGuid& ObjectId, const FMovieSceneSequenceID& SequenceID, TSharedRef<const UE::MovieScene::FSharedPlaybackState> SharedPlaybackState);
 };
 
 /**
diff --git a/Engine/Source/Runtime/MovieScene/Public/MovieSceneDynamicBinding.h b/Engine/Source/Runtime/MovieScene/Public/MovieSceneDynamicBinding.h
--- a/Engine/Source/Runtime/MovieScene/Public/MovieSceneDynamicBinding.h
+++ b/Engine/Source/Runtime/MovieScene/Public/MovieSceneDynamicBinding.h
@@ -87,7 +87,7 @@ struct FMovieSceneDynamicBindingResolveParams
 };
 
 USTRUCT(BlueprintType)
-struct MOVIESCENE_API FMovieSceneDynamicBindingResolveResult
+struct FMovieSceneDynamicBindingResolveResult
 {
 	GENERATED_BODY()
 
diff --git a/Engine/Source/Runtime/MovieScene/Public/MovieScenePossessable.h b/Engine/Source/Runtime/MovieScene/Public/MovieScenePossessable.h
--- a/Engine/Source/Runtime/MovieScene/Public/MovieScenePossessable.h
+++ b/Engine/Source/Runtime/MovieScene/Public/MovieScenePossessable.h
@@ -154,9 +154,9 @@ public:
 	UPROPERTY(EditAnywhere, AdvancedDisplay, Category=Actor)
 	TArray<FName> Tags;
 
-	// Property has been deprecated- 'Dynamic Bindings' are now handled by UMovieSceneReplaceableDirectorBlueprintBinding and the custom binding system
-	UPROPERTY()
-	FMovieSceneDynamicBinding DynamicBinding_DEPRECATED;
+	/** Optional user-defined possessable lookup information */
+	UPROPERTY(EditAnywhere, AdvancedDisplay, Category="Sequencer")
+	FMovieSceneDynamicBinding DynamicBinding;
 
 	/* Get the optional binding id for binding to a spawnable */
 	const FMovieSceneObjectBindingID& GetSpawnableObjectBindingID() const 
diff --git a/Engine/Source/Runtime/MovieScene/Public/MovieSceneSequence.h b/Engine/Source/Runtime/MovieScene/Public/MovieSceneSequence.h
--- a/Engine/Source/Runtime/MovieScene/Public/MovieSceneSequence.h
+++ b/Engine/Source/Runtime/MovieScene/Public/MovieSceneSequence.h
@@ -41,12 +41,6 @@ struct FMovieScenePossessable;
 struct FMovieSceneTimecodeSource;
 struct FUniversalObjectLocator;
 struct FMovieSceneBindingReferences;
-struct FMovieSceneDynamicBinding;
-
-#if WITH_EDITOR
-struct FSlateBrush;
-#endif
-
 
 namespace UE::MovieScene
 {
@@ -145,7 +139,6 @@ public:
 	 * @param Params				Resolve parameters specifying the context and fragment-specific parameters
 	 * @param OutObjects			Destination array to add found objects to
 	 */
-	UE_DEPRECATED(5.5, "Please use the version that also takes a SharedPlaybackState")
 	MOVIESCENE_API void LocateBoundObjects(const FGuid& ObjectId, const UE::UniversalObjectLocator::FResolveParams& ResolveParams, TArray<UObject*, TInlineAllocator<1>>& OutObjects) const;
 
 	/**
@@ -156,27 +149,13 @@ public:
 	 * @param Context				Optional context to use to find the required object (for instance, a parent spawnable object)
 	 * @return An array of all bound objects
 	 */
-	UE_DEPRECATED(5.5, "Please use the version that also takes a SharedPlaybackState")
 	TArray<UObject*, TInlineAllocator<1>> LocateBoundObjects(const FGuid& ObjectId, const UE::UniversalObjectLocator::FResolveParams& Context) const
 	{
 		TArray<UObject*, TInlineAllocator<1>> OutObjects;
-		PRAGMA_DISABLE_DEPRECATION_WARNINGS
 		LocateBoundObjects(ObjectId, Context, OutObjects);
-		PRAGMA_ENABLE_DEPRECATION_WARNINGS
 		return OutObjects;
 	}
 
-	/**
-	 * Locate all the objects that correspond to the specified object ID, using the specified parameters.
-	 * Calling this directly instead of IMovieScenePlayer::ResolveBoundObjects means that you won't be able to locate Spawnable objects or objects from Custom Bindings.
-	 *
-	 * @param ObjectId				The unique identifier of the object.
-	 * @param ResolveParams			Resolve parameters specifying the context and fragment-specific parameters
-	 * @param SharedPlaybackState   Optional ptr to SharedPlaybackState
-	 * @param OutObjects			Destination array to add found objects to
-	 */
-	MOVIESCENE_API virtual void LocateBoundObjects(const FGuid& ObjectId, const UE::UniversalObjectLocator::FResolveParams& ResolveParams, TSharedPtr<const FSharedPlaybackState> SharedPlaybackState, TArray<UObject*, TInlineAllocator<1>>& OutObjects) const;
-
 	/**
 	 * Attempt to find the guid relating to the specified object
 	 *
@@ -323,12 +302,6 @@ public:
 	UFUNCTION(BlueprintCallable, Category = "Sequencer|Sequence")
 	MOVIESCENE_API FMovieSceneTimecodeSource GetEarliestTimecodeSource() const;
 
-	MOVIESCENE_API virtual void IterateDynamicBindings(const TSharedRef<UE::MovieScene::FSharedPlaybackState> SharedPlaybackState, TFunction<void(const FGuid&, FMovieSceneDynamicBinding&)> InCallback) {}
-	
-#if WITH_EDITOR
-	MOVIESCENE_API virtual const FSlateBrush* GetCustomBrushForBinding(FGuid BindingID) const { return nullptr; }
-#endif
-
 public:
 
 	MOVIESCENE_API virtual void PostLoad() override;
diff --git a/Engine/Source/Runtime/MovieScene/Public/MovieSceneSpawnable.h b/Engine/Source/Runtime/MovieScene/Public/MovieSceneSpawnable.h
--- a/Engine/Source/Runtime/MovieScene/Public/MovieSceneSpawnable.h
+++ b/Engine/Source/Runtime/MovieScene/Public/MovieSceneSpawnable.h
@@ -262,9 +262,9 @@ public:
 	UPROPERTY(EditAnywhere, Category=Actor)
 	bool bNetAddressableName;
 
-	// Property has been deprecated- 'Dynamic Bindings' are now handled by UMovieSceneReplaceableDirectorBlueprintBinding and the custom binding system
-	UPROPERTY()
-	FMovieSceneDynamicBinding DynamicBinding_DEPRECATED;
+	/** Optional user-defined spawning information */
+	UPROPERTY(EditAnywhere, Category="Sequencer")
+	FMovieSceneDynamicBinding DynamicBinding;
 
 private:
 
diff --git a/Engine/Source/Runtime/MovieSceneTracks/Private/Bindings/MovieSceneSpawnableActorBinding.cpp b/Engine/Source/Runtime/MovieSceneTracks/Private/Bindings/MovieSceneSpawnableActorBinding.cpp
--- a/Engine/Source/Runtime/MovieSceneTracks/Private/Bindings/MovieSceneSpawnableActorBinding.cpp
+++ b/Engine/Source/Runtime/MovieSceneTracks/Private/Bindings/MovieSceneSpawnableActorBinding.cpp
@@ -76,11 +76,35 @@ UObject* UMovieSceneSpawnableActorBindingBase::SpawnObjectInternal(UWorld* World
 		return nullptr;
 	}
 
+	// Ensure this spawnable is not a preview actor. Preview actors will not have BeginPlay() called on them.
+#if WITH_EDITOR
+	SpawnedActor->bIsEditorPreviewActor = false;
+#endif
+
+	// tag this actor so we know it was spawned by sequencer
+	SpawnedActor->Tags.AddUnique(SequencerActorTag);
+
 	if (bNetAddressableName)
 	{
 		SpawnedActor->SetNetAddressable();
 	}
 
+#if WITH_EDITOR
+	if (GIsEditor)
+	{
+		// Explicitly set RF_Transactional on spawned actors so we can undo/redo properties on them.
+		// This particular UObject will be marked RF_Transactional by the caller, but we need to set it on the components.
+
+		for (UActorComponent* Component : SpawnedActor->GetComponents())
+		{
+			if (Component)
+			{
+				Component->SetFlags(RF_Transactional);
+			}
+		}
+	}
+#endif
+
 	if (UMovieSceneEntitySystemLinker* Linker = SharedPlaybackState->GetLinker())
 	{
 		if (UMovieSceneDeferredComponentMovementSystem* DeferredMovementSystem = Linker->FindSystem<UMovieSceneDeferredComponentMovementSystem>())
@@ -98,6 +122,33 @@ UObject* UMovieSceneSpawnableActorBindingBase::SpawnObjectInternal(UWorld* World
 	const bool bIsDefaultTransform = true;
 	SpawnedActor->FinishSpawning(SpawnTransformToUse, bIsDefaultTransform);
 
+#if WITH_EDITOR
+	// Don't set the actor label in PIE as this requires flushing async loading.
+	if (WorldContext->WorldType == EWorldType::Editor)
+	{
+		FString BindingName = GetDesiredBindingName();
+		FString ActorLabel = !BindingName.IsEmpty() ? BindingName : SpawnName.ToString();
+
+		if (FMovieScenePossessable* Possessable = MovieScene.FindPossessable(BindingId))
+		{
+			ActorLabel = Possessable->GetName();
+
+			if (UMovieSceneSequence* Sequence = MovieScene.GetTypedOuter<UMovieSceneSequence>())
+			{
+				if (const FMovieSceneBindingReferences* BindingReferences = Sequence->GetBindingReferences())
+				{
+					if (BindingReferences->GetReferences(BindingId).Num() > 1)
+					{
+						// If there are multiple bound objects, use the Object Template actor label instead of the possessable name
+						ActorLabel = ActorTemplate->GetActorLabel();
+					}
+				}
+			}
+		}
+		SpawnedActor->SetActorLabel(ActorLabel);
+	}
+#endif
+
 	return SpawnedActor;
 }
 
@@ -109,6 +160,21 @@ void UMovieSceneSpawnableActorBindingBase::DestroySpawnedObjectInternal(UObject*
 		return;
 	}
 
+#if WITH_EDITOR
+	if (GIsEditor)
+	{
+		// Explicitly remove RF_Transactional on spawned actors since we don't want to trasact spawn/destroy events
+		// This particular UObject will have RF_Transactional cleared by the caller, but we need to cleared it on the components.
+		for (UActorComponent* Component : Actor->GetComponents())
+		{
+			if (Component)
+			{
+				Component->ClearFlags(RF_Transactional);
+			}
+		}
+	}
+#endif
+
 	UWorld* World = Actor->GetWorld();
 	if (World)
 	{
@@ -387,7 +453,7 @@ UMovieSceneCustomBinding* UMovieSceneSpawnableActorBinding::CreateNewCustomBindi
 
 	auto CreateBinding = [&]()
 	{
-		return NewObject<UMovieSceneSpawnableActorBinding>(&OwnerMovieScene, GetClass(), InstancedBindingName, RF_Transactional);
+		return NewObject<UMovieSceneSpawnableActorBinding>(&OwnerMovieScene, UMovieSceneSpawnableActorBinding::StaticClass(), InstancedBindingName, RF_Transactional);
 	};
 
 	// Deal with creating a spawnable from an instance of an actor
diff --git a/Engine/Source/Runtime/MovieSceneTracks/Public/Bindings/MovieSceneSpawnableActorBinding.h b/Engine/Source/Runtime/MovieSceneTracks/Public/Bindings/MovieSceneSpawnableActorBinding.h
--- a/Engine/Source/Runtime/MovieSceneTracks/Public/Bindings/MovieSceneSpawnableActorBinding.h
+++ b/Engine/Source/Runtime/MovieSceneTracks/Public/Bindings/MovieSceneSpawnableActorBinding.h
@@ -83,10 +83,9 @@ protected:
 	/* Returns the transform to spawn the actor at*/
 	MOVIESCENETRACKS_API  FTransform GetSpawnTransform() const;
 
-	MOVIESCENETRACKS_API UClass* GetBoundObjectClass() const override { return GetActorClass(); }
 #if WITH_EDITOR
 	/* MovieSceneCustomBinding overrides*/
-	virtual int32 GetCustomBindingPriority() const override { return 10; }
+	MOVIESCENETRACKS_API UClass* GetBoundObjectClass() const override { return GetActorClass(); }
 #endif
 
 protected:
@@ -120,33 +119,33 @@ public:
 public:
 
 	/* MovieSceneSpawnableActorBindingBase overrides*/
-	MOVIESCENETRACKS_API bool SupportsObjectTemplates() const override { return true; }
+	bool SupportsObjectTemplates() const override { return true; }
 
 	/* Override and return the object template if the binding type supports object templates*/
-	MOVIESCENETRACKS_API UObject* GetObjectTemplate() override { return ActorTemplate.Get(); }
+	UObject* GetObjectTemplate() override { return ActorTemplate.Get(); }
 
 	/* MovieSceneSpawnableBindingBase overrides*/
-	MOVIESCENETRACKS_API void SetObjectTemplate(UObject* InObjectTemplate) override;
-	MOVIESCENETRACKS_API void CopyObjectTemplate(UObject* InSourceObject, UMovieSceneSequence& MovieSceneSequence) override;
+	void SetObjectTemplate(UObject* InObjectTemplate) override;
+	void CopyObjectTemplate(UObject* InSourceObject, UMovieSceneSequence& MovieSceneSequence) override;
 
 	/* MovieSceneCustomBinding overrides*/
-	MOVIESCENETRACKS_API virtual bool SupportsBindingCreationFromObject(const UObject* SourceObject) const override;
-	MOVIESCENETRACKS_API virtual UMovieSceneCustomBinding* CreateNewCustomBinding(UObject* SourceObject, UMovieScene& OwnerMovieScene) override;
+	bool SupportsBindingCreationFromObject(const UObject* SourceObject) const override;
+	UMovieSceneCustomBinding* CreateNewCustomBinding(UObject* SourceObject, UMovieScene& OwnerMovieScene) override;
 
 #if WITH_EDITOR
-	MOVIESCENETRACKS_API virtual bool SupportsConversionFromBinding(const FMovieSceneBindingReference& BindingReference, const UObject* SourceObject) const override;
-	MOVIESCENETRACKS_API virtual UMovieSceneCustomBinding* CreateCustomBindingFromBinding(const FMovieSceneBindingReference& BindingReference, UObject* SourceObject, UMovieScene& OwnerMovieScene) override;
-	MOVIESCENETRACKS_API virtual FText GetBindingTypePrettyName() const override;
+	bool SupportsConversionFromBinding(const FMovieSceneBindingReference& BindingReference, const UObject* SourceObject) const override;
+	UMovieSceneCustomBinding* CreateCustomBindingFromBinding(const FMovieSceneBindingReference& BindingReference, UObject* SourceObject, UMovieScene& OwnerMovieScene) override;
+	FText GetBindingTypePrettyName() const override;
 #endif
 
 protected:
 
 	/* MovieSceneSpawnableBindingBase overrides*/
-	MOVIESCENETRACKS_API UWorld* GetWorldContext(TSharedRef<const UE::MovieScene::FSharedPlaybackState> SharedPlaybackState) const override;
+	UWorld* GetWorldContext(TSharedRef<const UE::MovieScene::FSharedPlaybackState> SharedPlaybackState) const override;
 	
 	/* MovieSceneSpawnableActorBindingBase overrides*/
-	MOVIESCENETRACKS_API TSubclassOf<AActor> GetActorClass() const override;
-	MOVIESCENETRACKS_API AActor* GetActorTemplate() const override { return ActorTemplate; }
+	TSubclassOf<AActor> GetActorClass() const override;
+	AActor* GetActorTemplate() const override { return ActorTemplate; }
 
 
 private:
diff --git a/Engine/Source/Runtime/UMG/Private/Animation/WidgetAnimation.cpp b/Engine/Source/Runtime/UMG/Private/Animation/WidgetAnimation.cpp
--- a/Engine/Source/Runtime/UMG/Private/Animation/WidgetAnimation.cpp
+++ b/Engine/Source/Runtime/UMG/Private/Animation/WidgetAnimation.cpp
@@ -12,18 +12,11 @@
 #include "Tracks/MovieSceneMaterialParameterCollectionTrack.h"
 #include "UObject/SequencerObjectVersion.h"
 
-#if WITH_EDITOR
-#include "Styling/AppStyle.h"
-#endif
-
 #include UE_INLINE_GENERATED_CPP_BY_NAME(WidgetAnimation)
 
 
 #define LOCTEXT_NAMESPACE "UWidgetAnimation"
 
-#if WITH_EDITOR
-UWidgetAnimation::FFixupWidgetDynamicBindingsEvent UWidgetAnimation::FixupWidgetDynamicBindingsEvent;
-#endif
 
 /* UWidgetAnimation structors
  *****************************************************************************/
@@ -46,35 +39,6 @@ void UWidgetAnimation::PostLoad()
 		bLegacyFinishOnStop = false;
 	}
 
-	if (GetLinkerCustomVersion(FFortniteMainBranchObjectVersion::GUID) < FFortniteMainBranchObjectVersion::LevelSequenceUpgradeDynamicBindings)
-	{
-		bool bConvertedDynamicBinding = false;
-		for (int32 Index = 0; Index < MovieScene->GetPossessableCount(); ++Index)
-		{
-			FMovieScenePossessable& Possessable = MovieScene->GetPossessable(Index);
-			if (Possessable.DynamicBinding_DEPRECATED.Function)
-			{
-				bConvertedDynamicBinding = true;
-
-				for (FWidgetAnimationBinding& Binding : AnimationBindings)
-				{
-					if (Binding.AnimationGuid == Possessable.GetGuid())
-					{
-						Binding.DynamicBinding = Possessable.DynamicBinding_DEPRECATED;
-						Possessable.DynamicBinding_DEPRECATED = FMovieSceneDynamicBinding();
-						break;
-					}
-				}
-			}
-		}
-#if WITH_EDITOR
-		if (bConvertedDynamicBinding && FixupWidgetDynamicBindingsEvent.IsBound())
-		{
-			FixupWidgetDynamicBindingsEvent.Broadcast(this);
-		}
-#endif
-	}
-
 	Super::PostLoad();
 }
 
@@ -123,22 +87,6 @@ ETrackSupport UWidgetAnimation::IsTrackSupported(TSubclassOf<class UMovieSceneTr
 
 	return Super::IsTrackSupported(InTrackClass);
 }
-
-const FSlateBrush* UWidgetAnimation::GetCustomBrushForBinding(FGuid BindingID) const
-{
-	for (const FWidgetAnimationBinding& Binding : AnimationBindings)
-	{
-		if (Binding.AnimationGuid == BindingID)
-		{
-			if (Binding.DynamicBinding.Function)
-			{
-				return FAppStyle::GetBrush("Sequencer.DynamicBindingIconOverlay");
-			}
-		}
-	}
-	return nullptr;
-}
-
 #endif
 
 float UWidgetAnimation::GetStartTime() const
@@ -277,14 +225,14 @@ bool UWidgetAnimation::CanPossessObject(UObject& Object, UObject* InPlaybackCont
 	return (Object.IsA<UVisual>() && Object.IsIn(PreviewWidget));
 }
 
-void UWidgetAnimation::LocateBoundObjects(const FGuid& ObjectId, const UE::UniversalObjectLocator::FResolveParams& ResolveParams, TSharedPtr<const FSharedPlaybackState> SharedPlaybackState, TArray<UObject*, TInlineAllocator<1>>& OutObjects) const
+void UWidgetAnimation::LocateBoundObjects(const FGuid& ObjectId, UObject* InContext, TArray<UObject*, TInlineAllocator<1>>& OutObjects) const
 {
-	if (ResolveParams.Context == nullptr)
+	if (InContext == nullptr)
 	{
 		return;
 	}
 
-	UUserWidget* PreviewWidget = CastChecked<UUserWidget>(ResolveParams.Context);
+	UUserWidget* PreviewWidget = CastChecked<UUserWidget>(InContext);
 	if (PreviewWidget->WidgetTree == nullptr)
 	{
 		return;
@@ -294,7 +242,7 @@ void UWidgetAnimation::LocateBoundObjects(const FGuid& ObjectId, const UE::Unive
 	{
 		if (Binding.AnimationGuid == ObjectId)
 		{
-			UObject* FoundObject = Binding.FindRuntimeObject(*PreviewWidget->WidgetTree, *PreviewWidget, this, SharedPlaybackState);
+			UObject* FoundObject = Binding.FindRuntimeObject(*PreviewWidget->WidgetTree, *PreviewWidget);
 
 			if (FoundObject)
 			{
@@ -302,8 +250,10 @@ void UWidgetAnimation::LocateBoundObjects(const FGuid& ObjectId, const UE::Unive
 			}
 		}
 	}
+
 }
 
+
 UMovieScene* UWidgetAnimation::GetMovieScene() const
 {
 	return MovieScene;
@@ -316,14 +266,6 @@ UObject* UWidgetAnimation::CreateDirectorInstance(TSharedRef<const FSharedPlayba
 	return WidgetContext;
 }
 
-void UWidgetAnimation::IterateDynamicBindings(const TSharedRef<UE::MovieScene::FSharedPlaybackState> SharedPlaybackState, TFunction<void(const FGuid&, FMovieSceneDynamicBinding&)> InCallback)
-{
-	for (FWidgetAnimationBinding& AnimationBinding : AnimationBindings)
-	{
-		InCallback(AnimationBinding.AnimationGuid, AnimationBinding.DynamicBinding);
-	}
-}
-
 UObject* UWidgetAnimation::GetParentObject(UObject* Object) const
 {
 	UPanelSlot* Slot = Cast<UPanelSlot>(Object);
diff --git a/Engine/Source/Runtime/UMG/Private/Animation/WidgetAnimationBinding.cpp b/Engine/Source/Runtime/UMG/Private/Animation/WidgetAnimationBinding.cpp
--- a/Engine/Source/Runtime/UMG/Private/Animation/WidgetAnimationBinding.cpp
+++ b/Engine/Source/Runtime/UMG/Private/Animation/WidgetAnimationBinding.cpp
@@ -4,7 +4,6 @@
 #include "UObject/Object.h"
 #include "Components/Widget.h"
 #include "Blueprint/WidgetTree.h"
-#include "MovieSceneDynamicBindingInvoker.h"
 
 #include UE_INLINE_GENERATED_CPP_BY_NAME(WidgetAnimationBinding)
 
@@ -14,20 +13,6 @@
 
 UObject* FWidgetAnimationBinding::FindRuntimeObject(const UWidgetTree& WidgetTree, UUserWidget& UserWidget ) const
 {	
-	return FindRuntimeObject(WidgetTree, UserWidget, nullptr, nullptr);
-}
-
-UObject* FWidgetAnimationBinding::FindRuntimeObject(const UWidgetTree& WidgetTree, UUserWidget& UserWidget, const UMovieSceneSequence* Sequence, TSharedPtr<const UE::MovieScene::FSharedPlaybackState> SharedPlaybackState) const
-{
-	if (Sequence && SharedPlaybackState && DynamicBinding.Function)
-	{
-		FMovieSceneDynamicBindingResolveResult ResolveResult = FMovieSceneDynamicBindingInvoker::ResolveDynamicBinding(SharedPlaybackState.ToSharedRef(), const_cast<UMovieSceneSequence*>(Sequence), MovieSceneSequenceID::Root, AnimationGuid, DynamicBinding);
-		if (ResolveResult.Object)
-		{
-			return ResolveResult.Object;
-		}
-	}
-
 	if (bIsRootWidget)
 	{
 		return &UserWidget;
@@ -47,4 +32,4 @@ UObject* FWidgetAnimationBinding::FindRuntimeObject(const UWidgetTree& WidgetTre
 	}
 
 	return FoundObject;
-}
\ No newline at end of file
+}
diff --git a/Engine/Source/Runtime/UMG/Public/Animation/WidgetAnimation.h b/Engine/Source/Runtime/UMG/Public/Animation/WidgetAnimation.h
--- a/Engine/Source/Runtime/UMG/Public/Animation/WidgetAnimation.h
+++ b/Engine/Source/Runtime/UMG/Public/Animation/WidgetAnimation.h
@@ -13,10 +13,6 @@
 class UMovieScene;
 class UUserWidget;
 
-#if WITH_EDITOR
-struct FSlateBrush;
-#endif
-
 
 /**
  * 
@@ -101,16 +97,10 @@ public:
 	virtual void UnbindPossessableObjects(const FGuid& ObjectId) override;
 	virtual void UnbindObjects(const FGuid& ObjectId, const TArray<UObject*>& InObjects, UObject* InContext) override {}
 	virtual void UnbindInvalidObjects(const FGuid& ObjectId, UObject* InContext) override {}
-	virtual void LocateBoundObjects(const FGuid& ObjectId, const UE::UniversalObjectLocator::FResolveParams& ResolveParams, TSharedPtr<const FSharedPlaybackState> SharedPlaybackState, TArray<UObject*, TInlineAllocator<1>>& OutObjects) const override;
+	virtual void LocateBoundObjects(const FGuid& ObjectId, UObject* Context, TArray<UObject*, TInlineAllocator<1>>& OutObjects) const override;
 	virtual UObject* CreateDirectorInstance(TSharedRef<const FSharedPlaybackState> SharedPlaybackState, FMovieSceneSequenceID SequenceID) override;
-
-	virtual void IterateDynamicBindings(const TSharedRef<UE::MovieScene::FSharedPlaybackState> SharedPlaybackState, TFunction<void(const FGuid&, FMovieSceneDynamicBinding&)> InCallback) override;
 #if WITH_EDITOR
 	virtual ETrackSupport IsTrackSupported(TSubclassOf<class UMovieSceneTrack> InTrackClass) const override;
-	virtual const FSlateBrush* GetCustomBrushForBinding(FGuid BindingID) const override;
-
-	DECLARE_MULTICAST_DELEGATE_OneParam(FFixupWidgetDynamicBindingsEvent, UWidgetAnimation*);
-	static UMG_API FFixupWidgetDynamicBindingsEvent FixupWidgetDynamicBindingsEvent;
 #endif
 	// ~UMovieSceneAnimation overrides
 
diff --git a/Engine/Source/Runtime/UMG/Public/Animation/WidgetAnimationBinding.h b/Engine/Source/Runtime/UMG/Public/Animation/WidgetAnimationBinding.h
--- a/Engine/Source/Runtime/UMG/Public/Animation/WidgetAnimationBinding.h
+++ b/Engine/Source/Runtime/UMG/Public/Animation/WidgetAnimationBinding.h
@@ -5,19 +5,10 @@
 #include "CoreMinimal.h"
 #include "UObject/ObjectMacros.h"
 #include "Misc/Guid.h"
-#include "MovieSceneDynamicBinding.h"
-#include "UObject/FortniteMainBranchObjectVersion.h"
 #include "WidgetAnimationBinding.generated.h"
 
 class UUserWidget;
 class UWidgetTree;
-class UMovieSceneSequence;
-
-namespace UE::MovieScene
-{
-	struct FSharedPlaybackState;
-}
-
 
 /**
  * A single object bound to a UMG sequence.
@@ -39,20 +30,14 @@ struct FWidgetAnimationBinding
 	UPROPERTY()
 	bool bIsRootWidget = false;
 
-	UPROPERTY()
-	FMovieSceneDynamicBinding DynamicBinding;
-
 public:
 
 	/**
 	 * Locates a runtime object to animate from the provided tree of widgets.
 	 * @return the runtime object to animate or null if not found
 	 */
-	UE_DEPRECATED(5.5, "Please use the version that takes a SharedPlaybackState and Sequence")
 	UMG_API UObject* FindRuntimeObject(const UWidgetTree& WidgetTree, UUserWidget& UserWidget) const;
 
-	UMG_API UObject* FindRuntimeObject(const UWidgetTree& WidgetTree, UUserWidget& UserWidget, const UMovieSceneSequence* Sequence, TSharedPtr<const UE::MovieScene::FSharedPlaybackState> SharedPlaybackState) const;
-
 	/**
 	 * Compares two widget animation bindings for equality.
 	 *
@@ -62,6 +47,30 @@ public:
 	 */
 	friend bool operator==(const FWidgetAnimationBinding& X, const FWidgetAnimationBinding& Y)
 	{
-		return (X.WidgetName == Y.WidgetName) && (X.SlotWidgetName == Y.SlotWidgetName) && (X.AnimationGuid == Y.AnimationGuid) && (X.bIsRootWidget == Y.bIsRootWidget) && (X.DynamicBinding.Function == Y.DynamicBinding.Function);
+		return (X.WidgetName == Y.WidgetName) && (X.SlotWidgetName == Y.SlotWidgetName) && (X.AnimationGuid == Y.AnimationGuid) && (X.bIsRootWidget == Y.bIsRootWidget);
+	}
+
+	/**
+	 * Serializes a widget animation binding from/to the given archive.
+	 *
+	 * @param Ar The archive to serialize to/from.
+	 * @param Binding the binding to serialize.
+	 */
+	friend FArchive& operator<<(FArchive& Ar, FWidgetAnimationBinding& Binding)
+	{
+		Ar << Binding.WidgetName;
+		Ar << Binding.SlotWidgetName;
+		Ar << Binding.AnimationGuid;
+		Ar << Binding.bIsRootWidget;
+		return Ar;
+	}
+
+	friend void operator<<(FStructuredArchive::FSlot Slot, FWidgetAnimationBinding& Binding)
+	{
+		FStructuredArchive::FRecord Record = Slot.EnterRecord();
+		Record << SA_VALUE(TEXT("WidgetName"), Binding.WidgetName);
+		Record << SA_VALUE(TEXT("SlotWidgetName"), Binding.SlotWidgetName);
+		Record << SA_VALUE(TEXT("AnimationGuid"), Binding.AnimationGuid);
+		Record << SA_VALUE(TEXT("bIsRootWidget"), Binding.bIsRootWidget);
 	}
 };
