Weighted tangents: When checking equality when solving the cubic using IsNearlyEqual to make sure we handle numerical precisions issues with optimzations.

diff --git a/Engine/Source/Runtime/Engine/Private/Curves/CurveEvaluation.cpp b/Engine/Source/Runtime/Engine/Private/Curves/CurveEvaluation.cpp
--- a/Engine/Source/Runtime/Engine/Private/Curves/CurveEvaluation.cpp
+++ b/Engine/Source/Runtime/Engine/Private/Curves/CurveEvaluation.cpp
@@ -166,7 +166,7 @@ namespace UE
 				NewInterp = TNumericLimits<double>::Lowest(); //just need to be out of range
 				for (double Result : Results)
 				{
-					if ((Result >= 0.0) && (Result <= 1.0))
+					if ((Result > 0.0 || FMath::IsNearlyEqual(Result, 0.0)) && (Result < 1.0 || FMath::IsNearlyEqual(Result, 1.0))) 
 					{
 						if (NewInterp < 0.0 || Result > NewInterp)
 						{
diff --git a/Engine/Source/Runtime/MovieScene/Private/Channels/MovieSceneInterpolation.cpp b/Engine/Source/Runtime/MovieScene/Private/Channels/MovieSceneInterpolation.cpp
--- a/Engine/Source/Runtime/MovieScene/Private/Channels/MovieSceneInterpolation.cpp
+++ b/Engine/Source/Runtime/MovieScene/Private/Channels/MovieSceneInterpolation.cpp
@@ -218,7 +218,7 @@ double FWeightedCubicInterpolation::Evaluate(FFrameTime InTime) const
 		NewInterp = TNumericLimits<double>::Lowest(); //just need to be out of range
 		for (double Result : Results)
 		{
-			if ((Result >= 0.0) && (Result <= 1.0))
+			if ((Result > 0.0 || FMath::IsNearlyEqual(Result, 0.0)) && (Result < 1.0 || FMath::IsNearlyEqual(Result, 1.0))) 
 			{
 				if (NewInterp < 0.0 || Result > NewInterp)
 				{
