Sequencer: remove IMovieScenePlayer usage from core code

diff --git a/Engine/Plugins/MovieScene/TemplateSequence/Source/TemplateSequence/Private/CameraAnimationSequence.cpp b/Engine/Plugins/MovieScene/TemplateSequence/Source/TemplateSequence/Private/CameraAnimationSequence.cpp
--- a/Engine/Plugins/MovieScene/TemplateSequence/Source/TemplateSequence/Private/CameraAnimationSequence.cpp
+++ b/Engine/Plugins/MovieScene/TemplateSequence/Source/TemplateSequence/Private/CameraAnimationSequence.cpp
@@ -7,6 +7,7 @@
 UCameraAnimationSequence::UCameraAnimationSequence(const FObjectInitializer& ObjectInitializer)
     : UTemplateSequence(ObjectInitializer)
 {
+	// TODO: warn about saving property values in the spawnable template
 }
 
 #if WITH_EDITOR
diff --git a/Engine/Plugins/MovieScene/TemplateSequence/Source/TemplateSequence/Private/CameraAnimationSequenceSubsystem.cpp b/Engine/Plugins/MovieScene/TemplateSequence/Source/TemplateSequence/Private/CameraAnimationSequenceSubsystem.cpp
--- a/Engine/Plugins/MovieScene/TemplateSequence/Source/TemplateSequence/Private/CameraAnimationSequenceSubsystem.cpp
+++ b/Engine/Plugins/MovieScene/TemplateSequence/Source/TemplateSequence/Private/CameraAnimationSequenceSubsystem.cpp
@@ -8,7 +8,9 @@
 #include "EntitySystem/MovieSceneEntitySystemTask.h"
 #include "EntitySystem/MovieSceneInstanceRegistry.h"
 #include "EntitySystem/MovieSceneRootInstantiatorSystem.h"
+#include "EntitySystem/MovieSceneSharedPlaybackState.h"
 #include "IMovieScenePlayer.h"
+#include "MovieSceneSpawnRegister.h"
 #include "Systems/MovieScenePropertyInstantiator.h"
 
 #include UE_INLINE_GENERATED_CPP_BY_NAME(CameraAnimationSequenceSubsystem)
@@ -63,8 +65,8 @@ struct FCameraAnimationInstantiationMutation : IMovieSceneEntityMutation
 	void BindObjectImpl(const FInstanceHandle& InstanceHandle, const FGuid& ObjectBinding, UObject*& OutBoundObject) const
 	{
 		const FSequenceInstance& Instance = InstanceRegistry.GetInstance(InstanceHandle);
-		IMovieScenePlayer* Player = Instance.GetPlayer();
-		TArrayView<TWeakObjectPtr<>> BoundObjects = Player->FindBoundObjects(ObjectBinding, Instance.GetSequenceID());
+		TSharedRef<const FSharedPlaybackState> SharedPlaybackState = Instance.GetSharedPlaybackState();
+		TArrayView<TWeakObjectPtr<>> BoundObjects = SharedPlaybackState->FindBoundObjects(ObjectBinding, Instance.GetSequenceID());
 		if (ensure(BoundObjects.Num() > 0))
 		{
 			// In theory we should get the scene component from the object, but we know that camera animations are
@@ -109,10 +111,11 @@ void UCameraAnimationSpawnableSystem::OnRun(FSystemTaskPrerequisites& InPrerequi
 			// We won't actually be spawning anything, because our player's spawn register will simply 
 			// return the fake camera "stand-in" object.
 			const FSequenceInstance& Instance = InstanceRegistry->GetInstance(InstanceHandle);
-			IMovieScenePlayer* Player = Instance.GetPlayer();
-			const UMovieSceneSequence* Sequence = Player->State.FindSequence(Instance.GetSequenceID());
-			UObject* SpawnedObject = Player->GetSpawnRegister().SpawnObject(
-				SpawnableBinding, *Sequence->GetMovieScene(), Instance.GetSequenceID(), *Player);
+			TSharedRef<const FSharedPlaybackState> SharedPlaybackState = Instance.GetSharedPlaybackState();
+			const UMovieSceneSequence* Sequence = SharedPlaybackState->GetSequence(Instance.GetSequenceID());
+			FMovieSceneSpawnRegister* SpawnRegister = SharedPlaybackState->FindCapability<FMovieSceneSpawnRegister>();
+			UObject* SpawnedObject = SpawnRegister->SpawnObject(
+				SpawnableBinding, *Sequence->GetMovieScene(), Instance.GetSequenceID(), SharedPlaybackState);
 			ensure(SpawnedObject);
 		});
 }
diff --git a/Engine/Source/Runtime/MovieScene/Private/EntitySystem/MovieSceneBindingLifetimeSystem.cpp b/Engine/Source/Runtime/MovieScene/Private/EntitySystem/MovieSceneBindingLifetimeSystem.cpp
--- a/Engine/Source/Runtime/MovieScene/Private/EntitySystem/MovieSceneBindingLifetimeSystem.cpp
+++ b/Engine/Source/Runtime/MovieScene/Private/EntitySystem/MovieSceneBindingLifetimeSystem.cpp
@@ -51,33 +51,37 @@ void UMovieSceneBindingLifetimeSystem::OnRun(FSystemTaskPrerequisites& InPrerequ
 	}
 
 	FInstanceRegistry* InstanceRegistry = Linker->GetInstanceRegistry();
+	FPreAnimatedStateExtension& PreAnimatedState = Linker->PreAnimatedState;
 
 	bool bLink = false;
-	auto SetBindingActivation = [InstanceRegistry, &bLink](FMovieSceneEntityID EntityID, FInstanceHandle InstanceHandle, const FMovieSceneBindingLifetimeComponentData& BindingLifetime)
+	auto SetBindingActivation = [InstanceRegistry, &PreAnimatedState, &bLink](FMovieSceneEntityID EntityID, FInstanceHandle InstanceHandle, const FMovieSceneBindingLifetimeComponentData& BindingLifetime)
 	{
 		const FSequenceInstance& SequenceInstance = InstanceRegistry->GetInstance(InstanceHandle);
 
 		FMovieSceneSequenceID SequenceID = SequenceInstance.GetSequenceID();
-		IMovieScenePlayer* Player = SequenceInstance.GetPlayer();
-		if (Player)
+		TSharedRef<const FSharedPlaybackState> SharedPlaybackState = SequenceInstance.GetSharedPlaybackState();
+		FMovieSceneEvaluationState* EvaluationState = SharedPlaybackState->FindCapability<FMovieSceneEvaluationState>();
+		IMovieScenePlayer* Player = FPlayerIndexPlaybackCapability::GetPlayer(SharedPlaybackState);
+		if (EvaluationState)
 		{
 			// For now we use the linking/unlinking of the inactive ranges to set the binding activations
 			if (BindingLifetime.BindingLifetimeState == EMovieSceneBindingLifetimeState::Inactive)
 			{
-				Player->State.SetBindingActivation(BindingLifetime.BindingGuid, SequenceID, !bLink);
+				EvaluationState->SetBindingActivation(BindingLifetime.BindingGuid, SequenceID, !bLink);
 			}
 			else
 			{
-				for (TWeakObjectPtr<> WeakBoundObject : Player->FindBoundObjects(BindingLifetime.BindingGuid, SequenceID))
+				TArrayView<TWeakObjectPtr<>> BoundObjects = EvaluationState->FindBoundObjects(BindingLifetime.BindingGuid, SequenceID, SharedPlaybackState);
+				for (TWeakObjectPtr<> WeakBoundObject : BoundObjects)
 				{
 					if (UObject* BoundObject = WeakBoundObject.Get())
 					{
 						if (BoundObject->Implements<UMovieSceneBindingEventReceiverInterface>())
 						{
 							TScriptInterface<IMovieSceneBindingEventReceiverInterface> BindingEventReceiver = BoundObject;
-							if (BindingEventReceiver.GetObject())
+							if (BindingEventReceiver.GetObject() && Player)
 							{
-								FMovieSceneObjectBindingID BindingID = UE::MovieScene::FRelativeObjectBindingID(MovieSceneSequenceID::Root, SequenceID, BindingLifetime.BindingGuid, *Player);
+								FMovieSceneObjectBindingID BindingID = UE::MovieScene::FRelativeObjectBindingID(MovieSceneSequenceID::Root, SequenceID, BindingLifetime.BindingGuid, SharedPlaybackState);
 								if (bLink)
 								{
 									IMovieSceneBindingEventReceiverInterface::Execute_OnObjectBoundBySequencer(BindingEventReceiver.GetObject(), Cast<UMovieSceneSequencePlayer>(Player->AsUObject()), BindingID);
diff --git a/Engine/Source/Runtime/MovieScene/Private/EntitySystem/MovieSceneSequenceInstance.cpp b/Engine/Source/Runtime/MovieScene/Private/EntitySystem/MovieSceneSequenceInstance.cpp
--- a/Engine/Source/Runtime/MovieScene/Private/EntitySystem/MovieSceneSequenceInstance.cpp
+++ b/Engine/Source/Runtime/MovieScene/Private/EntitySystem/MovieSceneSequenceInstance.cpp
@@ -139,27 +139,31 @@ uint16 FSequenceInstance::GetPlayerIndex() const
 
 void FSequenceInstance::InitializeLegacyEvaluator()
 {
-	IMovieScenePlayer* Player = GetPlayer();
-	check(Player);
-
 	UMovieSceneCompiledDataManager* CompiledDataManager = SharedPlaybackState->GetCompiledDataManager();
 	const FMovieSceneCompiledDataID RootCompiledDataID = SharedPlaybackState->GetRootCompiledDataID();
 	const FMovieSceneCompiledDataEntry& CompiledEntry   = CompiledDataManager->GetEntryRef(RootCompiledDataID);
 
 	if (EnumHasAnyFlags(CompiledEntry.AccumulatedMask, EMovieSceneSequenceCompilerMask::EvaluationTemplate))
 	{
-		UpdateFlags |= ESequenceInstanceUpdateFlags::HasLegacyTemplates;
-
-		if (!LegacyEvaluator)
+		IMovieScenePlayer* Player = FPlayerIndexPlaybackCapability::GetPlayer(SharedPlaybackState);
+		if (ensureMsgf(Player, TEXT("Sequence has legacy templates but no player! Legacy tracks will be disabled!")))
 		{
-			LegacyEvaluator = MakeUnique<FMovieSceneTrackEvaluator>(CompiledEntry.GetSequence(), RootCompiledDataID, CompiledDataManager);
+			UpdateFlags |= ESequenceInstanceUpdateFlags::HasLegacyTemplates;
+
+			if (!LegacyEvaluator)
+			{
+				LegacyEvaluator = MakeUnique<FMovieSceneTrackEvaluator>(CompiledEntry.GetSequence(), RootCompiledDataID, CompiledDataManager);
+			}
 		}
 	}
 	else if (LegacyEvaluator)
 	{
-		LegacyEvaluator->Finish(*Player);
-		LegacyEvaluator = nullptr;
-
+		IMovieScenePlayer* Player = FPlayerIndexPlaybackCapability::GetPlayer(SharedPlaybackState);
+		if (ensureMsgf(Player, TEXT("Sequence needs to teardown legacy evaluator, but has no player! Legacy track evaluators will leak!")))
+		{
+			LegacyEvaluator->Finish(*Player);
+			LegacyEvaluator = nullptr;
+		}
 		UpdateFlags &= ~ESequenceInstanceUpdateFlags::HasLegacyTemplates;
 	}
 }
@@ -285,7 +289,7 @@ void FSequenceInstance::Start(const FMovieSceneContext& InContext)
 void FSequenceInstance::Update(const FMovieSceneContext& InContext)
 {
 	SCOPE_CYCLE_COUNTER(MovieSceneEval_SequenceInstanceUpdate);
-	SCOPE_CYCLE_UOBJECT(ContextScope, GetPlayer()->AsUObject());
+	SCOPE_CYCLE_UOBJECT(ContextScope, SharedPlaybackState->GetRootSequence());
 
 	ensureMsgf(bInitialized, TEXT("This instance hasn't been initialized yet!"));
 
@@ -364,8 +368,8 @@ void FSequenceInstance::PreEvaluation()
 
 	if (IsRootSequence())
 	{
-		IMovieScenePlayer* Player = GetPlayer();
-		if (ensure(Player))
+		IMovieScenePlayer* Player = FPlayerIndexPlaybackCapability::GetPlayer(SharedPlaybackState);
+		if (Player)
 		{
 			Player->PreEvaluation(Context);
 		}
@@ -377,7 +381,7 @@ void FSequenceInstance::RunLegacyTrackTemplates()
 	if (LegacyEvaluator)
 	{
 		IMovieScenePlayer* Player = FPlayerIndexPlaybackCapability::GetPlayer(SharedPlaybackState);
-		if (ensure(Player))
+		if (Player)
 		{
 			if (bFinished)
 			{
@@ -395,8 +399,8 @@ void FSequenceInstance::PostEvaluation()
 {
 	if (IsRootSequence() && EnumHasAnyFlags(UpdateFlags, ESequenceInstanceUpdateFlags::NeedsPostEvaluation))
 	{
-		IMovieScenePlayer* Player = GetPlayer();
-		if (ensure(Player))
+		IMovieScenePlayer* Player = FPlayerIndexPlaybackCapability::GetPlayer(SharedPlaybackState);
+		if (Player)
 		{
 			SCOPE_CYCLE_COUNTER(MovieSceneEval_SequenceInstancePostUpdate);
 
@@ -471,7 +475,7 @@ void FSequenceInstance::FindEntities(UObject* Owner, TArray<FMovieSceneEntityID>
 
 FSubSequencePath FSequenceInstance::GetSubSequencePath() const
 {
-	return FSubSequencePath(SequenceID, *GetPlayer());
+	return FSubSequencePath(SequenceID, SharedPlaybackState);
 }
 
 bool FSequenceInstance::ConditionalRecompile(UMovieSceneEntitySystemLinker* Linker)
diff --git a/Engine/Source/Runtime/MovieScene/Private/EntitySystem/MovieSceneSharedPlaybackState.cpp b/Engine/Source/Runtime/MovieScene/Private/EntitySystem/MovieSceneSharedPlaybackState.cpp
--- a/Engine/Source/Runtime/MovieScene/Private/EntitySystem/MovieSceneSharedPlaybackState.cpp
+++ b/Engine/Source/Runtime/MovieScene/Private/EntitySystem/MovieSceneSharedPlaybackState.cpp
@@ -63,6 +63,23 @@ UMovieSceneSequence* FSharedPlaybackState::GetSequence(FMovieSceneSequenceIDRef
 	}
 }
 
+TArrayView<TWeakObjectPtr<>> FSharedPlaybackState::FindBoundObjects(const FGuid& ObjectBindingID, FMovieSceneSequenceIDRef SequenceID) const
+{
+	if (FMovieSceneEvaluationState* EvaluationState = FindCapability<FMovieSceneEvaluationState>())
+	{
+		return EvaluationState->FindBoundObjects(ObjectBindingID, SequenceID, SharedThis(this));
+	}
+	return TArrayView<TWeakObjectPtr<>>();
+}
+
+void FSharedPlaybackState::ClearObjectCaches()
+{
+	if (FMovieSceneEvaluationState* EvaluationState = FindCapability<FMovieSceneEvaluationState>())
+	{
+		EvaluationState->ClearObjectCaches(SharedThis(this));
+	}
+}
+
 void FSharedPlaybackState::InvalidateCachedData()
 {
 	UMovieSceneEntitySystemLinker* Linker = GetLinker();
diff --git a/Engine/Source/Runtime/MovieScene/Private/EntitySystem/MovieSceneSpawnablesSystem.cpp b/Engine/Source/Runtime/MovieScene/Private/EntitySystem/MovieSceneSpawnablesSystem.cpp
--- a/Engine/Source/Runtime/MovieScene/Private/EntitySystem/MovieSceneSpawnablesSystem.cpp
+++ b/Engine/Source/Runtime/MovieScene/Private/EntitySystem/MovieSceneSpawnablesSystem.cpp
@@ -216,7 +216,7 @@ void UMovieSceneSpawnablesSystem::OnRun(FSystemTaskPrerequisites& InPrerequisite
 
 		const FSequenceInstance& SequenceInstance = InstanceRegistry->GetInstance(InstanceHandle);
 
-		TSharedRef<const FSharedPlaybackState> SharedPlaybackState = SequenceInstance.GetSharedPlaybackState();
+		TSharedRef<FSharedPlaybackState> SharedPlaybackState = SequenceInstance.GetSharedPlaybackState();
 		FMovieSceneSpawnRegister* SpawnRegister = SharedPlaybackState->FindCapability<FMovieSceneSpawnRegister>();
 		if (!SpawnRegister)
 		{
@@ -271,12 +271,17 @@ void UMovieSceneSpawnablesSystem::OnRun(FSystemTaskPrerequisites& InPrerequisite
 		{
 			UObject* SpawnedObject = SpawnRegister->SpawnObject(SpawnableBindingID, *Sequence->GetMovieScene(), SequenceID, SharedPlaybackState, BindingIndex);
 			IMovieScenePlayer* Player = FPlayerIndexPlaybackCapability::GetPlayer(SharedPlaybackState);
-			if (SpawnedObject && Player)
+			if (SpawnedObject)
 			{
 				FMovieSceneEvaluationOperand Operand(SequenceID, SpawnableBindingID);
-				Player->OnObjectSpawned(SpawnedObject, Operand);
 
-				Player->SavePreAnimatedState(*SpawnedObject, SpawnableAnimTypeID, FSpawnTrackPreAnimatedTokenProducer(Operand, BindingIndex));
+				if (Player)
+				{
+					Player->OnObjectSpawned(SpawnedObject, Operand);
+				}
+
+				FMovieSceneInstancePreAnimatedState& PreAnimatedState = SharedPlaybackState->GetPreAnimatedState();
+				PreAnimatedState.SavePreAnimatedState(*SpawnedObject, SpawnableAnimTypeID, FSpawnTrackPreAnimatedTokenProducer(Operand, BindingIndex));
 			}
 			return SpawnedObject;
 		};
diff --git a/Engine/Source/Runtime/MovieScene/Private/Evaluation/MovieSceneRootOverridePath.cpp b/Engine/Source/Runtime/MovieScene/Private/Evaluation/MovieSceneRootOverridePath.cpp
--- a/Engine/Source/Runtime/MovieScene/Private/Evaluation/MovieSceneRootOverridePath.cpp
+++ b/Engine/Source/Runtime/MovieScene/Private/Evaluation/MovieSceneRootOverridePath.cpp
@@ -15,6 +15,11 @@ namespace MovieScene
 FSubSequencePath::FSubSequencePath()
 {}
 
+FSubSequencePath::FSubSequencePath(FMovieSceneSequenceID LeafID, TSharedRef<const FSharedPlaybackState> SharedPlaybackState)
+{
+	Reset(LeafID, SharedPlaybackState->GetHierarchy());
+}
+
 FSubSequencePath::FSubSequencePath(FMovieSceneSequenceID LeafID, IMovieScenePlayer& Player)
 {
 	Reset(LeafID, Player.GetEvaluationTemplate().GetHierarchy());
diff --git a/Engine/Source/Runtime/MovieScene/Private/Evaluation/PreAnimatedState/MovieScenePreAnimatedStateExtension.cpp b/Engine/Source/Runtime/MovieScene/Private/Evaluation/PreAnimatedState/MovieScenePreAnimatedStateExtension.cpp
--- a/Engine/Source/Runtime/MovieScene/Private/Evaluation/PreAnimatedState/MovieScenePreAnimatedStateExtension.cpp
+++ b/Engine/Source/Runtime/MovieScene/Private/Evaluation/PreAnimatedState/MovieScenePreAnimatedStateExtension.cpp
@@ -812,14 +812,6 @@ bool FPreAnimatedStateExtension::HasActiveCaptureSource() const
 	return (CaptureSource && CaptureSource->bWantsRestoreState);
 }
 
-bool FPreAnimatedStateExtension::ShouldCaptureAnyState() const
-{
-	FScopedPreAnimatedCaptureSource* CaptureSource = FScopedPreAnimatedCaptureSource::GetCaptureSourcePtr();
-	ensureMsgf(!CaptureSource || CaptureSource->WeakLinker.Get() == Linker,
-			TEXT("The current capture source is related to a different linker. Are you missing setting a scope capture source?"));
-	return (CaptureSource && CaptureSource->bWantsRestoreState) || IsCapturingGlobalState();
-}
-
 void FPreAnimatedStateExtension::AddSourceMetaData(const UE::MovieScene::FPreAnimatedStateEntry& Entry)
 {
 	using namespace UE::MovieScene;
diff --git a/Engine/Source/Runtime/MovieScene/Private/Evaluation/PreAnimatedState/MovieSceneRestoreStateParams.cpp b/Engine/Source/Runtime/MovieScene/Private/Evaluation/PreAnimatedState/MovieSceneRestoreStateParams.cpp
--- a/Engine/Source/Runtime/MovieScene/Private/Evaluation/PreAnimatedState/MovieSceneRestoreStateParams.cpp
+++ b/Engine/Source/Runtime/MovieScene/Private/Evaluation/PreAnimatedState/MovieSceneRestoreStateParams.cpp
@@ -5,6 +5,7 @@
 #include "EntitySystem/MovieSceneEntitySystemLinker.h"
 #include "EntitySystem/MovieSceneInstanceRegistry.h"
 #include "EntitySystem/MovieSceneSequenceInstance.h"
+#include "IMovieScenePlayer.h"
 
 namespace UE
 {
@@ -15,14 +16,15 @@ IMovieScenePlayer* FRestoreStateParams::GetTerminalPlayer() const
 {
 	if (Linker && TerminalInstanceHandle.IsValid())
 	{
-		return Linker->GetInstanceRegistry()->GetInstance(TerminalInstanceHandle).GetPlayer();
+		const FSequenceInstance& Instance = Linker->GetInstanceRegistry()->GetInstance(TerminalInstanceHandle);
+		return UE::MovieScene::FPlayerIndexPlaybackCapability::GetPlayer(Instance.GetSharedPlaybackState());
 	}
 
 	ensureAlways(false);
 	return nullptr;
 }
 
-TSharedPtr<const FSharedPlaybackState> FRestoreStateParams::GetTerminalPlaybackState() const
+TSharedPtr<FSharedPlaybackState> FRestoreStateParams::GetTerminalPlaybackState() const
 {
 	if (Linker && TerminalInstanceHandle.IsValid())
 	{
diff --git a/Engine/Source/Runtime/MovieScene/Public/EntitySystem/MovieSceneSharedPlaybackState.h b/Engine/Source/Runtime/MovieScene/Public/EntitySystem/MovieSceneSharedPlaybackState.h
--- a/Engine/Source/Runtime/MovieScene/Public/EntitySystem/MovieSceneSharedPlaybackState.h
+++ b/Engine/Source/Runtime/MovieScene/Public/EntitySystem/MovieSceneSharedPlaybackState.h
@@ -100,6 +100,12 @@ public:
 	/** Gets a sub-sequence given an ID */
 	UMovieSceneSequence* GetSequence(FMovieSceneSequenceIDRef SequenceID) const;
 
+	/** Finds the bound objects for the given object binding in the given (sub)sequence */
+	TArrayView<TWeakObjectPtr<>> FindBoundObjects(const FGuid& ObjectBindingID, FMovieSceneSequenceIDRef SequenceID) const;
+
+	/** Clears object caches for the entire sequence hierarchy. */
+	void ClearObjectCaches();
+
 public:
 
 	/**
diff --git a/Engine/Source/Runtime/MovieScene/Public/Evaluation/MovieSceneRootOverridePath.h b/Engine/Source/Runtime/MovieScene/Public/Evaluation/MovieSceneRootOverridePath.h
--- a/Engine/Source/Runtime/MovieScene/Public/Evaluation/MovieSceneRootOverridePath.h
+++ b/Engine/Source/Runtime/MovieScene/Public/Evaluation/MovieSceneRootOverridePath.h
@@ -15,6 +15,7 @@ namespace UE
 namespace MovieScene
 {
 
+struct FSharedPlaybackState;
 
 /**
  * A path of unaccumulated sequence IDs ordered from child->parent->grandparent that is used to generate unique sequenceIDs for inner sequences
@@ -33,6 +34,9 @@ struct FSubSequencePath
 	 * @param LeafID 			ID of the child-most sequence to include in this path
 	 * @param Player 			Player from which to retrieve the hierarchy
 	 */
+	MOVIESCENE_API explicit FSubSequencePath(FMovieSceneSequenceID LeafID, TSharedRef<const FSharedPlaybackState> SharedPlaybackState);
+
+	UE_DEPRECATED(5.5, "Please use the version that takes an FSharedPlaybackState")
 	MOVIESCENE_API explicit FSubSequencePath(FMovieSceneSequenceID LeafID, IMovieScenePlayer& Player);
 
 	/**
diff --git a/Engine/Source/Runtime/MovieScene/Public/Evaluation/PreAnimatedState/MovieScenePreAnimatedStateExtension.h b/Engine/Source/Runtime/MovieScene/Public/Evaluation/PreAnimatedState/MovieScenePreAnimatedStateExtension.h
--- a/Engine/Source/Runtime/MovieScene/Public/Evaluation/PreAnimatedState/MovieScenePreAnimatedStateExtension.h
+++ b/Engine/Source/Runtime/MovieScene/Public/Evaluation/PreAnimatedState/MovieScenePreAnimatedStateExtension.h
@@ -283,8 +283,15 @@ public:
 	 */
 	MOVIESCENE_API bool ContainsAnyStateForInstanceHandle(FRootInstanceHandle RootInstanceHandle) const;
 
+public:
 
-	// Use FScopedPreAnimatedCaptureSource to capture from a specific source rather than globally
+	// These methods for globally saving pre-animated state shouldn't be used directly:
+	//
+	// - use FScopedPreAnimatedCaptureSource to capture from a specific source.
+	//
+	// - use the similar methods on SharedPlaybackState->GetPreAnimatedState() to save only when it's
+	//   needed for that given sequence instance.
+	//
 	MOVIESCENE_API void SavePreAnimatedState(FMovieSceneAnimTypeID InTokenType, const IMovieScenePreAnimatedGlobalTokenProducer& Producer);
 	MOVIESCENE_API void SavePreAnimatedState(UObject& InObject, FMovieSceneAnimTypeID InTokenType, const IMovieScenePreAnimatedTokenProducer& Producer);
 
@@ -326,7 +333,6 @@ public:
 private:
 
 	MOVIESCENE_API void FreeGroupInternal(FPreAnimatedStorageGroupHandle Handle);
-	MOVIESCENE_API bool ShouldCaptureAnyState() const;
 
 	MOVIESCENE_API void AddReferencedObjects(UMovieSceneEntitySystemLinker*, FReferenceCollector& ReferenceCollector);
 
diff --git a/Engine/Source/Runtime/MovieScene/Public/Evaluation/PreAnimatedState/MovieSceneRestoreStateParams.h b/Engine/Source/Runtime/MovieScene/Public/Evaluation/PreAnimatedState/MovieSceneRestoreStateParams.h
--- a/Engine/Source/Runtime/MovieScene/Public/Evaluation/PreAnimatedState/MovieSceneRestoreStateParams.h
+++ b/Engine/Source/Runtime/MovieScene/Public/Evaluation/PreAnimatedState/MovieSceneRestoreStateParams.h
@@ -42,7 +42,7 @@ struct FRestoreStateParams
 	 * 2 completely different sequences animated the same object at the same time.
 	 * May be null in test harnesses
 	 */
-	MOVIESCENE_API TSharedPtr<const FSharedPlaybackState> GetTerminalPlaybackState() const;
+	MOVIESCENE_API TSharedPtr<FSharedPlaybackState> GetTerminalPlaybackState() const;
 };
 
 
