Sequencer: mark FSequenceInstance::GetPlayer as deprecated

diff --git a/Engine/Source/Runtime/MovieScene/Public/EntitySystem/MovieSceneSequenceInstance.h b/Engine/Source/Runtime/MovieScene/Public/EntitySystem/MovieSceneSequenceInstance.h
--- a/Engine/Source/Runtime/MovieScene/Public/EntitySystem/MovieSceneSequenceInstance.h
+++ b/Engine/Source/Runtime/MovieScene/Public/EntitySystem/MovieSceneSequenceInstance.h
@@ -155,11 +155,13 @@ public:
 	 *
 	 * @return A pointer to this instance's player
 	 */
+	UE_DEPRECATED(5.5, "The player interface is now optional and will be deprecated in a future relase. Please use FPlayerIndexPlaybackCapability::GetPlayer()")
 	MOVIESCENE_API IMovieScenePlayer* GetPlayer() const;
 
 	/**
 	 * Retrieve the IMovieScenePlayer's unique index
 	 */
+	UE_DEPRECATED(5.5, "The player interface is now optional and will be deprecated in a future relase. Please use FPlayerIndexPlaybackCapability::GetPlayerIndex()")
 	MOVIESCENE_API uint16 GetPlayerIndex() const;
 
 	/**
