Sequencer: fix non-deterministic cook errors

Don't change guids during cook, and prevent data modifications after compiled data has been created.

diff --git a/Engine/Source/Runtime/MovieScene/Private/Compilation/MovieSceneCompiledDataManager.cpp b/Engine/Source/Runtime/MovieScene/Private/Compilation/MovieSceneCompiledDataManager.cpp
--- a/Engine/Source/Runtime/MovieScene/Private/Compilation/MovieSceneCompiledDataManager.cpp
+++ b/Engine/Source/Runtime/MovieScene/Private/Compilation/MovieSceneCompiledDataManager.cpp
@@ -484,6 +484,23 @@ void UMovieSceneCompiledDataManager::LoadCompiledData(UMovieSceneSequence* Seque
 	}
 }
 
+bool UMovieSceneCompiledDataManager::CanMarkSignedObjectAsChangedDuringCook(UMovieSceneSequence* Sequence) const
+{
+	const FMovieSceneCompiledDataID DataID = FindDataID(Sequence);
+	if (!DataID.IsValid())
+	{
+		// No data ID has been created, so this sequence hasn't been compiled yet.
+		// We're OK to modify it.
+		return true;
+	}
+
+	const FMovieSceneCompiledDataEntry* EntryPtr = GetEntryPtr(DataID);
+
+	// If the compiled signature is set, we have already compiled the sequence. In that
+	// case, it's not OK to modify data anymore.
+	return !EntryPtr->CompiledSignature.IsValid();
+}
+
 void UMovieSceneCompiledDataManager::SetEmulatedNetworkMask(EMovieSceneServerClientMask NewMask)
 {
 	DestroyAllData();
diff --git a/Engine/Source/Runtime/MovieScene/Private/MovieSceneSequence.cpp b/Engine/Source/Runtime/MovieScene/Private/MovieSceneSequence.cpp
--- a/Engine/Source/Runtime/MovieScene/Private/MovieSceneSequence.cpp
+++ b/Engine/Source/Runtime/MovieScene/Private/MovieSceneSequence.cpp
@@ -199,16 +199,14 @@ void UMovieSceneSequence::Serialize(FArchive& Ar)
 
 #if WITH_EDITOR
 
-bool UMovieSceneSequence::OptimizeForCook()
+void UMovieSceneSequence::OptimizeForCook()
 {
 	UMovieScene* MovieScene = GetMovieScene();
 	if (!MovieScene)
 	{
-		return false;
+		return;
 	}
 
-	bool bModified = false;
-
 	for (int32 TrackIndex = 0; TrackIndex < MovieScene->GetTracks().Num(); )
 	{
 		UMovieSceneTrack* Track = MovieScene->GetTracks()[TrackIndex];
@@ -217,7 +215,6 @@ bool UMovieSceneSequence::OptimizeForCook()
 			Track->RemoveForCook();
 			MovieScene->RemoveTrack(*Track);
 			UE_LOG(LogMovieScene, Display, TEXT("Removing muted track: %s from: %s"), *Track->GetDisplayName().ToString(), *GetPathName());
-			bModified = true;
 			continue;
 		}
 		++TrackIndex;
@@ -237,7 +234,6 @@ bool UMovieSceneSequence::OptimizeForCook()
 					Section->RemoveForCook();
 					Track->RemoveSection(*Section);
 					UE_LOG(LogMovieScene, Display, TEXT("Removing muted section: %s from: %s"), *Section->GetPathName(), *Track->GetDisplayName().ToString());
-					bModified = true;
 					continue;
 				}
 				++SectionIndex;
@@ -269,7 +265,6 @@ bool UMovieSceneSequence::OptimizeForCook()
 				Track->RemoveForCook();
 				MovieScene->RemoveTrack(*Track);
 				UE_LOG(LogMovieScene, Display, TEXT("Removing muted track: %s from: %s"), *Track->GetDisplayName().ToString(), *GetPathName());
-				bModified = true;
 				continue;
 			}
 			++TrackIndex;
@@ -289,7 +284,6 @@ bool UMovieSceneSequence::OptimizeForCook()
 						Section->RemoveForCook();
 						Track->RemoveSection(*Section);
 						UE_LOG(LogMovieScene, Display, TEXT("Removing muted section: %s from: %s"), *Section->GetPathName(), *Track->GetDisplayName().ToString());
-						bModified = true;
 						continue;
 					}
 					++SectionIndex;
@@ -301,22 +295,14 @@ bool UMovieSceneSequence::OptimizeForCook()
 		{
 			UE_LOG(LogMovieScene, Display, TEXT("Removing muted object: %s from: %s"), *MovieScene->GetBindings()[ObjectBindingIndex].GetName(), *GetPathName());
 			FGuid GuidToRemove = MovieScene->GetBindings()[ObjectBindingIndex].GetObjectGuid();
-			bModified |= MovieScene->RemoveSpawnable(GuidToRemove);
-			bModified |= MovieScene->RemovePossessable(GuidToRemove);
+			MovieScene->RemoveSpawnable(GuidToRemove);
+			MovieScene->RemovePossessable(GuidToRemove);
 		}
 		else
 		{
 			++ObjectBindingIndex;
 		}
 	}
-
-	if (bModified)
-	{
-		Modify();
-		MovieScene->Modify();
-	}
-
-	return bModified;
 }
 
 #endif
diff --git a/Engine/Source/Runtime/MovieScene/Private/MovieSceneSignedObject.cpp b/Engine/Source/Runtime/MovieScene/Private/MovieSceneSignedObject.cpp
--- a/Engine/Source/Runtime/MovieScene/Private/MovieSceneSignedObject.cpp
+++ b/Engine/Source/Runtime/MovieScene/Private/MovieSceneSignedObject.cpp
@@ -3,6 +3,7 @@
 #include "MovieSceneSignedObject.h"
 #include "Templates/Casts.h"
 #include "MovieSceneSequence.h"
+#include "Compilation/MovieSceneCompiledDataManager.h"
 #include "UObject/Package.h"
 #include "CoreGlobals.h"
 
@@ -116,6 +117,31 @@ void UMovieSceneSignedObject::MarkAsChanged()
 {
 	using namespace UE::MovieScene;
 
+	if (IsRunningCookCommandlet())
+	{
+		// During cooking, we don't want to regenerate new GUIDs, since they would be different
+		// every time and trigger non-detrministic cooking errors. We therefore bail out early
+		// and leave the original GUID.
+		//
+		// However, we also check that nobody is trying to modify data after we have compiled the
+		// sequence, as this would lead to a mismatch between the source data and the compiled
+		// data we use to run it. We therefore check with the compiled data manager whether our
+		// parent sequence has already been compiled or not.
+		UMovieSceneSequence* OuterSequence = GetTypedOuter<UMovieSceneSequence>();
+		if (OuterSequence)
+		{
+			UMovieSceneCompiledDataManager* CompiledDataManager = UMovieSceneCompiledDataManager::GetPrecompiledData();
+			check(CompiledDataManager);
+			const bool bCanMarkAsChanged = CompiledDataManager->CanMarkSignedObjectAsChangedDuringCook(OuterSequence);
+			ensureAlwaysMsgf(bCanMarkAsChanged,
+					TEXT("This object's signature was locked after its sequence was compiled, "
+						 "but someone attempted to mark it, or one of its inner objects, as changed! "
+						 "Current object: %s ; Current sequence: %s"),
+					*GetPathName(), *OuterSequence->GetPathName());
+		}
+		return;
+	}
+
 	// We always change the signature immediately to ensure that any external code that wants
 	// to directly check our signature (eg, to clear caches) can still do so even while there
 	// is an outstanding signal pending for this object
diff --git a/Engine/Source/Runtime/MovieScene/Public/Compilation/MovieSceneCompiledDataManager.h b/Engine/Source/Runtime/MovieScene/Public/Compilation/MovieSceneCompiledDataManager.h
--- a/Engine/Source/Runtime/MovieScene/Public/Compilation/MovieSceneCompiledDataManager.h
+++ b/Engine/Source/Runtime/MovieScene/Public/Compilation/MovieSceneCompiledDataManager.h
@@ -251,6 +251,12 @@ public:
 	MOVIESCENE_API void CopyCompiledData(UMovieSceneSequence* Sequence);
 	MOVIESCENE_API void LoadCompiledData(UMovieSceneSequence* Sequence);
 
+public:
+
+	// Internal API.
+
+	bool CanMarkSignedObjectAsChangedDuringCook(UMovieSceneSequence* Sequence) const;
+
 private:
 
 	MOVIESCENE_API void Gather(const FMovieSceneCompiledDataEntry& Entry, UMovieSceneSequence* Sequence, const FTrackGatherParameters& Params, FMovieSceneGatheredCompilerData* OutCompilerData) const;
@@ -293,6 +299,12 @@ private:
 		return &CompiledDataEntries[DataID.Value];
 	}
 
+	const FMovieSceneCompiledDataEntry* GetEntryPtr(FMovieSceneCompiledDataID DataID) const
+	{
+		check(CompiledDataEntries.IsValidIndex(DataID.Value));
+		return &CompiledDataEntries[DataID.Value];
+	}
+
 	MOVIESCENE_API void ConsoleVariableSink();
 
 private:
diff --git a/Engine/Source/Runtime/MovieScene/Public/MovieSceneSequence.h b/Engine/Source/Runtime/MovieScene/Public/MovieSceneSequence.h
--- a/Engine/Source/Runtime/MovieScene/Public/MovieSceneSequence.h
+++ b/Engine/Source/Runtime/MovieScene/Public/MovieSceneSequence.h
@@ -388,7 +388,7 @@ public:
 private:
 
 #if WITH_EDITOR
-	bool OptimizeForCook();
+	void OptimizeForCook();
 #endif
 
 	/** Serialized compiled data - should only be used through UMovieSceneCompiledDataManager */
