Sequencer: Closing sequencer would not clear selection on Spawnables and Replaceables.

diff --git a/Engine/Source/Runtime/MovieScene/Private/MovieSceneSpawnRegister.cpp b/Engine/Source/Runtime/MovieScene/Private/MovieSceneSpawnRegister.cpp
--- a/Engine/Source/Runtime/MovieScene/Private/MovieSceneSpawnRegister.cpp
+++ b/Engine/Source/Runtime/MovieScene/Private/MovieSceneSpawnRegister.cpp
@@ -173,6 +173,7 @@ void FMovieSceneSpawnRegister::DestroyObjectsByPredicate(TSharedRef<const FShare
 								{
 									bCustomBinding = true;
 									SpawnableBinding->PreDestroyObject(SpawnedObject, It.Key().BindingId, It.Key().BindingIndex, It.Key().TemplateID);
+									PreDestroyObject(*SpawnedObject, It.Key().BindingId, 0, It.Key().TemplateID); 
 									DestroySpawnedObject(*SpawnedObject, SpawnableBinding);
 								}
 							}
