[MovieScene]
* Added include to fix non unity compile error

diff --git a/Engine/Source/Runtime/MovieScene/Public/Evaluation/MovieSceneSequenceWeights.h b/Engine/Source/Runtime/MovieScene/Public/Evaluation/MovieSceneSequenceWeights.h
--- a/Engine/Source/Runtime/MovieScene/Public/Evaluation/MovieSceneSequenceWeights.h
+++ b/Engine/Source/Runtime/MovieScene/Public/Evaluation/MovieSceneSequenceWeights.h
@@ -4,6 +4,7 @@
 
 #include "CoreTypes.h"
 #include "Containers/SortedMap.h"
+#include "EntitySystem/MovieSceneEntityIDs.h"
 #include "UObject/WeakObjectPtrTemplates.h"
 #include "MovieSceneSequenceID.h"
 #include "EntitySystem/MovieSceneSequenceInstanceHandle.h"
