Sequencer: Automatically add spawn track for spawnable with director blueprint bindings.

diff --git a/Engine/Source/Runtime/MovieScene/Private/Bindings/MovieSceneSpawnableBinding.cpp b/Engine/Source/Runtime/MovieScene/Private/Bindings/MovieSceneSpawnableBinding.cpp
--- a/Engine/Source/Runtime/MovieScene/Private/Bindings/MovieSceneSpawnableBinding.cpp
+++ b/Engine/Source/Runtime/MovieScene/Private/Bindings/MovieSceneSpawnableBinding.cpp
@@ -157,10 +157,6 @@ void UMovieSceneSpawnableBindingBase::DestroySpawnedObject(UObject* Object)
 #if WITH_EDITOR
 void UMovieSceneSpawnableBindingBase::SetupDefaults(UObject* SpawnedObject, FGuid ObjectBindingId, UMovieScene& OwnerMovieScene, TSharedRef<const UE::MovieScene::FSharedPlaybackState> SharedPlaybackState)
 {
-	if (!SpawnedObject)
-	{
-		return;
-	}
 	Super::SetupDefaults(SpawnedObject, ObjectBindingId, OwnerMovieScene, SharedPlaybackState);
 
 	// TODO: For now we are not using binding lifetime track for this, though it will support it. We continue to use spawn track until we improve UX of splitting sections
