Sequencer: Only invalidate data-dependent cached data when calling InvalidateCachedData in response to re-compiled data

diff --git a/Engine/Source/Runtime/MovieScene/Private/EntitySystem/MovieSceneSequenceInstance.cpp b/Engine/Source/Runtime/MovieScene/Private/EntitySystem/MovieSceneSequenceInstance.cpp
--- a/Engine/Source/Runtime/MovieScene/Private/EntitySystem/MovieSceneSequenceInstance.cpp
+++ b/Engine/Source/Runtime/MovieScene/Private/EntitySystem/MovieSceneSequenceInstance.cpp
@@ -254,7 +254,7 @@ bool FSequenceInstance::ConditionalRecompile()
 	{
 		if (VolatilityManager->ConditionalRecompile())
 		{
-			InvalidateCachedData();
+			InvalidateCachedData(ESequenceInstanceInvalidationType::DataChanged);
 			return true;
 		}
 	}
