Sequencer: Fixed up and down casting of channel handles

Channel handles reference their channel within the proxy by typename, which has no knowledge of inheritance. In order to support up and down-casting between channel handles of related types, we need to guarantee that the cast is valid.
In order to support this we keep track of the struct types by name in a debug map that is only used when assertions are enabled. Casts are assumed to be valid when assertions are disabled.

diff --git a/Engine/Source/Runtime/MovieScene/Private/Channels/MovieSceneChannelHandle.cpp b/Engine/Source/Runtime/MovieScene/Private/Channels/MovieSceneChannelHandle.cpp
--- a/Engine/Source/Runtime/MovieScene/Private/Channels/MovieSceneChannelHandle.cpp
+++ b/Engine/Source/Runtime/MovieScene/Private/Channels/MovieSceneChannelHandle.cpp
@@ -3,6 +3,16 @@
 #include "Channels/MovieSceneChannelHandle.h"
 #include "Channels/MovieSceneChannelProxy.h"
 
+namespace UE::MovieScene
+{
+
+
+#if DO_CHECK
+	TMap<FName, TWeakObjectPtr<UStruct>> GChannelTypeNamesToClassTypes;
+#endif // DO_CHECK
+
+} // namespace UE::MovieScene
+
 FMovieSceneChannelHandle::FMovieSceneChannelHandle()
 	: ChannelTypeName(NAME_None)
 	, ChannelIndex(INDEX_NONE)
@@ -79,4 +89,36 @@ const void* FMovieSceneChannelHandle::GetExtendedEditorData() const
 	return nullptr;
 }
 
-#endif // WITH_EDITOR
\ No newline at end of file
+#endif // WITH_EDITOR
+
+
+
+#if DO_CHECK
+
+void FMovieSceneChannelHandle::TrackChannelTypeNameInternal(UStruct* ChannelType)
+{
+	FName ChannelName = ChannelType->GetFName();
+	if (!UE::MovieScene::GChannelTypeNamesToClassTypes.Contains(ChannelName))
+	{
+		UE::MovieScene::GChannelTypeNamesToClassTypes.Add(ChannelName, ChannelType);
+	}
+}
+UStruct* FMovieSceneChannelHandle::GetChannelTypeByName(FName ChannelName)
+{
+	return UE::MovieScene::GChannelTypeNamesToClassTypes.FindRef(ChannelName).Get();
+}
+
+bool FMovieSceneChannelHandle::IsCastValidInternal(UStruct* OtherType) const
+{
+	if (OtherType->GetFName() != ChannelTypeName)
+	{
+		UStruct* ThisType = GetChannelTypeByName(ChannelTypeName);
+		if (!ThisType || !(ThisType->IsChildOf(OtherType) || OtherType->IsChildOf(ThisType)))
+		{
+			return false;
+		}
+	}
+	return true;
+}
+
+#endif // DO_CHECK
\ No newline at end of file
diff --git a/Engine/Source/Runtime/MovieScene/Public/Channels/MovieSceneChannelHandle.h b/Engine/Source/Runtime/MovieScene/Public/Channels/MovieSceneChannelHandle.h
--- a/Engine/Source/Runtime/MovieScene/Public/Channels/MovieSceneChannelHandle.h
+++ b/Engine/Source/Runtime/MovieScene/Public/Channels/MovieSceneChannelHandle.h
@@ -54,6 +54,14 @@ public:
 		return !(A == B);
 	}
 
+	template<typename ChannelType>
+	static void TrackChannelTypeName()
+	{
+#if DO_CHECK
+		TrackChannelTypeNameInternal(ChannelType::StaticStruct());
+#endif
+	}
+
 public:
 
 	/**
@@ -105,6 +113,20 @@ public:
 
 #endif // WITH_EDITOR
 
+private:
+
+#if DO_CHECK
+	MOVIESCENE_API static void TrackChannelTypeNameInternal(UStruct* ChannelType);
+	MOVIESCENE_API static UStruct* GetChannelTypeByName(FName ChannelName);
+	MOVIESCENE_API bool IsCastValidInternal(UStruct* OtherType) const;
+
+	template<typename OtherChannelType>
+	bool IsCastValid() const
+	{
+		return IsCastValidInternal(OtherChannelType::StaticStruct());
+	}
+#endif // DO_CHECK
+
 private:
 
 	/** Weak pointer to the channel, proxy alisased to the channel proxy's shared reference controller to ensure it becomes null when the proxy is re-allocated */
@@ -182,12 +204,17 @@ protected:
 	TMovieSceneChannelHandle(TWeakPtr<FMovieSceneChannelProxy> InWeakChannelProxy, FName InChannelTypeName, int32 InChannelIndex)
 		: TMovieSceneChannelHandle<typename ChannelType::Super>(InWeakChannelProxy, InChannelTypeName, InChannelIndex)
 	{}
+
+	friend FMovieSceneChannelHandle;
 };
 
 
 template<typename OtherChannelType>
 TMovieSceneChannelHandle<OtherChannelType> FMovieSceneChannelHandle::Cast() const
 {
-	check(OtherChannelType::StaticStruct()->GetFName() == ChannelTypeName);
-	return TMovieSceneChannelHandle<OtherChannelType>(WeakChannelProxy, ChannelIndex);
+#if DO_CHECK
+	checkf(IsCastValid<OtherChannelType>(), TEXT("Invalid cast between unrelated channel types"));
+#endif
+
+	return TMovieSceneChannelHandle<OtherChannelType>(WeakChannelProxy, ChannelTypeName, ChannelIndex);
 }
diff --git a/Engine/Source/Runtime/MovieScene/Public/Channels/MovieSceneChannelProxy.h b/Engine/Source/Runtime/MovieScene/Public/Channels/MovieSceneChannelProxy.h
--- a/Engine/Source/Runtime/MovieScene/Public/Channels/MovieSceneChannelProxy.h
+++ b/Engine/Source/Runtime/MovieScene/Public/Channels/MovieSceneChannelProxy.h
@@ -298,6 +298,17 @@ public:
 		return MakeHandle(ChannelTypeName, Index).template Cast<ChannelType>();
 	}
 
+	/**
+	 * Make a new channel handle for the same channel type name and index for this channel proxy
+	 *
+	 * @return A handle to the supplied channel that will become nullptr when the proxy is reallocated, or nullptr if the index or channel type name are invalid.
+	 */
+	template<typename ChannelType>
+	TMovieSceneChannelHandle<ChannelType> CopyHandle(TMovieSceneChannelHandle<ChannelType> InOtherHandle)
+	{
+		return MakeHandle(InOtherHandle.GetChannelTypeName(), InOtherHandle.GetChannelIndex()).template Cast<ChannelType>();
+	}
+
 #if !WITH_EDITOR
 
 	/**
@@ -387,6 +398,8 @@ private:
 template<typename ChannelType>
 int32 FMovieSceneChannelProxyData::AddInternal(ChannelType& InChannel)
 {
+	FMovieSceneChannelHandle::TrackChannelTypeName<ChannelType>();
+
 	// Find the entry for this channel's type
 	FName ChannelTypeName = ChannelType::StaticStruct()->GetFName();
 
