Sequencer: Make tick interval options blueprint accessible

diff --git a/Engine/Source/Runtime/MovieScene/Public/MovieSceneSequencePlaybackSettings.h b/Engine/Source/Runtime/MovieScene/Public/MovieSceneSequencePlaybackSettings.h
--- a/Engine/Source/Runtime/MovieScene/Public/MovieSceneSequencePlaybackSettings.h
+++ b/Engine/Source/Runtime/MovieScene/Public/MovieSceneSequencePlaybackSettings.h
@@ -77,7 +77,7 @@ struct FMovieSceneSequencePlaybackSettings
 	FMovieSceneSequenceLoopCount LoopCount;
 
 	/** Overridable tick interval for this sequence to update at. When not overridden, the owning actor or component's tick interval will be used */
-	UPROPERTY(EditAnywhere, Category="Playback", meta=(ShowOnlyInnerProperties, Units=s, EditCondition="!bInheritTickIntervalFromOwner"))
+	UPROPERTY(EditAnywhere, BlueprintReadWrite, Category="Playback", meta=(ShowOnlyInnerProperties, Units=s, EditCondition="!bInheritTickIntervalFromOwner"))
 	FMovieSceneSequenceTickInterval TickInterval;
 
 	/** The rate at which to playback the animation */
diff --git a/Engine/Source/Runtime/MovieScene/Public/MovieSceneSequenceTickInterval.h b/Engine/Source/Runtime/MovieScene/Public/MovieSceneSequenceTickInterval.h
--- a/Engine/Source/Runtime/MovieScene/Public/MovieSceneSequenceTickInterval.h
+++ b/Engine/Source/Runtime/MovieScene/Public/MovieSceneSequenceTickInterval.h
@@ -34,19 +34,19 @@ struct FMovieSceneSequenceTickInterval
 public:
 
 	/** Defines the rate at which the sequence should update, in seconds */
-	UPROPERTY(EditAnywhere, Category="Playback", meta=(DisplayName="Tick Interval", Units=s))
+	UPROPERTY(EditAnywhere, BlueprintReadWrite, Category="Playback", meta=(DisplayName="Tick Interval", Units=s))
 	float TickIntervalSeconds = 0.f;
 
 	/** Defines an approximate budget for evaluation of this sequence (and any other sequences with the same tick interval) */
-	UPROPERTY(EditAnywhere, Category="Playback", meta=(DisplayName="Evaluation Budget", ForceUnits=us))
+	UPROPERTY(EditAnywhere, BlueprintReadWrite, Category="Playback", meta=(DisplayName="Evaluation Budget", ForceUnits=us))
 	float EvaluationBudgetMicroseconds = 0.f;
 
 	/** When true, the sequence will continue to tick and progress even when the world is paused */
-	UPROPERTY(EditAnywhere, Category="Playback")
+	UPROPERTY(EditAnywhere, BlueprintReadWrite, Category="Playback")
 	bool bTickWhenPaused = false;
 
 	/** When true, allow the sequence to be grouped with other sequences based on Sequencer.TickIntervalGroupingResolutionMs. Otherwise the interval will be used precisely. */
-	UPROPERTY(EditAnywhere, Category="Playback")
+	UPROPERTY(EditAnywhere, BlueprintReadWrite, Category="Playback")
 	bool bAllowRounding = true;
 
 public:
