Sequencer: Removed uncached channel evaluators

diff --git a/Engine/Source/Runtime/MovieScene/Private/Channels/MovieSceneCurveChannelImpl.cpp b/Engine/Source/Runtime/MovieScene/Private/Channels/MovieSceneCurveChannelImpl.cpp
--- a/Engine/Source/Runtime/MovieScene/Private/Channels/MovieSceneCurveChannelImpl.cpp
+++ b/Engine/Source/Runtime/MovieScene/Private/Channels/MovieSceneCurveChannelImpl.cpp
@@ -36,20 +36,6 @@ namespace UE
 namespace MovieScene
 {
 
-MOVIESCENE_API float GCachedChannelEvaluationParityThreshold = 0.f;
-static FAutoConsoleVariableRef CVarCachedChannelEvaluationParityThreshold(
-	TEXT("Sequencer.CachedChannelEvaluationParityThreshold"),
-	GCachedChannelEvaluationParityThreshold,
-	TEXT("Threshold for testing Evaluate parity with cached/uncached routines."),
-	ECVF_Default);
-
-MOVIESCENE_API bool GEnableCachedChannelEvaluation = true;
-static FAutoConsoleVariableRef CVarEnableCachedChannelEvaluation(
-	TEXT("Sequencer.EnableCachedChannelEvaluation"),
-	GEnableCachedChannelEvaluation,
-	TEXT("Toggles whether channel evaluation will use cached or non-cached evaluation."),
-	ECVF_Default);
-
 template<typename ChannelType>
 static typename ChannelType::CurveValueType
 	EvalForTwoKeys(
@@ -203,100 +189,6 @@ int32 TMovieSceneCurveChannelImpl<ChannelType>::AddCubicKey(ChannelType* InChann
 	return Index;
 }
 
-template<typename ChannelType>
-bool TMovieSceneCurveChannelImpl<ChannelType>::EvaluateExtrapolation(const ChannelType* InChannel, FFrameTime InTime, CurveValueType& OutValue)
-{
-	// If the time is outside of the curve, deal with extrapolation
-	if (InTime < InChannel->Times[0])
-	{
-		if (InChannel->PreInfinityExtrap == RCCE_None)
-		{
-			return false;
-		}
-
-		if (InChannel->PreInfinityExtrap == RCCE_Constant)
-		{
-			OutValue = InChannel->Values[0].Value;
-			return true;
-		}
-
-		if (InChannel->PreInfinityExtrap == RCCE_Linear)
-		{
-			const ChannelValueType FirstValue = InChannel->Values[0];
-
-			if (FirstValue.InterpMode == RCIM_Constant)
-			{
-				OutValue = FirstValue.Value;
-			}
-			else if(FirstValue.InterpMode == RCIM_Cubic)
-			{
-				FFrameTime Delta = FFrameTime(InChannel->Times[0]) - InTime;
-				OutValue = FirstValue.Value - Delta.AsDecimal() * FirstValue.Tangent.ArriveTangent;
-			}
-			else if(FirstValue.InterpMode == RCIM_Linear)
-			{
-				const int32 InterpStartFrame = InChannel->Times[1].Value;
-				const int32 DeltaFrame       = InterpStartFrame - InChannel->Times[0].Value;
-				if (DeltaFrame == 0)
-				{
-					OutValue = FirstValue.Value;
-				}
-				else
-				{
-					OutValue = FMath::Lerp(InChannel->Values[1].Value, FirstValue.Value, (InterpStartFrame - InTime.AsDecimal())/DeltaFrame);
-				}
-			}
-			return true;
-		}
-	}
-	else if (InTime > InChannel->Times.Last())
-	{
-		if (InChannel->PostInfinityExtrap == RCCE_None)
-		{
-			return false;
-		}
-
-		if (InChannel->PostInfinityExtrap == RCCE_Constant)
-		{
-			OutValue = InChannel->Values.Last().Value;
-			return true;
-		}
-
-		if (InChannel->PostInfinityExtrap == RCCE_Linear)
-		{
-			const ChannelValueType LastValue = InChannel->Values.Last();
-
-			if (LastValue.InterpMode == RCIM_Constant)
-			{
-				OutValue = LastValue.Value;
-			}
-			else if(LastValue.InterpMode == RCIM_Cubic)
-			{
-				FFrameTime Delta = InTime - InChannel->Times.Last();
-				OutValue = LastValue.Value + Delta.AsDecimal() * LastValue.Tangent.LeaveTangent;
-			}
-			else if(LastValue.InterpMode == RCIM_Linear)
-			{
-				const int32 NumKeys          = InChannel->Times.Num();
-				const int32 InterpStartFrame = InChannel->Times[NumKeys-2].Value;
-				const int32 DeltaFrame       = InChannel->Times.Last().Value-InterpStartFrame;
-
-				if (DeltaFrame == 0)
-				{
-					OutValue = LastValue.Value;
-				}
-				else
-				{
-					OutValue = FMath::Lerp(InChannel->Values[NumKeys-2].Value, LastValue.Value, (InTime.AsDecimal() - InterpStartFrame)/DeltaFrame);
-				}
-			}
-			return true;
-		}
-	}
-
-	return false;
-}
-
 template<typename ChannelType>
 bool TMovieSceneCurveChannelImpl<ChannelType>::CacheExtrapolation(const ChannelType* InChannel, FFrameTime InTime, UE::MovieScene::Interpolation::FCachedInterpolation& OutValue)
 {
@@ -500,382 +392,168 @@ UE::MovieScene::Interpolation::FCachedInterpolation TMovieSceneCurveChannelImpl<
 	}
 	else
 	{
-		// We have a valid pair of keys to cache on the spline.
-		const double DX = InChannel->Times[Index2].Value - InChannel->Times[Index1].Value;
-
-		// Cache the pair of keys from the spline at the correct time
-		ChannelValueType Key1 = InChannel->Values[Index1];
-		ChannelValueType Key2 = InChannel->Values[Index2];
-
-		// Manipulate they keys by translating them by the cycle count in the time-domain.
-		// By caching the control points in this translated state we can avoid having to 
-		// do any manipulation on the input time
-		const FFrameNumber CycleOffset = Params.CycleCount*Params.Duration;
-
-		// Compute the start and end values and tangents.
-		FFrameNumber Time1, Time2;
-
-		// Control point 1
-		double V1 = Params.ValueOffset + Key1.Value;
-		double T1 = Key1.Tangent.LeaveTangent;
-		double W1 = Key1.Tangent.LeaveTangentWeight;
-		bool bIsWeighted1 = (Key1.Tangent.TangentWeightMode == RCTWM_WeightedBoth || Key1.Tangent.TangentWeightMode == RCTWM_WeightedLeave);
-
-		// Control point 2
-		double V2 = Params.ValueOffset + Key2.Value;
-		double T2 = Key2.Tangent.ArriveTangent;
-		double W2 = Key2.Tangent.ArriveTangentWeight;
-		bool bIsWeighted2 = (Key2.Tangent.TangentWeightMode == RCTWM_WeightedBoth || Key2.Tangent.TangentWeightMode == RCTWM_WeightedArrive);
-
-		// Mirror the curve for oscillation by swapping the control points
-		// and mirroring the times based on the width.
-		if (Params.ShouldMirrorCurve())
-		{
-			Swap(V1, V2);
-			Swap(T1, T2);
-			Swap(W1, W2);
-			Swap(bIsWeighted1, bIsWeighted2);
-
-			T1 = -T1;
-			T2 = -T2;
-
-			// Mirror the times of the control points such that
-			// MinFrame                              MaxFrame
-			//  | x1                    x2   x3       |
-			//  |               becomes               |
-			//  |      x3   x2                     x1 |
-			// and offset them by the cycle offset
-			Time1 = MinFrame + (MaxFrame - InChannel->Times[Index2]) + CycleOffset;
-			Time2 = MinFrame + (MaxFrame - InChannel->Times[Index1]) + CycleOffset;
-		}
-		else
-		{
-			// No oscillation so just offset the control points to be in their final positions
-			Time1 = InChannel->Times[Index1] + CycleOffset;
-			Time2 = InChannel->Times[Index2] + CycleOffset;
-		}
-
-		// Cache this interpolation for any time between the two control points
-		FCachedInterpolationRange Range = FCachedInterpolationRange::Finite(Time1, Time2);
-
-		// Careful: We use the original Key1 rather than the oscillated interpmode to ensure that
-		//          we use the correct key to produce a mirror image
-		TEnumAsByte<ERichCurveInterpMode> InterpMode = Key1.InterpMode;
-		const int CheckBothLinear = GSequencerLinearCubicInterpolation;
-	 	if(InterpMode == RCIM_Linear && (CheckBothLinear  && Key2.InterpMode == RCIM_Cubic))
-		{
-			InterpMode = RCIM_Cubic;
-		}
-
-		switch (InterpMode)
-		{
-		case RCIM_Cubic:
-			{
-				// Weighted if either control point has weight on their tangent
-				if (bIsWeighted1 || bIsWeighted2)
-				{
-					return FCachedInterpolation(Range, FWeightedCubicInterpolation(
-						InChannel->TickResolution, Time1,
-						Time1, V1, T1, W1, bIsWeighted1,
-						Time2, V2, T2, W2, bIsWeighted2
-					));
-				}
-				else
-				{
-					return FCachedInterpolation(Range, FCubicInterpolation(Time1, DX, V1, V2, T1, T2));
-				}
-			}
-			break;
-
-		case RCIM_Linear:
-			{
-				const double DY = V2 - V1;
-				return FCachedInterpolation(Range, FLinearInterpolation(Time1, DY/DX, V1));
-			}
-
-		default:
-			return FCachedInterpolation(Range, FConstantValue(V1));
-		}
+		return GetInterpolationForKey(InChannel, Index1, Index2, &Params);
 	}
 }
 
 template<typename ChannelType>
-bool TMovieSceneCurveChannelImpl<ChannelType>::Evaluate(const ChannelType* InChannel, FFrameTime InTime, CurveValueType& OutValue) 
+UE::MovieScene::Interpolation::FCachedInterpolation TMovieSceneCurveChannelImpl<ChannelType>::GetInterpolationForKey(const ChannelType* InChannel, int32 Index, const UE::MovieScene::FCycleParams* Params)
 {
-	return EvaluateWithCache(InChannel, nullptr, InTime, OutValue);
-}
-
-template<typename ChannelType>
-bool TMovieSceneCurveChannelImpl<ChannelType>::EvaluateWithCache(const ChannelType* InChannel, FTimeEvaluationCache* InOutEvaluationCache, FFrameTime InTime, CurveValueType& OutValue) 
-{
-	using namespace UE::MovieScene;
-
-	const bool bResult = GEnableCachedChannelEvaluation ? EvaluateCached(InChannel, InOutEvaluationCache, InTime, OutValue) : EvaluateLegacy(InChannel, InOutEvaluationCache, InTime, OutValue);
+	using namespace UE::MovieScene::Interpolation;
+	check(InChannel->Times.IsValidIndex(Index));
 
-	// ------------------------------------------------------------------------
-	// Check against new cached codepath - eventually this will replace the code above
-	if (GCachedChannelEvaluationParityThreshold != 0.f)
+	if (Index == InChannel->Times.Num() - 1)
 	{
-		CurveValueType ParityValue;
-		const bool bParityResult = GEnableCachedChannelEvaluation ? EvaluateLegacy(InChannel, InOutEvaluationCache, InTime, ParityValue) : EvaluateCached(InChannel, InOutEvaluationCache, InTime, ParityValue);
-
-		if (bParityResult != bResult || !FMath::IsNearlyEqual((double)ParityValue, (double)OutValue, (double)GCachedChannelEvaluationParityThreshold))
-		{
-			UE_LOG(LogMovieScene, Warning, TEXT("Parity mismatch between cached and non-cached evaluation %.16f != %.16f!"), OutValue, ParityValue);
-
-			static bool bBreakDebugger = false;
-			ensureAlways(!bBreakDebugger);
-		}
+		FCachedInterpolationRange Range = FCachedInterpolationRange::Only(InChannel->Times.Last());
+		return FCachedInterpolation(Range, FConstantValue(InChannel->Values.Last().Value));
 	}
 
-	return bResult;
+	return GetInterpolationForKey(InChannel, Index, Index+1, Params);
 }
 
 template<typename ChannelType>
-bool TMovieSceneCurveChannelImpl<ChannelType>::EvaluateCached(const ChannelType* InChannel, FTimeEvaluationCache* InOutEvaluationCache, FFrameTime InTime, CurveValueType& OutValue)
+UE::MovieScene::Interpolation::FCachedInterpolation TMovieSceneCurveChannelImpl<ChannelType>::GetInterpolationForKey(const ChannelType* InChannel, int32 Index1, int32 Index2, const UE::MovieScene::FCycleParams* Params)
 {
-	double ResultValue = 0.0;
-	if (GetInterpolationForTime(InChannel, InOutEvaluationCache, InTime).Evaluate(InTime, ResultValue))
-	{
-		OutValue = static_cast<CurveValueType>(ResultValue);
-		return true;
-	}
-	return false;
-}
-
-template<typename ChannelType>
-bool TMovieSceneCurveChannelImpl<ChannelType>::EvaluateLegacy(const ChannelType* InChannel, FTimeEvaluationCache* InOutEvaluationCache, FFrameTime InTime, CurveValueType& OutValue)
-{
-	using namespace UE::MovieScene;
+	using namespace UE::MovieScene::Interpolation;
 
-	// ------------------------------------------------------------------------
-	// Legacy evaluate in-place codepath - only exists as a fallback
-	const int32 NumKeys = InChannel->Times.Num();
+	check(InChannel->Times.IsValidIndex(Index1) && InChannel->Times.IsValidIndex(Index2));
 
-	// No keys means default value, or nothing
-	if (NumKeys == 0)
-	{
-		if (InChannel->bHasDefaultValue)
-		{
-			OutValue = InChannel->DefaultValue;
-			return true;
-		}
-		return false;
-	}
+	// We have a valid pair of keys to cache on the spline.
+	const double DX = InChannel->Times[Index2].Value - InChannel->Times[Index1].Value;
 
-	// For single keys, we can only ever return that value
-	if (NumKeys == 1)
-	{
-		OutValue = InChannel->Values[0].Value;
-		return true;
-	}
+	// Cache the pair of keys from the spline at the correct time
+	ChannelValueType Key1 = InChannel->Values[Index1];
+	ChannelValueType Key2 = InChannel->Values[Index2];
 
-	// Evaluate with extrapolation if we're outside the bounds of the curve
-	if (EvaluateExtrapolation(InChannel, InTime, OutValue))
-	{
-		return true;
-	}
+	// Manipulate they keys by translating them by the cycle count in the time-domain.
+	// By caching the control points in this translated state we can avoid having to 
+	// do any manipulation on the input time
+	const FFrameNumber CycleOffset = Params ? Params->CycleCount*Params->Duration : 0;
+	const double       ValueOffset = Params ? Params->ValueOffset : 0.0;
 
 	const FFrameNumber MinFrame = InChannel->Times[0];
 	const FFrameNumber MaxFrame = InChannel->Times.Last();
 
-	// Compute the cycled time
-	FCycleParams Params = CycleTime(MinFrame, MaxFrame, InTime);
-
-	// Deal with offset cycles and oscillation
-	if (InTime < FFrameTime(MinFrame))
-	{
-		switch (InChannel->PreInfinityExtrap)
-		{
-		case RCCE_CycleWithOffset: Params.ComputePreValueOffset(InChannel->Values[0].Value, InChannel->Values[NumKeys-1].Value); break;
-		case RCCE_Oscillate:       Params.Oscillate(MinFrame.Value, MaxFrame.Value);                       break;
-		}
-	}
-	else if (InTime > FFrameTime(MaxFrame))
-	{
-		switch (InChannel->PostInfinityExtrap)
-		{
-		case RCCE_CycleWithOffset: Params.ComputePostValueOffset(InChannel->Values[0].Value, InChannel->Values[NumKeys-1].Value); break;
-		case RCCE_Oscillate:       Params.Oscillate(MinFrame.Value, MaxFrame.Value);                        break;
-		}
+	// Compute the start and end values and tangents.
+	FFrameNumber Time1, Time2;
+
+	// Control point 1
+	double V1 = ValueOffset + Key1.Value;
+	double T1 = Key1.Tangent.LeaveTangent;
+	double W1 = Key1.Tangent.LeaveTangentWeight;
+	bool bIsWeighted1 = (Key1.Tangent.TangentWeightMode == RCTWM_WeightedBoth || Key1.Tangent.TangentWeightMode == RCTWM_WeightedLeave);
+
+	// Control point 2
+	double V2 = ValueOffset + Key2.Value;
+	double T2 = Key2.Tangent.ArriveTangent;
+	double W2 = Key2.Tangent.ArriveTangentWeight;
+	bool bIsWeighted2 = (Key2.Tangent.TangentWeightMode == RCTWM_WeightedBoth || Key2.Tangent.TangentWeightMode == RCTWM_WeightedArrive);
+
+	// Mirror the curve for oscillation by swapping the control points
+	// and mirroring the times based on the width.
+	if (Params && Params->ShouldMirrorCurve())
+	{
+		Swap(V1, V2);
+		Swap(T1, T2);
+		Swap(W1, W2);
+		Swap(bIsWeighted1, bIsWeighted2);
+
+		T1 = -T1;
+		T2 = -T2;
+
+		// Mirror the times of the control points such that
+		// MinFrame                              MaxFrame
+		//  | x1                    x2   x3       |
+		//  |               becomes               |
+		//  |      x3   x2                     x1 |
+		// and offset them by the cycle offset
+		Time1 = MinFrame + (MaxFrame - InChannel->Times[Index2]) + CycleOffset;
+		Time2 = MinFrame + (MaxFrame - InChannel->Times[Index1]) + CycleOffset;
 	}
-
-	if (!ensureMsgf(Params.Time.FrameNumber >= MinFrame && Params.Time.FrameNumber <= MaxFrame, TEXT("Invalid time computed for float channel evaluation")))
+	else
 	{
-		return false;
+		// No oscillation so just offset the control points to be in their final positions
+		Time1 = InChannel->Times[Index1] + CycleOffset;
+		Time2 = InChannel->Times[Index2] + CycleOffset;
 	}
 
-	// Evaluate the curve data
-	double Interp = 0.0;
-	int32 Index1 = INDEX_NONE, Index2 = INDEX_NONE;
+	// Cache this interpolation for any time between the two control points
+	FCachedInterpolationRange Range = FCachedInterpolationRange::Finite(Time1, Time2);
 
-	// Initialize cache if not yet performed
-	if(InOutEvaluationCache && InOutEvaluationCache->CachedNumFrames == INDEX_NONE)
-	{
-		UE::MovieScene::EvaluateTime(InChannel->Times, Params.Time, InOutEvaluationCache->Index1, InOutEvaluationCache->Index2, InOutEvaluationCache->InterpValue);
-		InOutEvaluationCache->CachedNumFrames = InChannel->Times.Num();
-		InOutEvaluationCache->CacheFrameTime = Params.Time;
-	}
-
-	// If cache matches contained number of frames, copy data out rather than evaluating time data again
-	if (InOutEvaluationCache && InOutEvaluationCache->CachedNumFrames == InChannel->Times.Num() && InOutEvaluationCache->CacheFrameTime == Params.Time)
-	{
-		Interp = InOutEvaluationCache->InterpValue;
-		Index1 = InOutEvaluationCache->Index1;
-		Index2 = InOutEvaluationCache->Index2;
-	}
-	else
-	{
-		UE::MovieScene::EvaluateTime(InChannel->Times, Params.Time, Index1, Index2, Interp);
-	}
-	
+	// Careful: We use the original Key1 rather than the oscillated interpmode to ensure that
+	//          we use the correct key to produce a mirror image
+	TEnumAsByte<ERichCurveInterpMode> InterpMode = Key1.InterpMode;
 	const int CheckBothLinear = GSequencerLinearCubicInterpolation;
-
-	if (Index1 == INDEX_NONE)
+ 	if(InterpMode == RCIM_Linear && (CheckBothLinear  && Key2.InterpMode == RCIM_Cubic))
 	{
-		OutValue = Params.ValueOffset + InChannel->Values[Index2].Value;
+		InterpMode = RCIM_Cubic;
 	}
-	else if (Index2 == INDEX_NONE)
-	{
-		OutValue = Params.ValueOffset + InChannel->Values[Index1].Value;
-	}
-	else
+
+	switch (InterpMode)
 	{
-		ChannelValueType Key1 = InChannel->Values[Index1];
-		ChannelValueType Key2 = InChannel->Values[Index2];
-		TEnumAsByte<ERichCurveInterpMode> InterpMode = Key1.InterpMode;
-	    if(InterpMode == RCIM_Linear && (CheckBothLinear  && Key2.InterpMode == RCIM_Cubic))
+	case RCIM_Cubic:
 		{
-			InterpMode = RCIM_Cubic;
+			if (V1 == V2)
+			{
+				return FCachedInterpolation(Range, FConstantValue(V1));
+			}
+			// Weighted if either control point has weight on their tangent
+			else if (bIsWeighted1 || bIsWeighted2)
+			{
+				return FCachedInterpolation(Range, FWeightedCubicInterpolation(
+					InChannel->TickResolution, Time1,
+					Time1, V1, T1, W1, bIsWeighted1,
+					Time2, V2, T2, W2, bIsWeighted2
+				));
+			}
+			else
+			{
+				return FCachedInterpolation(Range, FCubicInterpolation(Time1, DX, V1, V2, T1, T2));
+			}
 		}
-		
-		switch (InterpMode)
-		{
-		case RCIM_Cubic:
+		break;
+
+	case RCIM_Linear:
 		{
-			const double OneThird = 1.0 / 3.0;
-			if ((Key1.Tangent.TangentWeightMode == RCTWM_WeightedNone || Key1.Tangent.TangentWeightMode == RCTWM_WeightedArrive)
-				&& (Key2.Tangent.TangentWeightMode == RCTWM_WeightedNone || Key2.Tangent.TangentWeightMode == RCTWM_WeightedLeave))
+			const double DY = V2 - V1;
+			if (DY == 0.0)
 			{
-				const int32 Diff = InChannel->Times[Index2].Value - InChannel->Times[Index1].Value;
-				const double P0 = Key1.Value;
-				const double P1 = P0 + (Key1.Tangent.LeaveTangent * Diff * OneThird);
-				const double P3 = Key2.Value;
-				const double P2 = P3 - (Key2.Tangent.ArriveTangent * Diff * OneThird);
-
-				OutValue = Params.ValueOffset + UE::Curves::BezierInterp(P0, P1, P2, P3, Interp);
-				break;
+				return FCachedInterpolation(Range, FConstantValue(V1));
 			}
-			else //its weighted
+			else
 			{
-				const float TimeInterval = InChannel->TickResolution.AsInterval();
-				const float ToSeconds = 1.0f / TimeInterval;
-
-				const double Time1 = InChannel->TickResolution.AsSeconds(InChannel->Times[Index1].Value);
-				const double Time2 = InChannel->TickResolution.AsSeconds(InChannel->Times[Index2].Value);
-				const double X = Time2 - Time1;
-				double CosAngle, SinAngle;
-				double Angle = FMath::Atan(Key1.Tangent.LeaveTangent * ToSeconds);
-				FMath::SinCos(&SinAngle, &CosAngle, Angle);
-				double LeaveWeight;
-				if (Key1.Tangent.TangentWeightMode == RCTWM_WeightedNone || Key1.Tangent.TangentWeightMode == RCTWM_WeightedArrive)
-				{
-					const double LeaveTangentNormalized = Key1.Tangent.LeaveTangent / (TimeInterval);
-					const double Y = LeaveTangentNormalized * X;
-					LeaveWeight = FMath::Sqrt(X*X + Y * Y) * OneThird;
-				}
-				else
-				{
-					LeaveWeight = Key1.Tangent.LeaveTangentWeight;
-				}
-				const double Key1TanX = CosAngle * LeaveWeight + Time1;
-				const double Key1TanY = SinAngle * LeaveWeight + Key1.Value;
-
-				Angle = FMath::Atan(Key2.Tangent.ArriveTangent * ToSeconds);
-				FMath::SinCos(&SinAngle, &CosAngle, Angle);
-				double ArriveWeight;
-				if (Key2.Tangent.TangentWeightMode == RCTWM_WeightedNone || Key2.Tangent.TangentWeightMode == RCTWM_WeightedLeave)
-				{
-					const double ArriveTangentNormalized = Key2.Tangent.ArriveTangent / (TimeInterval);
-					const double Y = ArriveTangentNormalized * X;
-					ArriveWeight = FMath::Sqrt(X*X + Y * Y) * OneThird;
-				}
-				else
-				{
-					ArriveWeight =  Key2.Tangent.ArriveTangentWeight;
-				}
-				const double Key2TanX = -CosAngle * ArriveWeight + Time2;
-				const double Key2TanY = -SinAngle * ArriveWeight + Key2.Value;
-
-				//Normalize the Time Range
-				const double RangeX = Time2 - Time1;
-
-				const double Dx1 = Key1TanX - Time1;
-				const double Dx2 = Key2TanX - Time1;
-
-				// Normalize values
-				const double NormalizedX1 = Dx1 / RangeX;
-				const double NormalizedX2 = Dx2 / RangeX;
-				
-				double Coeff[4];
-				double Results[3];
-
-				//Convert Bezier to Power basis, also float to double for precision for root finding.
-				UE::Curves::BezierToPower(
-					0.0, NormalizedX1, NormalizedX2, 1.0,
-					&(Coeff[3]), &(Coeff[2]), &(Coeff[1]), &(Coeff[0])
-				);
-
-				Coeff[0] = Coeff[0] - Interp;
-				
-				const int32 NumResults = UE::Curves::SolveCubic(Coeff, Results);
-				double NewInterp = Interp;
-				if (NumResults == 1)
-				{
-					NewInterp = Results[0];
-				}
-				else
-				{
-					NewInterp = TNumericLimits<float>::Lowest(); //just need to be out of range
-					for (double Result : Results)
-					{
-						if ((Result >= 0.0f) && (Result <= 1.0f))
-						{
-							if (NewInterp < 0.0f || Result > NewInterp)
-							{
-								NewInterp = Result;
-							}
-						}
-					}
-
-					if (NewInterp == TNumericLimits<float>::Lowest())
-					{
-						NewInterp = 0.f;
-					}
-
-				}
-				//now use NewInterp and adjusted tangents plugged into the Y (Value) part of the graph.
-				const double P0 = Key1.Value;
-				const double P1 = Key1TanY;
-				const double P3 = Key2.Value;
-				const double P2 = Key2TanY;
-
-				OutValue = Params.ValueOffset + UE::Curves::BezierInterp(P0, P1, P2, P3,  NewInterp);
+				return FCachedInterpolation(Range, FLinearInterpolation(Time1, DY/DX, V1));
 			}
-			break;
 		}
 
-		case RCIM_Linear:
-			OutValue = Params.ValueOffset + FMath::Lerp(Key1.Value, Key2.Value, Interp);
-			break;
-
-		default:
-			OutValue = Params.ValueOffset + Key1.Value;
-			break;
-		}
+	default:
+		return FCachedInterpolation(Range, FConstantValue(V1));
 	}
+}
 
-	return true;
+template<typename ChannelType>
+bool TMovieSceneCurveChannelImpl<ChannelType>::Evaluate(const ChannelType* InChannel, FFrameTime InTime, CurveValueType& OutValue) 
+{
+	return EvaluateCached(InChannel, nullptr, InTime, OutValue);
+}
+
+template<typename ChannelType>
+bool TMovieSceneCurveChannelImpl<ChannelType>::EvaluateWithCache(const ChannelType* InChannel, FTimeEvaluationCache* InOutEvaluationCache, FFrameTime InTime, CurveValueType& OutValue) 
+{
+	using namespace UE::MovieScene;
+	return EvaluateCached(InChannel, InOutEvaluationCache, InTime, OutValue);
+}
+
+template<typename ChannelType>
+bool TMovieSceneCurveChannelImpl<ChannelType>::EvaluateCached(const ChannelType* InChannel, FTimeEvaluationCache* InOutEvaluationCache, FFrameTime InTime, CurveValueType& OutValue)
+{
+	double ResultValue = 0.0;
+	if (GetInterpolationForTime(InChannel, InOutEvaluationCache, InTime).Evaluate(InTime, ResultValue))
+	{
+		OutValue = static_cast<CurveValueType>(ResultValue);
+		return true;
+	}
+	return false;
 }
 
 template <typename ChannelValueType>
diff --git a/Engine/Source/Runtime/MovieScene/Private/Channels/MovieSceneCurveChannelImpl.h b/Engine/Source/Runtime/MovieScene/Private/Channels/MovieSceneCurveChannelImpl.h
--- a/Engine/Source/Runtime/MovieScene/Private/Channels/MovieSceneCurveChannelImpl.h
+++ b/Engine/Source/Runtime/MovieScene/Private/Channels/MovieSceneCurveChannelImpl.h
@@ -48,6 +48,9 @@ struct TMovieSceneCurveChannelImpl
 	static UE::MovieScene::Interpolation::FCachedInterpolation GetInterpolationForTime(const ChannelType* InChannel, FFrameTime InTime);
 	static UE::MovieScene::Interpolation::FCachedInterpolation GetInterpolationForTime(const ChannelType* InChannel, FTimeEvaluationCache* InOutEvaluationCache, FFrameTime InTime);
 
+	static UE::MovieScene::Interpolation::FCachedInterpolation GetInterpolationForKey(const ChannelType* InChannel, int32 KeyIndex, const UE::MovieScene::FCycleParams* Params = nullptr);
+	static UE::MovieScene::Interpolation::FCachedInterpolation GetInterpolationForKey(const ChannelType* InChannel, int32 KeyIndex1, int32 KeyIndex2, const UE::MovieScene::FCycleParams* Params = nullptr);
+
 	/** Evaluate this channel at provided FrameTime, using or populating cached time to frame-number(s) calculation */
 	static bool EvaluateWithCache(const ChannelType* InChannel, FTimeEvaluationCache* InOutEvaluationCache, FFrameTime InTime, CurveValueType& OutValue);
 	
@@ -109,6 +112,15 @@ struct TMovieSceneCurveChannelImpl
 	 * @return Interpolation mode to use at that frame
 	 */
 	static EMovieSceneKeyInterpolation GetInterpolationMode(ChannelType* InChannel, const FFrameNumber& InTime, EMovieSceneKeyInterpolation DefaultInterpolationMode);
+
+	/**
+	 * Evaluate this channel's extrapolation by populating a cachable structure. Assumes more than 1 key is present.
+	 *
+	 * @param InTime     The time to evaluate at
+	 * @param OutValue   A value to receive the result
+	 * @return true if the time was evaluated with extrapolation, false otherwise
+	 */
+	static bool CacheExtrapolation(const ChannelType* InChannel, FFrameTime InTime, UE::MovieScene::Interpolation::FCachedInterpolation& OutValue);
 	
 	/** Dilate channel data.*/
 	static void Dilate(ChannelType* InChannel, FFrameNumber Origin, float DilationFactor);
@@ -144,24 +156,6 @@ private:
 	 */
 	static int32 InsertKeyInternal(ChannelType* InChannel, FFrameNumber InTime);
 
-	/**
-	 * Evaluate this channel's extrapolation. Assumes more than 1 key is present.
-	 *
-	 * @param InTime     The time to evaluate at
-	 * @param OutValue   A value to receive the result
-	 * @return true if the time was evaluated with extrapolation, false otherwise
-	 */
-	static bool EvaluateExtrapolation(const ChannelType* InChannel, FFrameTime InTime, CurveValueType& OutValue);
-
-	/**
-	 * Evaluate this channel's extrapolation by populating a cachable structure. Assumes more than 1 key is present.
-	 *
-	 * @param InTime     The time to evaluate at
-	 * @param OutValue   A value to receive the result
-	 * @return true if the time was evaluated with extrapolation, false otherwise
-	 */
-	static bool CacheExtrapolation(const ChannelType* InChannel, FFrameTime InTime, UE::MovieScene::Interpolation::FCachedInterpolation& OutValue);
-
 	/**
 	 * Adds median points between each of the supplied points if their evaluated value is significantly different than the linear interpolation of those points
 	 *
@@ -172,7 +166,6 @@ private:
 	 */
 	static void RefineCurvePoints(const ChannelType* InChannel, FFrameRate TickResolution, double TimeThreshold, CurveValueType ValueThreshold, TArray<TTuple<double, double>>& InOutPoints);
 
-	static bool EvaluateLegacy(const ChannelType* InChannel, FTimeEvaluationCache* InOutEvaluationCache, FFrameTime InTime, CurveValueType& OutValue);
 	static bool EvaluateCached(const ChannelType* InChannel, FTimeEvaluationCache* InOutEvaluationCache, FFrameTime InTime, CurveValueType& OutValue);
 };
 
diff --git a/Engine/Source/Runtime/MovieSceneTracks/Private/Systems/DoubleChannelEvaluatorSystem.cpp b/Engine/Source/Runtime/MovieSceneTracks/Private/Systems/DoubleChannelEvaluatorSystem.cpp
--- a/Engine/Source/Runtime/MovieSceneTracks/Private/Systems/DoubleChannelEvaluatorSystem.cpp
+++ b/Engine/Source/Runtime/MovieSceneTracks/Private/Systems/DoubleChannelEvaluatorSystem.cpp
@@ -23,8 +23,6 @@ DECLARE_CYCLE_STAT(TEXT("MovieScene: Evaluate double channels"), MovieSceneEval_
 namespace UE::MovieScene
 {
 
-MOVIESCENE_API extern bool GEnableCachedChannelEvaluation;
-
 struct FDoubleChannelTypeAssociation
 {
 	TComponentTypeID<FSourceDoubleChannel> ChannelType;
@@ -37,18 +35,6 @@ TArray<FDoubleChannelTypeAssociation, TInlineAllocator<4>> GDoubleChannelTypeAss
 // @todo: for multi-bindings we currently re-evaluate the double channel for each binding, even though the time is the same.
 // Do we need to optimize for this case using something like the code below, while pessimizing the common (non-multi-bind) codepath??
 
-/** Entity-component task that evaluates using the non-cached codepath for testing parity */
-struct FEvaluateDoubleChannels_Uncached
-{
-	void ForEachEntity(FSourceDoubleChannel DoubleChannel, FFrameTime FrameTime, double& OutResult) const
-	{
-		if (!DoubleChannel.Source->Evaluate(FrameTime, OutResult))
-		{
-			OutResult = MIN_dbl;
-		}
-	}
-};
-
 /** Entity-component task that evaluates using a cached interpolation if possible */
 struct FEvaluateDoubleChannels_Cached
 {
@@ -176,28 +162,15 @@ void UDoubleChannelEvaluatorSystem::OnSchedulePersistentTasks(UE::MovieScene::IE
 
 	for (const FDoubleChannelTypeAssociation& ChannelType : GDoubleChannelTypeAssociations)
 	{
-		if (GEnableCachedChannelEvaluation)
-		{
-			// Evaluate double channels per instance and write the evaluated value into the output
-			FEntityTaskBuilder()
-			.Read(ChannelType.ChannelType)
-			.Read(BuiltInComponents->EvalTime)
-			.Write(ChannelType.CachedInterpolationType)
-			.Write(ChannelType.ResultType)
-			.FilterNone({ BuiltInComponents->Tags.Ignored })
-			.SetStat(GET_STATID(MovieSceneEval_EvaluateDoubleChannelTask))
-			.Fork_PerEntity<FEvaluateDoubleChannels_Cached>(&Linker->EntityManager, TaskScheduler);
-		}
-		else
-		{
-			FEntityTaskBuilder()
-			.Read(ChannelType.ChannelType)
-			.Read(BuiltInComponents->EvalTime)
-			.Write(ChannelType.ResultType)
-			.FilterNone({ BuiltInComponents->Tags.Ignored })
-			.SetStat(GET_STATID(MovieSceneEval_EvaluateDoubleChannelTask))
-			.Fork_PerEntity<FEvaluateDoubleChannels_Uncached>(&Linker->EntityManager, TaskScheduler);
-		}
+		// Evaluate double channels per instance and write the evaluated value into the output
+		FEntityTaskBuilder()
+		.Read(ChannelType.ChannelType)
+		.Read(BuiltInComponents->EvalTime)
+		.Write(ChannelType.CachedInterpolationType)
+		.Write(ChannelType.ResultType)
+		.FilterNone({ BuiltInComponents->Tags.Ignored })
+		.SetStat(GET_STATID(MovieSceneEval_EvaluateDoubleChannelTask))
+		.Fork_PerEntity<FEvaluateDoubleChannels_Cached>(&Linker->EntityManager, TaskScheduler);
 	}
 }
 
@@ -225,21 +198,6 @@ void UDoubleChannelEvaluatorSystem::OnRun(FSystemTaskPrerequisites& InPrerequisi
 	}
 	else if (Runner->GetCurrentPhase() == ESystemPhase::Evaluation)
 	{
-		if (!GEnableCachedChannelEvaluation)
-		{
-			for (const FDoubleChannelTypeAssociation& ChannelType : GDoubleChannelTypeAssociations)
-			{
-				FEntityTaskBuilder()
-				.Read(ChannelType.ChannelType)
-				.Read(BuiltInComponents->EvalTime)
-				.Write(ChannelType.ResultType)
-				.FilterNone({ BuiltInComponents->Tags.Ignored })
-				.SetStat(GET_STATID(MovieSceneEval_EvaluateDoubleChannelTask))
-				.Dispatch_PerEntity<FEvaluateDoubleChannels_Uncached>(&Linker->EntityManager, InPrerequisites, &Subsequents);
-			}
-			return;
-		}
-
 		for (const FDoubleChannelTypeAssociation& ChannelType : GDoubleChannelTypeAssociations)
 		{
 			// Evaluate double channels per instance and write the evaluated value into the output
diff --git a/Engine/Source/Runtime/MovieSceneTracks/Private/Systems/FloatChannelEvaluatorSystem.cpp b/Engine/Source/Runtime/MovieSceneTracks/Private/Systems/FloatChannelEvaluatorSystem.cpp
--- a/Engine/Source/Runtime/MovieSceneTracks/Private/Systems/FloatChannelEvaluatorSystem.cpp
+++ b/Engine/Source/Runtime/MovieSceneTracks/Private/Systems/FloatChannelEvaluatorSystem.cpp
@@ -23,8 +23,6 @@ DECLARE_CYCLE_STAT(TEXT("MovieScene: Evaluate float channels"), MovieSceneEval_E
 namespace UE::MovieScene
 {
 
-MOVIESCENE_API extern bool GEnableCachedChannelEvaluation;
-
 struct FFloatChannelTypeAssociation
 {
 	TComponentTypeID<FSourceFloatChannel> ChannelType;
@@ -37,23 +35,6 @@ TArray<FFloatChannelTypeAssociation, TInlineAllocator<4>> GFloatChannelTypeAssoc
 // @todo: for multi-bindings we currently re-evaluate the float channel for each binding, even though the time is the same.
 // Do we need to optimize for this case using something like the code below, while pessimizing the common (non-multi-bind) codepath??
 
-/** Entity-component task that evaluates using the non-cached codepath for testing parity */
-struct FEvaluateFloatChannels_Uncached
-{
-	static void ForEachEntity(FSourceFloatChannel FloatChannel, FFrameTime FrameTime, double& OutResult)
-	{
-		float Result;
-		if (FloatChannel.Source->Evaluate(FrameTime, Result))
-		{
-			OutResult = Result;
-		}
-		else
-		{
-			OutResult = MIN_dbl;
-		}
-	}
-};
-
 /** Entity-component task that evaluates using a cached interpolation if possible */
 struct FEvaluateFloatChannels_Cached
 {
@@ -184,33 +165,18 @@ void UFloatChannelEvaluatorSystem::OnSchedulePersistentTasks(UE::MovieScene::IEn
 	using namespace UE::MovieScene;
 
 	FBuiltInComponentTypes* BuiltInComponents = FBuiltInComponentTypes::Get();
-	if (GEnableCachedChannelEvaluation)
-	{
-		for (const FFloatChannelTypeAssociation& ChannelType : GFloatChannelTypeAssociations)
-		{
-			// Evaluate float channels per instance and write the evaluated value into the output
-			FEntityTaskBuilder()
-			.Read(ChannelType.ChannelType)
-			.Read(BuiltInComponents->EvalTime)
-			.Write(ChannelType.CachedInterpolationType)
-			.Write(ChannelType.ResultType)
-			.FilterNone({ BuiltInComponents->Tags.Ignored })
-			.SetStat(GET_STATID(MovieSceneEval_EvaluateFloatChannelTask))
-			.Fork_PerEntity<FEvaluateFloatChannels_Cached>(&Linker->EntityManager, TaskScheduler);
-		}
-	}
-	else
+
+	for (const FFloatChannelTypeAssociation& ChannelType : GFloatChannelTypeAssociations)
 	{
-		for (const FFloatChannelTypeAssociation& ChannelType : GFloatChannelTypeAssociations)
-		{
-			FEntityTaskBuilder()
-			.Read(ChannelType.ChannelType)
-			.Read(BuiltInComponents->EvalTime)
-			.Write(ChannelType.ResultType)
-			.FilterNone({ BuiltInComponents->Tags.Ignored })
-			.SetStat(GET_STATID(MovieSceneEval_EvaluateFloatChannelTask))
-			.Fork_PerEntity<FEvaluateFloatChannels_Uncached>(&Linker->EntityManager, TaskScheduler);
-		}
+		// Evaluate float channels per instance and write the evaluated value into the output
+		FEntityTaskBuilder()
+		.Read(ChannelType.ChannelType)
+		.Read(BuiltInComponents->EvalTime)
+		.Write(ChannelType.CachedInterpolationType)
+		.Write(ChannelType.ResultType)
+		.FilterNone({ BuiltInComponents->Tags.Ignored })
+		.SetStat(GET_STATID(MovieSceneEval_EvaluateFloatChannelTask))
+		.Fork_PerEntity<FEvaluateFloatChannels_Cached>(&Linker->EntityManager, TaskScheduler);
 	}
 }
 
@@ -238,21 +204,6 @@ void UFloatChannelEvaluatorSystem::OnRun(FSystemTaskPrerequisites& InPrerequisit
 	}
 	else if (Runner->GetCurrentPhase() == ESystemPhase::Evaluation)
 	{
-		if (!GEnableCachedChannelEvaluation)
-		{
-			for (const FFloatChannelTypeAssociation& ChannelType : GFloatChannelTypeAssociations)
-			{
-				FEntityTaskBuilder()
-				.Read(ChannelType.ChannelType)
-				.Read(BuiltInComponents->EvalTime)
-				.Write(ChannelType.ResultType)
-				.FilterNone({ BuiltInComponents->Tags.Ignored })
-				.SetStat(GET_STATID(MovieSceneEval_EvaluateFloatChannelTask))
-				.Dispatch_PerEntity<FEvaluateFloatChannels_Uncached>(&Linker->EntityManager, InPrerequisites, &Subsequents);
-			}
-			return;
-		}
-
 		for (const FFloatChannelTypeAssociation& ChannelType : GFloatChannelTypeAssociations)
 		{
 			// Evaluate float channels per instance and write the evaluated value into the output
