Sequencer: Fixed integer math for oscillation with large ranges

diff --git a/Engine/Source/Runtime/MovieScene/Public/Channels/MovieSceneChannel.h b/Engine/Source/Runtime/MovieScene/Public/Channels/MovieSceneChannel.h
--- a/Engine/Source/Runtime/MovieScene/Public/Channels/MovieSceneChannel.h
+++ b/Engine/Source/Runtime/MovieScene/Public/Channels/MovieSceneChannel.h
@@ -85,7 +85,9 @@ struct FCycleParams
 		if (FMath::Abs(CycleCount) % 2 == 1)
 		{
 			bMirrorCurve = true;
-			Time = MinFrame + (FFrameTime(MaxFrame) - Time);
+
+			const int64 Offset = int64(MaxFrame) - int64(Time.FrameNumber.Value);
+			Time = FFrameTime(MinFrame) + FFrameTime(static_cast<int32>(Offset)) - FFrameTime(0, Time.GetSubFrame());
 		}
 	}
 };
