Sequencer: Fix regression where externally owned spawnables were not persisting beyond the end of the sequence.

diff --git a/Engine/Source/Runtime/MovieScene/Private/MovieSceneSpawnRegister.cpp b/Engine/Source/Runtime/MovieScene/Private/MovieSceneSpawnRegister.cpp
--- a/Engine/Source/Runtime/MovieScene/Private/MovieSceneSpawnRegister.cpp
+++ b/Engine/Source/Runtime/MovieScene/Private/MovieSceneSpawnRegister.cpp
@@ -66,6 +66,8 @@ UObject* FMovieSceneSpawnRegister::SpawnObject(const FGuid& BindingId, UMovieSce
 						return nullptr;
 					}
 
+					SpawnOwnership = SpawnableBinding->SpawnOwnership;
+
 					// Call the Spawnable binding itself to spawn the object
 					SpawnedActor = SpawnableBinding->SpawnObject(BindingId, BindingIndex, MovieScene, TemplateID, SharedPlaybackState);
 				}
