Sequencer: Added an option to determinism fences for the frame to be inclusive

This is useful for test cases, or situations where content requires that a specific frame must be evaluated precisely.

diff --git a/Engine/Source/Runtime/MovieScene/Private/Compilation/MovieSceneCompiledDataManager.cpp b/Engine/Source/Runtime/MovieScene/Private/Compilation/MovieSceneCompiledDataManager.cpp
--- a/Engine/Source/Runtime/MovieScene/Private/Compilation/MovieSceneCompiledDataManager.cpp
+++ b/Engine/Source/Runtime/MovieScene/Private/Compilation/MovieSceneCompiledDataManager.cpp
@@ -775,7 +775,7 @@ void UMovieSceneCompiledDataManager::Compile(FMovieSceneCompiledDataID DataID, U
 			{
 				if (Mark.bIsDeterminismFence)
 				{
-					GatheredData.DeterminismData.Fences.Add(Mark.FrameNumber);
+					GatheredData.DeterminismData.Fences.Emplace(Mark.FrameNumber, Mark.bIsInclusiveTime);
 				}
 			}
 
@@ -860,7 +860,7 @@ void UMovieSceneCompiledDataManager::Compile(FMovieSceneCompiledDataID DataID, U
 	Entry.DeterminismFences = MoveTemp(GatheredData.DeterminismData.Fences);
 	if (Entry.DeterminismFences.Num())
 	{
-		Algo::Sort(Entry.DeterminismFences);
+		Algo::SortBy(Entry.DeterminismFences, &FMovieSceneDeterminismFence::FrameNumber);
 		const int32 NewNum = Algo::Unique(Entry.DeterminismFences);
 		if (NewNum != Entry.DeterminismFences.Num())
 		{
diff --git a/Engine/Source/Runtime/MovieScene/Private/EntitySystem/MovieSceneSequenceUpdaters.cpp b/Engine/Source/Runtime/MovieScene/Private/EntitySystem/MovieSceneSequenceUpdaters.cpp
--- a/Engine/Source/Runtime/MovieScene/Private/EntitySystem/MovieSceneSequenceUpdaters.cpp
+++ b/Engine/Source/Runtime/MovieScene/Private/EntitySystem/MovieSceneSequenceUpdaters.cpp
@@ -29,6 +29,12 @@ namespace UE
 namespace MovieScene
 {
 
+struct FMovieSceneDeterminismFenceWithSubframe
+{
+	FFrameTime FrameTime;
+	uint8 bInclusive : 1;
+};
+
 /** Flat sequence updater (ie, no hierarchy) */
 struct FSequenceUpdater_Flat : ISequenceUpdater
 {
@@ -51,7 +57,7 @@ private:
 
 	TRange<FFrameNumber> CachedEntityRange;
 
-	TOptional<TArray<FFrameTime>> CachedDeterminismFences;
+	TOptional<TArray<FMovieSceneDeterminismFence>> CachedDeterminismFences;
 	FMovieSceneCompiledDataID CompiledDataID;
 
 	TOptional<bool> bDynamicWeighting;
@@ -100,7 +106,41 @@ private:
 	TOptional<bool> bDynamicWeighting;
 };
 
-void DissectRange(TArrayView<const FFrameTime> InDissectionTimes, const TRange<FFrameTime>& Bounds, TArray<TRange<FFrameTime>>& OutDissections)
+void DissectRange(TArrayView<const FMovieSceneDeterminismFenceWithSubframe> InDissectionTimes, const TRange<FFrameTime>& Bounds, TArray<TRange<FFrameTime>>& OutDissections)
+{
+	if (InDissectionTimes.Num() == 0)
+	{
+		return;
+	}
+
+	TRangeBound<FFrameTime> LowerBound = Bounds.GetLowerBound();
+
+	for (int32 Index = 0; Index < InDissectionTimes.Num(); ++Index)
+	{
+		FMovieSceneDeterminismFenceWithSubframe DissectionFence = InDissectionTimes[Index];
+
+		TRange<FFrameTime> Dissection = DissectionFence.bInclusive
+			? TRange<FFrameTime>(LowerBound, TRangeBound<FFrameTime>::Inclusive(DissectionFence.FrameTime))
+			: TRange<FFrameTime>(LowerBound, TRangeBound<FFrameTime>::Exclusive(DissectionFence.FrameTime));
+
+		if (!Dissection.IsEmpty())
+		{
+			ensureAlwaysMsgf(Bounds.Contains(Dissection), TEXT("Dissection specified for a range outside of the current bounds"));
+
+			OutDissections.Add(Dissection);
+
+			LowerBound = TRangeBound<FFrameTime>::FlipInclusion(Dissection.GetUpperBound());
+		}
+	}
+
+	TRange<FFrameTime> TailRange(LowerBound, Bounds.GetUpperBound());
+	if (!TailRange.IsEmpty())
+	{
+		OutDissections.Add(TailRange);
+	}
+}
+
+void DissectRange(TArrayView<const FMovieSceneDeterminismFence> InDissectionTimes, const TRange<FFrameTime>& Bounds, TArray<TRange<FFrameTime>>& OutDissections)
 {
 	if (InDissectionTimes.Num() == 0)
 	{
@@ -111,16 +151,19 @@ void DissectRange(TArrayView<const FFrameTime> InDissectionTimes, const TRange<F
 
 	for (int32 Index = 0; Index < InDissectionTimes.Num(); ++Index)
 	{
-		FFrameTime DissectionTime = InDissectionTimes[Index];
+		FMovieSceneDeterminismFence DissectionFence = InDissectionTimes[Index];
+
+		TRange<FFrameTime> Dissection = DissectionFence.bInclusive
+			? TRange<FFrameTime>(LowerBound, TRangeBound<FFrameTime>::Inclusive(DissectionFence.FrameNumber))
+			: TRange<FFrameTime>(LowerBound, TRangeBound<FFrameTime>::Exclusive(DissectionFence.FrameNumber));
 
-		TRange<FFrameTime> Dissection(LowerBound, TRangeBound<FFrameTime>::Exclusive(DissectionTime));
 		if (!Dissection.IsEmpty())
 		{
 			ensureAlwaysMsgf(Bounds.Contains(Dissection), TEXT("Dissection specified for a range outside of the current bounds"));
 
 			OutDissections.Add(Dissection);
 
-			LowerBound = TRangeBound<FFrameTime>::Inclusive(DissectionTime);
+			LowerBound = TRangeBound<FFrameTime>::FlipInclusion(Dissection.GetUpperBound());
 		}
 	}
 
@@ -131,25 +174,35 @@ void DissectRange(TArrayView<const FFrameTime> InDissectionTimes, const TRange<F
 	}
 }
 
-TArrayView<const FFrameTime> GetFencesWithinRange(TArrayView<const FFrameTime> Fences, const TRange<FFrameTime>& Boundary)
+TArrayView<const FMovieSceneDeterminismFence> GetFencesWithinRange(TArrayView<const FMovieSceneDeterminismFence> Fences, const TRange<FFrameTime>& Boundary)
 {
 	if (Fences.Num() == 0 || Boundary.IsEmpty())
 	{
-		return TArrayView<const FFrameTime>();
+		return TArrayView<const FMovieSceneDeterminismFence>();
 	}
 
-	// Take care to include or exclude the lower bound of the range if it's on a whole frame numbe
-	const int32 StartFence = Boundary.GetLowerBound().IsClosed() ? Algo::UpperBound(Fences, Boundary.GetLowerBoundValue()) : 0;
+	// Take care to include or exclude the lower bound of the range if it's on a whole frame number
+	const int32 StartFence = Boundary.GetLowerBound().IsOpen()
+		? 0
+		: Boundary.GetLowerBound().IsInclusive() && Boundary.GetLowerBoundValue().GetSubFrame() == 0.0
+			? Algo::LowerBoundBy(Fences, Boundary.GetLowerBoundValue().FrameNumber, &FMovieSceneDeterminismFence::FrameNumber)
+			: Algo::UpperBoundBy(Fences, Boundary.GetLowerBoundValue().FrameNumber, &FMovieSceneDeterminismFence::FrameNumber);
+
 	if (StartFence >= Fences.Num())
 	{
-		return TArrayView<const FFrameTime>();
+		return TArrayView<const FMovieSceneDeterminismFence>();
 	}
 
-	const int32 EndFence = Boundary.GetUpperBound().IsClosed() ? Algo::UpperBound(Fences, Boundary.GetUpperBoundValue()) : Fences.Num();
+	const int32 EndFence = Boundary.GetUpperBound().IsOpen()
+		? 0
+		: Boundary.GetUpperBound().IsInclusive() && Boundary.GetUpperBoundValue().GetSubFrame() == 0.0
+			? Algo::LowerBoundBy(Fences, Boundary.GetUpperBoundValue().FrameNumber, &FMovieSceneDeterminismFence::FrameNumber)
+			: Algo::UpperBoundBy(Fences, Boundary.GetUpperBoundValue().FrameNumber, &FMovieSceneDeterminismFence::FrameNumber);
+
 	const int32 NumFences = FMath::Max(0, EndFence - StartFence);
 	if (NumFences == 0)
 	{
-		return TArrayView<const FFrameTime>();
+		return TArrayView<const FMovieSceneDeterminismFence>();
 	}
 
 	return MakeArrayView(Fences.GetData() + StartFence, NumFences);
@@ -200,12 +253,12 @@ void FSequenceUpdater_Flat::PopulateUpdateFlags(TSharedRef<const FSharedPlayback
 {
 	if (!CachedDeterminismFences.IsSet())
 	{
-		UMovieSceneCompiledDataManager* CompiledDataManager = SharedPlaybackState->GetCompiledDataManager();
-		TArrayView<const FFrameTime>    DeterminismFences   = CompiledDataManager->GetEntryRef(CompiledDataID).DeterminismFences;
+		UMovieSceneCompiledDataManager*               CompiledDataManager = SharedPlaybackState->GetCompiledDataManager();
+		TArrayView<const FMovieSceneDeterminismFence> DeterminismFences   = CompiledDataManager->GetEntryRef(CompiledDataID).DeterminismFences;
 
 		if (DeterminismFences.Num() != 0)
 		{
-			CachedDeterminismFences = TArray<FFrameTime>(DeterminismFences.GetData(), DeterminismFences.Num());
+			CachedDeterminismFences = TArray<FMovieSceneDeterminismFence>(DeterminismFences.GetData(), DeterminismFences.Num());
 		}
 		else
 		{
@@ -228,12 +281,12 @@ void FSequenceUpdater_Flat::DissectContext(TSharedRef<const FSharedPlaybackState
 {
 	if (!CachedDeterminismFences.IsSet())
 	{
-		UMovieSceneCompiledDataManager* CompiledDataManager = SharedPlaybackState->GetCompiledDataManager();
-		TArrayView<const FFrameTime>    DeterminismFences   = CompiledDataManager->GetEntryRef(CompiledDataID).DeterminismFences;
+		UMovieSceneCompiledDataManager*               CompiledDataManager = SharedPlaybackState->GetCompiledDataManager();
+		TArrayView<const FMovieSceneDeterminismFence> DeterminismFences   = CompiledDataManager->GetEntryRef(CompiledDataID).DeterminismFences;
 
 		if (DeterminismFences.Num() != 0)
 		{
-			CachedDeterminismFences = TArray<FFrameTime>(DeterminismFences.GetData(), DeterminismFences.Num());
+			CachedDeterminismFences = TArray<FMovieSceneDeterminismFence>(DeterminismFences.GetData(), DeterminismFences.Num());
 		}
 		else
 		{
@@ -243,7 +296,7 @@ void FSequenceUpdater_Flat::DissectContext(TSharedRef<const FSharedPlaybackState
 
 	if (CachedDeterminismFences->Num() != 0)
 	{
-		TArrayView<const FFrameTime> TraversedFences = GetFencesWithinRange(CachedDeterminismFences.GetValue(), Context.GetRange());
+		TArrayView<const FMovieSceneDeterminismFence> TraversedFences = GetFencesWithinRange(CachedDeterminismFences.GetValue(), Context.GetRange());
 		UE::MovieScene::DissectRange(TraversedFences, Context.GetRange(), OutDissections);
 	}
 }
@@ -407,12 +460,12 @@ void FSequenceUpdater_Hierarchical::DissectContext(TSharedRef<const FSharedPlayb
 {
 	UMovieSceneCompiledDataManager* CompiledDataManager = SharedPlaybackState->GetCompiledDataManager();
 
-	TRange<FFrameNumber> TraversedRange = Context.GetFrameNumberRange();
-	TArray<FFrameTime>   RootDissectionTimes;
+	TRange<FFrameNumber>                TraversedRange = Context.GetFrameNumberRange();
+	TArray<FMovieSceneDeterminismFenceWithSubframe> RootDissectionTimes;
 
 	{
-		const FMovieSceneCompiledDataEntry& DataEntry       = CompiledDataManager->GetEntryRef(CompiledDataID);
-		TArrayView<const FFrameTime>        TraversedFences = GetFencesWithinRange(DataEntry.DeterminismFences, Context.GetRange());
+		const FMovieSceneCompiledDataEntry&           DataEntry       = CompiledDataManager->GetEntryRef(CompiledDataID);
+		TArrayView<const FMovieSceneDeterminismFence> TraversedFences = GetFencesWithinRange(DataEntry.DeterminismFences, Context.GetRange());
 
 		UE::MovieScene::DissectRange(TraversedFences, Context.GetRange(), OutDissections);
 	}
@@ -428,7 +481,7 @@ void FSequenceUpdater_Hierarchical::DissectContext(TSharedRef<const FSharedPlayb
 			// When Context.GetRange() does not fall on whole frame boundaries, we can sometimes end up with a range that clamps to being empty, even though the range overlapped
 			// the traversed range. ie if we evaluated range (1.5, 10], our traversed range would be [2, 11). If we have a sub sequence range of (10, 20), it would still be iterated here
 			// because [2, 11) overlaps (10, 20), but when clamped to the evaluated range, the range is (10, 10], which is empty.
-			if (!RootClampRange.IsEmpty())
+			if (RootClampRange.IsEmpty())
 			{
 				continue;
 			}
@@ -445,12 +498,12 @@ void FSequenceUpdater_Hierarchical::DissectContext(TSharedRef<const FSharedPlayb
 					continue;
 				}
 
-				TArrayView<const FFrameTime> SubDeterminismFences = CompiledDataManager->GetEntryRef(SubDataID).DeterminismFences;
+				TArrayView<const FMovieSceneDeterminismFence> SubDeterminismFences = CompiledDataManager->GetEntryRef(SubDataID).DeterminismFences;
 				if (SubDeterminismFences.Num() > 0)
 				{
 					TRange<FFrameTime>   InnerRange           = SubData->RootToSequenceTransform.TransformRangeUnwarped(RootClampRange);
 
-					TArrayView<const FFrameTime> TraversedFences  = GetFencesWithinRange(SubDeterminismFences, InnerRange);
+					TArrayView<const FMovieSceneDeterminismFence> TraversedFences  = GetFencesWithinRange(SubDeterminismFences, InnerRange);
 					if (TraversedFences.Num() > 0)
 					{
 						FMovieSceneWarpCounter WarpCounter;
@@ -458,7 +511,7 @@ void FSequenceUpdater_Hierarchical::DissectContext(TSharedRef<const FSharedPlayb
 						SubData->RootToSequenceTransform.TransformTime(RootClampRange.GetLowerBoundValue(), Unused, WarpCounter);
 
 						FMovieSceneSequenceTransform InverseTransform = SubData->RootToSequenceTransform.InverseFromLoop(WarpCounter);
-						Algo::Transform(TraversedFences, RootDissectionTimes, [InverseTransform](FFrameTime In){ return In * InverseTransform; });
+						Algo::Transform(TraversedFences, RootDissectionTimes, [InverseTransform](const FMovieSceneDeterminismFence& In){ return FMovieSceneDeterminismFenceWithSubframe{In.FrameNumber * InverseTransform, In.bInclusive }; });
 					}
 				}
 			}
@@ -467,7 +520,7 @@ void FSequenceUpdater_Hierarchical::DissectContext(TSharedRef<const FSharedPlayb
 
 	if (RootDissectionTimes.Num() > 0)
 	{
-		Algo::Sort(RootDissectionTimes);
+		Algo::SortBy(RootDissectionTimes, &FMovieSceneDeterminismFenceWithSubframe::FrameTime);
 		UE::MovieScene::DissectRange(RootDissectionTimes, Context.GetRange(), OutDissections);
 	}
 }
diff --git a/Engine/Source/Runtime/MovieScene/Public/Compilation/IMovieSceneDeterminismSource.h b/Engine/Source/Runtime/MovieScene/Public/Compilation/IMovieSceneDeterminismSource.h
--- a/Engine/Source/Runtime/MovieScene/Public/Compilation/IMovieSceneDeterminismSource.h
+++ b/Engine/Source/Runtime/MovieScene/Public/Compilation/IMovieSceneDeterminismSource.h
@@ -5,6 +5,7 @@
 #include "UObject/Interface.h"
 #include "Containers/Array.h"
 #include "Misc/FrameTime.h"
+#include "Compilation/MovieSceneDeterminismFence.h"
 #include "IMovieSceneDeterminismSource.generated.h"
 
 /** Determinism data that is generated on compile of a UMovieSceneSequence */
@@ -18,7 +19,7 @@ struct FMovieSceneDeterminismData
 	 * Duplicates are allowed during compilation, but will be removed in the final compiled data.
 	 */
 	UPROPERTY()
-	TArray<FFrameTime> Fences;
+	TArray<FMovieSceneDeterminismFence> Fences;
 
 	/** True if this sequence should include a fence on the lower bound of any sub sequence's that include it */
 	UPROPERTY()
diff --git a/Engine/Source/Runtime/MovieScene/Public/Compilation/MovieSceneCompiledDataManager.h b/Engine/Source/Runtime/MovieScene/Public/Compilation/MovieSceneCompiledDataManager.h
--- a/Engine/Source/Runtime/MovieScene/Public/Compilation/MovieSceneCompiledDataManager.h
+++ b/Engine/Source/Runtime/MovieScene/Public/Compilation/MovieSceneCompiledDataManager.h
@@ -3,6 +3,7 @@
 #pragma once
 
 #include "Compilation/MovieSceneCompiledDataID.h"
+#include "Compilation/MovieSceneDeterminismFence.h"
 #include "Containers/Array.h"
 #include "Containers/Map.h"
 #include "Containers/Set.h"
@@ -112,7 +113,7 @@ private:
 
 	/** 16 Bytes */
 	UPROPERTY()
-	TArray<FFrameTime> DeterminismFences;
+	TArray<FMovieSceneDeterminismFence> DeterminismFences;
 
 	/** 16 bytes */
 	UPROPERTY()
@@ -149,7 +150,7 @@ struct FMovieSceneCompiledDataEntry
 	FGuid CompiledSignature;
 
 	/** 16 Bytes */
-	TArray<FFrameTime> DeterminismFences;
+	TArray<FMovieSceneDeterminismFence> DeterminismFences;
 
 	/** 8 Bytes */
 	FObjectKey SequenceKey;
diff --git a/Engine/Source/Runtime/MovieScene/Public/Compilation/MovieSceneDeterminismFence.h b/Engine/Source/Runtime/MovieScene/Public/Compilation/MovieSceneDeterminismFence.h
new file mode 100644
--- /dev/null
+++ b/Engine/Source/Runtime/MovieScene/Public/Compilation/MovieSceneDeterminismFence.h
@@ -0,0 +1,44 @@
+// Copyright Epic Games, Inc. All Rights Reserved.
+
+#pragma once
+
+#include "Misc/FrameNumber.h"
+#include "MovieSceneDeterminismFence.generated.h"
+
+
+/** Structure that defines a specific determinism fence */
+USTRUCT()
+struct FMovieSceneDeterminismFence
+{
+	GENERATED_BODY()
+
+	FMovieSceneDeterminismFence()
+		: bInclusive(false)
+	{}
+
+	FMovieSceneDeterminismFence(FFrameNumber InFrameNumber, bool bInInclusive = false)
+		: FrameNumber(InFrameNumber)
+		, bInclusive(bInInclusive)
+	{}
+
+	friend bool operator==(FMovieSceneDeterminismFence A, FMovieSceneDeterminismFence B)
+	{
+		return A.FrameNumber == B.FrameNumber && A.bInclusive == B.bInclusive;
+	}
+
+	/** True if this sequence should include a fence on the lower bound of any sub sequence's that include it */
+	UPROPERTY()
+	FFrameNumber FrameNumber;
+
+	/**
+	 * Default: false. When true, specifies that this fence should be evaluated exactly on the specified time.
+	 * When false, all times up to, but not including FrameNumber will be evaluated.
+	 * 
+	 * Exclusive should be used for a fence at the end of a subsection with an exclusive time to ensure that the
+	 *     sub-section is entirely evaluated before evaluation returns.
+	 * 
+	 * Inclusive should be used if an exact time must be evaluated (such as for testing purposes)
+	 */
+	UPROPERTY()
+	uint8 bInclusive : 1;
+};
\ No newline at end of file
diff --git a/Engine/Source/Runtime/MovieScene/Public/MovieSceneMarkedFrame.h b/Engine/Source/Runtime/MovieScene/Public/MovieSceneMarkedFrame.h
--- a/Engine/Source/Runtime/MovieScene/Public/MovieSceneMarkedFrame.h
+++ b/Engine/Source/Runtime/MovieScene/Public/MovieSceneMarkedFrame.h
@@ -20,6 +20,7 @@ struct FMovieSceneMarkedFrame
 		, CustomColor(0.f, 1.f, 1.f, 0.4f)	
 #endif
 		, bIsDeterminismFence(false)
+		, bIsInclusiveTime(false)
 	{}
 
 	FMovieSceneMarkedFrame(FFrameNumber InFrameNumber)
@@ -30,6 +31,7 @@ struct FMovieSceneMarkedFrame
 		, CustomColor(0.f, 1.f, 1.f, 0.4f)
 #endif
 		, bIsDeterminismFence(false)
+		, bIsInclusiveTime(false)
 	{}
 
 	UPROPERTY(EditAnywhere, BlueprintReadWrite, Category = "Marked Frame")
@@ -54,4 +56,7 @@ struct FMovieSceneMarkedFrame
 
 	UPROPERTY(EditAnywhere, BlueprintReadWrite, Category = "Marked Frame", DisplayName="Is Determinism Fence?", meta=(Tooltip="When checked, treat this mark as a fence for evaluation purposes. Fences cannot be crossed in a single evaluation, and force the evaluation to be split into 2 separate parts."))
 	bool bIsDeterminismFence;
+
+	UPROPERTY(EditAnywhere, BlueprintReadWrite, Category = "Marked Frame", DisplayName="Inclusive?", meta=(EditCondition="bIsDeterminismFence", Tooltip="Defines how this determinism fence determines the previous and next range: when true, the range will be dissected as (X, Y] -> (Y, Z], when false it will be (X, Y) -> [Y, Z]."))
+	bool bIsInclusiveTime;
 };
