Sequencer: Added time-warp capabilities to sequences, sub-sections and skeletal animation sections

diff --git a/Engine/Source/Runtime/MovieScene/Private/Channels/MovieSceneCurveChannelImpl.cpp b/Engine/Source/Runtime/MovieScene/Private/Channels/MovieSceneCurveChannelImpl.cpp
--- a/Engine/Source/Runtime/MovieScene/Private/Channels/MovieSceneCurveChannelImpl.cpp
+++ b/Engine/Source/Runtime/MovieScene/Private/Channels/MovieSceneCurveChannelImpl.cpp
@@ -89,7 +89,13 @@ static typename ChannelType::CurveValueType
 
 FCycleParams CycleTime(FFrameNumber MinFrame, FFrameNumber MaxFrame, FFrameTime InTime)
 {
-	FCycleParams Params(InTime, MaxFrame.Value - MinFrame.Value);
+	int64 Duration = int64(MaxFrame.Value) - int64(MinFrame.Value);
+	if (Duration > MAX_int32)
+	{
+		return FCycleParams(InTime, 0);
+	}
+
+	FCycleParams Params(InTime, static_cast<int32>(Duration));
 	if (Params.Duration == 0)
 	{
 		Params.Time = MaxFrame;
@@ -97,17 +103,19 @@ FCycleParams CycleTime(FFrameNumber MinFrame, FFrameNumber MaxFrame, FFrameTime
 	}
 	else if (InTime < MinFrame)
 	{
-		const int32 CycleCount = ((MaxFrame - InTime) / Params.Duration).FloorToFrame().Value;
+		const int64 CycleCount = (int64(MaxFrame.Value) - int64(InTime.FrameNumber.Value)) / Params.Duration;
 
-		Params.CycleCount = -CycleCount;
-		Params.Time = InTime + FFrameTime(Params.Duration * CycleCount);
+		Params.CycleCount = static_cast<int32>(-CycleCount);
+		Params.Time = InTime;
+		Params.Time.FrameNumber.Value = static_cast<int32>(int64(Params.Time.FrameNumber.Value) + int64(Params.Duration)*CycleCount);
 	}
 	else if (InTime > MaxFrame)
 	{
-		const int32 CycleCount = ((InTime - MinFrame) / Params.Duration).FloorToFrame().Value;
+		const int64 CycleCount = (int64(InTime.FrameNumber.Value) - int64(MinFrame.Value)) / Params.Duration;
 
-		Params.CycleCount = CycleCount;
-		Params.Time = InTime - FFrameTime(Params.Duration * CycleCount);
+		Params.CycleCount = static_cast<int32>(CycleCount);
+		Params.Time = InTime;
+		Params.Time.FrameNumber.Value = static_cast<int32>(int64(Params.Time.FrameNumber.Value) - int64(Params.Duration) * CycleCount);
 	}
 
 	return Params;
@@ -206,7 +214,7 @@ bool TMovieSceneCurveChannelImpl<ChannelType>::CacheExtrapolation(const ChannelT
 
 		if (InChannel->PreInfinityExtrap == RCCE_Constant)
 		{
-			OutValue = FCachedInterpolation(Range, FConstantValue(InChannel->Values[0].Value));
+			OutValue = FCachedInterpolation(Range, FConstantValue(InChannel->Times[0], InChannel->Values[0].Value));
 			return true;
 		}
 
@@ -216,7 +224,7 @@ bool TMovieSceneCurveChannelImpl<ChannelType>::CacheExtrapolation(const ChannelT
 
 			if (FirstValue.InterpMode == RCIM_Constant)
 			{
-				OutValue = FCachedInterpolation(Range, FConstantValue(FirstValue.Value));
+				OutValue = FCachedInterpolation(Range, FConstantValue(InChannel->Times[0], FirstValue.Value));
 			}
 			else if(FirstValue.InterpMode == RCIM_Cubic)
 			{
@@ -246,7 +254,7 @@ bool TMovieSceneCurveChannelImpl<ChannelType>::CacheExtrapolation(const ChannelT
 
 		if (InChannel->PostInfinityExtrap == RCCE_Constant)
 		{
-			OutValue = FCachedInterpolation(Range, FConstantValue(InChannel->Values.Last().Value));
+			OutValue = FCachedInterpolation(Range, FConstantValue(InChannel->Times.Last(), InChannel->Values.Last().Value));
 			return true;
 		}
 
@@ -256,7 +264,7 @@ bool TMovieSceneCurveChannelImpl<ChannelType>::CacheExtrapolation(const ChannelT
 
 			if (LastValue.InterpMode == RCIM_Constant)
 			{
-				OutValue = FCachedInterpolation(Range, FConstantValue(LastValue.Value));
+				OutValue = FCachedInterpolation(Range, FConstantValue(InChannel->Times.Last(), LastValue.Value));
 			}
 			else if(LastValue.InterpMode == RCIM_Cubic)
 			{
@@ -300,7 +308,7 @@ UE::MovieScene::Interpolation::FCachedInterpolation TMovieSceneCurveChannelImpl<
 	{
 		if (InChannel->bHasDefaultValue)
 		{
-			return FCachedInterpolation(FCachedInterpolationRange::Infinite(), FConstantValue(InChannel->DefaultValue));
+			return FCachedInterpolation(FCachedInterpolationRange::Infinite(), FConstantValue(0, InChannel->DefaultValue));
 		}
 
 		return FCachedInterpolation();
@@ -309,7 +317,7 @@ UE::MovieScene::Interpolation::FCachedInterpolation TMovieSceneCurveChannelImpl<
 	// For single keys, we can only ever return that value
 	if (NumKeys == 1)
 	{
-		return FCachedInterpolation(FCachedInterpolationRange::Infinite(), FConstantValue(InChannel->Values[0].Value));
+		return FCachedInterpolation(FCachedInterpolationRange::Infinite(), FConstantValue(InChannel->Times[0], InChannel->Values[0].Value));
 	}
 
 	// Cache extrapolation if we're outside the bounds of the curve
@@ -381,14 +389,14 @@ UE::MovieScene::Interpolation::FCachedInterpolation TMovieSceneCurveChannelImpl<
 		// No starting key - we are probably evaluating directly on the first or last key
 		//   we explicitly only cache this for the current time to ensure that subsequent caches can cache the correct pair
 		FCachedInterpolationRange Range = FCachedInterpolationRange::Only(Params.Time.GetFrame());
-		return FCachedInterpolation(Range, FConstantValue(Params.ValueOffset + InChannel->Values[Index2].Value));
+		return FCachedInterpolation(Range, FConstantValue(InChannel->Times[Index2], Params.ValueOffset + InChannel->Values[Index2].Value));
 	}
 	else if (Index2 == INDEX_NONE)
 	{
 		// No ending key - we are probably evaluating directly on the first or last key
 		//   we explicitly only cache this for the current time to ensure that subsequent caches can cache the correct pair
 		FCachedInterpolationRange Range = FCachedInterpolationRange::Only(Params.Time.GetFrame());
-		return FCachedInterpolation(Range, FConstantValue(Params.ValueOffset + InChannel->Values[Index1].Value));
+		return FCachedInterpolation(Range, FConstantValue(InChannel->Times[Index1], Params.ValueOffset + InChannel->Values[Index1].Value));
 	}
 	else
 	{
@@ -405,7 +413,7 @@ UE::MovieScene::Interpolation::FCachedInterpolation TMovieSceneCurveChannelImpl<
 	if (Index == InChannel->Times.Num() - 1)
 	{
 		FCachedInterpolationRange Range = FCachedInterpolationRange::Only(InChannel->Times.Last());
-		return FCachedInterpolation(Range, FConstantValue(InChannel->Values.Last().Value));
+		return FCachedInterpolation(Range, FConstantValue(Range.Start, InChannel->Values.Last().Value));
 	}
 
 	return GetInterpolationForKey(InChannel, Index, Index+1, Params);
@@ -495,7 +503,7 @@ UE::MovieScene::Interpolation::FCachedInterpolation TMovieSceneCurveChannelImpl<
 		{
 			if (V1 == V2)
 			{
-				return FCachedInterpolation(Range, FConstantValue(V1));
+				return FCachedInterpolation(Range, FConstantValue(Time1, V1));
 			}
 			// Weighted if either control point has weight on their tangent
 			else if (bIsWeighted1 || bIsWeighted2)
@@ -508,7 +516,7 @@ UE::MovieScene::Interpolation::FCachedInterpolation TMovieSceneCurveChannelImpl<
 			}
 			else
 			{
-				return FCachedInterpolation(Range, FCubicInterpolation(Time1, DX, V1, V2, T1, T2));
+				return FCachedInterpolation(Range, FCubicBezierInterpolation(Time1, DX, V1, V2, T1, T2));
 			}
 		}
 		break;
@@ -518,7 +526,7 @@ UE::MovieScene::Interpolation::FCachedInterpolation TMovieSceneCurveChannelImpl<
 			const double DY = V2 - V1;
 			if (DY == 0.0)
 			{
-				return FCachedInterpolation(Range, FConstantValue(V1));
+				return FCachedInterpolation(Range, FConstantValue(Time1, V1));
 			}
 			else
 			{
@@ -527,7 +535,7 @@ UE::MovieScene::Interpolation::FCachedInterpolation TMovieSceneCurveChannelImpl<
 		}
 
 	default:
-		return FCachedInterpolation(Range, FConstantValue(V1));
+		return FCachedInterpolation(Range, FConstantValue(Time1, V1));
 	}
 }
 
diff --git a/Engine/Source/Runtime/MovieScene/Private/Channels/MovieSceneDoubleChannel.cpp b/Engine/Source/Runtime/MovieScene/Private/Channels/MovieSceneDoubleChannel.cpp
--- a/Engine/Source/Runtime/MovieScene/Private/Channels/MovieSceneDoubleChannel.cpp
+++ b/Engine/Source/Runtime/MovieScene/Private/Channels/MovieSceneDoubleChannel.cpp
@@ -5,6 +5,8 @@
 #include "Channels/MovieSceneCurveChannelImpl.h"
 #include "Channels/MovieSceneFloatChannel.h"
 #include "Channels/MovieSceneInterpolation.h"
+#include "Channels/MovieScenePiecewiseCurve.h"
+#include "Channels/MovieScenePiecewiseCurveUtils.inl"
 #include "HAL/Platform.h"
 #include "MovieSceneFrameMigration.h"
 #include "MovieSceneFwd.h"
@@ -18,6 +20,80 @@ static_assert(
 		"The size of the float channel value has changed. You need to update the padding byte at the end of the structure. "
 		"You also need to update the layout in FMovieSceneDoubleValue so that they match!");
 
+
+namespace UE::MovieScene::Interpolation
+{
+
+	struct FDoubleChannelPiecewiseData
+	{
+		const FMovieSceneDoubleChannel* Channel;
+
+		bool HasDefaultValue() const
+		{
+			return Channel->GetDefault().IsSet();
+		}
+		double GetDefaultValue() const
+		{
+			return Channel->GetDefault().Get(0.0);
+		}
+		double PreExtrapolate(const FFrameTime& InTime) const
+		{
+			double Result = 0.0;
+			TMovieSceneCurveChannelImpl<FMovieSceneDoubleChannel>::Evaluate(Channel, InTime, Result);
+			return Result;
+		}
+		double PostExtrapolate(const FFrameTime& InTime) const
+		{
+			double Result = 0.0;
+			TMovieSceneCurveChannelImpl<FMovieSceneDoubleChannel>::Evaluate(Channel, InTime, Result);
+			return Result;
+		}
+		int32 NumPieces() const
+		{
+			return FMath::Max(0, Channel->GetData().GetValues().Num() - 1);
+		}
+		int32 GetIndexOfPieceByTime(const FFrameTime& Time) const
+		{
+			TArrayView<const FFrameNumber> Times = Channel->GetData().GetTimes();
+			return FMath::Max(Algo::UpperBound(Times, Time)-1, 0);
+		}
+		Interpolation::FCachedInterpolation GetPieceByIndex(int32 Index) const
+		{
+			return TMovieSceneCurveChannelImpl<FMovieSceneDoubleChannel>::GetInterpolationForKey(Channel, Index);
+		}
+		Interpolation::FCachedInterpolation GetPieceByTime(const FFrameTime& Time) const
+		{
+			return TMovieSceneCurveChannelImpl<FMovieSceneDoubleChannel>::GetInterpolationForTime(Channel, Time);
+		}
+		FFrameNumber GetFiniteStart() const
+		{
+			return Channel->GetData().GetTimes()[0];
+		}
+		FFrameNumber GetFiniteEnd() const
+		{
+			return Channel->GetData().GetTimes().Last();
+		}
+		ERichCurveExtrapolation GetPreExtrapolation() const
+		{
+			return Channel->PreInfinityExtrap;
+		}
+		ERichCurveExtrapolation GetPostExtrapolation() const
+		{
+			return Channel->PostInfinityExtrap;
+		}
+		double GetStartingValue() const
+		{
+			return Channel->GetData().GetValues()[0].Value;
+		}
+		double GetEndingValue() const
+		{
+			return Channel->GetData().GetValues().Last().Value;
+		}
+	};
+
+} // namespace UE::MovieScene
+
+
 bool FMovieSceneDoubleValue::Serialize(FArchive& Ar)
 {
 	return TMovieSceneCurveChannelImpl<FMovieSceneDoubleChannel>::SerializeChannelValue(*this, Ar);
@@ -178,6 +254,362 @@ void FMovieSceneDoubleChannel::ClearDefault()
 	bHasDefaultValue = false;
 }
 
+int32 FMovieSceneDoubleChannel::GetCycleCount(FFrameTime InTime) const
+{
+	if (Times.Num() == 0)
+	{
+		return 0;
+	}
+
+	if (InTime < Times[0])
+	{
+		switch (PreInfinityExtrap)
+		{
+		case RCCE_None:     return -1;
+		case RCCE_Constant: return -1;
+		case RCCE_Linear:   return -1;
+		default:            break;
+		}
+	}
+	else if (InTime > Times.Last())
+	{
+		switch (PostInfinityExtrap)
+		{
+		case RCCE_None:     return 1;
+		case RCCE_Constant: return 1;
+		case RCCE_Linear:   return 1;
+		default:            break;
+		}
+	}
+
+	return UE::MovieScene::CycleTime(Times[0], Times.Last(), InTime).CycleCount;
+}
+
+TRange<FFrameNumber> FMovieSceneDoubleChannel::GetCycleRange(int32 InCycleCount) const
+{
+	if (Times.Num() <= 0)
+	{
+		return TRange<FFrameNumber>::All();
+	}
+
+	const bool bCyclePre  = PreInfinityExtrap == RCCE_Cycle  || PreInfinityExtrap == RCCE_CycleWithOffset  || PreInfinityExtrap == RCCE_Oscillate;
+	const bool bCyclePost = PostInfinityExtrap == RCCE_Cycle || PostInfinityExtrap == RCCE_CycleWithOffset || PostInfinityExtrap == RCCE_Oscillate;
+
+	if (InCycleCount == 0 || (InCycleCount < 0 && bCyclePre) || (InCycleCount > 0 && bCyclePost))
+	{
+		const FFrameNumber MinFrame = Times[0];
+		const FFrameNumber MaxFrame = Times.Last();
+		const FFrameNumber Offset   = (MaxFrame - MinFrame) * InCycleCount;
+
+		return TRange<FFrameNumber>::Inclusive(MinFrame+Offset, MaxFrame+Offset);
+	}
+	else if (InCycleCount < 0 && (PreInfinityExtrap == RCCE_Linear || PreInfinityExtrap == RCCE_Constant))
+	{
+		return TRange<FFrameNumber>(TRangeBound<FFrameNumber>::Open(), Times[0]);
+	}
+	else if (InCycleCount > 0 && (PostInfinityExtrap == RCCE_Linear || PostInfinityExtrap == RCCE_Constant))
+	{
+		return TRange<FFrameNumber>(TRangeBound<FFrameNumber>::Exclusive(Times.Last()), TRangeBound<FFrameNumber>::Open());
+	}
+
+	return TRange<FFrameNumber>::Empty();
+}
+
+bool FMovieSceneDoubleChannel::InverseEvaluateBetween(double InValue, FFrameTime StartTime, FFrameTime EndTime, const TFunctionRef<bool(FFrameTime)>& VisitorCallback) const
+{
+	using namespace UE::MovieScene;
+
+	if (Values.Num() == 0)
+	{
+		if (bHasDefaultValue && InValue == DefaultValue)
+		{
+			// Infinite number of solutions - just pick one
+			return VisitorCallback(0);
+		}
+
+		// No solution
+		return true;
+	}
+
+	if (Values.Num() == 1)
+	{
+		if (InValue == Values[0].Value)
+		{
+			// Infinite number of solutions - just pick one
+			return VisitorCallback(0);
+		}
+
+		// No solution
+		return true;
+	}
+
+	FFrameTime TmpSolutions[4];
+
+	Interpolation::FCachedInterpolation Interp = FMovieSceneDoubleChannelImpl::GetInterpolationForTime(this, StartTime);
+	while (Interp.IsValid())
+	{
+		const int32 LocalNumSolutions = Interp.InverseEvaluate(InValue, TmpSolutions);
+
+		for (int32 SolutionIndex = 0; SolutionIndex < LocalNumSolutions; ++SolutionIndex)
+		{
+			if (!VisitorCallback(TmpSolutions[SolutionIndex]))
+			{
+				return false;
+			}
+		}
+
+		// Move on to the next one
+		FFrameNumber ThisInterpEnd = Interp.GetRange().End;
+		if (ThisInterpEnd != TNumericLimits<FFrameNumber>::Max() && ThisInterpEnd < EndTime)
+		{
+			Interp = FMovieSceneDoubleChannelImpl::GetInterpolationForTime(this, ThisInterpEnd+1);
+		}
+		else
+		{
+			Interp = Interpolation::FCachedInterpolation();
+		}
+	}
+
+	return true;
+}
+
+TOptional<FFrameTime> FMovieSceneDoubleChannel::InverseEvaluate(double InValue, FFrameTime InTimeHint, UE::MovieScene::EInverseEvaluateFlags Flags) const
+{
+	using namespace UE::MovieScene;
+
+	if (Values.Num() == 0)
+	{
+		if (bHasDefaultValue && InValue == DefaultValue)
+		{
+			// Infinite number of solutions - just pick one
+			return FFrameTime(0);
+		}
+
+		// No solution
+		return TOptional<FFrameTime>();
+	}
+
+	if (Values.Num() == 1)
+	{
+		if (InValue == Values[0].Value)
+		{
+			// Infinite number of solutions - just pick one
+			return FFrameTime(0);
+		}
+
+		// No solution
+		return TOptional<FFrameTime>();
+	}
+
+	// Never walk more than this number of iterations away from the time hint unless we have a cycle with offset mode
+	int32 MaxIterations = Times.Num()*2;
+	if (EnumHasAnyFlags(Flags, EInverseEvaluateFlags::Cycle) && (PreInfinityExtrap == RCCE_CycleWithOffset || PostInfinityExtrap == RCCE_CycleWithOffset))
+	{
+		const double ValueOffset = Values.Last().Value - Values[0].Value;
+		if (!FMath::IsNearlyZero(ValueOffset))
+		{
+			MaxIterations = 1000;
+		}
+	}
+
+	const FFrameNumber MinFrame = Times[0];
+	const FFrameNumber MaxFrame = Times.Last();
+	const int32 TimeHintCycle = CycleTime(MinFrame, MaxFrame, InTimeHint).CycleCount;
+
+	// Use the hint to find our first interpolation
+	Interpolation::FCachedInterpolation NextInterp = FMovieSceneDoubleChannelImpl::GetInterpolationForTime(this, InTimeHint);
+	if (!NextInterp.IsValid())
+	{
+		return TOptional<FFrameTime>();
+	}
+
+	// Compute the preceeding interpolation
+	Interpolation::FCachedInterpolation PrevInterp = EnumHasAnyFlags(Flags, EInverseEvaluateFlags::Backwards)
+		? FMovieSceneDoubleChannelImpl::GetInterpolationForTime(this, NextInterp.GetRange().Start-1)
+		: Interpolation::FCachedInterpolation();
+	
+
+	// Choose the nearest of all the solutions
+	FFrameTime TmpSolutions[4];
+	TOptional<FFrameTime> Result;
+	TOptional<int32> ResultCycleDiff;
+	
+	double Difference = std::numeric_limits<double>::max();
+
+	int32 IterationCount = 0;
+
+
+	auto ReportSolution = [Flags, InTimeHint, TimeHintCycle, MinFrame, MaxFrame, &Result, &Difference, &ResultCycleDiff](FFrameTime InResult)
+	{
+		// Reject solutions that occur at the same time hint if we're not searching with the equal flag
+		if (!EnumHasAnyFlags(Flags, EInverseEvaluateFlags::Equal) && InResult == InTimeHint)
+		{
+			return false;
+		}
+
+		// Reject solutions that occur before the time hint if we're not searching backwards
+		if (!EnumHasAnyFlags(Flags, EInverseEvaluateFlags::Backwards) && InResult < InTimeHint)
+		{
+			return false;
+		}
+
+		// Reject solutions that occur after the time hint if we're not searching forwards
+		if (!EnumHasAnyFlags(Flags, EInverseEvaluateFlags::Forwards) && InResult > InTimeHint)
+		{
+			return false;
+		}
+
+		const int32 SolutionCycle = CycleTime(MinFrame, MaxFrame, InResult).CycleCount;
+
+		// Reject solutions that occur in a different cycle if we don't allow cycling
+		if(!EnumHasAnyFlags(Flags, EInverseEvaluateFlags::Cycle) && SolutionCycle != TimeHintCycle)
+		{
+			return false;
+		}
+
+		const double ThisDiff  = FMath::Abs((InResult - InTimeHint).AsDecimal());
+		const int32  CycleDiff = FMath::Abs(SolutionCycle - TimeHintCycle);
+		
+		if (Result.IsSet())
+		{
+			if (CycleDiff > ResultCycleDiff.GetValue())
+			{
+				return false;
+			}
+			if (ThisDiff > Difference)
+			{
+				return false;
+			}
+		}
+
+		Result = InResult;
+		Difference = ThisDiff;
+		ResultCycleDiff = CycleDiff;
+		return true;
+	};
+
+
+	// Walk forwards
+	while (NextInterp.IsValid() && IterationCount < MaxIterations)
+	{
+		++IterationCount;
+
+		const int32 LocalNumSolutions = NextInterp.InverseEvaluate(InValue, TmpSolutions);
+
+		for (int32 SolutionIndex = 0; SolutionIndex < LocalNumSolutions; ++SolutionIndex)
+		{
+			ReportSolution(TmpSolutions[SolutionIndex]);
+		}
+
+		if (!Result.IsSet() && EnumHasAnyFlags(Flags, EInverseEvaluateFlags::Forwards))
+		{
+			// Move on to the next one if possible
+			FFrameNumber ThisInterpEnd = NextInterp.GetRange().End;
+			if (ThisInterpEnd < TNumericLimits<FFrameNumber>::Max())
+			{
+				NextInterp = FMovieSceneDoubleChannelImpl::GetInterpolationForTime(this, ThisInterpEnd+1);
+				continue;
+			}
+		}
+
+		// Should only get here if there were solutions, or we're at the end of the range
+		break;
+	}
+
+	// Walk backwards
+	while (PrevInterp.IsValid() && IterationCount < MaxIterations)
+	{
+		++IterationCount;
+
+		const int32 LocalNumSolutions = PrevInterp.InverseEvaluate(InValue, TmpSolutions);
+		for (int32 SolutionIndex = 0; SolutionIndex < LocalNumSolutions; ++SolutionIndex)
+		{
+			ReportSolution(TmpSolutions[SolutionIndex]);
+		}
+
+		if (!Result.IsSet())
+		{
+			// Move on to the previous one if possible
+			FFrameNumber ThisInterpStart = PrevInterp.GetRange().Start;
+			if (ThisInterpStart > TNumericLimits<FFrameNumber>::Lowest())
+			{
+				PrevInterp = FMovieSceneDoubleChannelImpl::GetInterpolationForTime(this, ThisInterpStart-1);
+				continue;
+			}
+		}
+
+		// Should only get here if there were solutions, or we're at the start of the range
+		break;
+	}
+
+	return Result;
+}
+
+UE::MovieScene::Interpolation::FInterpolationExtents FMovieSceneDoubleChannel::ComputeExtents(FFrameTime StartTime, FFrameTime EndTime) const
+{
+	using namespace UE::MovieScene;
+	using namespace UE::MovieScene::Interpolation;
+
+	return ComputePiecewiseExtents(FDoubleChannelPiecewiseData{ this }, StartTime, EndTime);
+}
+
+UE::MovieScene::FPiecewiseCurve FMovieSceneDoubleChannel::AsPiecewiseCurve() const
+{
+	using namespace UE::MovieScene;
+	using namespace UE::MovieScene::Interpolation;
+
+	FPiecewiseCurve Curve;
+
+	if (Times.Num() == 0)
+	{
+		if (bHasDefaultValue)
+		{
+			Curve.Values.Add(FCachedInterpolation(FCachedInterpolationRange::Infinite(), FConstantValue(0, DefaultValue)));
+		}
+		return Curve;
+	}
+
+	if (PreInfinityExtrap != RCCE_None)
+	{
+		Interpolation::FCachedInterpolation PreExtrap;
+		if (FMovieSceneDoubleChannelImpl::CacheExtrapolation(this, Times[0] - 1, PreExtrap))
+		{
+			Curve.Values.Emplace(PreExtrap);
+		}
+		else
+		{
+			ensureMsgf(false, TEXT("Unrepresentable extrapolation mode encountered for piecewise curve"));
+		}
+	}
+
+	if (Times.Num() > 1)
+	{
+		for (int32 Index = 0; Index < Times.Num()-1; ++Index)
+		{
+			Interpolation::FCachedInterpolation Interp = FMovieSceneDoubleChannelImpl::GetInterpolationForKey(this, Index);
+			if (ensure(Interp.IsValid()))
+			{
+				Curve.Values.Emplace(Interp);
+			}
+		}
+	}
+
+	if (PostInfinityExtrap != RCCE_None)
+	{
+		Interpolation::FCachedInterpolation PostExtrap;
+		if (FMovieSceneDoubleChannelImpl::CacheExtrapolation(this, Times.Last() + 1, PostExtrap))
+		{
+			Curve.Values.Emplace(PostExtrap);
+		}
+		else
+		{
+			ensureMsgf(false, TEXT("Unrepresentable extrapolation mode encountered for piecewise curve"));
+		}
+	}
+
+	return Curve;
+}
+
 EMovieSceneKeyInterpolation GetInterpolationMode(FMovieSceneDoubleChannel* InChannel, const FFrameNumber& InTime, EMovieSceneKeyInterpolation DefaultInterpolationMode)
 {
 	return TMovieSceneCurveChannelImpl<FMovieSceneDoubleChannel>::GetInterpolationMode(InChannel, InTime, DefaultInterpolationMode);
diff --git a/Engine/Source/Runtime/MovieScene/Private/Channels/MovieSceneInterpolation.cpp b/Engine/Source/Runtime/MovieScene/Private/Channels/MovieSceneInterpolation.cpp
--- a/Engine/Source/Runtime/MovieScene/Private/Channels/MovieSceneInterpolation.cpp
+++ b/Engine/Source/Runtime/MovieScene/Private/Channels/MovieSceneInterpolation.cpp
@@ -2,10 +2,39 @@
 
 #include "Channels/MovieSceneInterpolation.h"
 #include "Curves/CurveEvaluation.h"
+#include "Algo/Partition.h"
 
 namespace UE::MovieScene::Interpolation
 {
 
+bool FInterpolationExtents::IsValid() const
+{
+	return MinValue != std::numeric_limits<double>::max() && MaxValue != std::numeric_limits<double>::lowest();
+}
+
+void FInterpolationExtents::AddPoint(double Value, FFrameTime Time)
+{
+	if (Value < MinValue)
+	{
+		MinValue = Value;
+		MinValueTime = Time;
+	}
+	if (Value > MaxValue)
+	{
+		MaxValue = Value;
+		MaxValueTime = Time;
+	}
+}
+
+void FInterpolationExtents::Combine(const FInterpolationExtents& Other)
+{
+	if (Other.IsValid())
+	{
+		AddPoint(Other.MinValue, Other.MinValueTime);
+		AddPoint(Other.MaxValue, Other.MaxValueTime);
+	}
+}
+
 FCachedInterpolationRange FCachedInterpolationRange::Empty()
 {
 	return FCachedInterpolationRange{ 0, -1 };
@@ -39,6 +68,11 @@ bool FCachedInterpolationRange::Contains(FFrameNumber FrameNumber) const
 	return FrameNumber >= Start && (FrameNumber < End || End == TNumericLimits<FFrameNumber>::Max());
 }
 
+bool FCachedInterpolationRange::IsEmpty() const
+{
+	return Start >= End;
+}
+
 FCachedInterpolation::FCachedInterpolation()
 	: Data(TInPlaceType<FInvalidValue>())
 	, Range(FCachedInterpolationRange::Empty())
@@ -57,23 +91,146 @@ FCachedInterpolation::FCachedInterpolation(const FCachedInterpolationRange& InRa
 {
 }
 
+FCachedInterpolation::FCachedInterpolation(const FCachedInterpolationRange& InRange, const FQuadraticInterpolation& Quadratic)
+	: Data(TInPlaceType<FQuadraticInterpolation>(), Quadratic)
+	, Range(InRange)
+{
+}
+
 FCachedInterpolation::FCachedInterpolation(const FCachedInterpolationRange& InRange, const FCubicInterpolation& Cubic)
 	: Data(TInPlaceType<FCubicInterpolation>(), Cubic)
 	, Range(InRange)
 {
 }
 
+FCachedInterpolation::FCachedInterpolation(const FCachedInterpolationRange& InRange, const FQuarticInterpolation& Quartic)
+	: Data(TInPlaceType<FQuarticInterpolation>(), Quartic)
+	, Range(InRange)
+{
+}
+
+FCachedInterpolation::FCachedInterpolation(const FCachedInterpolationRange& InRange, const FCubicBezierInterpolation& Cubic)
+	: Data(TInPlaceType<FCubicBezierInterpolation>(), Cubic)
+	, Range(InRange)
+{
+}
+
 FCachedInterpolation::FCachedInterpolation(const FCachedInterpolationRange& InRange, const FWeightedCubicInterpolation& WeightedCubic)
 	: Data(TInPlaceType<FWeightedCubicInterpolation>(), WeightedCubic)
 	, Range(InRange)
 {
 }
 
+bool FCachedInterpolation::IsValid() const
+{
+	return !Range.IsEmpty();
+}
+
 bool FCachedInterpolation::IsCacheValidForTime(FFrameNumber FrameNumber) const
 {
 	return Range.Contains(FrameNumber);
 }
 
+FCachedInterpolationRange FCachedInterpolation::GetRange() const
+{
+	return Range;
+}
+
+FInterpolationExtents FCachedInterpolation::ComputeExtents() const
+{
+	return ComputeExtents(Range.Start, Range.End);
+}
+
+FInterpolationExtents FCachedInterpolation::ComputeExtents(FFrameTime From, FFrameTime To) const
+{
+	FInterpolationExtents Extents;
+
+	if (const FConstantValue* Constant = Data.TryGet<FConstantValue>())
+	{
+		Extents.AddPoint(Constant->Value, From);
+		Extents.AddPoint(Constant->Value, To);
+	}
+	else if (const FLinearInterpolation* Linear = Data.TryGet<FLinearInterpolation>())
+	{
+		Extents.AddPoint(Linear->Evaluate(From), From);
+		Extents.AddPoint(Linear->Evaluate(To), To);
+	}
+	else if (const FQuadraticInterpolation* Quadratic = Data.TryGet<FQuadraticInterpolation>())
+	{
+		Extents.AddPoint(Quadratic->Evaluate(From), From);
+		Extents.AddPoint(Quadratic->Evaluate(To), To);
+
+		FLinearInterpolation Derivative = Quadratic->Derivative();
+
+		FFrameTime Solutions[1];
+		const int32 NumSolutions = Derivative.SolveWithin(From, To, 0.0, Solutions);
+		check(NumSolutions <= 1);
+		for (int32 Index = 0; Index < NumSolutions; ++Index)
+		{
+			Extents.AddPoint(Quadratic->Evaluate(Solutions[Index]), Solutions[Index]);
+		}
+	}
+	else if (const FCubicInterpolation* Cubic = Data.TryGet<FCubicInterpolation>())
+	{
+		Extents.AddPoint(Cubic->Evaluate(From), From);
+		Extents.AddPoint(Cubic->Evaluate(To), To);
+
+		FQuadraticInterpolation Derivative = Cubic->Derivative();
+
+		FFrameTime Solutions[2];
+		const int32 NumSolutions = Derivative.SolveWithin(From, To, 0.0, Solutions);
+		check(NumSolutions <= 2);
+		for (int32 Index = 0; Index < NumSolutions; ++Index)
+		{
+			Extents.AddPoint(Cubic->Evaluate(Solutions[Index]), Solutions[Index]);
+		}
+	}
+	else if (const FQuarticInterpolation* Quartic = Data.TryGet<FQuarticInterpolation>())
+	{
+		Extents.AddPoint(Quartic->Evaluate(From), From);
+		Extents.AddPoint(Quartic->Evaluate(To), To);
+
+		FCubicInterpolation Derivative = Quartic->Derivative();
+
+		FFrameTime Solutions[3];
+		const int32 NumSolutions = Derivative.SolveWithin(From, To, 0.0, Solutions);
+		check(NumSolutions <= 3);
+		for (int32 Index = 0; Index < NumSolutions; ++Index)
+		{
+			Extents.AddPoint(Quartic->Evaluate(Solutions[Index]), Solutions[Index]);
+		}
+	}
+	else if (const FCubicBezierInterpolation* CubicBezier = Data.TryGet<FCubicBezierInterpolation>())
+	{
+		const double Origin = CubicBezier->Origin.Value;
+		const double DX = CubicBezier->DX;
+
+		Extents.AddPoint(CubicBezier->Evaluate(From), From);
+		Extents.AddPoint(CubicBezier->Evaluate(To), To);
+
+		FQuadraticInterpolation Derivative = CubicBezier->Derivative();
+
+		FFrameTime Solutions[2];
+		const int32 NumSolutions = Derivative.SolveWithin(From, To, 0.0, Solutions);
+		check(NumSolutions <= 2);
+		for (int32 Index = 0; Index < NumSolutions; ++Index)
+		{
+			Extents.AddPoint(CubicBezier->Evaluate(Solutions[Index]), Solutions[Index]);
+		}
+	}
+	else if (const FWeightedCubicInterpolation* WeightedCubic = Data.TryGet<FWeightedCubicInterpolation>())
+	{
+		Extents.AddPoint(WeightedCubic->Evaluate(From), From);
+		Extents.AddPoint(WeightedCubic->Evaluate(To),   To);
+	}
+
+	check(
+		Extents.MinValueTime >= From && Extents.MinValueTime <= To &&
+		Extents.MaxValueTime >= From && Extents.MaxValueTime <= To);
+
+	return Extents;
+}
+
 bool FCachedInterpolation::Evaluate(FFrameTime Time, double& OutResult) const
 {
 	if (const FConstantValue* Constant = Data.TryGet<FConstantValue>())
@@ -86,11 +243,26 @@ bool FCachedInterpolation::Evaluate(FFrameTime Time, double& OutResult) const
 		OutResult = Linear->Evaluate(Time);
 		return true;
 	}
+	else if (const FQuadraticInterpolation* Quadratic = Data.TryGet<FQuadraticInterpolation>())
+	{
+		OutResult = Quadratic->Evaluate(Time);
+		return true;
+	}
 	else if (const FCubicInterpolation* Cubic = Data.TryGet<FCubicInterpolation>())
 	{
 		OutResult = Cubic->Evaluate(Time);
 		return true;
 	}
+	else if (const FQuarticInterpolation* Quartic = Data.TryGet<FQuarticInterpolation>())
+	{
+		OutResult = Quartic->Evaluate(Time);
+		return true;
+	}
+	else if (const FCubicBezierInterpolation* CubicBezier = Data.TryGet<FCubicBezierInterpolation>())
+	{
+		OutResult = CubicBezier->Evaluate(Time);
+		return true;
+	}
 	else if (const FWeightedCubicInterpolation* WeightedCubic = Data.TryGet<FWeightedCubicInterpolation>())
 	{
 		OutResult = WeightedCubic->Evaluate(Time);
@@ -100,7 +272,463 @@ bool FCachedInterpolation::Evaluate(FFrameTime Time, double& OutResult) const
 	return false;
 }
 
-FCubicInterpolation::FCubicInterpolation(FFrameNumber InOrigin, double InDX, double InStartValue, double InEndValue, double InStartTangent, double InEndTangent)
+void FCachedInterpolation::Offset(double Amount)
+{
+	if (FConstantValue* Constant = Data.TryGet<FConstantValue>())
+	{
+		Constant->Value += Amount;
+	}
+	else if (FLinearInterpolation* Linear = Data.TryGet<FLinearInterpolation>())
+	{
+		Linear->Constant += Amount;
+	}
+	else if (FQuadraticInterpolation* Quadratic = Data.TryGet<FQuadraticInterpolation>())
+	{
+		Quadratic->Constant += Amount;
+	}
+	else if (FCubicInterpolation* Cubic = Data.TryGet<FCubicInterpolation>())
+	{
+		Cubic->Constant += Amount;
+	}
+	else if (FQuarticInterpolation* Quartic = Data.TryGet<FQuarticInterpolation>())
+	{
+		Quartic->Constant += Amount;
+	}
+	else if (FCubicBezierInterpolation* CubicBezier = Data.TryGet<FCubicBezierInterpolation>())
+	{
+		CubicBezier->P0 += Amount;
+		CubicBezier->P1 += Amount;
+		CubicBezier->P2 += Amount;
+		CubicBezier->P3 += Amount;
+	}
+	else if (FWeightedCubicInterpolation* WeightedCubic = Data.TryGet<FWeightedCubicInterpolation>())
+	{
+		WeightedCubic->StartKeyValue += Amount;
+		WeightedCubic->EndKeyValue   += Amount;
+	}
+}
+
+TOptional<FCachedInterpolation> FCachedInterpolation::ComputeIntegral(double ConstantOffset) const
+{
+	if (const FConstantValue* Constant = Data.TryGet<FConstantValue>())
+	{
+		return FCachedInterpolation(Range, Constant->Integral(ConstantOffset));
+	}
+	else if (const FLinearInterpolation* Linear = Data.TryGet<FLinearInterpolation>())
+	{
+		return FCachedInterpolation(Range, Linear->Integral(ConstantOffset));
+	}
+	else if (const FQuadraticInterpolation* Quadratic = Data.TryGet<FQuadraticInterpolation>())
+	{
+		return FCachedInterpolation(Range, Quadratic->Integral(ConstantOffset));
+	}
+	else if (const FCubicInterpolation* Cubic = Data.TryGet<FCubicInterpolation>())
+	{
+		return FCachedInterpolation(Range, Cubic->Integral(ConstantOffset));
+	}
+	else if (const FQuarticInterpolation* Quartic = Data.TryGet<FQuarticInterpolation>())
+	{
+		checkf(false, TEXT("Unable to compute the integral of a quartic. Although the math is easy, quintic curves are practically impossible to solve."));
+	}
+	else if (const FCubicBezierInterpolation* CubicBezier = Data.TryGet<FCubicBezierInterpolation>())
+	{
+		return FCachedInterpolation(Range, CubicBezier->Integral(ConstantOffset));
+	}
+	else if (const FWeightedCubicInterpolation* WeightedCubic = Data.TryGet<FWeightedCubicInterpolation>())
+	{
+	}
+
+	return TOptional<FCachedInterpolation>();
+}
+
+TOptional<FCachedInterpolation> FCachedInterpolation::ComputeDerivative() const
+{
+	if (const FConstantValue* Constant = Data.TryGet<FConstantValue>())
+	{
+		return FCachedInterpolation(Range, Constant->Derivative());
+	}
+	else if (const FLinearInterpolation* Linear = Data.TryGet<FLinearInterpolation>())
+	{
+		return FCachedInterpolation(Range, Linear->Derivative());
+	}
+	else if (const FQuadraticInterpolation* Quadratic = Data.TryGet<FQuadraticInterpolation>())
+	{
+		return FCachedInterpolation(Range, Quadratic->Derivative());
+	}
+	else if (const FCubicInterpolation* Cubic = Data.TryGet<FCubicInterpolation>())
+	{
+		return FCachedInterpolation(Range, Cubic->Derivative());
+	}
+	else if (const FQuarticInterpolation* Quartic = Data.TryGet<FQuarticInterpolation>())
+	{
+		return FCachedInterpolation(Range, Quartic->Derivative());
+	}
+	else if (const FCubicBezierInterpolation* CubicBezier = Data.TryGet<FCubicBezierInterpolation>())
+	{
+		return FCachedInterpolation(Range, CubicBezier->Derivative());
+	}
+	else if (const FWeightedCubicInterpolation* WeightedCubic = Data.TryGet<FWeightedCubicInterpolation>())
+	{
+	}
+
+	return TOptional<FCachedInterpolation>();
+}
+
+int32 FCachedInterpolation::InverseEvaluate(double InValue, TInterpSolutions<FFrameTime, 4> OutResults) const
+{
+	int32 NumSolutions = 0;
+	if (const FConstantValue* Constant = Data.TryGet<FConstantValue>())
+	{
+		if (InValue == Constant->Value)
+		{
+			if (Range.Start == TNumericLimits<FFrameNumber>::Lowest())
+			{
+				if (Range.End == TNumericLimits<FFrameNumber>::Lowest())
+				{
+					OutResults[0] = FFrameTime();
+				}
+				else
+				{
+					OutResults[0] = Range.End;
+				}
+			}
+			else
+			{
+				OutResults[0] = Range.Start;
+			}
+			NumSolutions = 1;
+		}
+
+		check(NumSolutions <= 1);
+	}
+	else if (const FLinearInterpolation* Linear = Data.TryGet<FLinearInterpolation>())
+	{
+		NumSolutions = Linear->Solve(InValue, OutResults);
+		check(NumSolutions <= 1);
+	}
+	else if (const FQuadraticInterpolation* Quadratic = Data.TryGet<FQuadraticInterpolation>())
+	{
+		NumSolutions = Quadratic->Solve(InValue, OutResults);
+		check(NumSolutions <= 2);
+	}
+	else if (const FCubicInterpolation* Cubic = Data.TryGet<FCubicInterpolation>())
+	{
+		NumSolutions = Cubic->Solve(InValue, OutResults);
+		check(NumSolutions <= 3);
+	}
+	else if (const FQuarticInterpolation* Quartic = Data.TryGet<FQuarticInterpolation>())
+	{
+		NumSolutions = Quartic->Solve(InValue, OutResults);
+		check(NumSolutions <= 4);
+	}
+	else if (const FCubicBezierInterpolation* CubicBezier = Data.TryGet<FCubicBezierInterpolation>())
+	{
+		NumSolutions = CubicBezier->Solve(InValue, OutResults);
+		check(NumSolutions <= 2);
+	}
+	else if (const FWeightedCubicInterpolation* WeightedCubic = Data.TryGet<FWeightedCubicInterpolation>())
+	{
+		NumSolutions = WeightedCubic->Solve(InValue, OutResults);
+		check(NumSolutions <= 2);
+	}
+
+	// Only accept solutions within our acceptable range
+	int32 WritePos = 0;
+	for (int32 Index = 0; Index < NumSolutions; ++Index)
+	{
+		if (OutResults[Index] >= Range.Start && OutResults[Index] <= Range.End)
+		{
+			if (Index != WritePos)
+			{
+				OutResults[WritePos] = OutResults[Index];
+			}
+			++WritePos;
+		}
+	}
+
+	return WritePos;
+}
+
+FConstantValue FConstantValue::Derivative() const
+{
+	return FConstantValue(Origin, 0.0);
+}
+
+FLinearInterpolation FConstantValue::Integral(double ConstantOffset) const
+{
+	return FLinearInterpolation(Origin, Value, ConstantOffset);
+}
+
+double FLinearInterpolation::Evaluate(FFrameTime InTime) const
+{
+	return Coefficient*(InTime - Origin).AsDecimal() + Constant;
+}
+
+int32 FLinearInterpolation::Solve(double Value, TInterpSolutions<FFrameTime, 1> OutResults) const
+{
+	if (Coefficient != 0.0)
+	{
+		OutResults[0] = FFrameTime::FromDecimal((Value - Constant) / Coefficient) + Origin;
+		return 1;
+	}
+	else if (Value == Constant)
+	{
+		OutResults[0] = FFrameTime::FromDecimal(Constant);
+		return 1;
+	}
+	return 0;
+}
+
+int32 FLinearInterpolation::SolveWithin(FFrameTime Start, FFrameTime End, double Value, TInterpSolutions<FFrameTime, 1> OutResults) const
+{
+	if (Value == Constant)
+	{
+		OutResults[0] = Start;
+		return 1;
+	}
+	if (Solve(Value, OutResults) == 1 && OutResults[0] >= Start && OutResults[0] < End)
+	{
+		return 1;
+	}
+	return 0;
+}
+
+FConstantValue FLinearInterpolation::Derivative() const
+{
+	return FConstantValue(Origin, Coefficient);
+}
+
+FQuadraticInterpolation FLinearInterpolation::Integral(double ConstantOffset) const
+{
+	return FQuadraticInterpolation(Origin, 0.5*Coefficient, Constant, ConstantOffset);
+}
+
+double FQuadraticInterpolation::Evaluate(FFrameTime InTime) const
+{
+	const double X = (InTime - Origin).AsDecimal();
+	return A*X*X + B*X + Constant;
+}
+
+int32 FQuadraticInterpolation::Solve(double Value, TInterpSolutions<FFrameTime, 2> OutResults) const
+{
+	const double C = Constant - Value;
+
+	// Solve using the quadratic formula
+	OutResults[0] = FFrameTime::FromDecimal( (-B + FMath::Sqrt(B*B - 4.0*A*C)) / (2.0*A) ) + Origin;
+	OutResults[1] = FFrameTime::FromDecimal( (-B - FMath::Sqrt(B*B - 4.0*A*C)) / (2.0*A) ) + Origin;
+
+	return 2;
+}
+
+int32 FQuadraticInterpolation::SolveWithin(FFrameTime Start, FFrameTime End, double Value, TInterpSolutions<FFrameTime, 2> OutResults) const
+{
+	Solve(Value, OutResults);
+
+	if (OutResults[0] < Start || OutResults[0] >= End)
+	{
+		if (OutResults[1] < Start || OutResults[1] >= End)
+		{
+			return 0;
+		}
+		OutResults[0] = OutResults[1];
+		return 1;
+	}
+	if (OutResults[1] < Start || OutResults[1] >= End)
+	{
+		return 1;
+	}
+	return 2;
+}
+
+FLinearInterpolation FQuadraticInterpolation::Derivative() const
+{
+	return FLinearInterpolation(Origin, 2.0*A, B);
+}
+
+FCubicInterpolation FQuadraticInterpolation::Integral(double ConstantOffset) const
+{
+	return FCubicInterpolation(Origin, A/3.0, B/2.0, Constant, ConstantOffset);
+}
+
+double FCubicInterpolation::Evaluate(FFrameTime InTime) const
+{
+	const double X = (InTime - Origin).AsDecimal() / DX;
+	return A*X*X*X + B*X*X + C*X + Constant;
+}
+
+int32 FCubicInterpolation::Solve(double Value, TInterpSolutions<FFrameTime, 3> OutResults) const
+{
+	double Coefficients[4] = {
+		Constant-Value,
+		C,
+		B,
+		A,
+	};
+
+	double Solutions[3];
+	const int32 NumRealSolutions = UE::Curves::SolveCubic(Coefficients, Solutions);
+
+	for (int32 Index = 0; Index < NumRealSolutions; ++Index)
+	{
+		OutResults[Index] = DX*FFrameTime::FromDecimal(Solutions[Index])+Origin;
+	}
+
+	return NumRealSolutions;
+}
+
+int32 FCubicInterpolation::SolveWithin(FFrameTime Start, FFrameTime End, double Value, TInterpSolutions<FFrameTime, 3> OutResults) const
+{
+	int32 NumSolutions = Solve(Value, OutResults);
+
+	// Only accept solutions within our acceptable range
+	int32 WritePos = 0;
+	for (int32 Index = 0; Index < NumSolutions; ++Index)
+	{
+		if (OutResults[Index] >= Start && OutResults[Index] <= End)
+		{
+			if (Index != WritePos)
+			{
+				OutResults[WritePos] = OutResults[Index];
+			}
+			++WritePos;
+		}
+	}
+
+	return WritePos;
+}
+
+FQuadraticInterpolation FCubicInterpolation::Derivative() const
+{
+	return FQuadraticInterpolation(Origin, 3.0*A/(DX*DX*DX), 2.0*B/(DX*DX), C/DX);
+}
+
+FQuarticInterpolation FCubicInterpolation::Integral(double ConstantOffset) const
+{
+	return FQuarticInterpolation(Origin, A*DX/4.0, B*DX/3.0, C*DX/2.0, DX*Constant, ConstantOffset, DX);
+}
+
+double FQuarticInterpolation::Evaluate(FFrameTime InTime) const
+{
+	const double X = (InTime - Origin).AsDecimal() / DX;
+	return A*X*X*X*X + B*X*X*X + C*X*X + D*X + Constant;
+}
+
+FCubicInterpolation FQuarticInterpolation::Derivative() const
+{
+	return FCubicInterpolation(Origin, 4.0*A, 3.0*B, 2.0*C, D, DX);
+}
+
+int32 FQuarticInterpolation::Solve(double Value, TInterpSolutions<FFrameTime, 4> OutResults) const
+{
+	// Solving ax^4 + bx^3 + cx^2 + dx + e = 0 is very difficult, but we can solve it using some tricks.
+	// First off, convert our coefficients to make this a monic quartic of the form X^4 + bx^3 + cx^2 + dx + e:
+	const double A_ = 1.0;
+	const double B_ = B/A;
+	const double C_ = C/A;
+	const double D_ = D/A;
+	const double E_ = (Constant - Value)/A;
+
+	// Now convert the monic quartic to a depressed quartic of the form y^4 + py^2 + qy + r by substituting x = y - b/4 (since a=1.0)
+	const double P = C_ - (3.0*B_*B_) / 8.0;
+	const double Q = D_ - (B_*C_)/2.0 + (B_*B_*B_)/8.0;
+	const double R = E_ - (B_*D_)/4.0 + (B_*B_*C_)/16.0 - (3.0*B_*B_*B_*B_)/256.0;
+
+	// Step 2: Factor the depressed quartic into two quadratics:
+	//         (y^2 + sy + t)(y^2 + uy + v)
+	//
+	//         which results in
+	//         y^4 + (s+u)y^3 + (t+v+su)y^2 + (sv+tu)y + tv
+	//
+	//         Since y^3 is 0 in our depressed quartic, we can deduce that
+	//         s+u = 0
+	//           p = t + v + su
+	//           q = sv + tu
+	//           r = tv
+	//
+	//         From this we can deduce that s = -u:
+	//         p + u^2 = t + v
+	//               q = u(t-v)
+	//               r = tv
+	//
+	//         Eliminating t and v by substitution and substituting U = u^2:
+	//         U^3 + 2pU^2 + (p^2 - 4r)U - q^2 = 0
+
+	// Solve the cubic for U using Cardano's algorithm:
+	// SolveCubic expects coefficients in increasing exponent order
+	//    ie d + cx + bx^2 + ax^3
+	double Coefficients[4] = {
+		-(Q*Q),         // d
+		(P*P - 4.0*R),  // c
+		(2.0*P),        // b
+		1.0,            // a
+	};
+
+	double Solutions[3];
+	const int32 NumSolutions = UE::Curves::SolveCubic(Coefficients, Solutions);
+
+	int32 SolutionToUse = INDEX_NONE;
+
+	// Any real solution will do, but fallback to a complex solution if there's only one
+	for (int32 Index = 0; Index < NumSolutions; ++Index)
+	{
+		SolutionToUse = Index;
+
+		// Use first real solution if possible
+		if (Solutions[Index] > 0.0)
+		{
+			break;
+		}
+	}
+
+	int32 SolutionIndex = 0;
+	if (SolutionToUse != INDEX_NONE)
+	{
+		const double U = FMath::Sqrt(FMath::Abs(Solutions[0]));
+		const double S = -U;
+		const double T = (U*U*U + P*U + Q) / (2.0*U);
+		const double V = T - (Q / U);
+
+		const double USquaredMinusFourV = U*U - 4.0*V;
+		const double SSquaredMinusFourT = S*S - 4.0*T;
+		const double BOverFourA         = B_ / (4.0*A_);
+
+		if (USquaredMinusFourV >= 0.0)
+		{
+			OutResults[SolutionIndex++] = FFrameTime::FromDecimal(( (-U + FMath::Sqrt(USquaredMinusFourV)) / 2.0 ) - BOverFourA) + Origin;
+			OutResults[SolutionIndex++] = FFrameTime::FromDecimal(( (-U - FMath::Sqrt(USquaredMinusFourV)) / 2.0 ) - BOverFourA) + Origin;
+		}
+		if (SSquaredMinusFourT >= 0.0)
+		{
+			OutResults[SolutionIndex++] = FFrameTime::FromDecimal(( (-S + FMath::Sqrt(SSquaredMinusFourT)) / 2.0 ) - BOverFourA) + Origin;
+			OutResults[SolutionIndex++] = FFrameTime::FromDecimal(( (-S - FMath::Sqrt(SSquaredMinusFourT)) / 2.0 ) - BOverFourA) + Origin;
+		}
+	}
+
+	return SolutionIndex;
+}
+
+int32 FQuarticInterpolation::SolveWithin(FFrameTime Start, FFrameTime End, double Value, TInterpSolutions<FFrameTime, 4> OutResults) const
+{
+	int32 NumSolutions = Solve(Value, OutResults);
+
+	// Only accept solutions within our acceptable range
+	int32 WritePos = 0;
+	for (int32 Index = 0; Index < NumSolutions; ++Index)
+	{
+		if (OutResults[Index] >= Start && OutResults[Index] <= End)
+		{
+			if (Index != WritePos)
+			{
+				OutResults[WritePos] = OutResults[Index];
+			}
+			++WritePos;
+		}
+	}
+
+	return WritePos;
+}
+
+FCubicBezierInterpolation::FCubicBezierInterpolation(FFrameNumber InOrigin, double InDX, double InStartValue, double InEndValue, double InStartTangent, double InEndTangent)
 	: DX(InDX)
 	, Origin(InOrigin)
 {
@@ -112,12 +740,75 @@ FCubicInterpolation::FCubicInterpolation(FFrameNumber InOrigin, double InDX, dou
 	P2 = P3 - (InEndTangent * DX * OneThird);
 }
 
-double FCubicInterpolation::Evaluate(FFrameTime InTime) const
+FCubicInterpolation FCubicBezierInterpolation::AsCubic() const
+{
+	// Beziers are interpolated using a normalized input so we have to factor out that normalization to
+	//    each term of the resulting polynomial.
+
+	const double OneOverDxCubed   = 1.0 / (DX * DX * DX);
+	const double OneOverDxSquared = 1.0 / (DX * DX);
+	const double OneOverDx        = 1.0 / (DX);
+
+	const double A = OneOverDxCubed*(-P0 + 3.0*P1 - 3.0*P2 + P3);
+	const double B = OneOverDxSquared*(3.0*P0 - 6.0*P1 + 3.0*P2);
+	const double C = OneOverDx*3.0*(-P0 + P1);
+	const double D = P0;
+
+	return FCubicInterpolation(
+		Origin,
+		A,
+		B,
+		C,
+		D,
+		1.0);// DX);
+}
+
+double FCubicBezierInterpolation::Evaluate(FFrameTime InTime) const
 {
 	const float Interp = static_cast<float>((InTime - Origin).AsDecimal() / DX);
 	return UE::Curves::BezierInterp(P0, P1, P2, P3, Interp);
 }
 
+FQuadraticInterpolation FCubicBezierInterpolation::Derivative() const
+{
+	return AsCubic().Derivative();
+}
+FQuarticInterpolation FCubicBezierInterpolation::Integral(double ConstantOffset) const
+{
+	return AsCubic().Integral(ConstantOffset);
+}
+
+int32 FCubicBezierInterpolation::Solve(double InValue, TInterpSolutions<FFrameTime, 4> OutResults) const
+{
+	// Offset the curve by InValue to find roots
+
+	// SolveCubic expects coefficients in increasing exponent order
+	//    ie d + cx + bx^2 + ax^3
+	double Coefficients[4] = {
+		(P0-InValue),                       // d
+		(-3.0 * P0 + 3.0 * P1),             // c
+		(3.0 * P0 - 6.0 * P1 + 3.0 * P2),   // b
+		(-P0 + 3.0 * P1 - 3.0 * P2 + P3),   // a
+	};
+
+	double Solutions[3];
+	const int32 NumRealSolutions = UE::Curves::SolveCubic(Coefficients, Solutions);
+
+	int32 NumSolutionsWithinRange = 0;
+	for (int32 Index = 0; Index < NumRealSolutions; ++Index)
+	{
+		double Solution = Solutions[Index];
+		// Only accept solutions within the 0-1 range, allowing for some rounding discrepancies
+		if (Solution >= -UE_SMALL_NUMBER && Solution <= 1.0 + UE_SMALL_NUMBER)
+		{
+			Solution = FMath::Clamp(Solution, 0.0, 1.0);
+			OutResults[NumSolutionsWithinRange++] = FFrameTime::FromDecimal(Solution*DX) + Origin;
+		}
+	}
+
+	return NumSolutionsWithinRange;
+}
+
 FWeightedCubicInterpolation::FWeightedCubicInterpolation(
 	FFrameRate TickResolution,
 	FFrameNumber InOrigin,
@@ -242,4 +933,45 @@ double FWeightedCubicInterpolation::Evaluate(FFrameTime InTime) const
 	return UE::Curves::BezierInterp(P0, P1, P2, P3, static_cast<float>(NewInterp));
 }
 
+int32 FWeightedCubicInterpolation::Solve(double InValue, TInterpSolutions<FFrameTime, 4> OutResults) const
+{
+	// Solve the control points first
+	const double P0 = StartKeyValue;
+	const double P1 = StartKeyTanY;
+	const double P3 = EndKeyValue;
+	const double P2 = EndKeyTanY;
+
+	// Offset the curve by InValue to find roots
+
+	// SolveCubic expects coefficients in increasing exponent order
+	//    ie d + cx + bx^2 + ax^3
+	double Coefficients[4] = {
+		(P0-InValue),                       // d
+		(-3.0 * P0 + 3.0 * P1),             // c
+		(3.0 * P0 - 6.0 * P1 + 3.0 * P2),   // b
+		(-P0 + 3.0 * P1 - 3.0 * P2 + P3),   // a
+	};
+
+	double Solutions[3];
+	const int32 NumRealSolutions = UE::Curves::SolveCubic(Coefficients, Solutions);
+
+	FFrameTime InitialSolutions[3];
+	int32 NumSolutionsWithinRange = 0;
+	for (int32 Index = 0; Index < NumRealSolutions; ++Index)
+	{
+		double Solution = Solutions[Index];
+		// Only accept solutions within the 0-1 range, allowing for some rounding discrepancies
+		if (Solution >= -UE_SMALL_NUMBER && Solution <= 1.0 + UE_SMALL_NUMBER)
+		{
+			Solution = FMath::Clamp(Solution, 0.0, 1.0);
+
+			// Now solve this interp for the power basis
+			Solution = UE::Curves::BezierInterp(0.0, NormalizedStartTanDX, NormalizedEndTanDX, 1.0, static_cast<float>(Solution));
+			OutResults[NumSolutionsWithinRange++] = FFrameTime::FromDecimal(Solution*DX) + Origin;
+		}
+	}
+
+	return NumSolutionsWithinRange;
+}
+
 } // namespace UE::MovieScene
diff --git a/Engine/Source/Runtime/MovieScene/Private/Compilation/MovieSceneCompiledDataManager.cpp b/Engine/Source/Runtime/MovieScene/Private/Compilation/MovieSceneCompiledDataManager.cpp
--- a/Engine/Source/Runtime/MovieScene/Private/Compilation/MovieSceneCompiledDataManager.cpp
+++ b/Engine/Source/Runtime/MovieScene/Private/Compilation/MovieSceneCompiledDataManager.cpp
@@ -11,8 +11,10 @@
 #include "MovieSceneSequence.h"
 #include "Sections/MovieSceneSubSection.h"
 #include "Tracks/MovieSceneSubTrack.h"
+#include "Tracks/MovieSceneTimeWarpTrack.h"
 #include "IMovieSceneModule.h"
 #include "MovieSceneTimeHelpers.h"
+#include "MovieSceneTransformTypes.h"
 
 #include "Algo/Sort.h"
 #include "Algo/Unique.h"
@@ -147,20 +149,27 @@ struct FGatherParameters
 
 	FGatherParameters CreateForSubData(const FMovieSceneSubSequenceData& SubData, FMovieSceneSequenceID InSubSequenceID) const
 	{
-		return CreateForSubData(SubData, InSubSequenceID, FMovieSceneWarpCounter());
-	}
+		using namespace UE::MovieScene;
 
-	FGatherParameters CreateForSubData(const FMovieSceneSubSequenceData& SubData, FMovieSceneSequenceID InSubSequenceID, FMovieSceneWarpCounter WarpCounter) const
-	{
-		FGatherParameters SubParams = *this;
+		FGatherParameters SubParams;
 
+		SubParams.SequenceID                = InSubSequenceID;
+		SubParams.RootClampRange            = this->RootClampRange;
+		SubParams.LocalClampRange           = UE::MovieScene::ConvertToDiscreteRange(SubData.RootToSequenceTransform.ComputeTraversedHull(this->RootClampRange));
+		SubParams.Flags                     = this->Flags;
 		SubParams.RootToSequenceTransform   = SubData.RootToSequenceTransform;
+#if WITH_EDITORONLY_DATA
+		SubParams.StartTimeBreadcrumbs      = SubData.StartTimeBreadcrumbs;
+		SubParams.EndTimeBreadcrumbs        = SubData.EndTimeBreadcrumbs;
+#else
+		SubParams.StartTimeBreadcrumbs      = this->StartTimeBreadcrumbs;
+		SubParams.EndTimeBreadcrumbs        = this->EndTimeBreadcrumbs;
+		SubParams.RootToSequenceTransform.TransformTime(DiscreteInclusiveLower(SubData.ParentPlayRange.Value), FTransformTimeParams().AppendBreadcrumbs(SubParams.StartTimeBreadcrumbs));
+		SubParams.RootToSequenceTransform.TransformTime(DiscreteExclusiveUpper(SubData.ParentPlayRange.Value), FTransformTimeParams().AppendBreadcrumbs(SubParams.EndTimeBreadcrumbs));
+#endif
 		SubParams.HierarchicalBias          = SubData.HierarchicalBias;
 		SubParams.AccumulatedFlags          = SubData.AccumulatedFlags;
-		SubParams.SequenceID                = InSubSequenceID;
-		SubParams.RootToSequenceWarpCounter = WarpCounter;
-
-		SubParams.LocalClampRange           = SubData.RootToSequenceTransform.TransformRangeUnwarped(SubParams.RootClampRange);
+		SubParams.NetworkMask               = this->NetworkMask;
 
 		return SubParams;
 	}
@@ -168,7 +177,7 @@ struct FGatherParameters
 	void SetClampRange(TRange<FFrameNumber> InNewRootClampRange)
 	{
 		RootClampRange  = InNewRootClampRange;
-		LocalClampRange = RootToSequenceTransform.TransformRangeUnwarped(InNewRootClampRange);
+		LocalClampRange = UE::MovieScene::ConvertToDiscreteRange(RootToSequenceTransform.ComputeTraversedHull(InNewRootClampRange));
 	}
 
 	/** Clamp the specified range to the current clamp range (in root space) */
@@ -177,6 +186,101 @@ struct FGatherParameters
 		return TRange<FFrameNumber>::Intersection(RootClampRange, InRootRange);
 	}
 
+	void TransformLocalRange(const TRange<FFrameNumber>& InLocalRange, TFunctionRef<bool(TRange<FFrameTime>)> InVisitor) const
+	{
+		using namespace UE::MovieScene;
+
+		TRange<FFrameTime> Range = ConvertToFrameTimeRange(InLocalRange);
+
+		FMovieSceneInverseSequenceTransform SequenceToRootTransform = RootToSequenceTransform.Inverse();
+
+		// Linear transforms are easy
+		if (SequenceToRootTransform.IsLinear())
+		{
+			FMovieSceneTimeTransform LinearTransform = SequenceToRootTransform.AsLinear();
+
+			if (!Range.GetLowerBound().IsOpen())
+			{
+				Range.SetLowerBoundValue(Range.GetLowerBoundValue() * LinearTransform);
+			}
+			if (!Range.GetUpperBound().IsOpen())
+			{
+				Range.SetUpperBoundValue(Range.GetUpperBoundValue() * LinearTransform);
+			}
+
+			InVisitor(Range);
+			return;
+		}
+
+		// Warping transforms are a bit harder
+
+	
+		// First off, intersect with the clamp range
+		if (Range.GetLowerBound().IsOpen() || Range.GetUpperBound().IsOpen())
+		{
+			Range = TRange<FFrameTime>::Intersection(Range, RootToSequenceTransform.ComputeTraversedHull(ConvertToFrameTimeRange(RootClampRange)));
+		}
+
+		// Make the range finite based on clamp ranges if possible
+		if (Range.GetLowerBound().IsOpen() && !RootClampRange.GetLowerBound().IsOpen())
+		{
+			TRangeBound<FFrameNumber> LowerBound = RootClampRange.GetLowerBound();
+			FFrameTime NewTime = RootToSequenceTransform.TransformTime(LowerBound.GetValue());
+			if (LowerBound.IsInclusive())
+			{
+				Range.SetLowerBound(TRangeBound<FFrameTime>::Inclusive(NewTime));
+			}
+			else
+			{
+				Range.SetLowerBound(TRangeBound<FFrameTime>::Exclusive(NewTime));
+			}
+		}
+		if (Range.GetUpperBound().IsOpen() && !RootClampRange.GetUpperBound().IsOpen())
+		{
+			TRangeBound<FFrameNumber> UpperBound = RootClampRange.GetUpperBound();
+			FFrameTime NewTime = RootToSequenceTransform.TransformTime(UpperBound.GetValue());
+			if (UpperBound.IsInclusive())
+			{
+				Range.SetUpperBound(TRangeBound<FFrameTime>::Inclusive(NewTime));
+			}
+			else
+			{
+				Range.SetUpperBound(TRangeBound<FFrameTime>::Exclusive(NewTime));
+			}
+		}
+
+		if (Range.GetLowerBound().IsOpen() && Range.GetUpperBound().IsOpen())
+		{
+			// If the range is infinite we just have to add it all since there's no way for us to transform it.
+			InVisitor(Range);
+		}
+		else if (!Range.GetLowerBound().IsOpen() && !Range.GetUpperBound().IsOpen())
+		{
+			// We have a finite range so transform it as many times as it exists in the root space
+			SequenceToRootTransform.TransformFiniteRangeWithinRange(Range, InVisitor, StartTimeBreadcrumbs, EndTimeBreadcrumbs);
+		}
+		else if (Range.GetLowerBound().IsOpen())
+		{
+			// Open lower bound so just transform the the upper bound once and compile that
+			TOptional<FFrameTime> Time = SequenceToRootTransform.TryTransformTime(Range.GetUpperBoundValue(), EndTimeBreadcrumbs);
+			if (Time)
+			{
+				Range.SetUpperBoundValue(Time->FloorToFrame());
+				InVisitor(Range);
+			}
+		}
+		else if (Range.GetUpperBound().IsOpen())
+		{
+			// Open upper bound so just transform the the lower bound once and compile that
+			TOptional<FFrameTime> Time = SequenceToRootTransform.TryTransformTime(Range.GetLowerBoundValue(), StartTimeBreadcrumbs);
+			if (Time)
+			{
+				Range.SetLowerBoundValue(Time->FloorToFrame());
+				InVisitor(Range);
+			}
+		}
+	}
+
 	/** The ID of the sequence being compiled */
 	FMovieSceneSequenceID SequenceID;
 
@@ -190,9 +294,8 @@ struct FGatherParameters
 
 	/** Transform from the root time-space to the current sequence's time-space */
 	FMovieSceneSequenceTransform RootToSequenceTransform;
-
-	/** Loop counts from the root to the current sequence */
-	FMovieSceneWarpCounter RootToSequenceWarpCounter;
+	FMovieSceneTransformBreadcrumbs StartTimeBreadcrumbs;
+	FMovieSceneTransformBreadcrumbs EndTimeBreadcrumbs;
 
 	/** Current accumulated hierarchical bias */
 	int16 HierarchicalBias;
@@ -210,17 +313,22 @@ struct FTrackGatherParameters : FGatherParameters
 		: TemplateGenerator(InCompiledDataManager)
 	{}
 
-	FTrackGatherParameters CreateForSubData(const FMovieSceneSubSequenceData& SubData, FMovieSceneSequenceID InSubSequenceID, FMovieSceneWarpCounter WarpCounter) const
+	FTrackGatherParameters CreateForSubData(const FMovieSceneSubSequenceData& SubData, FMovieSceneSequenceID InSubSequenceID) const
 	{
-		FTrackGatherParameters SubParams = *this;
-		static_cast<FGatherParameters&>(SubParams) = FGatherParameters::CreateForSubData(SubData, InSubSequenceID, WarpCounter);
-
+		FTrackGatherParameters SubParams;
+		static_cast<FGatherParameters&>(SubParams) = FGatherParameters::CreateForSubData(SubData, InSubSequenceID);
+		SubParams.TemplateGenerator = this->TemplateGenerator;
 		return SubParams;
 	}
 
 
 	/** Store from which to retrieve templates */
 	mutable FMovieSceneCompileDataManagerGenerator TemplateGenerator;
+
+private:
+	FTrackGatherParameters()
+		: TemplateGenerator(nullptr)
+	{}
 };
 
 
@@ -962,7 +1070,7 @@ void UMovieSceneCompiledDataManager::CompileSubSequences(const FMovieSceneSequen
 			UMovieSceneSequence* SubSequence = SubData->GetSequence();
 			if (SubSequence)
 			{
-				FTrackGatherParameters SubSectionGatherParams = Params.CreateForSubData(*SubData, SubSequenceID, SubSequenceEntry.RootToSequenceWarpCounter);
+				FTrackGatherParameters SubSectionGatherParams = Params.CreateForSubData(*SubData, SubSequenceID);
 				SubSectionGatherParams.Flags |= SubSequenceEntry.Flags;
 				SubSectionGatherParams.SetClampRange(SubSequenceIt.Range());
 
@@ -1358,67 +1466,74 @@ void UMovieSceneCompiledDataManager::GatherTrack(const FMovieSceneBinding* Objec
 		// Iterate everything in the field
 		for (const FMovieSceneTrackEvaluationFieldEntry& Entry : EvaluationField.Entries)
 		{
-			FMovieSceneSequenceTransform SequenceToRootTransform  = Params.RootToSequenceTransform.InverseFromLoop(Params.RootToSequenceWarpCounter);
-			TRange<FFrameNumber>         ClampedRangeRoot         = Params.ClampRoot(SequenceToRootTransform.TransformRangeUnwarped(Entry.Range));
-			UMovieSceneSection*          Section                  = Entry.Section;
+			// Iterate all the valid ranges this translates to in the root
+			FMovieSceneInverseSequenceTransform SequenceToRootTransform  = Params.RootToSequenceTransform.Inverse();
 
-			if (Section && Track->IsRowEvalDisabled(Section->GetRowIndex()))
+			auto VisitWarpedRootRange = [this, &Entry, &Params, &OutCompilerData, TrackIdentifier, TrackTemplate, Track](TRange<FFrameTime> InRange)
 			{
-				continue;
-			}
+				TRange<FFrameNumber> ClampedRangeRoot = Params.ClampRoot(ConvertToDiscreteRange(InRange));
+				UMovieSceneSection*  Section          = Entry.Section;
 
-			if (ClampedRangeRoot.IsEmpty())
-			{
-				continue;
-			}
-
-			check(TrackTemplate);
-			const FMovieSceneEvaluationTrack* EvaluationTrack = TrackTemplate->FindTrack(TrackIdentifier);
-			check(EvaluationTrack);
+				if (Section && Track->IsRowEvalDisabled(Section->GetRowIndex()))
+				{
+					return true;
+				}
 
-			// Get the correct template for the sub sequence
-			FCompileOnTheFlyData CompileData;
+				if (ClampedRangeRoot.IsEmpty())
+				{
+					return true;
+				}
 
-			CompileData.Track                   = FMovieSceneEvaluationFieldTrackPtr(Params.SequenceID, TrackIdentifier);
-			CompileData.EvaluationPriority      = EvaluationTrack->GetEvaluationPriority();
-			CompileData.EvaluationGroup         = EvaluationTrack->GetEvaluationGroup();
-			CompileData.GroupEvaluationPriority = GetMovieSceneModule().GetEvaluationGroupParameters(CompileData.EvaluationGroup).EvaluationPriority;
-			CompileData.HierarchicalBias        = Params.HierarchicalBias;
-			CompileData.bPriorityTearDown       = EvaluationTrack->HasTearDownPriority();
+				const FMovieSceneEvaluationTrack* EvaluationTrack = TrackTemplate->FindTrack(TrackIdentifier);
+				check(EvaluationTrack);
 
-			auto FindChildWithSection = [Section](const FMovieSceneEvalTemplatePtr& ChildTemplate)
-			{
-				return ChildTemplate.IsValid() && ChildTemplate->GetSourceSection() == Section;
-			};
+				// Get the correct template for the sub sequence
+				FCompileOnTheFlyData CompileData;
 
-			const int32 ChildTemplateIndex = Section ? EvaluationTrack->GetChildTemplates().IndexOfByPredicate(FindChildWithSection) : INDEX_NONE;
-			if (ChildTemplateIndex != INDEX_NONE)
-			{
-				check(ChildTemplateIndex >= 0 && ChildTemplateIndex < TNumericLimits<uint16>::Max());
+				CompileData.Track                   = FMovieSceneEvaluationFieldTrackPtr(Params.SequenceID, TrackIdentifier);
+				CompileData.EvaluationPriority      = EvaluationTrack->GetEvaluationPriority();
+				CompileData.EvaluationGroup         = EvaluationTrack->GetEvaluationGroup();
+				CompileData.GroupEvaluationPriority = GetMovieSceneModule().GetEvaluationGroupParameters(CompileData.EvaluationGroup).EvaluationPriority;
+				CompileData.HierarchicalBias        = Params.HierarchicalBias;
+				CompileData.bPriorityTearDown       = EvaluationTrack->HasTearDownPriority();
 
-				ESectionEvaluationFlags Flags = Params.Flags == ESectionEvaluationFlags::None ? Entry.Flags : Params.Flags;
+				auto FindChildWithSection = [Section](const FMovieSceneEvalTemplatePtr& ChildTemplate)
+				{
+					return ChildTemplate.IsValid() && ChildTemplate->GetSourceSection() == Section;
+				};
 
-				if (EnumHasAnyFlags(Params.AccumulatedFlags, EMovieSceneSubSectionFlags::OverrideRestoreState))
+				const int32 ChildTemplateIndex = Section ? EvaluationTrack->GetChildTemplates().IndexOfByPredicate(FindChildWithSection) : INDEX_NONE;
+				if (ChildTemplateIndex != INDEX_NONE)
 				{
-					Flags |= ESectionEvaluationFlags::ForceRestoreState;
+					check(ChildTemplateIndex >= 0 && ChildTemplateIndex < TNumericLimits<uint16>::Max());
+
+					ESectionEvaluationFlags Flags = Params.Flags == ESectionEvaluationFlags::None ? Entry.Flags : Params.Flags;
+
+					if (EnumHasAnyFlags(Params.AccumulatedFlags, EMovieSceneSubSectionFlags::OverrideRestoreState))
+					{
+						Flags |= ESectionEvaluationFlags::ForceRestoreState;
+					}
+					else if (EnumHasAnyFlags(Params.AccumulatedFlags, EMovieSceneSubSectionFlags::OverrideKeepState))
+					{
+						Flags |= ESectionEvaluationFlags::ForceKeepState;
+					}
+
+					CompileData.ChildPriority = Entry.LegacySortOrder;
+					CompileData.Child         = FMovieSceneFieldEntry_ChildTemplate((uint16)ChildTemplateIndex, Flags, Entry.ForcedTime);
+					CompileData.bRequiresInit = EvaluationTrack->GetChildTemplate(ChildTemplateIndex).RequiresInitialization();
 				}
-				else if (EnumHasAnyFlags(Params.AccumulatedFlags, EMovieSceneSubSectionFlags::OverrideKeepState))
+				else
 				{
-					Flags |= ESectionEvaluationFlags::ForceKeepState;
+					CompileData.ChildPriority = 0;
+					CompileData.Child         = FMovieSceneFieldEntry_ChildTemplate{};
+					CompileData.bRequiresInit = false;
 				}
 
-				CompileData.ChildPriority = Entry.LegacySortOrder;
-				CompileData.Child         = FMovieSceneFieldEntry_ChildTemplate((uint16)ChildTemplateIndex, Flags, Entry.ForcedTime);
-				CompileData.bRequiresInit = EvaluationTrack->GetChildTemplate(ChildTemplateIndex).RequiresInitialization();
-			}
-			else
-			{
-				CompileData.ChildPriority = 0;
-				CompileData.Child         = FMovieSceneFieldEntry_ChildTemplate{};
-				CompileData.bRequiresInit = false;
-			}
+				OutCompilerData->TrackTemplates.Add(ClampedRangeRoot, CompileData);
+				return true;
+			};
 
-			OutCompilerData->TrackTemplates.Add(ClampedRangeRoot, CompileData);
+			Params.TransformLocalRange(Entry.Range, VisitWarpedRootRange);
 		}
 	}
 }
@@ -1432,19 +1547,53 @@ bool UMovieSceneCompiledDataManager::CompileHierarchy(UMovieSceneSequence* Seque
 
 bool UMovieSceneCompiledDataManager::CompileHierarchy(UMovieSceneSequence* Sequence, const FGatherParameters& Params, FMovieSceneSequenceHierarchy* InOutHierarchy)
 {
+	using namespace UE::MovieScene;
+
 	UE::MovieScene::FSubSequencePath RootPath;
 
+	FGatherParameters TimeWarpedParams;
+	const FGatherParameters* ParamsToUse = &Params;
+
+	bool bContainsTimeWarp = false;
+
+	if (Params.SequenceID == MovieSceneSequenceID::Root)
+	{
+		UMovieScene* MovieScene = Sequence->GetMovieScene();
+
+		// Generate time warp transforms before any sub sequences
+		for (UMovieSceneTrack* Track : MovieScene->GetTracks())
+		{
+			UMovieSceneTimeWarpTrack* TimeWarpTrack = Cast<UMovieSceneTimeWarpTrack>(Track);
+			if (TimeWarpTrack && !TimeWarpTrack->IsEvalDisabled())
+			{
+				FMovieSceneNestedSequenceTransform TimeWarpTransform = TimeWarpTrack->GenerateTransform();
+
+				// Don't do anything for identity transforms
+				if (!TimeWarpTransform.IsIdentity())
+				{
+					InOutHierarchy->SetRootTransform(FMovieSceneSequenceTransform(MoveTemp(TimeWarpTransform)));
+					bContainsTimeWarp = true;
+				}
+
+				// Only 1 time-wrap track supported
+				break;
+			}
+		}
+	}
+
 	// Compile all the sub data for every part of the hierarchy
-	const bool bContainsSubSequences = GenerateSubSequenceData(Sequence, Params, FMovieSceneEvaluationOperand(), &RootPath, InOutHierarchy);
+	const bool bContainsSubSequences = GenerateSubSequenceData(Sequence, *ParamsToUse, FMovieSceneEvaluationOperand(), &RootPath, InOutHierarchy);
 
 	// Populate the sub sequence tree that defines which sub sequences happen at a given time
-	PopulateSubSequenceTree(Sequence, Params, &RootPath, InOutHierarchy);
+	PopulateSubSequenceTree(Sequence, *ParamsToUse, &RootPath, InOutHierarchy);
 
-	return bContainsSubSequences;
+	return bContainsSubSequences || bContainsTimeWarp;
 }
 
 bool UMovieSceneCompiledDataManager::GenerateSubSequenceData(UMovieSceneSequence* SubSequence, const FGatherParameters& Params, const FMovieSceneEvaluationOperand& Operand, UE::MovieScene::FSubSequencePath* RootPath, FMovieSceneSequenceHierarchy* InOutHierarchy)
 {
+	using namespace UE::MovieScene;
+
 	UMovieScene* MovieScene = SubSequence ? SubSequence->GetMovieScene() : nullptr;
 	if (!MovieScene)
 	{
@@ -1481,6 +1630,8 @@ bool UMovieSceneCompiledDataManager::GenerateSubSequenceData(UMovieSceneSequence
 
 bool UMovieSceneCompiledDataManager::GenerateSubSequenceData(UMovieSceneSubTrack* SubTrack, const FGatherParameters& Params, const FMovieSceneEvaluationOperand& Operand, UE::MovieScene::FSubSequencePath* RootPath, FMovieSceneSequenceHierarchy* InOutHierarchy)
 {
+	using namespace UE::MovieScene;
+
 	bool bContainsSubSequences = false;
 
 	check(SubTrack && RootPath);
@@ -1522,13 +1673,21 @@ bool UMovieSceneCompiledDataManager::GenerateSubSequenceData(UMovieSceneSubTrack
 
 		// LocalClampRange here is in SubTrack's space, so we need to multiply that by the OuterToInnerTransform 
 		// (which is the same as RootToSequenceTransform here before we transform it)
-		TRange<FFrameNumber> InnerClampRange = NewSubData.RootToSequenceTransform.TransformRangeUnwarped(Params.LocalClampRange);
+		TRange<FFrameNumber> InnerClampRange = (Params.LocalClampRange.GetLowerBound().IsOpen() || Params.LocalClampRange.GetUpperBound().IsOpen())
+			? Params.LocalClampRange
+			: ConvertToDiscreteRange(NewSubData.OuterToInnerTransform.ComputeTraversedHull(Params.LocalClampRange));
 
+		// Put the root play range in the new root space
 		NewSubData.PlayRange               = TRange<FFrameNumber>::Intersection(InnerClampRange, NewSubData.PlayRange.Value);
 		NewSubData.RootToSequenceTransform = NewSubData.RootToSequenceTransform * Params.RootToSequenceTransform;
 		NewSubData.HierarchicalBias        = Params.HierarchicalBias + NewSubData.HierarchicalBias;
 		NewSubData.AccumulatedFlags        = UE::MovieScene::AccumulateChildSubSectionFlags(Params.AccumulatedFlags, NewSubData.AccumulatedFlags);
 
+#if WITH_EDITORONLY_DATA
+		NewSubData.StartTimeBreadcrumbs.CombineWithOuterBreadcrumbs(Params.StartTimeBreadcrumbs);
+		NewSubData.EndTimeBreadcrumbs.CombineWithOuterBreadcrumbs(Params.EndTimeBreadcrumbs);
+#endif // WITH_EDITORONLY_DATA
+
 		// Add the sub data to the root hierarchy
 		InOutHierarchy->Add(NewSubData, InnerSequenceID, ParentSequenceID);
 
@@ -1625,129 +1784,40 @@ void UMovieSceneCompiledDataManager::PopulateSubSequenceTree(UMovieSceneSubTrack
 		}
 
 		InOutHierarchy->AccumulateNetworkMask(SubSection->GetNetworkMask());
-		
-
-		const FMovieSceneSequenceTransform SequenceToRootTransform = Params.RootToSequenceTransform.InverseFromLoop(Params.RootToSequenceWarpCounter);
-
-		// In the case the Sequence to Root Transform contains an infinite timescale, then one of the timescales above us is zero. In this case, we cannot
-		// rely on a simple multiply by an inverse transform to figure out an effective range of this entry in root space, as this is non-deterministic.
-		// It all comes down to whether the single frame is inside or outside of this entry. 
-		// If inside, the effective range in root space is the entire clamp range.
-		// If outside, the effective range in root space is empty, and we should not include this entry.
-
-		TRange<FFrameNumber> EffectiveRange = TRange<FFrameNumber>::Empty();
-		if (FMath::IsFinite(SequenceToRootTransform.GetTimeScale()) || !TRange<FFrameNumber>::Intersection(Params.LocalClampRange, Entry.Range).IsEmpty())
-		{
-			EffectiveRange = Params.ClampRoot(SequenceToRootTransform.TransformRangeConstrained(Entry.Range));
-		}
-
-		if (EffectiveRange.IsEmpty())
-		{
-			continue;
-		}
 
 		const FMovieSceneSequenceID       SubSequenceID = RootPath->ResolveChildSequenceID(SubSection->GetSequenceID());
-		const FMovieSceneSubSequenceData* SubData       = InOutHierarchy->FindSubData(SubSequenceID);
+		const FMovieSceneSubSequenceData* SubData = InOutHierarchy->FindSubData(SubSequenceID);
 
 		checkf(SubData, TEXT("Unable to locate sub-data for a sub section that appears in the track's evaluation field - this indicates that the section is being evaluated even though it is not active"));
 
-		const FMovieSceneFrameRange& SubSectionRange = SubSection->SectionRange;
-		const FMovieSceneSectionParameters& SubSectionParams = SubSection->Parameters;
-
-		const ESectionEvaluationFlags SubEntryFlags = Entry.Flags | Params.Flags;
-
-		// If we can't loop, or our timescale is zero, and therefore looping is irrelevant
-		if (!SubSectionParams.bCanLoop || FMath::IsNearlyZero(SubData->RootToSequenceTransform.GetTimeScale()))
+		auto AddRange = [&Params, &Entry, &RootPath, InOutHierarchy, NewMask, SubData, SubSequenceID](TRange<FFrameTime> Range)
 		{
-			FGatherParameters SubParams = Params.CreateForSubData(*SubData, SubSequenceID, Params.RootToSequenceWarpCounter);
-			SubParams.SetClampRange(EffectiveRange);
-			SubParams.Flags |= Entry.Flags;
-			SubParams.NetworkMask = NewMask;
+			TRange<FFrameNumber> FrameRange = Params.ClampRoot(ConvertToDiscreteRange(Range));
 
-			for (int i = 0; i < SubParams.RootToSequenceTransform.NestedTransforms.Num(); ++i)
+			if (!FrameRange.IsEmpty())
 			{
-				if (SubParams.RootToSequenceTransform.NestedTransforms[i].IsLooping())
-				{
-					SubParams.RootToSequenceWarpCounter.AddWarpingLevel(1);
-				}
-				else
-				{
-					SubParams.RootToSequenceWarpCounter.AddNonWarpingLevel();
-				}
-			}
-			
-			// The section isn't looping, so we can just add it to the tree.
-			InOutHierarchy->AddRange(EffectiveRange, SubSequenceID, SubEntryFlags, SubParams.RootToSequenceWarpCounter);
+				FGatherParameters SubParams = Params.CreateForSubData(*SubData, SubSequenceID);
+				SubParams.SetClampRange(FrameRange);
+				SubParams.Flags |= Entry.Flags;
+				SubParams.NetworkMask = NewMask;
 
-			// Recurse into the sub sequence
-			RootPath->PushGeneration(SubSequenceID, SubData->DeterministicSequenceID);
-			{
-				PopulateSubSequenceTree(SubData->GetSequence(), SubParams, RootPath, InOutHierarchy);
-			}
-			RootPath->PopGenerations(1);
-		}
-		else
-		{
-			// The section is looping so we need to add its contents to the tree as many times as it has loops.
-			const float RootToSubSequenceTimeScale = SubData->RootToSequenceTransform.GetTimeScale();
-			const float SubSequenceToRootTimeScale = (RootToSubSequenceTimeScale != 0.f) ? 1.0f / RootToSubSequenceTimeScale : 1.f;
+				const ESectionEvaluationFlags SubEntryFlags = Entry.Flags | Params.Flags;
 
-			UMovieSceneSequence* SubSequence = SubData->GetSequence();
+				// The section isn't looping, so we can just add it to the tree.
+				InOutHierarchy->AddRange(FrameRange, SubSequenceID, SubEntryFlags);
 
-			const TRange<FFrameNumber> SubSequencePlaybackRange = SubSequence->GetMovieScene()->GetPlaybackRange();
-			const FFrameNumber SubSequenceLength = DiscreteSize(SubSequencePlaybackRange);
-			const FFrameTime RootLoopLength = FMath::Max(
-					FFrameTime(SubSequenceLength - SubSectionParams.StartFrameOffset - SubSectionParams.EndFrameOffset) * SubSequenceToRootTimeScale,
-					FFrameTime(1));
-			const FFrameTime FirstRootLoopLength = FMath::Max(
-					RootLoopLength - SubSectionParams.FirstLoopStartFrameOffset * SubSequenceToRootTimeScale,
-					FFrameTime(1));
-			const TOptional<FFrameNumber> LocalSectionExclusiveEndTime = GetLoopingSubSectionEndTime(OuterSequence, SubSection, Params);
-
-			if (LocalSectionExclusiveEndTime.IsSet() && !SubSectionRange.GetLowerBound().IsOpen())
-			{
-				uint32 LoopCount = 0;
-				FFrameTime CurRootRangeStart = DiscreteInclusiveLower(SubSectionRange.GetLowerBound()) * SequenceToRootTransform;
-				TRange<FFrameNumber> CurRootRange(CurRootRangeStart.FloorToFrame(), (CurRootRangeStart + FirstRootLoopLength).CeilToFrame());
-				const FFrameNumber RootSectionEndTime = (LocalSectionExclusiveEndTime.GetValue() * SequenceToRootTransform).CeilToFrame();
-
-				while (CurRootRange.GetLowerBoundValue() < RootSectionEndTime)
+				// Recurse into the sub sequence
+				RootPath->PushGeneration(SubSequenceID, SubData->DeterministicSequenceID);
 				{
-					if (CurRootRange.Overlaps(Params.RootClampRange))
-					{
-						// Clamp the sub-sequence's range by the containing section's range and the current compilation range.
-						const TRange<FFrameNumber> ClampedCurRootRange = TRange<FFrameNumber>::Intersection(CurRootRange, EffectiveRange);
-
-						FGatherParameters CurLoopParams = Params.CreateForSubData(*SubData, SubSequenceID, Params.RootToSequenceWarpCounter);
-						CurLoopParams.SetClampRange(ClampedCurRootRange);
-						CurLoopParams.Flags |= Entry.Flags;
-						CurLoopParams.NetworkMask = NewMask;
-						CurLoopParams.RootToSequenceWarpCounter.AddWarpingLevel(LoopCount);
-
-						// Add the section to the tree for the current loop.
-						InOutHierarchy->AddRange(ClampedCurRootRange, SubSequenceID, SubEntryFlags, CurLoopParams.RootToSequenceWarpCounter);
-
-						// Recurse into this loop's sub sequence.
-						RootPath->PushGeneration(SubSequenceID, SubData->DeterministicSequenceID);
-						{
-							PopulateSubSequenceTree(SubData->GetSequence(), CurLoopParams, RootPath, InOutHierarchy);
-						}
-						RootPath->PopGenerations(1);
-					}
-
-					// Move on to the next loop.
-					CurRootRangeStart = CurRootRange.GetUpperBoundValue();
-					CurRootRange = TRange<FFrameNumber>(CurRootRangeStart.FloorToFrame(), (CurRootRangeStart + RootLoopLength).FloorToFrame());
-					if (CurRootRange.GetUpperBoundValue() > RootSectionEndTime)
-						CurRootRange.SetUpperBoundValue(RootSectionEndTime);
-					++LoopCount;
+					PopulateSubSequenceTree(SubData->GetSequence(), SubParams, RootPath, InOutHierarchy);
 				}
+				RootPath->PopGenerations(1);
 			}
-			// Faced with the cosmic horror of infinites, we choose to shield our sanity and skip this sub-section.
-			// (it either has an open-ended start time, which means we needed to loop since before time began, which means
-			//  we don't know where loops are in the present... or it means the section and root sequence have open-ended
-			//  end times, which means we would need to compile loops forever)
-		}
+
+			return true;
+		};
+
+		Params.TransformLocalRange(Entry.Range, AddRange);
 	}
 }
 
diff --git a/Engine/Source/Runtime/MovieScene/Private/EntitySystem/MovieSceneEvaluationHookSystem.cpp b/Engine/Source/Runtime/MovieScene/Private/EntitySystem/MovieSceneEvaluationHookSystem.cpp
--- a/Engine/Source/Runtime/MovieScene/Private/EntitySystem/MovieSceneEvaluationHookSystem.cpp
+++ b/Engine/Source/Runtime/MovieScene/Private/EntitySystem/MovieSceneEvaluationHookSystem.cpp
@@ -43,11 +43,12 @@ struct FEvaluationHookUpdater
 			}
 
 			const FSequenceInstance& SequenceInstance = InstanceRegistry->GetInstance(InstanceHandles[Index]);
+			FMovieSceneInverseSequenceTransform SequenceToRootTransform = SequenceInstance.GetContext().GetRootToSequenceTransform().Inverse();
 
 			FMovieSceneEvaluationHookEvent NewEvent;
 			NewEvent.Hook          = Hooks[Index];
 			NewEvent.Type          = EEvaluationHookEvent::Update;
-			NewEvent.RootTime      = EvalTimes[Index] * SequenceInstance.GetContext().GetSequenceToRootSequenceTransform();
+			NewEvent.RootTime      = SequenceToRootTransform.TryTransformTime(EvalTimes[Index], SequenceInstance.GetContext().GetRootToSequenceWarpCounter()).Get(FFrameTime());
 			NewEvent.RootInstanceHandle = SequenceInstance.GetRootInstanceHandle();
 			NewEvent.SequenceID    = SequenceInstance.GetSequenceID();
 			NewEvent.bRestoreState = bRestoreState;
@@ -150,11 +151,12 @@ void UMovieSceneEvaluationHookSystem::UpdateHooks()
 		for (int32 Index = 0; Index < Num; ++Index)
 		{
 			const FSequenceInstance& SequenceInstance = InstanceRegistry->GetInstance(InstanceHandles[Index]);
+			FMovieSceneInverseSequenceTransform SequenceToRootTransform = SequenceInstance.GetContext().GetRootToSequenceTransform().Inverse();
 
 			FMovieSceneEvaluationHookEvent NewEvent;
 			NewEvent.Hook          = Hooks[Index];
 			NewEvent.Type          = EEvaluationHookEvent::Begin;
-			NewEvent.RootTime      = EvalTimes[Index] * SequenceInstance.GetContext().GetSequenceToRootSequenceTransform();
+			NewEvent.RootTime      = SequenceToRootTransform.TryTransformTime(EvalTimes[Index], SequenceInstance.GetContext().GetRootToSequenceWarpCounter()).Get(FFrameTime());
 			NewEvent.RootInstanceHandle = SequenceInstance.GetRootInstanceHandle();
 			NewEvent.SequenceID    = SequenceInstance.GetSequenceID();
 			NewEvent.bRestoreState = bRestoreState;
@@ -171,11 +173,12 @@ void UMovieSceneEvaluationHookSystem::UpdateHooks()
 		for (int32 Index = 0; Index < Num; ++Index)
 		{
 			const FSequenceInstance& SequenceInstance = InstanceRegistry->GetInstance(InstanceHandles[Index]);
+			FMovieSceneInverseSequenceTransform SequenceToRootTransform = SequenceInstance.GetContext().GetRootToSequenceTransform().Inverse();
 
 			FMovieSceneEvaluationHookEvent NewEvent;
 			NewEvent.Hook          = Hooks[Index];
 			NewEvent.Type          = EEvaluationHookEvent::End;
-			NewEvent.RootTime      = EvalTimes[Index] * SequenceInstance.GetContext().GetSequenceToRootSequenceTransform();
+			NewEvent.RootTime      = SequenceToRootTransform.TryTransformTime(EvalTimes[Index], SequenceInstance.GetContext().GetRootToSequenceWarpCounter()).Get(FFrameTime());
 			NewEvent.RootInstanceHandle = SequenceInstance.GetRootInstanceHandle();
 			NewEvent.SequenceID    = SequenceInstance.GetSequenceID();
 			NewEvent.bRestoreState = bRestoreState;
diff --git a/Engine/Source/Runtime/MovieScene/Private/EntitySystem/MovieSceneSequenceUpdaters.cpp b/Engine/Source/Runtime/MovieScene/Private/EntitySystem/MovieSceneSequenceUpdaters.cpp
--- a/Engine/Source/Runtime/MovieScene/Private/EntitySystem/MovieSceneSequenceUpdaters.cpp
+++ b/Engine/Source/Runtime/MovieScene/Private/EntitySystem/MovieSceneSequenceUpdaters.cpp
@@ -11,6 +11,7 @@
 #include "Containers/BitArray.h"
 #include "Containers/SortedMap.h"
 
+#include "MovieSceneTransformTypes.h"
 #include "MovieSceneSequence.h"
 #include "MovieSceneSequenceID.h"
 #include "Evaluation/MovieScenePlayback.h"
@@ -476,7 +477,7 @@ void FSequenceUpdater_Hierarchical::DissectContext(TSharedRef<const FSharedPlayb
 		FMovieSceneEvaluationTreeRangeIterator SubSequenceIt = Hierarchy->GetTree().IterateFromLowerBound(TraversedRange.GetLowerBound());
 		for ( ; SubSequenceIt && SubSequenceIt.Range().Overlaps(TraversedRange); ++SubSequenceIt)
 		{
-			TRange<FFrameTime> RootClampRange = TRange<FFrameTime>::Intersection(ConvertRange<FFrameNumber, FFrameTime>(SubSequenceIt.Range()), Context.GetRange());
+			TRange<FFrameTime> RootClampRange = TRange<FFrameTime>::Intersection(ConvertToFrameTimeRange(SubSequenceIt.Range()), Context.GetRange());
 
 			// When Context.GetRange() does not fall on whole frame boundaries, we can sometimes end up with a range that clamps to being empty, even though the range overlapped
 			// the traversed range. ie if we evaluated range (1.5, 10], our traversed range would be [2, 11). If we have a sub sequence range of (10, 20), it would still be iterated here
@@ -501,17 +502,34 @@ void FSequenceUpdater_Hierarchical::DissectContext(TSharedRef<const FSharedPlayb
 				TArrayView<const FMovieSceneDeterminismFence> SubDeterminismFences = CompiledDataManager->GetEntryRef(SubDataID).DeterminismFences;
 				if (SubDeterminismFences.Num() > 0)
 				{
-					TRange<FFrameTime>   InnerRange           = SubData->RootToSequenceTransform.TransformRangeUnwarped(RootClampRange);
+					TRange<FFrameTime> InnerRange = SubData->RootToSequenceTransform.ComputeTraversedHull(RootClampRange);
 
-					TArrayView<const FMovieSceneDeterminismFence> TraversedFences  = GetFencesWithinRange(SubDeterminismFences, InnerRange);
-					if (TraversedFences.Num() > 0)
+					// Time-warp can result in inside-out ranges
+					if (InnerRange.GetLowerBound().IsClosed() && InnerRange.GetUpperBound().IsClosed() && InnerRange.GetLowerBoundValue() > InnerRange.GetUpperBoundValue())
 					{
-						FMovieSceneWarpCounter WarpCounter;
-						FFrameTime Unused;
-						SubData->RootToSequenceTransform.TransformTime(RootClampRange.GetLowerBoundValue(), Unused, WarpCounter);
+						TRangeBound<FFrameTime> OldLower = InnerRange.GetLowerBound();
+						TRangeBound<FFrameTime> OldUpper = InnerRange.GetUpperBound();
+						InnerRange.SetLowerBound(OldUpper);
+						InnerRange.SetUpperBound(OldLower);
+					}
 
-						FMovieSceneSequenceTransform InverseTransform = SubData->RootToSequenceTransform.InverseFromLoop(WarpCounter);
-						Algo::Transform(TraversedFences, RootDissectionTimes, [InverseTransform](const FMovieSceneDeterminismFence& In){ return FMovieSceneDeterminismFenceWithSubframe{In.FrameNumber * InverseTransform, In.bInclusive }; });
+					TArrayView<const FMovieSceneDeterminismFence> TraversedFences = GetFencesWithinRange(SubDeterminismFences, InnerRange);
+					if (TraversedFences.Num() > 0)
+					{
+						// Find the breadcrumbs for this range
+						FMovieSceneTransformBreadcrumbs Breadcrumbs;
+						SubData->RootToSequenceTransform.TransformTime(RootClampRange.GetLowerBoundValue(), FTransformTimeParams().HarvestBreadcrumbs(Breadcrumbs));
+
+						FMovieSceneInverseSequenceTransform InverseTransform = SubData->RootToSequenceTransform.Inverse();
+
+						for (FMovieSceneDeterminismFence Fence : TraversedFences)
+						{
+							TOptional<FFrameTime> RootTime = InverseTransform.TryTransformTime(Fence.FrameNumber, Breadcrumbs);
+							if (RootTime)
+							{
+								RootDissectionTimes.Emplace(FMovieSceneDeterminismFenceWithSubframe{ RootTime.GetValue(), Fence.bInclusive });
+							}
+						}
 					}
 				}
 			}
@@ -616,6 +634,10 @@ void FSequenceUpdater_Hierarchical::Update(TSharedRef<const FSharedPlaybackState
 			ActiveSequences.Add(RootOverrideSequenceID);
 		}
 	}
+	else if (RootHierarchy && !RootHierarchy->GetRootTransform().IsIdentity())
+	{
+		RootContext = Context.Transform(RootHierarchy->GetRootTransform(), Context.GetFrameRate());
+	}
 
 	FFrameNumber ImportTime = RootContext.GetEvaluationFieldTime();
 	const bool bGatherEntities = !CachedEntityRange.Contains(ImportTime);
@@ -770,24 +792,36 @@ void FSequenceUpdater_Hierarchical::Update(TSharedRef<const FSharedPlaybackState
 					EntitiesScratch.Reset();
 
 					TRange<FFrameNumber> SubEntityRange = UpdateEntitiesForSequence(SubComponentField, SubSequenceTime, EntitiesScratch);
+					SubEntityRange = TRange<FFrameNumber>::Intersection(SubEntityRange, SubData->PlayRange.Value);
 
 					SubSequenceInstance.Ledger.UpdateEntities(Linker, Params, SubComponentField, EntitiesScratch);
 
-					// Clamp to the current warp loop if necessary
-					FMovieSceneWarpCounter WarpCounter;
-					FFrameTime Unused;
-					SubData->RootToSequenceTransform.TransformTime(ImportTime, Unused, WarpCounter);
+					// Convert sub entity range into root space
+					// 
+					// Sometimes the bounds can be unset if the lower bound does not map to any valid time in the root sequence.
+					//   If this happens, we rely in the intersection with SubSequenceIt.Range() to clamp to the bounds of the current sub sequence range
+					FMovieSceneInverseSequenceTransform Inv = SubContext.GetSequenceToRootSequenceTransform();
 
-					if (WarpCounter.WarpCounts.Num() > 0)
+					TRange<FFrameNumber> SubCachedRange = TRange<FFrameNumber>::All();
+					if (!SubEntityRange.GetLowerBound().IsOpen())
 					{
-						FMovieSceneSequenceTransform InverseTransform = SubData->RootToSequenceTransform.InverseFromLoop(WarpCounter);
-						CachedEntityRange = TRange<FFrameNumber>::Intersection(CachedEntityRange, InverseTransform.TransformRangeConstrained(SubData->PlayRange.Value));
+						TOptional<FFrameTime> LowerBoundRootSpace = Inv.TryTransformTime(SubEntityRange.GetLowerBoundValue(), SubContext.GetRootToSequenceWarpCounter());
+						if (LowerBoundRootSpace)
+						{
+							SubCachedRange.SetLowerBound(TRangeBound<FFrameNumber>::Inclusive(LowerBoundRootSpace.GetValue().CeilToFrame()));
+						}
 					}
 
-					const FMovieSceneSequenceTransform SequenceToRootOverrideTransform = SubContext.GetSequenceToRootSequenceTransform() * RootContext.GetRootToSequenceTransform();
-					SubEntityRange = SequenceToRootOverrideTransform.TransformRangeConstrained(SubEntityRange);
+					if (!SubEntityRange.GetUpperBound().IsOpen())
+					{
+						TOptional<FFrameTime> UpperBoundRootSpace = Inv.TryTransformTime(SubEntityRange.GetUpperBoundValue(), SubContext.GetRootToSequenceWarpCounter());
+						if (UpperBoundRootSpace)
+						{
+							SubCachedRange.SetUpperBound(TRangeBound<FFrameNumber>::Exclusive(UpperBoundRootSpace.GetValue().FloorToFrame()));
+						}
+					}
 
-					CachedEntityRange = TRange<FFrameNumber>::Intersection(CachedEntityRange, SubEntityRange);
+					CachedEntityRange = TRange<FFrameNumber>::Intersection(CachedEntityRange, SubCachedRange);
 				}
 
 				// Update any one-shot entities for the sub sequence
diff --git a/Engine/Source/Runtime/MovieScene/Private/Evaluation/MovieScenePlayback.cpp b/Engine/Source/Runtime/MovieScene/Private/Evaluation/MovieScenePlayback.cpp
--- a/Engine/Source/Runtime/MovieScene/Private/Evaluation/MovieScenePlayback.cpp
+++ b/Engine/Source/Runtime/MovieScene/Private/Evaluation/MovieScenePlayback.cpp
@@ -1,7 +1,7 @@
 // Copyright Epic Games, Inc. All Rights Reserved.
 
 #include "Evaluation/MovieScenePlayback.h"
-
+#include "MovieSceneTransformTypes.h"
 #include "MovieScene.h"
 
 namespace 
@@ -142,6 +142,44 @@ TRange<FFrameTime> FMovieSceneEvaluationRange::NumberRangeToTimeRange(const TRan
 	return FrameTimeRange;
 }
 
+
+FMovieSceneContext FMovieSceneContext::Transform(const FMovieSceneSequenceTransform& InTransform, FFrameRate NewFrameRate) const
+{
+	using namespace UE::MovieScene;
+
+	FMovieSceneContext NewContext = *this;
+	NewContext.RootToSequenceTransform = NewContext.RootToSequenceTransform * InTransform;
+	NewContext.CurrentFrameRate = NewFrameRate;
+
+	NewContext.EvaluationRange = InTransform.ComputeTraversedHull(EvaluationRange);
+
+	if (NewContext.EvaluationRange.GetLowerBound().IsClosed() && NewContext.EvaluationRange.GetUpperBound().IsClosed() && NewContext.EvaluationRange.GetLowerBoundValue() > NewContext.EvaluationRange.GetUpperBoundValue())
+	{
+		TRangeBound<FFrameTime> OldLower = NewContext.EvaluationRange.GetLowerBound();
+		TRangeBound<FFrameTime> OldUpper = NewContext.EvaluationRange.GetUpperBound();
+		NewContext.EvaluationRange.SetLowerBound(OldUpper);
+		NewContext.EvaluationRange.SetUpperBound(OldLower);
+		NewContext.Direction = EPlayDirection::Backwards;
+	}
+
+	// Transform the current time so we get an idea in what loop(s) we are relative to the root sequence.
+	InTransform.TransformTime(GetTime(), FTransformTimeParams().AppendBreadcrumbs(NewContext.RootToSequenceWarpCounter));
+
+	return NewContext;
+}
+
+PRAGMA_DISABLE_DEPRECATION_WARNINGS
+FMovieSceneTimeTransform FMovieSceneContext::GetSequenceToRootTransform() const
+{
+	return RootToSequenceTransform.Inverse().AsLegacyLinearTimeTransform();
+}
+PRAGMA_ENABLE_DEPRECATION_WARNINGS
+
+FMovieSceneInverseSequenceTransform FMovieSceneContext::GetSequenceToRootSequenceTransform() const
+{
+	return RootToSequenceTransform.Inverse();
+}
+
 void FMovieScenePlaybackPosition::CheckInvariants() const
 {
 	checkf(InputRate.IsValid() && OutputRate.IsValid(), TEXT("Invalid input or output rate. SetTimeBase must be called before any use of this class."))
diff --git a/Engine/Source/Runtime/MovieScene/Private/Evaluation/MovieSceneSequenceHierarchy.cpp b/Engine/Source/Runtime/MovieScene/Private/Evaluation/MovieSceneSequenceHierarchy.cpp
--- a/Engine/Source/Runtime/MovieScene/Private/Evaluation/MovieSceneSequenceHierarchy.cpp
+++ b/Engine/Source/Runtime/MovieScene/Private/Evaluation/MovieSceneSequenceHierarchy.cpp
@@ -7,11 +7,27 @@
 #include "MovieSceneSequence.h"
 #include "MovieScene.h"
 #include "MovieSceneTimeHelpers.h"
+#include "MovieSceneTransformTypes.h"
 #include "UObject/ReleaseObjectVersion.h"
 #include "UObject/FortniteMainBranchObjectVersion.h"
 
 #include UE_INLINE_GENERATED_CPP_BY_NAME(MovieSceneSequenceHierarchy)
 
+PRAGMA_DISABLE_DEPRECATION_WARNINGS
+FMovieSceneSubSequenceData::FMovieSceneSubSequenceData(const FMovieSceneSubSequenceData&) = default;
+FMovieSceneSubSequenceData& FMovieSceneSubSequenceData::operator=(const FMovieSceneSubSequenceData&) = default;
+FMovieSceneSubSequenceData::FMovieSceneSubSequenceData(FMovieSceneSubSequenceData&&) = default;
+FMovieSceneSubSequenceData& FMovieSceneSubSequenceData::operator=(FMovieSceneSubSequenceData&&) = default;
+FMovieSceneSubSequenceData::~FMovieSceneSubSequenceData() = default;
+
+FMovieSceneSubSequenceTreeEntry::FMovieSceneSubSequenceTreeEntry() = default;
+FMovieSceneSubSequenceTreeEntry::FMovieSceneSubSequenceTreeEntry(const FMovieSceneSubSequenceTreeEntry&) = default;
+FMovieSceneSubSequenceTreeEntry& FMovieSceneSubSequenceTreeEntry::operator=(const FMovieSceneSubSequenceTreeEntry&) = default;
+FMovieSceneSubSequenceTreeEntry::FMovieSceneSubSequenceTreeEntry(FMovieSceneSubSequenceTreeEntry&&) = default;
+FMovieSceneSubSequenceTreeEntry& FMovieSceneSubSequenceTreeEntry::operator=(FMovieSceneSubSequenceTreeEntry&&) = default;
+FMovieSceneSubSequenceTreeEntry::~FMovieSceneSubSequenceTreeEntry() = default;
+PRAGMA_ENABLE_DEPRECATION_WARNINGS
+
 FMovieSceneSubSequenceData::FMovieSceneSubSequenceData()
 	: Sequence(nullptr)
 	, HierarchicalBias(0)
@@ -22,61 +38,59 @@ FMovieSceneSubSequenceData::FMovieSceneSubSequenceData(const UMovieSceneSubSecti
 	: Sequence(InSubSection.GetSequence())
 	, DeterministicSequenceID(InSubSection.GetSequenceID())
 	, ParentPlayRange(InSubSection.GetTrueRange())
-	, ParentStartFrameOffset(InSubSection.Parameters.StartFrameOffset)
-	, ParentEndFrameOffset(InSubSection.Parameters.EndFrameOffset)
-	, ParentFirstLoopStartFrameOffset(InSubSection.Parameters.FirstLoopStartFrameOffset)
-	, bCanLoop(InSubSection.Parameters.bCanLoop)
 	, HierarchicalBias(InSubSection.Parameters.HierarchicalBias)
 	, AccumulatedFlags(InSubSection.Parameters.Flags)
+	, bCanLoop(InSubSection.Parameters.bCanLoop)
 #if WITH_EDITORONLY_DATA
 	, SectionPath(*InSubSection.GetPathNameInMovieScene())
 #endif
 	, SubSectionSignature(InSubSection.GetSignature())
 {
+	using namespace UE::MovieScene;
+
 	PreRollRange.Value  = TRange<FFrameNumber>::Empty();
 	PostRollRange.Value = TRange<FFrameNumber>::Empty();
 
+
 	UMovieSceneSequence* SequencePtr   = GetSequence();
 	UMovieScene*         MovieScenePtr = SequencePtr ? SequencePtr->GetMovieScene() : nullptr;
 
 	checkf(MovieScenePtr, TEXT("Attempting to construct sub sequence data with a null sequence."));
 
 	TickResolution = MovieScenePtr->GetTickResolution();
-	FullPlayRange.Value = MovieScenePtr->GetPlaybackRange();
 
-	TRange<FFrameNumber> SubSectionRange = InSubSection.GetTrueRange();
-	checkf(SubSectionRange.GetLowerBound().IsClosed() && SubSectionRange.GetUpperBound().IsClosed(), TEXT("Use of open (infinite) bounds with sub sections is not supported."));
+	checkf(ParentPlayRange.GetLowerBound().IsClosed() && ParentPlayRange.GetUpperBound().IsClosed(), TEXT("Use of open (infinite) bounds with sub sections is not supported."));
 
 	// Get the transform from the given section to its inner sequence.
 	// Note that FMovieSceneCompiler will accumulate RootToSequenceTransform for us a bit later so that it ends up
 	// being truly the full transform.
-	OuterToInnerTransform = RootToSequenceTransform = InSubSection.OuterToInnerTransform();
+	OuterToInnerTransform = InSubSection.OuterToInnerTransform_NoInnerTimeWarp();
 
-	if (!InSubSection.Parameters.bCanLoop || FMath::IsNearlyZero(RootToSequenceTransform.GetTimeScale()))
-	{
-		PlayRange.Value = RootToSequenceTransform.TransformRangeUnwarped(SubSectionRange);
-		UnwarpedPlayRange.Value = PlayRange.Value;
-	}
-	else
-	{
-		// If we're looping, there's a good chance we need the entirety of the sub-sequence to be compiled.
-		PlayRange.Value = MovieScenePtr->GetPlaybackRange();
-		PlayRange.Value.SetLowerBoundValue(PlayRange.Value.GetLowerBoundValue() + InSubSection.Parameters.StartFrameOffset);
-		PlayRange.Value.SetUpperBoundValue(FMath::Max(
-			PlayRange.Value.GetUpperBoundValue() - InSubSection.Parameters.EndFrameOffset,
-			PlayRange.Value.GetLowerBoundValue() + 1));
-		UnwarpedPlayRange.Value = RootToSequenceTransform.TransformRangeUnwarped(SubSectionRange);
-	}
+	InSubSection.AppendInnerTimeWarpTransform(OuterToInnerTransform);
+	RootToSequenceTransform = OuterToInnerTransform;
 
-	// Make sure pre/postroll *ranges* are in the inner sequence's time space. Pre/PostRollFrames are in the outer sequence space.
+	PlayRange.Value = UE::MovieScene::ConvertToDiscreteRange(OuterToInnerTransform.ComputeTraversedHull(ParentPlayRange.Value));
+
+#if WITH_EDITORONLY_DATA
+	// Set initial breadcrumbs
+	OuterToInnerTransform.TransformTime(InSubSection.GetInclusiveStartFrame(), FTransformTimeParams().HarvestBreadcrumbs(StartTimeBreadcrumbs));
+	OuterToInnerTransform.TransformTime(InSubSection.GetExclusiveEndFrame(), FTransformTimeParams().HarvestBreadcrumbs(EndTimeBreadcrumbs));
+#endif // WITH_EDITORONLY_DATA
 
+	// Make sure pre/postroll *ranges* are in the inner sequence's time space. Pre/PostRollFrames are in the outer sequence space.
 	if (InSubSection.GetPreRollFrames() > 0)
 	{
-		PreRollRange = RootToSequenceTransform.TransformRangeUnwarped(UE::MovieScene::MakeDiscreteRangeFromUpper(TRangeBound<FFrameNumber>::FlipInclusion(SubSectionRange.GetLowerBound()), InSubSection.GetPreRollFrames()));
+		FFrameNumber PreRollStart = OuterToInnerTransform.TransformTime(InSubSection.GetInclusiveStartFrame() - InSubSection.GetPreRollFrames()).FloorToFrame();
+		FFrameNumber PreRollEnd   = OuterToInnerTransform.TransformTime(InSubSection.GetInclusiveStartFrame()).FloorToFrame();
+
+		PreRollRange = TRange<FFrameNumber>(PreRollStart, PreRollEnd);
 	}
 	if (InSubSection.GetPostRollFrames() > 0)
 	{
-		PostRollRange = RootToSequenceTransform.TransformRangeUnwarped(UE::MovieScene::MakeDiscreteRangeFromLower(TRangeBound<FFrameNumber>::FlipInclusion(SubSectionRange.GetUpperBound()), InSubSection.GetPostRollFrames()));
+		FFrameNumber PostRollStart = OuterToInnerTransform.TransformTime(InSubSection.GetExclusiveEndFrame()).FloorToFrame();
+		FFrameNumber PostRollEnd   = OuterToInnerTransform.TransformTime(InSubSection.GetExclusiveEndFrame() + InSubSection.GetPostRollFrames()).FloorToFrame();
+
+		PostRollRange = TRange<FFrameNumber>(PostRollStart, PostRollEnd);
 	}
 }
 
@@ -110,13 +124,7 @@ bool FMovieSceneSubSequenceData::IsDirty(const UMovieSceneSubSection& InSubSecti
 
 FMovieSceneSectionParameters FMovieSceneSubSequenceData::ToSubSectionParameters() const
 {
-	FMovieSceneSectionParameters Parameters;
-	Parameters.StartFrameOffset = ParentStartFrameOffset;
-	Parameters.bCanLoop = bCanLoop;
-	Parameters.EndFrameOffset = ParentEndFrameOffset;
-	Parameters.FirstLoopStartFrameOffset = ParentFirstLoopStartFrameOffset;
-	Parameters.TimeScale = OuterToInnerTransform.GetTimeScale();
-	return Parameters;
+	return FMovieSceneSectionParameters();
 }
 
 void FMovieSceneSequenceHierarchy::Add(const FMovieSceneSubSequenceData& Data, FMovieSceneSequenceIDRef ThisSequenceID, FMovieSceneSequenceIDRef ParentID)
@@ -190,31 +198,24 @@ void FMovieSceneSequenceHierarchy::Remove(TArrayView<const FMovieSceneSequenceID
 
 void FMovieSceneSequenceHierarchy::AddRange(const TRange<FFrameNumber>& RootSpaceRange, FMovieSceneSequenceIDRef InSequenceID, ESectionEvaluationFlags InFlags, FMovieSceneWarpCounter RootToSequenceWarpCounter)
 {
-	Tree.Data.AddUnique(RootSpaceRange, FMovieSceneSubSequenceTreeEntry{ InSequenceID, InFlags, RootToSequenceWarpCounter });
+	AddRange(RootSpaceRange, InSequenceID, InFlags);
 }
 
-FArchive& operator<<(FArchive& Ar, FMovieSceneSubSequenceTreeEntry& InOutEntry)
+void FMovieSceneSequenceHierarchy::AddRange(const TRange<FFrameNumber>& RootSpaceRange, FMovieSceneSequenceIDRef InSequenceID, ESectionEvaluationFlags InFlags)
 {
-	Ar.UsingCustomVersion(FReleaseObjectVersion::GUID);
-	Ar.UsingCustomVersion(FFortniteMainBranchObjectVersion::GUID);
+	Tree.Data.AddUnique(RootSpaceRange, FMovieSceneSubSequenceTreeEntry{ InSequenceID, InFlags });
+}
 
+FArchive& operator<<(FArchive& Ar, FMovieSceneSubSequenceTreeEntry& InOutEntry)
+{
 	Ar << InOutEntry.SequenceID << InOutEntry.Flags;
-
-	if (Ar.CustomVer(FReleaseObjectVersion::GUID) >= FReleaseObjectVersion::AddedSubSequenceEntryWarpCounter ||
-		Ar.CustomVer(FFortniteMainBranchObjectVersion::GUID) >= FFortniteMainBranchObjectVersion::AddedSubSequenceEntryWarpCounter)
-	{
-		FMovieSceneWarpCounter::StaticStruct()->SerializeTaggedProperties(
-				Ar, (uint8*)&InOutEntry.RootToSequenceWarpCounter, FMovieSceneWarpCounter::StaticStruct(), nullptr);
-	}
-
 	return Ar;
 }
 
 bool operator==(const FMovieSceneSubSequenceTreeEntry& A, const FMovieSceneSubSequenceTreeEntry& B)
 {
 	return A.SequenceID == B.SequenceID 
-		&& A.Flags == B.Flags
-		&& A.RootToSequenceWarpCounter == B.RootToSequenceWarpCounter;
+		&& A.Flags == B.Flags;
 }
 
 #if !NO_LOGGING
@@ -247,7 +248,7 @@ void FMovieSceneSequenceHierarchy::LogHierarchy() const
 					*Indent, *CurData->GetSequence()->GetName(),
 					*LexToString(CurData->bCanLoop),
 					CurData->HierarchicalBias,
-					*LexToString(CurData->UnwarpedPlayRange.Value),
+					*LexToString(CurData->ParentPlayRange.Value),
 					*LexToString(CurData->RootToSequenceTransform));
 		}
 
@@ -271,9 +272,8 @@ void FMovieSceneSequenceHierarchy::LogSubSequenceTree() const
 	Formatter.DataFormatter = TreeFormatter::FOnFormatData::CreateLambda(
 			[](const FMovieSceneSubSequenceTreeEntry& Entry, TStringBuilder<256>& Builder)
 			{
-				Builder.Appendf(TEXT("ID=%d Warps=%s"), 
-						Entry.SequenceID.GetInternalValue(),
-						*LexToString(Entry.RootToSequenceWarpCounter));
+				Builder.Appendf(TEXT("ID=%d"),
+					Entry.SequenceID.GetInternalValue());
 			});
 	Formatter.LogTree();
 }
