Sequencer: Added time-warp capabilities to sequences, sub-sections and skeletal animation sections

diff --git a/Engine/Source/Runtime/MovieScene/Private/Evaluation/MovieSceneSequenceTransform.cpp b/Engine/Source/Runtime/MovieScene/Private/Evaluation/MovieSceneSequenceTransform.cpp
--- a/Engine/Source/Runtime/MovieScene/Private/Evaluation/MovieSceneSequenceTransform.cpp
+++ b/Engine/Source/Runtime/MovieScene/Private/Evaluation/MovieSceneSequenceTransform.cpp
@@ -1,9 +1,61 @@
 // Copyright Epic Games, Inc. All Rights Reserved.
 
 #include "Evaluation/MovieSceneSequenceTransform.h"
+#include "MovieSceneTimeHelpers.h"
+#include "Misc/FrameRate.h"
+#include "Variants/MovieSceneTimeWarpGetter.h"
+#include "Variants/MovieSceneTimeWarpVariantPayloads.h"
+#include "MovieSceneTransformTypes.h"
+#include "MovieSceneTimeHelpers.h"
+#include "Algo/AllOf.h"
+#include "Algo/AnyOf.h"
 
 #include UE_INLINE_GENERATED_CPP_BY_NAME(MovieSceneSequenceTransform)
 
+namespace UE::MovieScene
+{
+	TRange<FFrameTime> TranslateRange(const TRange<FFrameTime>& InRange, FFrameTime Offset)
+	{
+		TRange<FFrameTime> Result = InRange;
+		if (InRange.HasLowerBound())
+		{
+			Result.SetLowerBoundValue(InRange.GetLowerBoundValue() + Offset);
+		}
+		if (InRange.HasUpperBound())
+		{
+			Result.SetUpperBoundValue(InRange.GetUpperBoundValue() + Offset);
+		}
+		return Result;
+	}
+
+
+	FTransformTimeParams& FTransformTimeParams::HarvestBreadcrumbs(FMovieSceneTransformBreadcrumbs& OutBreadcrumbs)
+	{
+		OutBreadcrumbs.Reset();
+		Breadcrumbs = &OutBreadcrumbs;
+		return *this;
+	}
+
+	FTransformTimeParams& FTransformTimeParams::AppendBreadcrumbs(FMovieSceneTransformBreadcrumbs& OutBreadcrumbs)
+	{
+		Breadcrumbs = &OutBreadcrumbs;
+		return *this;
+	}
+
+	FTransformTimeParams& FTransformTimeParams::TrackCycleCounts(TOptional<int32>* OutCycleCounter)
+	{
+		CycleCount = OutCycleCounter;
+		return *this;
+	}
+
+	FTransformTimeParams& FTransformTimeParams::IgnoreClamps()
+	{
+		bIgnoreClamps = true;
+		return *this;
+	}
+
+} // namespace UE::MovieScene;
+
 FString LexToString(const FMovieSceneSequenceTransform& InTransform)
 {
 	TStringBuilder<256> Builder;
@@ -13,27 +65,10 @@ FString LexToString(const FMovieSceneSequenceTransform& InTransform)
 	int32 NestedIndex = 0;
 	for (const FMovieSceneNestedSequenceTransform& Nested : InTransform.NestedTransforms)
 	{
-		const bool bHasLinear = !Nested.LinearTransform.IsIdentity();
-		const bool bHasWarping = Nested.Warping.IsValid();
-
-		if (bHasLinear || bHasWarping)
+		if (Nested.IsIdentity())
 		{
 			Builder.Appendf(TEXT(" [ %d = "), NestedIndex);
-		}
-		if (bHasLinear)
-		{
-			Builder.Append(LexToString(Nested.LinearTransform));
-		}
-		if (bHasWarping)
-		{
-			if (bHasLinear)
-			{
-				Builder.Append(TEXT(" | "));
-			}
-			Builder.Append(LexToString(Nested.Warping));
-		}
-		if (bHasLinear || bHasWarping)
-		{
+			Nested.ToString(Builder);
 			Builder.Append(TEXT(" ]"));
 		}
 
@@ -45,7 +80,7 @@ FString LexToString(const FMovieSceneSequenceTransform& InTransform)
 
 FString LexToString(const FMovieSceneWarpCounter& InCounter)
 {
-	if (InCounter.WarpCounts.Num() == 0)
+	if (InCounter.Num() == 0)
 	{
 		return FString(TEXT("[]"));
 	}
@@ -54,13 +89,13 @@ FString LexToString(const FMovieSceneWarpCounter& InCounter)
 
 	Builder.Append(TEXT("["));
 	int32 Index = 0;
-	for (uint16 Loop : InCounter.WarpCounts)
+	for (FFrameTime Breadcrumb : InCounter)
 	{
 		if (Index > 0)
 		{
 			Builder.Append(TEXT(","));
 		}
-		Builder.Appendf(TEXT("%u"), (uint32)Loop);
+		Builder.Appendf(TEXT("%.3f"), Breadcrumb.AsDecimal());
 		++Index;
 	}
 	Builder.Append(TEXT("]"));
@@ -69,4 +104,1067 @@ FString LexToString(const FMovieSceneWarpCounter& InCounter)
 	return OutString;
 }
 
+FFrameTime operator*(FFrameTime InTime, const FMovieSceneSequenceTransform& RHS)
+{
+	const uint32 NestedTransformsSize = RHS.NestedTransforms.Num();
+	if (NestedTransformsSize == 0)
+	{
+		return InTime * RHS.LinearTransform;
+	}
+	else
+	{
+		FFrameTime OutTime = InTime * RHS.LinearTransform;
+		for (const FMovieSceneNestedSequenceTransform& NestedTransform : RHS.NestedTransforms)
+		{
+			OutTime = NestedTransform.TransformTime(OutTime);
+		}
+		return OutTime;
+	}
+}
+
+FMovieSceneTransformBreadcrumbs::operator TArrayView<const FFrameTime>() const
+{
+	return Breadcrumbs;
+}
+
+PRAGMA_DISABLE_DEPRECATION_WARNINGS
+FMovieSceneWarpCounter::FMovieSceneWarpCounter() = default;
+FMovieSceneWarpCounter::FMovieSceneWarpCounter(const FMovieSceneWarpCounter&) = default;
+FMovieSceneWarpCounter& FMovieSceneWarpCounter::operator=(const FMovieSceneWarpCounter&) = default;
+FMovieSceneWarpCounter::FMovieSceneWarpCounter(FMovieSceneWarpCounter&&) = default;
+FMovieSceneWarpCounter& FMovieSceneWarpCounter::operator=(FMovieSceneWarpCounter&&) = default;
+FMovieSceneWarpCounter::~FMovieSceneWarpCounter() = default;
+PRAGMA_ENABLE_DEPRECATION_WARNINGS
+
+void FMovieSceneNestedSequenceTransform::ToString(TStringBuilderBase<TCHAR>& OutBuilder) const
+{
+	switch (TimeScale.GetType())
+	{
+	case EMovieSceneTimeWarpType::FixedPlayRate:
+		OutBuilder.Append(LexToString(FMovieSceneTimeTransform(Offset, TimeScale.AsFixedPlayRateFloat())));
+		return;  // Explicit return because there's nothing else top do
+
+	case EMovieSceneTimeWarpType::FixedTime:
+		{
+			FMovieSceneTimeWarpFixedFrame Value = TimeScale.AsFixedTime();
+			OutBuilder.Appendf(TEXT("Fixed Frame: %s"), *LexToString(FFrameTime(Value.FrameNumber) + Offset));
+		}
+		break;
+	case EMovieSceneTimeWarpType::FrameRate:
+		{
+			FFrameRate FrameRate = TimeScale.AsFrameRate().GetFrameRate();
+			OutBuilder.Appendf(TEXT("Frame Rate: [%i/%i]"), FrameRate.Numerator, FrameRate.Denominator);
+		}
+		break;
+	case EMovieSceneTimeWarpType::Loop:
+		{
+			FMovieSceneTimeWarpLoop Value = TimeScale.AsLoop();
+			OutBuilder.Appendf(TEXT("Loop [%s:%s)"), *LexToString(-Offset), *LexToString(-Offset+Value.Duration));
+		}
+		break;
+	case EMovieSceneTimeWarpType::Clamp:
+		{
+			FMovieSceneTimeWarpClamp Value = TimeScale.AsClamp();
+			OutBuilder.Appendf(TEXT("Clamp [%s:%s)"), *LexToString(-Offset), *LexToString(-Offset+Value.Max));
+		}
+		break;
+	case EMovieSceneTimeWarpType::LoopFloat:
+		{
+			FMovieSceneTimeWarpLoopFloat Value = TimeScale.AsLoopFloat();
+			OutBuilder.Appendf(TEXT("Loop [%s:%s)"), *LexToString(-Offset), *LexToString(-Offset+FFrameTime::FromDecimal(Value.Duration)));
+		}
+		break;
+	case EMovieSceneTimeWarpType::ClampFloat:
+		{
+			FMovieSceneTimeWarpClampFloat Value = TimeScale.AsClampFloat();
+			OutBuilder.Appendf(TEXT("Clamp [%s:%s)"), *LexToString(-Offset), *LexToString(-Offset+FFrameTime::FromDecimal(Value.Max)));
+		}
+		break;
+	case EMovieSceneTimeWarpType::Custom:
+		if (UMovieSceneTimeWarpGetter* Custom = TimeScale.AsCustom())
+		{
+			OutBuilder.Append(Custom->GetName());
+		}
+		break;
+	}
+
+	if (Offset != 0)
+	{
+		OutBuilder.Appendf(TEXT(" + %s"), *LexToString(Offset));
+	}
+}
+
+FFrameTime FMovieSceneNestedSequenceTransform::TransformTime(FFrameTime InTime) const
+{
+	switch (TimeScale.GetType())
+	{
+	case EMovieSceneTimeWarpType::FixedPlayRate:
+		return InTime * FMovieSceneTimeTransform(Offset, TimeScale.AsFixedPlayRateFloat());
+
+	case EMovieSceneTimeWarpType::FixedTime:
+		return TimeScale.AsFixedTime().FrameNumber + Offset;
+
+	case EMovieSceneTimeWarpType::FrameRate:
+		return ConvertFrameTime(InTime + Offset, TimeScale.AsFrameRate().GetFrameRate(), FFrameRate(1, 1));
+
+	case EMovieSceneTimeWarpType::Loop:
+		return TimeScale.AsLoop().LoopTime(InTime + Offset) - Offset;
+
+	case EMovieSceneTimeWarpType::Clamp:
+		return TimeScale.AsClamp().Clamp(InTime + Offset) - Offset;
+
+	case EMovieSceneTimeWarpType::LoopFloat:
+		return TimeScale.AsLoopFloat().LoopTime(InTime + Offset) - Offset;
+
+	case EMovieSceneTimeWarpType::ClampFloat:
+		return TimeScale.AsClampFloat().Clamp(InTime + Offset) - Offset;
+
+	case EMovieSceneTimeWarpType::Custom:
+		if (UMovieSceneTimeWarpGetter* Custom = TimeScale.AsCustom())
+		{
+			return Custom->RemapTime(InTime + Offset);
+		}
+		break;
+	}
+
+	return InTime;
+}
+
+FFrameTime FMovieSceneNestedSequenceTransform::TransformTime(FFrameTime InTime, const UE::MovieScene::FTransformTimeParams& Params) const
+{
+	switch (TimeScale.GetType())
+	{
+	case EMovieSceneTimeWarpType::FixedPlayRate:
+		return InTime * FMovieSceneTimeTransform(Offset, TimeScale.AsFixedPlayRateFloat());
+
+	case EMovieSceneTimeWarpType::FixedTime:
+		return TimeScale.AsFixedTime().FrameNumber + Offset;
+
+	case EMovieSceneTimeWarpType::FrameRate:
+		return ConvertFrameTime(InTime + Offset, TimeScale.AsFrameRate().GetFrameRate(), FFrameRate(1, 1));
+
+	case EMovieSceneTimeWarpType::Loop:
+		if (Params.CycleCount)
+		{
+			*Params.CycleCount = 0;
+			return TimeScale.AsLoop().LoopTime(InTime + Offset, Params.CycleCount->GetValue()) - Offset;
+		}
+		else
+		{
+			return TimeScale.AsLoop().LoopTime(InTime + Offset) - Offset;
+		}
+
+	case EMovieSceneTimeWarpType::Clamp:
+		if (Params.bIgnoreClamps)
+		{
+			break;
+		}
+		return TimeScale.AsClamp().Clamp(InTime + Offset) - Offset;
+
+	case EMovieSceneTimeWarpType::LoopFloat:
+		return TimeScale.AsLoopFloat().LoopTime(InTime + Offset) - Offset;
+
+	case EMovieSceneTimeWarpType::ClampFloat:
+		if (Params.bIgnoreClamps)
+		{
+			break;
+		}
+		return TimeScale.AsClampFloat().Clamp(InTime + Offset) - Offset;
+
+	case EMovieSceneTimeWarpType::Custom:
+		if (UMovieSceneTimeWarpGetter* Custom = TimeScale.AsCustom())
+		{
+			return Custom->RemapTime(InTime + Offset);
+		}
+		break;
+	}
+
+	return InTime;
+}
+
+TRange<FFrameTime> FMovieSceneNestedSequenceTransform::ComputeTraversedHull(const TRange<FFrameTime>& InRange) const
+{
+	using namespace UE::MovieScene;
+
+	EMovieSceneTimeWarpType Type = TimeScale.GetType();
+	if (Type == EMovieSceneTimeWarpType::FixedPlayRate)
+	{
+		FMovieSceneTimeTransform LinearTransform(Offset, TimeScale.AsFixedPlayRate());
+
+		TRange<FFrameTime> Result = InRange;
+		if (Result.HasLowerBound())
+		{
+			Result.SetLowerBoundValue(Result.GetLowerBoundValue() * LinearTransform);
+		}
+		if (Result.HasUpperBound())
+		{
+			Result.SetUpperBoundValue(Result.GetUpperBoundValue() * LinearTransform);
+		}
+		return Result;
+	}
+	else if (Type == EMovieSceneTimeWarpType::FixedTime)
+	{
+		FFrameTime FixedFrame = Offset + TimeScale.AsFixedTime().FrameNumber;
+		return TRange<FFrameTime>::Inclusive(FixedFrame, FixedFrame);
+	}
+	else if (Type == EMovieSceneTimeWarpType::FrameRate)
+	{
+		FFrameRate Rate = TimeScale.AsFrameRate().GetFrameRate();
+
+		TRange<FFrameTime> Result = InRange;
+		if (Result.HasLowerBound())
+		{
+			Result.SetLowerBoundValue(ConvertFrameTime(Result.GetLowerBoundValue() + Offset, Rate, FFrameRate(1, 1)));
+		}
+		if (Result.HasUpperBound())
+		{
+			Result.SetUpperBoundValue(ConvertFrameTime(Result.GetUpperBoundValue() + Offset, Rate, FFrameRate(1, 1)));
+		}
+		return Result;
+	}
+
+	// Handle other types
+	TRange<FFrameTime> OffsetRange = TranslateRange(InRange, Offset);
+
+	switch (Type)
+	{
+
+	case EMovieSceneTimeWarpType::Loop:
+	{
+		OffsetRange = TimeScale.AsLoop().ComputeTraversedHull(OffsetRange);
+		break;
+	}
+
+	case EMovieSceneTimeWarpType::Clamp:
+	{
+		OffsetRange = TimeScale.AsClamp().ComputeTraversedHull(OffsetRange);
+		break;
+	}
+
+	case EMovieSceneTimeWarpType::LoopFloat:
+	{
+		OffsetRange = TimeScale.AsLoopFloat().ComputeTraversedHull(OffsetRange);
+		break;
+	}
+
+	case EMovieSceneTimeWarpType::ClampFloat:
+	{
+		OffsetRange = TimeScale.AsClampFloat().ComputeTraversedHull(OffsetRange);
+		break;
+	}
+
+	case EMovieSceneTimeWarpType::Custom:
+		if (UMovieSceneTimeWarpGetter* Custom = TimeScale.AsCustom())
+		{
+			OffsetRange = Custom->ComputeTraversedHull(OffsetRange);
+		}
+		return OffsetRange;
+	}
+
+	OffsetRange = TranslateRange(OffsetRange, -Offset);
+	return OffsetRange;
+}
+
+bool FMovieSceneNestedSequenceTransform::ExtractBoundariesWithinRange(FFrameTime Start, FFrameTime End, const TFunctionRef<bool(FFrameTime)>& InVisitor) const
+{
+	EMovieSceneTimeWarpType WarpType = TimeScale.GetType();
+	switch (WarpType)
+	{
+	case EMovieSceneTimeWarpType::Loop:
+		return TimeScale.AsLoop().ExtractBoundariesWithinRange(Start, End, InVisitor);
+	case EMovieSceneTimeWarpType::LoopFloat:
+		return TimeScale.AsLoopFloat().ExtractBoundariesWithinRange(Start, End, InVisitor);
+	default:
+		return false;
+	}
+}
+
+bool FMovieSceneNestedSequenceTransform::SupportsBoundaries() const
+{
+	switch (TimeScale.GetType())
+	{
+	case EMovieSceneTimeWarpType::Loop:
+	case EMovieSceneTimeWarpType::LoopFloat:
+		return true;
+
+	default:
+		return false;
+	}
+}
+
+FMovieSceneInverseNestedSequenceTransform FMovieSceneNestedSequenceTransform::Inverse() const
+{
+	if (TimeScale.GetType() == EMovieSceneTimeWarpType::FixedPlayRate)
+	{
+		const double PlayRate = TimeScale.AsFixedPlayRate();
+		checkf(!FMath::IsNearlyZero(PlayRate), TEXT("Play rate cannot be zero! This should be expressed as a nullptr TimeScale with FLAG_Zero."));
+
+		FMovieSceneInverseNestedSequenceTransform InverseTransform;
+		InverseTransform.Offset    = -Offset / PlayRate;
+		InverseTransform.TimeScale = 1.0 / PlayRate;
+		return InverseTransform;
+	}
+
+	FMovieSceneInverseNestedSequenceTransform InverseTransform;
+	InverseTransform.Offset    = Offset;
+	InverseTransform.TimeScale = TimeScale;
+	return InverseTransform;
+}
+
+FMovieSceneTimeTransform FMovieSceneInverseNestedSequenceTransform::AsLinear() const
+{
+	const float PlayRate = TimeScale.AsFixedPlayRateFloat();
+	checkf(!FMath::IsNearlyZero(PlayRate), TEXT("Play rate cannot be zero! This should be expressed as a nullptr TimeScale with FLAG_Zero."));
+	return FMovieSceneTimeTransform(Offset, TimeScale.AsFixedPlayRateFloat());
+}
+
+bool FMovieSceneInverseNestedSequenceTransform::TransformTimeWithinRange(FFrameTime InTime, const TFunctionRef<bool(FFrameTime)>& InVisitor, FFrameTime RangeStart, FFrameTime RangeEnd) const
+{
+	auto OffsetVisitor = [this, &InVisitor](FFrameTime VisitTime)
+	{
+		return InVisitor(VisitTime - this->Offset);
+	};
+
+	FFrameTime TransformedTime;
+	switch (TimeScale.GetType())
+	{
+	case EMovieSceneTimeWarpType::FixedPlayRate:
+		TransformedTime = InTime * AsLinear();
+		break;
+
+	case EMovieSceneTimeWarpType::FixedTime:
+		TransformedTime = TimeScale.AsFixedTime().FrameNumber + Offset;
+		break;
+
+	case EMovieSceneTimeWarpType::FrameRate:
+		TransformedTime = ConvertFrameTime(InTime + Offset, TimeScale.AsFrameRate().GetFrameRate(), FFrameRate(1, 1));
+		break;
+
+	case EMovieSceneTimeWarpType::Loop:
+		if (RangeStart.FrameNumber.Value != MIN_int32)
+		{
+			RangeStart += Offset;
+		}
+		if (RangeEnd.FrameNumber.Value != MAX_int32)
+		{
+			RangeEnd += Offset;
+		}
+
+		return TimeScale.AsLoop().InverseRemapTimeWithinRange(InTime + Offset, RangeStart, RangeEnd, OffsetVisitor);
+
+	case EMovieSceneTimeWarpType::Clamp:
+		if (InTime < -Offset && InTime > TimeScale.AsClamp().Max - Offset)
+		{
+			return true;
+		}
+
+		TransformedTime = InTime;
+		break;
+
+	case EMovieSceneTimeWarpType::LoopFloat:
+		if (RangeStart.FrameNumber.Value != MIN_int32)
+		{
+			RangeStart += Offset;
+		}
+		if (RangeEnd.FrameNumber.Value != MAX_int32)
+		{
+			RangeEnd += Offset;
+		}
+
+		return TimeScale.AsLoopFloat().InverseRemapTimeWithinRange(InTime + Offset, RangeStart, RangeEnd, OffsetVisitor);
+
+	case EMovieSceneTimeWarpType::ClampFloat:
+		if (InTime < -Offset && InTime > FFrameTime::FromDecimal(TimeScale.AsClampFloat().Max) - Offset)
+		{
+			return true;
+		}
+
+		TransformedTime = InTime;
+		break;
+
+	case EMovieSceneTimeWarpType::Custom:
+		if (UMovieSceneTimeWarpGetter* Custom = TimeScale.AsCustom())
+		{
+			return Custom->InverseRemapTimeWithinRange(InTime, RangeStart, RangeEnd, OffsetVisitor);
+		}
+		return true;
+	}
+
+	if (TransformedTime >= RangeStart && TransformedTime <= RangeEnd)
+	{
+		return InVisitor(TransformedTime);
+	}
+
+	return true;
+}
+
+TOptional<FFrameTime> FMovieSceneInverseNestedSequenceTransform::TryTransformTime(FFrameTime InTime, FFrameTime Breadcrumb) const
+{
+	return TryTransformTime(InTime, Breadcrumb, UE::MovieScene::FInverseTransformTimeParams());
+}
+
+TOptional<FFrameTime> FMovieSceneInverseNestedSequenceTransform::TryTransformTime(FFrameTime InTime, FFrameTime Breadcrumb, const UE::MovieScene::FInverseTransformTimeParams& Params) const
+{
+	using namespace UE::MovieScene;
+
+	switch (TimeScale.GetType())
+	{
+	case EMovieSceneTimeWarpType::FixedPlayRate:
+		return InTime * FMovieSceneTimeTransform(Offset, TimeScale.AsFixedPlayRateFloat());
+
+	case EMovieSceneTimeWarpType::FixedTime:
+		if (InTime == TimeScale.AsFixedTime().FrameNumber + Offset)
+		{
+			return InTime;
+		}
+		break;
+
+	case EMovieSceneTimeWarpType::FrameRate:
+		return ConvertFrameTime(InTime, FFrameRate(1, 1), TimeScale.AsFrameRate().GetFrameRate()) - Offset;
+
+	case EMovieSceneTimeWarpType::Loop:
+		{
+			TOptional<FFrameTime> Result = TimeScale.AsLoop().InverseRemapTimeCycled(InTime + Offset, Breadcrumb + Offset, Params);
+			if (Result)
+			{
+				return Result.GetValue() - Offset;
+			}
+		}
+		break;
+
+	case EMovieSceneTimeWarpType::Clamp:
+		if (EnumHasAnyFlags(Params.Flags, EInverseEvaluateFlags::IgnoreClamps) || ( InTime >= -Offset && InTime <= TimeScale.AsClamp().Max - Offset) )
+		{
+			return InTime;
+		}
+		break;
+
+	case EMovieSceneTimeWarpType::LoopFloat:
+		{
+			TOptional<FFrameTime> Result = TimeScale.AsLoopFloat().InverseRemapTimeCycled(InTime + Offset, Breadcrumb + Offset, Params);
+			if (Result)
+			{
+				return Result.GetValue() - Offset;
+			}
+		}
+		break;
+
+	case EMovieSceneTimeWarpType::ClampFloat:
+		if (EnumHasAnyFlags(Params.Flags, EInverseEvaluateFlags::IgnoreClamps) || ( InTime >= -Offset && InTime <= FFrameTime::FromDecimal(TimeScale.AsClampFloat().Max) - Offset) )
+		{
+			return InTime;
+		}
+		break;
+
+	case EMovieSceneTimeWarpType::Custom:
+		if (UMovieSceneTimeWarpGetter* Custom = TimeScale.AsCustom())
+		{
+			TOptional<FFrameTime> Result = Custom->InverseRemapTimeCycled(InTime, Breadcrumb, Params);
+			if (Result)
+			{
+				return Result.GetValue() - Offset;
+			}
+		}
+		break;
+	}
+
+	return TOptional<FFrameTime>();
+}
+
+TOptional<FFrameTime> FMovieSceneInverseSequenceTransform::TryTransformTime(FFrameTime InTime) const
+{
+	return TryTransformTime(InTime, UE::MovieScene::FInverseTransformTimeParams());
+}
+
+TOptional<FFrameTime> FMovieSceneInverseSequenceTransform::TryTransformTime(FFrameTime InTime, const UE::MovieScene::FInverseTransformTimeParams& Params) const
+{
+	FFrameTime OutTime = InTime;
+
+	for (const FMovieSceneInverseNestedSequenceTransform& NestedTransform : NestedTransforms)
+	{
+		FFrameTime Breadcrumb;
+
+		TOptional<FFrameTime> NewTime = NestedTransform.TryTransformTime(OutTime, Breadcrumb, Params);
+		if (!NewTime)
+		{
+			return TOptional<FFrameTime>();
+		}
+		OutTime = NewTime.GetValue();
+	}
+
+	return OutTime * LinearTransform;
+}
+
+
+
+TOptional<FFrameTime> FMovieSceneInverseSequenceTransform::TryTransformTime(FFrameTime InTime, const FMovieSceneTransformBreadcrumbs& InBreadcrumbs) const
+{
+	return TryTransformTime(InTime, InBreadcrumbs, UE::MovieScene::FInverseTransformTimeParams());
+}
+
+TOptional<FFrameTime> FMovieSceneInverseSequenceTransform::TryTransformTime(FFrameTime InTime, const FMovieSceneTransformBreadcrumbs& InBreadcrumbs, const UE::MovieScene::FInverseTransformTimeParams& Params) const
+{
+	FFrameTime OutTime = InTime;
 
+	int32 BreadcrumbIndex = InBreadcrumbs.Num() - 1;
+	for (const FMovieSceneInverseNestedSequenceTransform& NestedTransform : NestedTransforms)
+	{
+		FFrameTime Breadcrumb;
+
+		const bool bShouldHaveBreadcrumb = (InBreadcrumbs.GetMode() == EMovieSceneBreadcrumbMode::Dense || NestedTransform.NeedsBreadcrumb());
+		if (bShouldHaveBreadcrumb && InBreadcrumbs.IsValidIndex(BreadcrumbIndex))
+		{
+			Breadcrumb = InBreadcrumbs[BreadcrumbIndex--];
+		}
+
+		TOptional<FFrameTime> NewTime = NestedTransform.TryTransformTime(OutTime, Breadcrumb, Params);
+		if (!NewTime)
+		{
+			return TOptional<FFrameTime>();
+		}
+		OutTime = NewTime.GetValue();
+	}
+
+	return OutTime * LinearTransform;
+}
+
+bool FMovieSceneInverseSequenceTransform::RecursiveTransformTimeWithinRange(int32 NestingIndex, FFrameTime InTime, const TFunctionRef<bool(FFrameTime)>& FinalVisitor, TArrayView<const FFrameTime> StartBreadcrumbs, TArrayView<const FFrameTime> EndBreadcrumbs) const
+{
+	for ( ; NestingIndex < NestedTransforms.Num(); ++NestingIndex)
+	{
+		const FMovieSceneInverseNestedSequenceTransform& NestedTransform = NestedTransforms[NestingIndex];
+
+		// Linear transform is easy - keep looping them
+		if (NestedTransform.IsLinear())
+		{
+			InTime = InTime * NestedTransform.AsLinear();
+		}
+		// Warped ranges may map to zero or more times in the outer sequence
+		//     so perform a complete recursive expansion on all of them
+		else if (ensureMsgf(StartBreadcrumbs.Num() > 0 && EndBreadcrumbs.Num() > 0,
+			TEXT("Breadcrumb count mismatch in inverse transform computation")))
+		{
+			auto TransformNext = [this, FinalVisitor, NestingIndex, StartBreadcrumbs, EndBreadcrumbs](FFrameTime NextTime)
+			{
+				return this->RecursiveTransformTimeWithinRange(
+					NestingIndex+1,
+					NextTime,
+					FinalVisitor,
+					StartBreadcrumbs.LeftChop(1),
+					EndBreadcrumbs.LeftChop(1)
+				);
+			};
+
+			// RecursiveTransformTimeWithinRange will complete the recursion
+			return NestedTransform.TransformTimeWithinRange(
+				InTime,
+				TransformNext,
+				StartBreadcrumbs.Last(),
+				EndBreadcrumbs.Last()
+			);
+		}
+	}
+
+	return FinalVisitor(InTime * LinearTransform);
+}
+
+bool FMovieSceneInverseSequenceTransform::TransformFiniteRangeWithinRange(const TRange<FFrameTime>& InRange, TFunctionRef<bool(TRange<FFrameTime>)> InVisitor, const FMovieSceneTransformBreadcrumbs& StartBreadcrumbs, const FMovieSceneTransformBreadcrumbs& EndBreadcrumbs) const
+{
+	check(!InRange.GetLowerBound().IsOpen() && !InRange.GetUpperBound().IsOpen());
+
+	if (NestedTransforms.Num() == 0)
+	{
+		// Only one solution
+		return InVisitor(InRange * LinearTransform);
+	}
+
+	TArray<FFrameTime> LowerBounds, UpperBounds;
+
+	// Transform lower bounds
+	auto VisitLower = [&LowerBounds](FFrameTime InFrameTime)
+	{
+		LowerBounds.Add(InFrameTime);
+		return true;
+	};
+	TransformTimeWithinRange(InRange.GetLowerBoundValue(), VisitLower, StartBreadcrumbs, EndBreadcrumbs);
+
+	// Transform upper bounds
+	auto VisitUpper = [&UpperBounds](FFrameTime InFrameTime)
+	{
+		UpperBounds.Add(InFrameTime);
+		return true;
+	};
+	TransformTimeWithinRange(InRange.GetUpperBoundValue(), VisitUpper, StartBreadcrumbs, EndBreadcrumbs);
+
+	Algo::Sort(LowerBounds);
+	Algo::Sort(UpperBounds);
+
+	int32 LwrIndex = 0, UprIndex = 0;
+
+	// Handle leading upper bounds - should only be one?
+	while (UprIndex < UpperBounds.Num())
+	{
+		if (LwrIndex < LowerBounds.Num() && UpperBounds[UprIndex] >= LowerBounds[LwrIndex])
+		{
+			break;
+		}
+
+		// Maintain bound exclusivity
+		TRange<FFrameTime> Result = InRange;
+		if (UprIndex < UpperBounds.Num() - 1)
+		{
+			TRangeBound<FFrameTime> NewLower = Result.GetUpperBound();
+			NewLower.SetValue(UpperBounds[UprIndex + 1]);
+			Result.SetLowerBound(TRangeBound<FFrameTime>::FlipInclusion(NewLower));
+		}
+		else
+		{
+			Result.SetLowerBound(TRangeBound<FFrameTime>::Open());
+		}
+
+		Result.SetUpperBoundValue(UpperBounds[UprIndex]);
+		if (!InVisitor(Result))
+		{
+			return false;
+		}
+
+		++UprIndex;
+	}
+
+	// Handle finite ranges
+	while (LwrIndex < LowerBounds.Num() && UprIndex < UpperBounds.Num())
+	{
+		FFrameTime LowerBound = LowerBounds[LwrIndex];
+
+		TRange<FFrameTime> Result = InRange;
+		Result.SetLowerBoundValue(LowerBound);
+
+		// Skip any upper bounds that are < this lower bound
+		while (UprIndex < UpperBounds.Num() && UpperBounds[UprIndex] <= LowerBound)
+		{
+			++UprIndex;
+		}
+		
+		if (UprIndex < UpperBounds.Num())
+		{
+			Result.SetUpperBoundValue(UpperBounds[UprIndex]);
+			if (!Result.IsEmpty() && !InVisitor(Result))
+			{
+				return false;
+			}
+		}
+
+		++LwrIndex;
+		++UprIndex;
+	}
+
+	// Handle trailing lower bounds - there can be cases where there are multiple, especially in the case of nested
+	//    if a looping subsequences that have their end cropped
+	while (LwrIndex < LowerBounds.Num())
+	{
+		// Maintain bound exclusivity
+		TRange<FFrameTime> Result = InRange;
+		Result.SetLowerBoundValue(LowerBounds[LwrIndex]);
+		if (LwrIndex < LowerBounds.Num()-1)
+		{
+			TRangeBound<FFrameTime> NewUpper = Result.GetLowerBound();
+			NewUpper.SetValue(LowerBounds[LwrIndex + 1]);
+			Result.SetUpperBound(TRangeBound<FFrameTime>::FlipInclusion(NewUpper));
+		}
+		else
+		{
+			Result.SetUpperBound(TRangeBound<FFrameTime>::Open());
+		}
+
+		if (!Result.IsEmpty() && !InVisitor(Result))
+		{
+			return false;
+		}
+
+		++LwrIndex;
+	}
+
+	return true;
+}
+
+bool FMovieSceneInverseSequenceTransform::TransformTimeWithinRange(FFrameTime InTime, const TFunctionRef<bool(FFrameTime)>& InVisitor, const FMovieSceneTransformBreadcrumbs& StartBreadcrumbs, const FMovieSceneTransformBreadcrumbs& EndBreadcrumbs) const
+{
+	return RecursiveTransformTimeWithinRange(0, InTime, InVisitor, StartBreadcrumbs, EndBreadcrumbs);
+}
+
+FMovieSceneTimeTransform FMovieSceneInverseSequenceTransform::AsLegacyLinearTimeTransform() const
+{
+	FMovieSceneTimeTransform Result;
+	for (const FMovieSceneInverseNestedSequenceTransform& Nested : NestedTransforms)
+	{
+		if (Nested.IsLinear())
+		{
+			Result = Result * Nested.AsLinear();
+		}
+	}
+	// Linear Transform should apply last, and transform multiplication applies RtL
+	return LinearTransform * Result;
+}
+
+PRAGMA_DISABLE_DEPRECATION_WARNINGS
+FMovieSceneNestedSequenceTransform::FMovieSceneNestedSequenceTransform() = default;
+FMovieSceneNestedSequenceTransform::FMovieSceneNestedSequenceTransform(const FMovieSceneNestedSequenceTransform&) = default;
+FMovieSceneNestedSequenceTransform& FMovieSceneNestedSequenceTransform::operator=(const FMovieSceneNestedSequenceTransform&) = default;
+FMovieSceneNestedSequenceTransform::FMovieSceneNestedSequenceTransform(FMovieSceneNestedSequenceTransform&&) = default;
+FMovieSceneNestedSequenceTransform& FMovieSceneNestedSequenceTransform::operator=(FMovieSceneNestedSequenceTransform&&) = default;
+FMovieSceneNestedSequenceTransform::~FMovieSceneNestedSequenceTransform() = default;
+PRAGMA_ENABLE_DEPRECATION_WARNINGS
+
+void FMovieSceneSequenceTransform::Add(FMovieSceneTimeTransform InTransform)
+{
+	if (InTransform.IsIdentity())
+	{
+		return;
+	}
+
+	if (IsLinear())
+	{
+		LinearTransform = InTransform * LinearTransform;
+	}
+	else
+	{
+		NestedTransforms.Emplace(InTransform);
+	}
+}
+
+void FMovieSceneSequenceTransform::Add(FFrameTime InOffset, FMovieSceneTimeWarpVariant&& InTimeWarp)
+{
+	if (InTimeWarp.GetType() == EMovieSceneTimeWarpType::FixedPlayRate)
+	{
+		Add(FMovieSceneTimeTransform(InOffset, InTimeWarp.AsFixedPlayRateFloat()));
+	}
+	else
+	{
+		NestedTransforms.Emplace(InOffset, MoveTemp(InTimeWarp));
+	}
+}
+
+FFrameTime FMovieSceneSequenceTransform::TransformTime(FFrameTime InTime) const
+{
+	if (NestedTransforms.Num() == 0)
+	{
+		return InTime * LinearTransform;
+	}
+
+	FFrameTime OutTime = InTime * LinearTransform;
+	for (const FMovieSceneNestedSequenceTransform& NestedTransform : NestedTransforms)
+	{
+		OutTime = NestedTransform.TransformTime(OutTime);
+	}
+	return OutTime;
+}
+
+FFrameTime FMovieSceneSequenceTransform::TransformTime(FFrameTime InTime, const UE::MovieScene::FTransformTimeParams& Params) const
+{
+	FFrameTime OutTime = InTime * LinearTransform;
+	if (Params.Breadcrumbs && Params.Breadcrumbs->GetMode() == EMovieSceneBreadcrumbMode::Dense)
+	{
+		Params.Breadcrumbs->AddBreadcrumb(OutTime);
+	}
+
+	if (NestedTransforms.Num() == 0)
+	{
+		return OutTime;
+	}
+
+	for (const FMovieSceneNestedSequenceTransform& NestedTransform : NestedTransforms)
+	{
+		if (Params.Breadcrumbs && (Params.Breadcrumbs->GetMode() == EMovieSceneBreadcrumbMode::Dense || NestedTransform.NeedsBreadcrumb()) )
+		{
+			Params.Breadcrumbs->AddBreadcrumb(OutTime);
+		}
+
+		OutTime = NestedTransform.TransformTime(OutTime, Params);
+	}
+	return OutTime;
+}
+
+TRange<FFrameTime> FMovieSceneSequenceTransform::ComputeTraversedHull(const TRange<FFrameTime>& Range) const
+{
+	TRange<FFrameTime> Result = Range * LinearTransform;
+
+	for (const FMovieSceneNestedSequenceTransform& NestedTransform : NestedTransforms)
+	{
+		Result = NestedTransform.ComputeTraversedHull(Result);
+		if (Result.IsEmpty())
+		{
+			return Result;
+		}
+	}
+
+	return Result;
+}
+TRange<FFrameTime> FMovieSceneSequenceTransform::ComputeTraversedHull(const TRange<FFrameNumber>& Range) const
+{
+	return ComputeTraversedHull(UE::MovieScene::ConvertToFrameTimeRange(Range));
+}
+
+bool FMovieSceneSequenceTransform::ExtractBoundariesWithinRange(FFrameTime Start, FFrameTime End, const TFunctionRef<bool(FFrameTime)>& InVisitor) const
+{
+	FMovieSceneInverseSequenceTransform Inverse;
+
+	FMovieSceneTransformBreadcrumbs Breadcrumbs;
+
+	const bool bHasStart = Start != MIN_int32;
+	const bool bHasEnd   = End != MAX_int32;
+
+	if (bHasStart)
+	{
+		Start = Start * LinearTransform;
+	}
+	if (bHasEnd)
+	{
+		End = End * LinearTransform;
+	}
+
+	for (int32 NestedIndex = 0; NestedIndex < NestedTransforms.Num(); ++NestedIndex)
+	{
+		const FMovieSceneNestedSequenceTransform& NestedTransform = NestedTransforms[NestedIndex];
+
+		// Find the first transform that has any boundaries
+		if (!NestedTransform.SupportsBoundaries())
+		{
+			Breadcrumbs.AddBreadcrumb(Start + (End - Start) / 2);
+			if (bHasStart)
+			{
+				Start = NestedTransform.TransformTime(Start);
+			}
+			if (bHasEnd)
+			{
+				End = NestedTransform.TransformTime(End);
+			}
+			continue;
+		}
+
+
+		FMovieSceneSequenceTransform RootToParentTransform = *this;
+		RootToParentTransform.NestedTransforms.SetNum(NestedIndex);
+
+		FMovieSceneInverseSequenceTransform ParentToRootTransform = RootToParentTransform.Inverse();
+
+
+		auto VisitWrapper = [&InVisitor, &ParentToRootTransform, &Breadcrumbs](FFrameTime InBoundary)
+		{
+			if (TOptional<FFrameTime> RootTime = ParentToRootTransform.TryTransformTime(InBoundary, Breadcrumbs))
+			{
+				return InVisitor(RootTime.GetValue());
+			}
+			return true;
+		};
+
+		return NestedTransform.ExtractBoundariesWithinRange(Start, End, VisitWrapper);
+	}
+
+	return false;
+}
+
+void FMovieSceneSequenceTransform::AddLoop(FFrameNumber InStart, FFrameNumber InEnd)
+{
+	check(InStart < InEnd);
+	// Offset by -InStart because our looping variant can only loop from 0:Max
+	NestedTransforms.Emplace(-InStart, FMovieSceneTimeWarpVariant(FMovieSceneTimeWarpLoop{ InEnd - InStart }));
+}
+
+bool FMovieSceneSequenceTransform::IsIdentity() const
+{
+	return LinearTransform.IsIdentity() && 
+		Algo::AllOf(NestedTransforms, &FMovieSceneNestedSequenceTransform::IsIdentity);
+}
+
+FMovieSceneInverseSequenceTransform FMovieSceneSequenceTransform::Inverse() const
+{
+	FMovieSceneInverseSequenceTransform Result;
+
+	if (NestedTransforms.Num() == 0)
+	{
+		Result.LinearTransform = LinearTransform.Inverse();
+		return Result;
+	}
+
+	// Start accumulating the inverse transforms in reverse order.
+	Result.NestedTransforms.Reserve(NestedTransforms.Num());
+
+	for(int Index = NestedTransforms.Num() - 1; Index >= 0; --Index)
+	{
+		const FMovieSceneNestedSequenceTransform& Nested = NestedTransforms[Index];
+		if (Nested.IsLinear())
+		{
+			Result.LinearTransform = Result.LinearTransform * Nested.AsLinear().Inverse();
+			continue;
+		}
+
+		if (!Result.LinearTransform.IsIdentity())
+		{
+			// If we have any linear transform that needs to get added to the stack before this one's inverse
+			Result.NestedTransforms.Add(Result.LinearTransform);
+			Result.LinearTransform = FMovieSceneTimeTransform();
+		}
+
+		Result.NestedTransforms.Add(NestedTransforms[Index].Inverse());
+	}
+
+	// Add the inverse of the main linear transform if not identity
+	if (!LinearTransform.IsIdentity())
+	{
+		Result.LinearTransform = LinearTransform.Inverse() * Result.LinearTransform;
+	}
+
+	return Result;
+}
+
+
+FMovieSceneSequenceTransform FMovieSceneSequenceTransform::operator*(const FMovieSceneSequenceTransform& RHS) const
+{
+	if (IsLinear() && RHS.IsLinear())
+	{
+		// None of the transforms are warping... we can combine them into another linear transform.
+		return FMovieSceneSequenceTransform(LinearTransform * RHS.LinearTransform);
+	}
+	else if (IsLinear())
+	{
+		// LHS is linear, but RHS is warping. Since transforms are supposed to apply from right to left,
+		// we need to append LHS at the "bottom" of RHS, i.e. add a new nested transform that's LHS. However
+		// if LHS is identity, we have nothing to do, and if both LHS and RHS' deeper transform are linear,
+		// we can combine both.
+		FMovieSceneSequenceTransform Result(RHS);
+		if (!LinearTransform.IsIdentity())
+		{
+			FMovieSceneNestedSequenceTransform& LastNested = Result.NestedTransforms.Last();
+			if (LastNested.IsLinear())
+			{
+				FMovieSceneTimeTransform NewLinear = LinearTransform * LastNested.AsLinear();
+				LastNested = FMovieSceneNestedSequenceTransform(NewLinear);
+			}
+			else
+			{
+				Result.NestedTransforms.Emplace(LinearTransform);
+			}
+		}
+		return Result;
+	}
+	else if (RHS.IsLinear())
+	{
+		// RHS isn't warping, but LHS is, so we combine the linear transform parts and start looping
+		// from there.
+		FMovieSceneSequenceTransform Result;
+		Result.LinearTransform = LinearTransform * RHS.LinearTransform;
+		Result.NestedTransforms = NestedTransforms;
+		return Result;
+	}
+	else
+	{
+		// Both are looping, we need to combine them. Usually, a warping transform doesn't use its linear part,
+		// because whatever linear placement/scaling it has would be in the linear part of the nested transform
+		// struct.
+		FMovieSceneSequenceTransform Result(RHS);
+		const bool bHasOnlyNested = LinearTransform.IsIdentity();
+		if (!bHasOnlyNested)
+		{
+			Result.NestedTransforms.Add(FMovieSceneNestedSequenceTransform(LinearTransform));
+		}
+		Result.NestedTransforms.Append(NestedTransforms);
+		return Result;
+	}
+}
+
+PRAGMA_DISABLE_DEPRECATION_WARNINGS
+bool FMovieSceneSequenceTransform::IsLooping() const
+{
+	return Algo::AnyOf(NestedTransforms, &FMovieSceneNestedSequenceTransform::IsLooping);
+}
+void FMovieSceneSequenceTransform::TransformTime(FFrameTime InTime, FFrameTime& OutTime, FMovieSceneWarpCounter& OutWarpCounter) const
+{
+	OutTime = TransformTime(InTime, UE::MovieScene::FTransformTimeParams().HarvestBreadcrumbs(OutWarpCounter));
+}
+float FMovieSceneSequenceTransform::GetTimeScale() const
+{
+	float TimeScale = LinearTransform.TimeScale;
+	for (const FMovieSceneNestedSequenceTransform& NestedTransform : NestedTransforms)
+	{
+		if (NestedTransform.IsLinear())
+		{
+			TimeScale *= NestedTransform.AsLinear().TimeScale;
+		}
+	}
+	return TimeScale;
+}
+TRange<FFrameTime> FMovieSceneSequenceTransform::TransformRangeConstrained(const TRange<FFrameTime>& Range) const
+{
+	return ComputeTraversedHull(Range);
+}
+TRange<FFrameTime> FMovieSceneSequenceTransform::TransformRangePure(const TRange<FFrameTime>& Range) const
+{
+	return ComputeTraversedHull(Range);
+}
+TRange<FFrameTime> FMovieSceneSequenceTransform::TransformRangeUnwarped(const TRange<FFrameTime>& Range) const
+{
+	return ComputeTraversedHull(Range);
+}
+TRange<FFrameNumber> FMovieSceneSequenceTransform::TransformRangePure(const TRange<FFrameNumber>& Range) const
+{
+	using namespace UE::MovieScene;
+
+	TRange<FFrameTime> TimeRange = TransformRangePure(ConvertRange<FFrameNumber, FFrameTime>(Range));
+	return ConvertRange<FFrameTime, FFrameNumber>(TimeRange);
+}
+TRange<FFrameNumber> FMovieSceneSequenceTransform::TransformRangeUnwarped(const TRange<FFrameNumber>& Range) const
+{
+	using namespace UE::MovieScene;
+
+	TRange<FFrameTime> TimeRange = ConvertRange<FFrameNumber, FFrameTime>(Range);
+	TimeRange = TransformRangeUnwarped(TimeRange);
+	return ConvertRange<FFrameTime, FFrameNumber>(TimeRange);
+}
+TRange<FFrameNumber> FMovieSceneSequenceTransform::TransformRangeConstrained(const TRange<FFrameNumber>& Range) const
+{
+	using namespace UE::MovieScene;
+
+	TRange<FFrameTime> TimeRange = ConvertRange<FFrameNumber, FFrameTime>(Range);
+	TimeRange = TransformRangeConstrained(TimeRange);
+	return ConvertRange<FFrameTime, FFrameNumber>(TimeRange);
+}
+FMovieSceneTimeTransform FMovieSceneSequenceTransform::InverseLinearOnly() const
+{
+	ensureMsgf(!FMath::IsNearlyZero(LinearTransform.TimeScale), TEXT("Inverse of a zero timescale transform is undefined in a FMovieSceneTimeTransform. Please use InverseNoLooping for proper behavior."));
+	return LinearTransform.Inverse();
+}
+FMovieSceneSequenceTransform FMovieSceneSequenceTransform::InverseNoLooping() const
+{
+	return FMovieSceneSequenceTransform();
+}
+FMovieSceneTimeTransform FMovieSceneSequenceTransform::InverseFromAllFirstWarps() const
+{
+	return FMovieSceneTimeTransform();
+}
+FMovieSceneSequenceTransform FMovieSceneSequenceTransform::InverseFromAllFirstLoops() const 
+{
+	return FMovieSceneSequenceTransform();
+}
+FMovieSceneTimeTransform FMovieSceneSequenceTransform::InverseFromWarp(const FMovieSceneWarpCounter& WarpCounter) const
+{
+	return FMovieSceneTimeTransform();
+}
+FMovieSceneTimeTransform FMovieSceneSequenceTransform::InverseFromWarp(const TArrayView<const uint32>& WarpCounts) const
+{
+	return FMovieSceneTimeTransform();
+}
+FMovieSceneSequenceTransform FMovieSceneSequenceTransform::InverseFromLoop(const FMovieSceneWarpCounter& LoopCounter) const
+{
+	return FMovieSceneSequenceTransform();
+}
+
+FMovieSceneSequenceTransform FMovieSceneSequenceTransform::InverseFromLoop(const TArrayView<const FFrameTime>& Breadcrumbs) const
+{
+	return FMovieSceneSequenceTransform();
+}
+
+FMovieSceneSequenceTransform FMovieSceneSequenceTransform::InverseFromLoop(const TArrayView<const uint32>& LoopCounts) const
+{
+	return FMovieSceneSequenceTransform();
+}
+PRAGMA_ENABLE_DEPRECATION_WARNINGS
\ No newline at end of file
diff --git a/Engine/Source/Runtime/MovieScene/Private/MovieSceneSection.cpp b/Engine/Source/Runtime/MovieScene/Private/MovieSceneSection.cpp
--- a/Engine/Source/Runtime/MovieScene/Private/MovieSceneSection.cpp
+++ b/Engine/Source/Runtime/MovieScene/Private/MovieSceneSection.cpp
@@ -228,7 +228,22 @@ void UMovieSceneSection::MoveSectionImpl(FFrameNumber DeltaFrame)
 		{
 			SectionRange.Value.SetUpperBoundValue(SectionRange.Value.GetUpperBoundValue() + DeltaFrame);
 		}
+#if WITH_EDITOR
+		const bool bHasStartFrame = HasStartFrame();
 
+		for (const FMovieSceneChannelEntry& Entry : GetChannelProxy().GetAllEntries())
+		{
+			TArrayView<FMovieSceneChannel* const>        Channels = Entry.GetChannels();
+			TArrayView<const FMovieSceneChannelMetaData> MetaData = Entry.GetMetaData();
+			for (int32 Index = 0; Index < Channels.Num(); ++Index)
+			{
+				if (!MetaData[Index].bRelativeToSection || !bHasStartFrame)
+				{
+					Channels[Index]->Offset(DeltaFrame);
+				}
+			}
+		}
+#else
 		for (const FMovieSceneChannelEntry& Entry : GetChannelProxy().GetAllEntries())
 		{
 			for (FMovieSceneChannel* Channel : Entry.GetChannels())
@@ -236,9 +251,39 @@ void UMovieSceneSection::MoveSectionImpl(FFrameNumber DeltaFrame)
 				Channel->Offset(DeltaFrame);
 			}
 		}
+#endif
 	}
 }
 
+void UMovieSceneSection::FixupRelativeKeyframes(FFrameNumber Offset)
+{
+#if WITH_EDITOR
+	if (!TryModify() || !HasStartFrame())
+	{
+		return;
+	}
+
+	for (const FMovieSceneChannelEntry& Entry : GetChannelProxy().GetAllEntries())
+	{
+		TArrayView<FMovieSceneChannel* const>        Channels = Entry.GetChannels();
+		TArrayView<const FMovieSceneChannelMetaData> MetaData = Entry.GetMetaData();
+		for (int32 Index = 0; Index < Channels.Num(); ++Index)
+		{
+			if (MetaData[Index].bRelativeToSection)
+			{
+				UObject* OwningObject = MetaData[Index].WeakOwningObject.Get();
+				if (OwningObject && OwningObject != this)
+				{
+					OwningObject->Modify();
+				}
+
+				Channels[Index]->Offset(-Offset);
+			}
+		}
+	}
+#endif
+}
+
 void UMovieSceneSection::MoveSection(FFrameNumber DeltaFrame)
 {
 	MoveSectionImpl(DeltaFrame);
diff --git a/Engine/Source/Runtime/MovieScene/Private/MovieSceneSequenceVisitor.cpp b/Engine/Source/Runtime/MovieScene/Private/MovieSceneSequenceVisitor.cpp
--- a/Engine/Source/Runtime/MovieScene/Private/MovieSceneSequenceVisitor.cpp
+++ b/Engine/Source/Runtime/MovieScene/Private/MovieSceneSequenceVisitor.cpp
@@ -5,6 +5,7 @@
 #include "MovieScene.h"
 #include "MovieSceneTrack.h"
 #include "MovieSceneSection.h"
+#include "MovieSceneTransformTypes.h"
 #include "Tracks/MovieSceneSubTrack.h"
 #include "Sections/MovieSceneSubSection.h"
 #include "Evaluation/MovieSceneRootOverridePath.h"
@@ -52,30 +53,43 @@ void VisitSubTrackImpl(const FSequenceVisitParams& InParams, UMovieSceneSubTrack
 			continue;
 		}
 
-		TRange<FFrameNumber> EffectiveRange = TRange<FFrameNumber>::Intersection(SubSequenceSpace->RootToSequenceTransform.InverseNoLooping().TransformRangeConstrained(Entry.Range), SubSequenceSpace->RootClampRange);
-		if (EffectiveRange.IsEmpty())
+		auto VisitRange = [&InOutRootPath, &InVisitor, &InParams, SubSection, &ObjectBinding, SubSequence, SubSequenceSpace](TRange<FFrameTime> Range)
 		{
-			continue;
-		}
-
-		FMovieSceneSequenceTransform OuterToInnerTransform = SubSection->OuterToInnerTransform();
-		FMovieSceneSequenceID        LocalSequenceID       = SubSection->GetSequenceID();
-		const FMovieSceneSequenceID  SubSequenceID         = InOutRootPath->ResolveChildSequenceID(LocalSequenceID);
-
-		FSubSequenceSpace LocalSpace;
-		LocalSpace.RootToSequenceTransform = OuterToInnerTransform * SubSequenceSpace->RootToSequenceTransform;
-		LocalSpace.SequenceID = InOutRootPath->ResolveChildSequenceID(LocalSequenceID);
-		LocalSpace.RootClampRange = EffectiveRange;
-		LocalSpace.LocalClampRange = LocalSpace.RootToSequenceTransform.TransformRangeUnwarped(EffectiveRange);
-		LocalSpace.HierarchicalBias = SubSequenceSpace->HierarchicalBias + SubSection->Parameters.HierarchicalBias;
+			TRange<FFrameNumber> EffectiveRange = ConvertToDiscreteRange(TRange<FFrameTime>::Intersection(Range, ConvertToFrameTimeRange(SubSequenceSpace->RootClampRange)));
+			if (!EffectiveRange.IsEmpty())
+			{
+				FMovieSceneSequenceTransform OuterToInnerTransform = SubSection->OuterToInnerTransform();
+				FMovieSceneSequenceID        LocalSequenceID       = SubSection->GetSequenceID();
+				const FMovieSceneSequenceID  SubSequenceID         = InOutRootPath->ResolveChildSequenceID(LocalSequenceID);
+
+				FSubSequenceSpace LocalSpace;
+				LocalSpace.RootToSequenceTransform = OuterToInnerTransform * SubSequenceSpace->RootToSequenceTransform;
+				LocalSpace.StartBreadcrumbs = SubSequenceSpace->StartBreadcrumbs;
+				LocalSpace.EndBreadcrumbs = SubSequenceSpace->EndBreadcrumbs;
+				LocalSpace.SequenceID = InOutRootPath->ResolveChildSequenceID(LocalSequenceID);
+				LocalSpace.RootClampRange = EffectiveRange;
+				LocalSpace.LocalClampRange = ConvertToDiscreteRange(LocalSpace.RootToSequenceTransform.ComputeTraversedHull(EffectiveRange));
+				LocalSpace.HierarchicalBias = SubSequenceSpace->HierarchicalBias + SubSection->Parameters.HierarchicalBias;
+
+				if (OuterToInnerTransform.NeedsBreadcrumbs())
+				{
+					OuterToInnerTransform.TransformTime(DiscreteInclusiveLower(EffectiveRange.GetLowerBoundValue()), FTransformTimeParams().AppendBreadcrumbs(LocalSpace.StartBreadcrumbs));
+					OuterToInnerTransform.TransformTime(DiscreteExclusiveUpper(EffectiveRange.GetUpperBoundValue()), FTransformTimeParams().AppendBreadcrumbs(LocalSpace.EndBreadcrumbs));
+				}
+
+				// Recurse into the sub sequence
+				InOutRootPath->PushGeneration(SubSequenceID, LocalSequenceID);
+				{
+					InVisitor.VisitSubSequence(SubSequence, ObjectBinding, LocalSpace);
+					VisitSequenceImpl(SubSequence, InParams, InVisitor, InOutRootPath, &LocalSpace);
+				}
+				InOutRootPath->PopGenerations(1);
+			}
+			return true;
+		};
 
-		// Recurse into the sub sequence
-		InOutRootPath->PushGeneration(SubSequenceID, LocalSequenceID);
-		{
-			InVisitor.VisitSubSequence(SubSequence, ObjectBinding, LocalSpace);
-			VisitSequenceImpl(SubSequence, InParams, InVisitor, InOutRootPath, &LocalSpace);
-		}
-		InOutRootPath->PopGenerations(1);
+		FMovieSceneInverseSequenceTransform SequenceToRootTransform = SubSequenceSpace->RootToSequenceTransform.Inverse();
+		SequenceToRootTransform.TransformFiniteRangeWithinRange(ConvertToFrameTimeRange(Entry.Range), VisitRange, SubSequenceSpace->StartBreadcrumbs, SubSequenceSpace->EndBreadcrumbs);
 	}
 
 }
diff --git a/Engine/Source/Runtime/MovieScene/Private/Sections/MovieSceneHookSection.cpp b/Engine/Source/Runtime/MovieScene/Private/Sections/MovieSceneHookSection.cpp
--- a/Engine/Source/Runtime/MovieScene/Private/Sections/MovieSceneHookSection.cpp
+++ b/Engine/Source/Runtime/MovieScene/Private/Sections/MovieSceneHookSection.cpp
@@ -41,14 +41,20 @@ void UMovieSceneHookSection::ImportTriggerEntity(UMovieSceneEntitySystemLinker*
 		return;
 	}
 
-	const FSequenceInstance& ThisInstance   = EntityLinker->GetInstanceRegistry()->GetInstance(Params.Sequence.InstanceHandle);
-	FMovieSceneContext       Context        = ThisInstance.GetContext();
+	const FSequenceInstance&  ThisInstance   = EntityLinker->GetInstanceRegistry()->GetInstance(Params.Sequence.InstanceHandle);
+	const FMovieSceneContext& Context        = ThisInstance.GetContext();
+
+	TOptional<FFrameTime> RootTime = Context.GetSequenceToRootSequenceTransform().TryTransformTime(Times[EventIndex], Context.GetRootToSequenceWarpCounter());
+	if (!RootTime)
+	{
+		return;
+	}
 
 	FMovieSceneEvaluationHookEvent NewEvent;
 	NewEvent.Hook.Interface = this;
 	NewEvent.Hook.ObjectBindingID = Params.GetObjectBindingID();
 
-	NewEvent.RootTime = Times[EventIndex] * Context.GetSequenceToRootSequenceTransform();
+	NewEvent.RootTime = RootTime.GetValue();
 	NewEvent.Type = EEvaluationHookEvent::Trigger;
 	NewEvent.SequenceID = ThisInstance.GetSequenceID();
 	NewEvent.TriggerIndex = EventIndex;
diff --git a/Engine/Source/Runtime/MovieScene/Private/Sections/MovieSceneSubSection.cpp b/Engine/Source/Runtime/MovieScene/Private/Sections/MovieSceneSubSection.cpp
--- a/Engine/Source/Runtime/MovieScene/Private/Sections/MovieSceneSubSection.cpp
+++ b/Engine/Source/Runtime/MovieScene/Private/Sections/MovieSceneSubSection.cpp
@@ -10,10 +10,15 @@
 #include "Evaluation/MovieSceneEvaluationTemplate.h"
 #include "Misc/FrameRate.h"
 #include "Logging/MessageLog.h"
+#include "MovieSceneTransformTypes.h"
 #include "Evaluation/MovieSceneRootOverridePath.h"
 #include "EntitySystem/MovieSceneEntitySystemLinker.h"
 #include "EntitySystem/MovieSceneInstanceRegistry.h"
 #include "EntitySystem/IMovieSceneEntityProvider.h"
+#include "Channels/MovieSceneChannelProxy.h"
+#include "Tracks/MovieSceneTimeWarpTrack.h"
+#include "Sections/MovieSceneSectionTimingParameters.h"
+#include "Variants/MovieSceneTimeWarpGetter.h"
 
 #include UE_INLINE_GENERATED_CPP_BY_NAME(MovieSceneSubSection)
 
@@ -33,7 +38,45 @@ UMovieSceneSubSection::UMovieSceneSubSection(const FObjectInitializer& ObjInitia
 	SetBlendType(EMovieSceneBlendType::Absolute);
 }
 
-FMovieSceneSequenceTransform UMovieSceneSubSection::OuterToInnerTransform() const
+void UMovieSceneSubSection::DeleteChannels(TArrayView<const FName> ChannelNames)
+{
+	bool bDeletedAny = false;
+
+	if (Parameters.TimeScale.GetType() == EMovieSceneTimeWarpType::Custom && TryModify())
+	{
+		if (UMovieSceneTimeWarpGetter* Getter = Parameters.TimeScale.AsCustom())
+		{
+			for (FName ChannelName : ChannelNames)
+			{
+				bDeletedAny |= Getter->DeleteChannel(Parameters.TimeScale, ChannelName);
+			}
+		}
+	}
+
+	if (bDeletedAny)
+	{
+		ChannelProxy = nullptr;
+	}
+}
+
+EMovieSceneChannelProxyType UMovieSceneSubSection::CacheChannelProxy()
+{
+	FMovieSceneChannelProxyData Channels;
+
+	if (Parameters.TimeScale.GetType() == EMovieSceneTimeWarpType::Custom)
+	{
+		if (UMovieSceneTimeWarpGetter* Curve = Parameters.TimeScale.AsCustom())
+		{
+			Curve->PopulateChannelProxy(Channels, UMovieSceneTimeWarpGetter::EAllowTopLevelChannels::No);
+		}
+	}
+
+	ChannelProxy = MakeShared<FMovieSceneChannelProxy>(MoveTemp(Channels));
+	return EMovieSceneChannelProxyType::Dynamic;
+}
+
+
+FMovieSceneSequenceTransform UMovieSceneSubSection::OuterToInnerTransform_NoInnerTimeWarp() const
 {
 	UMovieSceneSequence* SequencePtr   = GetSequence();
 	if (!SequencePtr)
@@ -51,44 +94,63 @@ FMovieSceneSequenceTransform UMovieSceneSubSection::OuterToInnerTransform() cons
 
 	const FFrameRate   InnerFrameRate = MovieScenePtr->GetTickResolution();
 	const FFrameRate   OuterFrameRate = GetTypedOuter<UMovieScene>()->GetTickResolution();
-	const float        FrameRateScale = (OuterFrameRate == InnerFrameRate) ? 1.f : (InnerFrameRate / OuterFrameRate).AsDecimal();
 
 	const TRange<FFrameNumber> MovieScenePlaybackRange = GetValidatedInnerPlaybackRange(Parameters, *MovieScenePtr);
-	const FFrameNumber InnerStartTime = UE::MovieScene::DiscreteInclusiveLower(MovieScenePlaybackRange);
-	const FFrameNumber OuterStartTime = UE::MovieScene::DiscreteInclusiveLower(SubRange);
 
+	FMovieSceneSectionTimingParametersFrames TimingParams = {
+		Parameters.TimeScale.ShallowCopy(),
+		Parameters.StartFrameOffset,
+		Parameters.EndFrameOffset,
+		Parameters.FirstLoopStartFrameOffset,
+		Parameters.bCanLoop,
+		false, // do not clamp sub-sections by default
+		false
+	};
+
+	return TimingParams.MakeTransform(OuterFrameRate, SubRange, InnerFrameRate, MovieScenePtr->GetPlaybackRange());
+}
 
-	FMovieSceneSequenceTransform Result;
-	// We have to special case 0 and infinite timescale so we can keep hold of each transform separately for property inverting.
-	// The linear transform for this special case remains identity.
-	if (FMath::IsNearlyZero(Parameters.TimeScale) || !FMath::IsFinite(Parameters.TimeScale))
-	{
-		Result.NestedTransforms.Add(FMovieSceneTimeTransform(-OuterStartTime));
-		Result.NestedTransforms.Add(FMovieSceneTimeTransform(0, 0));
-		Result.NestedTransforms.Add(FMovieSceneTimeTransform(InnerStartTime));
-	}
-	else
+FMovieSceneSequenceTransform UMovieSceneSubSection::OuterToInnerTransform() const
+{
+	FMovieSceneSequenceTransform OuterToInner = OuterToInnerTransform_NoInnerTimeWarp();
+	AppendInnerTimeWarpTransform(OuterToInner);
+	return OuterToInner;
+}
+
+void UMovieSceneSubSection::AppendInnerTimeWarpTransform(FMovieSceneSequenceTransform& OutTransform) const
+{
+	UMovieSceneSequence* Sequence   = GetSequence();
+	UMovieScene*         MovieScene = Sequence ? Sequence->GetMovieScene() : nullptr;
+
+	if (!MovieScene)
 	{
-		// This is the transform for the "placement" (position and scaling) of the sub-sequence.
-		Result.LinearTransform = 
-			// Inner play offset
-			FMovieSceneTimeTransform(InnerStartTime)
-			// Inner play rate
-			* FMovieSceneTimeTransform(0, Parameters.TimeScale * FrameRateScale)
-			// Outer section start time
-			* FMovieSceneTimeTransform(-OuterStartTime);
+		return;
 	}
 
-	if (Parameters.bCanLoop)
+	// Look for any time warp tracks inside the sub sequence
+	for (UMovieSceneTrack* Track : MovieScene->GetTracks())
 	{
-		const FFrameNumber InnerEndTime = UE::MovieScene::DiscreteExclusiveUpper(MovieScenePlaybackRange);
-		const FMovieSceneTimeWarping LoopingTransform(InnerStartTime, InnerEndTime);
+		UMovieSceneTimeWarpTrack* TimeWarpTrack = Cast<UMovieSceneTimeWarpTrack>(Track);
+		if (TimeWarpTrack && !TimeWarpTrack->IsEvalDisabled())
+		{
+			FMovieSceneNestedSequenceTransform TimeWarpTransform = TimeWarpTrack->GenerateTransform();
 
-		Result.NestedTransforms.Add(FMovieSceneNestedSequenceTransform(FMovieSceneTimeTransform(Parameters.FirstLoopStartFrameOffset) * Result.LinearTransform, LoopingTransform));
-		Result.LinearTransform = FMovieSceneTimeTransform();
-		return Result;
+			if (!TimeWarpTransform.IsIdentity())
+			{
+				if (TimeWarpTransform.IsLinear() && OutTransform.IsLinear())
+				{
+					OutTransform = FMovieSceneSequenceTransform(OutTransform.AsLinear() * TimeWarpTransform.AsLinear());
+				}
+				else
+				{
+					OutTransform.NestedTransforms.Add(TimeWarpTransform);
+				}
+			}
+
+			// Only 1 timewarp track supported
+			return;
+		}
 	}
-	return Result;
 }
 
 bool UMovieSceneSubSection::GetValidatedInnerPlaybackRange(TRange<FFrameNumber>& OutInnerPlaybackRange) const
@@ -300,6 +362,11 @@ void UMovieSceneSubSection::PostEditChangeProperty(FPropertyChangedEvent& Proper
 		PreviousSubSequence = nullptr;
 	}
 
+	if (PropertyChangedEvent.Property && PropertyChangedEvent.Property->GetFName() == GET_MEMBER_NAME_CHECKED(FMovieSceneSectionParameters, TimeScale))
+	{
+		ChannelProxy = nullptr;
+	}
+
 	Super::PostEditChangeProperty(PropertyChangedEvent);
 
 	// recreate runtime instance when sequence is changed
@@ -317,11 +384,11 @@ TOptional<TRange<FFrameNumber> > UMovieSceneSubSection::GetAutoSizeRange() const
 	{
 		// We probably want to just auto-size the section to the sub-sequence's scaled playback range... if this section
 		// is looping, however, it's hard to know what we want to do. Let's just size it to one loop.
-		const FMovieSceneSequenceTransform InnerToOuter = OuterToInnerTransform().InverseNoLooping();
+		const FMovieSceneInverseSequenceTransform InnerToOuter = OuterToInnerTransform().Inverse();
 		const TRange<FFrameNumber> InnerPlaybackRange = UMovieSceneSubSection::GetValidatedInnerPlaybackRange(Parameters, *MovieScene);
 
-		const FFrameTime IncAutoStartTime = FFrameTime(UE::MovieScene::DiscreteInclusiveLower(InnerPlaybackRange)) * InnerToOuter;
-		const FFrameTime ExcAutoEndTime = FFrameTime(UE::MovieScene::DiscreteExclusiveUpper(InnerPlaybackRange)) * InnerToOuter;
+		const FFrameTime IncAutoStartTime = InnerToOuter.TryTransformTime(UE::MovieScene::DiscreteInclusiveLower(InnerPlaybackRange)).Get(InnerPlaybackRange.GetLowerBoundValue());
+		const FFrameTime ExcAutoEndTime   = InnerToOuter.TryTransformTime(UE::MovieScene::DiscreteExclusiveUpper(InnerPlaybackRange)).Get(InnerPlaybackRange.GetUpperBoundValue());
 
 		return TRange<FFrameNumber>(GetInclusiveStartFrame(), GetInclusiveStartFrame() + (ExcAutoEndTime.RoundToFrame() - IncAutoStartTime.RoundToFrame()));
 	}
@@ -382,8 +449,13 @@ void UMovieSceneSubSection::TrimSection( FQualifiedFrameTime TrimTime, bool bTri
 
 void UMovieSceneSubSection::GetSnapTimes(TArray<FFrameNumber>& OutSnapTimes, bool bGetSectionBorders) const
 {
+	using namespace UE::MovieScene;
+
 	Super::GetSnapTimes(OutSnapTimes, bGetSectionBorders);
 
+	const FFrameNumber StartFrame = GetInclusiveStartFrame();
+	const FFrameNumber EndFrame   = GetExclusiveEndFrame();
+
 	UMovieSceneSequence* Sequence = GetSequence();
 	UMovieScene* MovieScene = Sequence ? Sequence->GetMovieScene() : nullptr;
 	if (!MovieScene)
@@ -391,47 +463,31 @@ void UMovieSceneSubSection::GetSnapTimes(TArray<FFrameNumber>& OutSnapTimes, boo
 		return;
 	}
 
-    TRange<FFrameNumber> PlaybackRange = MovieScene->GetPlaybackRange();
-	const FFrameNumber StartFrame = GetInclusiveStartFrame();
-	const FFrameNumber EndFrame = GetExclusiveEndFrame() - 1; // -1 because we don't need to add the end frame twice
-
-	if (Parameters.bCanLoop)
+	auto VisitBoundary = [&OutSnapTimes](FFrameTime InTime)
 	{
-	    const float InvTimeScale = FMath::IsNearlyZero(Parameters.TimeScale) ? 1.0f : 1.0f / Parameters.TimeScale;
-		const TRange<FFrameNumber> InnerPlaybackRange = UMovieSceneSubSection::GetValidatedInnerPlaybackRange(Parameters, *MovieScene);
+		OutSnapTimes.Add(InTime.RoundToFrame());
+		return true;
+	};
 
-		const FFrameNumber InnerSubSeqLength = UE::MovieScene::DiscreteSize(InnerPlaybackRange);
-		const FFrameNumber InnerSubSeqFirstLoopLength = InnerSubSeqLength - Parameters.FirstLoopStartFrameOffset;
+	FMovieSceneSequenceTransform OuterToInner = OuterToInnerTransform();
 
-		const FFrameRate OuterFrameRate = GetTypedOuter<UMovieScene>()->GetTickResolution();
-		const FFrameRate InnerFrameRate = MovieScene->GetTickResolution();
-		const FFrameNumber OuterSubSeqLength = (ConvertFrameTime(InnerSubSeqLength, InnerFrameRate, OuterFrameRate) * InvTimeScale).FrameNumber;
-		const FFrameNumber OuterSubSeqFirstLoopLength = (ConvertFrameTime(InnerSubSeqFirstLoopLength, InnerFrameRate, OuterFrameRate) * InvTimeScale).FrameNumber;
+	if (!OuterToInner.ExtractBoundariesWithinRange(StartFrame, EndFrame, VisitBoundary))
+	{
+		FMovieSceneInverseSequenceTransform InnerToOuterTransform = OuterToInner.Inverse();
 
-		FFrameNumber CurOffsetFrame = FMath::Max(OuterSubSeqFirstLoopLength, FFrameNumber(0));
-	    
-		while (CurOffsetFrame < EndFrame)
-		{
-			const int32 CurOffset = CurOffsetFrame.Value;
-			      
-			OutSnapTimes.Add(StartFrame + CurOffset);
+		TRange<FFrameNumber> PlaybackRange = MovieScene->GetPlaybackRange();
 
-			CurOffsetFrame += OuterSubSeqLength;
-		}
-	}
-	else
-	{
-		const FMovieSceneSequenceTransform InnerToOuterTransform = OuterToInnerTransform().InverseNoLooping();
-		const FFrameNumber PlaybackStart = (UE::MovieScene::DiscreteInclusiveLower(PlaybackRange) * InnerToOuterTransform).FloorToFrame();
-		if (GetRange().Contains(PlaybackStart))
+		TOptional<FFrameTime> SequenceStart = InnerToOuterTransform.TryTransformTime(PlaybackRange.GetLowerBoundValue());
+		TOptional<FFrameTime> SequenceEnd   = InnerToOuterTransform.TryTransformTime(PlaybackRange.GetUpperBoundValue());
+
+		if (SequenceStart && SequenceStart.GetValue() >= StartFrame && SequenceStart.GetValue() < EndFrame)
 		{
-			OutSnapTimes.Add(PlaybackStart);
+			VisitBoundary(SequenceStart.GetValue());
 		}
 
-		const FFrameNumber PlaybackEnd = (UE::MovieScene::DiscreteExclusiveUpper(PlaybackRange) * InnerToOuterTransform).FloorToFrame();
-		if (GetRange().Contains(PlaybackEnd))
+		if (SequenceEnd && SequenceEnd.GetValue() >= StartFrame && SequenceEnd.GetValue() < EndFrame)
 		{
-			OutSnapTimes.Add(PlaybackEnd);
+			VisitBoundary(SequenceEnd.GetValue());
 		}
 	}
 }
@@ -462,7 +518,7 @@ FMovieSceneSubSequenceData UMovieSceneSubSection::GenerateSubSequenceData(const
 
 FFrameNumber UMovieSceneSubSection::MapTimeToSectionFrame(FFrameTime InPosition) const
 {
-	FFrameNumber LocalPosition = ((InPosition - Parameters.StartFrameOffset) * Parameters.TimeScale).GetFrame();
+	FFrameNumber LocalPosition = ((InPosition - Parameters.StartFrameOffset) * OuterToInnerTransform()).GetFrame();
 	return LocalPosition;
 }
 
diff --git a/Engine/Source/Runtime/MovieScene/Private/Tests/MovieSceneTransformTests.cpp b/Engine/Source/Runtime/MovieScene/Private/Tests/MovieSceneTransformTests.cpp
--- a/Engine/Source/Runtime/MovieScene/Private/Tests/MovieSceneTransformTests.cpp
+++ b/Engine/Source/Runtime/MovieScene/Private/Tests/MovieSceneTransformTests.cpp
@@ -3,9 +3,11 @@
 #include "CoreMinimal.h"
 #include "Evaluation/MovieSceneSequenceTransform.h"
 #include "Evaluation/MovieSceneSectionParameters.h"
+#include "Variants/MovieSceneTimeWarpVariantPayloads.h"
 #include "Containers/ArrayView.h"
 #include "Misc/AutomationTest.h"
 #include "MovieSceneTimeHelpers.h"
+#include "UObject/Package.h"
 
 #define LOCTEXT_NAMESPACE "MovieSceneTransformTests"
 
@@ -49,11 +51,6 @@ FFrameNumber TransformToFrameNumber(TTransform Transform, FFrameNumber Value)
 {
 	return (Value * Transform).FloorToFrame();
 }
-template<>
-FFrameNumber TransformToFrameNumber(FMovieSceneTimeWarping Transform, FFrameNumber Value)
-{
-	return Value * Transform;
-}
 
 // Generic method for testing the transform of frames and times.
 template<typename TTransform>
@@ -175,73 +172,6 @@ bool FMovieSceneSubSectionCoreLinearTransformsTest::RunTest(const FString& Param
 	return bSuccess;
 }
 
-// Warping transform tests
-IMPLEMENT_SIMPLE_AUTOMATION_TEST(FMovieSceneSubSectionCoreWarpTransformsTest,
-		"System.Engine.Sequencer.Core.WarpTransforms",
-		EAutomationTestFlags::EditorContext | EAutomationTestFlags::EngineFilter)
-bool FMovieSceneSubSectionCoreWarpTransformsTest::RunTest(const FString& Parameters)
-{
-	FFrameNumber SourceTimes[] = {
-		FFrameNumber(0),
-		FFrameNumber(25),
-		FFrameNumber(50),
-		FFrameNumber(60),
-		FFrameNumber(120)
-	};
-
-	bool bSuccess = true;
-
-	{
-		FFrameNumber ExpectedTimes[] = {
-			FFrameNumber(0),
-			FFrameNumber(25),
-			FFrameNumber(0),
-			FFrameNumber(10),
-			FFrameNumber(20)
-		};
-		FMovieSceneTimeWarping Warping(0, 50);
-		bSuccess = TestTransform(*this, Warping, SourceTimes, ExpectedTimes, TEXT("SimpleWarping")) && bSuccess;
-
-		FMovieSceneTimeTransform Transform = Warping.InverseFromWarp(0);
-		bSuccess = TestTransform(*this, Transform, MakeArrayView(ExpectedTimes, 2), MakeArrayView(SourceTimes, 2), TEXT("SimpleWarpingInverseLoop0")) && bSuccess;
-
-		FMovieSceneTimeTransform Transform2 = Warping.InverseFromWarp(1);
-		bSuccess = TestTransform(*this, Transform2, MakeArrayView(ExpectedTimes + 2, 2), MakeArrayView(SourceTimes + 2, 2), TEXT("SimpleWarpingInverseLoop1")) && bSuccess;
-
-		FMovieSceneTimeTransform Transform3 = Warping.InverseFromWarp(2);
-		bSuccess = TestTransform(*this, Transform3, MakeArrayView(ExpectedTimes + 4, 1), MakeArrayView(SourceTimes + 4, 1), TEXT("SimpleWarpingInverseLoop2")) && bSuccess;
-	}
-
-	SourceTimes[0] = FFrameNumber(3);
-	SourceTimes[1] = FFrameNumber(28);
-	SourceTimes[2] = FFrameNumber(53);
-	SourceTimes[3] = FFrameNumber(63);
-	SourceTimes[4] = FFrameNumber(123);
-
-	{
-		FFrameNumber ExpectedTimes[] {
-			FFrameNumber(3),
-			FFrameNumber(28),
-			FFrameNumber(14),
-			FFrameNumber(24),
-			FFrameNumber(6)
-		};
-		FMovieSceneTimeWarping Warping(3, 42);
-		bSuccess = TestTransform(*this, Warping, SourceTimes, ExpectedTimes, TEXT("WarpingWithTrim")) && bSuccess;
-
-		FMovieSceneTimeTransform Transform = Warping.InverseFromWarp(0);
-		bSuccess = TestTransform(*this, Transform, MakeArrayView(ExpectedTimes, 2), MakeArrayView(SourceTimes, 2), TEXT("WarpingWithTrimInverseLoop0")) && bSuccess;
-
-		FMovieSceneTimeTransform Transform2 = Warping.InverseFromWarp(1);
-		bSuccess = TestTransform(*this, Transform2, MakeArrayView(ExpectedTimes + 2, 2), MakeArrayView(SourceTimes + 2, 2), TEXT("WarpingWithTrimInverseLoop1")) && bSuccess;
-
-		FMovieSceneTimeTransform Transform3 = Warping.InverseFromWarp(3);  // We lapsed one full loop
-		bSuccess = TestTransform(*this, Transform3, MakeArrayView(ExpectedTimes + 4, 1), MakeArrayView(SourceTimes + 4, 1), TEXT("WarpingWithTrimInverseLoop2")) && bSuccess;
-	}
-
-	return bSuccess;
-}
-
 // Sequence transform tests 
 IMPLEMENT_SIMPLE_AUTOMATION_TEST(FMovieSceneSubSectionCoreSequenceTransformsTest,
 		"System.Engine.Sequencer.Core.SequenceTransforms",
@@ -318,12 +248,16 @@ bool FMovieSceneSubSectionCoreSequenceTransformsTest::RunTest(const FString& Par
 			AddError(FString::Printf(TEXT("Accumulated transform does not have the same effect as separate transformations (%i+%.5f != %i+%.5f)"), AccumValue.FrameNumber.Value, AccumValue.GetSubFrame(), SeedValue.FrameNumber.Value, SeedValue.GetSubFrame()));
 		}
 
-		FMovieSceneSequenceTransform InverseTransform = AccumulatedTransform.InverseNoLooping();
+		FMovieSceneInverseSequenceTransform InverseTransform = AccumulatedTransform.Inverse();
 
-		FFrameTime InverseValue = AccumValue * InverseTransform;
-		if (InverseValue != 10)
+		TOptional<FFrameTime> InverseValue = InverseTransform.TryTransformTime(AccumValue);
+		if (!InverseValue.IsSet())
+		{
+			AddError(FString::Printf(TEXT("Inverse accumulated transform did not return a valid time")));
+		}
+		else if (InverseValue.GetValue() != 10)
 		{
-			AddError(FString::Printf(TEXT("Inverse accumulated transform does not return value back to its original value (%i+%.5f != 10)"), InverseValue.FrameNumber.Value, InverseValue.GetSubFrame()));
+			AddError(FString::Printf(TEXT("Inverse accumulated transform does not return value back to its original value (%i+%.5f != 10)"), InverseValue->FrameNumber.Value, InverseValue->GetSubFrame()));
 		}
 	}
 
@@ -339,102 +273,145 @@ bool FMovieSceneSubSectionCoreWarpingAndScalingTransformsTest::RunTest(const FSt
 	{
 		// Sub-sequence at 0, playing at x2
 		FMovieSceneSequenceTransform Transform;
-		Transform.NestedTransforms.Emplace(FMovieSceneTimeTransform(0, 2.f), FMovieSceneTimeWarping(0, 30));
+		Transform.LinearTransform = FMovieSceneTimeTransform(0, 2.f);
+		Transform.AddLoop(0, 30);
 		TestEqual("Transform time 1", FFrameNumber(10) * Transform, FFrameTime(20));
-		TestEqual("Transform time 2", FFrameNumber(18) * Transform, FFrameTime(6));
 
-		FMovieSceneSequenceTransform Inv = Transform.InverseFromLoop(TArray<uint32> { 0 });
-		TestEqual("Inverse time 3", FFrameNumber(20) * Inv, FFrameTime(10));
-		Inv = Transform.InverseFromLoop(TArray<uint32> { 1 });
-		TestEqual("Inverse time 4", FFrameNumber(20) * Inv, FFrameTime(25));
+
+		FMovieSceneTransformBreadcrumbs Breadcrumbs;
+		FMovieSceneInverseSequenceTransform Inv = Transform.Inverse();
+
+		Breadcrumbs.AddBreadcrumb(15);
+		TestEqual("Inverse time 3", Inv.TryTransformTime(FFrameNumber(20), Breadcrumbs), TOptional<FFrameTime>(10));
+
+		Breadcrumbs.Reset();
+		Breadcrumbs.AddBreadcrumb(45);
+		TestEqual("Inverse time 4", Inv.TryTransformTime(FFrameNumber(20), Breadcrumbs), TOptional<FFrameTime>(FFrameTime(25)));
 	}
 
 	{
 		// Sub-sequence at 0, playing at x2, with start offset 20
 		FMovieSceneSequenceTransform Transform;
-		Transform.NestedTransforms.Emplace(FMovieSceneTimeTransform(20, 2.f), FMovieSceneTimeWarping(20, 50));
+		Transform.LinearTransform = FMovieSceneTimeTransform(20, 2.f);
+		Transform.AddLoop(20, 50);
 		TestEqual("Transform time 5", FFrameNumber(10) * Transform, FFrameTime(40));
 		TestEqual("Transform time 6", FFrameNumber(18) * Transform, FFrameTime(26));
 
-		FMovieSceneSequenceTransform Inv = Transform.InverseFromLoop(TArray<uint32> { 0 });
-		TestEqual("Inverse time 7", FFrameNumber(40) * Inv, FFrameTime(10));
-		Inv = Transform.InverseFromLoop(TArray<uint32> { 1 });
-		TestEqual("Inverse time 8", FFrameNumber(40) * Inv, FFrameTime(25));
+		FMovieSceneTransformBreadcrumbs Breadcrumbs;
+		FMovieSceneInverseSequenceTransform Inv = Transform.Inverse();
+
+		Breadcrumbs.AddBreadcrumb(35);	// 35 should be in the middle of the first loop
+		TestEqual("Inverse time 7", Inv.TryTransformTime(FFrameNumber(40), Breadcrumbs), TOptional<FFrameTime>(FFrameTime(10)));
+
+		Breadcrumbs.Reset();
+		Breadcrumbs.AddBreadcrumb(65);	// 65 should be in the middle of the second loop
+		TestEqual("Inverse time 8", Inv.TryTransformTime(FFrameNumber(40), Breadcrumbs), TOptional<FFrameTime>(FFrameTime(25)));
 	}
 
 	{
 		// Sub-sequence at 3, playing at x2
 		FMovieSceneSequenceTransform Transform;
-		Transform.NestedTransforms.Emplace(FMovieSceneTimeTransform(-6, 2.f), FMovieSceneTimeWarping(0, 30));
+		Transform.LinearTransform = FMovieSceneTimeTransform(-6, 2.f);
+		Transform.AddLoop(0, 30);
 		TestEqual("Transform time 9", FFrameNumber(13) * Transform, FFrameTime(20));
 		TestEqual("Transform time 10", FFrameNumber(21) * Transform, FFrameTime(6));
 
-		FMovieSceneSequenceTransform Inv = Transform.InverseFromLoop(TArray<uint32> { 0 });
-		TestEqual("Inverse time 11", FFrameNumber(20) * Inv, FFrameTime(13));
-		Inv = Transform.InverseFromLoop(TArray<uint32> { 1 });
-		TestEqual("Inverse time 12", FFrameNumber(20) * Inv, FFrameTime(28));
+		FMovieSceneTransformBreadcrumbs Breadcrumbs;
+		FMovieSceneInverseSequenceTransform Inv = Transform.Inverse();
+
+		Breadcrumbs.AddBreadcrumb(15);	// 15 should be in the middle of the first loop
+		TestEqual("Inverse time 11", Inv.TryTransformTime(FFrameNumber(20), Breadcrumbs), TOptional<FFrameTime>(FFrameTime(13)));
+
+		Breadcrumbs.Reset();
+		Breadcrumbs.AddBreadcrumb(45);	// 45 should be in the middle of the second loop
+		TestEqual("Inverse time 12", Inv.TryTransformTime(FFrameNumber(20), Breadcrumbs), TOptional<FFrameTime>(FFrameTime(28)));
 	}
 
 	{
 		// Sub-sequence at 3, playing at x2, with start offset 20
 		FMovieSceneSequenceTransform Transform;
-		Transform.NestedTransforms.Emplace(FMovieSceneTimeTransform(-6 + 20, 2.f), FMovieSceneTimeWarping(20, 50));
+		Transform.LinearTransform = FMovieSceneTimeTransform(-6 + 20, 2.f);
+		Transform.AddLoop(20, 50);
 		TestEqual("Transform time 13", FFrameNumber(13) * Transform, FFrameTime(40));
 		TestEqual("Transform time 14", FFrameNumber(21) * Transform, FFrameTime(26));
 
-		FMovieSceneSequenceTransform Inv = Transform.InverseFromLoop(TArray<uint32> { 0 });
-		TestEqual("Inverse time 15", FFrameNumber(40) * Inv, FFrameTime(13));
-		Inv = Transform.InverseFromLoop(TArray<uint32> { 1 });
-		TestEqual("Inverse time 16", FFrameNumber(40) * Inv, FFrameTime(28));
+		FMovieSceneTransformBreadcrumbs Breadcrumbs;
+		FMovieSceneInverseSequenceTransform Inv = Transform.Inverse();
+
+		Breadcrumbs.AddBreadcrumb(35);	// 35 should be in the middle of the first loop
+		TestEqual("Inverse time 15", Inv.TryTransformTime(FFrameNumber(40), Breadcrumbs), TOptional<FFrameTime>(FFrameTime(13)));
+
+		Breadcrumbs.Reset();
+		Breadcrumbs.AddBreadcrumb(65);	// 65 should be in the middle of the second loop
+		TestEqual("Inverse time 16", Inv.TryTransformTime(FFrameNumber(40), Breadcrumbs), TOptional<FFrameTime>(FFrameTime(28)));
 	}
 
 	{
 		// Two levels of sub-sequences: one placed at 10 and warping, the second placed at 6 with x2 scaling
 		FMovieSceneSequenceTransform Transform;
-		Transform.NestedTransforms.Emplace(FMovieSceneTimeTransform(-10, 1.f), FMovieSceneTimeWarping(0, 30));
+		Transform.LinearTransform = FMovieSceneTimeTransform(-10, 1.f);
+		Transform.AddLoop(0, 30);
 		Transform.NestedTransforms.Emplace(FMovieSceneTimeTransform(-12, 2.f));
 		TestEqual("Transform time 17", FFrameNumber(18) * Transform, FFrameTime(4));
 		TestEqual("Transform time 18", FFrameNumber(55) * Transform, FFrameTime(18));
 
-		FMovieSceneSequenceTransform Inv = Transform.InverseFromLoop(TArray<uint32> { 0, FMovieSceneTimeWarping::InvalidWarpCount });
-		TestEqual("Inverse time 17", FFrameNumber(4) * Inv, FFrameTime(18));
-		Inv = Transform.InverseFromLoop(TArray<uint32> { 1, FMovieSceneTimeWarping::InvalidWarpCount });
-		TestEqual("Inverse time 18", FFrameNumber(18) * Inv, FFrameTime(55));
+		FMovieSceneTransformBreadcrumbs Breadcrumbs;
+		FMovieSceneInverseSequenceTransform Inv = Transform.Inverse();
+
+		Breadcrumbs.AddBreadcrumb(15);
+		TestEqual("Inverse time 17", Inv.TryTransformTime(FFrameNumber(4), Breadcrumbs), TOptional<FFrameTime>(FFrameTime(18)));
+
+		Breadcrumbs.Reset();
+		Breadcrumbs.AddBreadcrumb(45);
+		TestEqual("Inverse time 18", Inv.TryTransformTime(FFrameNumber(18), Breadcrumbs), TOptional<FFrameTime>(FFrameTime(55)));
 	}
 
 	{
 		// Two levels of sub-sequences: one placed at 10, the second placed at 6 with x2 scaling and warping
 		FMovieSceneSequenceTransform Transform;
-		Transform.NestedTransforms.Emplace(FMovieSceneTimeTransform(-10, 1.f));
-		Transform.NestedTransforms.Emplace(FMovieSceneTimeTransform(-12, 2.f), FMovieSceneTimeWarping(0, 14));
+		Transform.LinearTransform = FMovieSceneTimeTransform(-10, 1.f);
+		Transform.NestedTransforms.Emplace(FMovieSceneTimeTransform(-12, 2.f));
+		Transform.AddLoop(0, 14);
 		TestEqual("Transform time 17", FFrameNumber(19) * Transform, FFrameTime(6));
 		TestEqual("Transform time 18", FFrameNumber(32) * Transform, FFrameTime(4));
 
-		FMovieSceneSequenceTransform Inv = Transform.InverseFromLoop(TArray<uint32> { FMovieSceneTimeWarping::InvalidWarpCount, 0 });
-		TestEqual("Inverse time 17", FFrameNumber(6) * Inv, FFrameTime(19));
-		Inv = Transform.InverseFromLoop(TArray<uint32> { FMovieSceneTimeWarping::InvalidWarpCount, 2 });
-		TestEqual("Inverse time 18", FFrameNumber(4) * Inv, FFrameTime(32));
+		FMovieSceneTransformBreadcrumbs Breadcrumbs;
+		FMovieSceneInverseSequenceTransform Inv = Transform.Inverse();
+
+		Breadcrumbs.AddBreadcrumb(7); // half way through loop 0
+		TestEqual("Inverse time 17", Inv.TryTransformTime(FFrameNumber(6), Breadcrumbs), TOptional<FFrameTime>(FFrameTime(19)));
+
+		Breadcrumbs.Reset();
+		Breadcrumbs.AddBreadcrumb(28); // half way through loop 2
+		TestEqual("Inverse time 18", Inv.TryTransformTime(FFrameNumber(4), Breadcrumbs), TOptional<FFrameTime>(FFrameTime(32)));
 	}
 
 	{
 		// Sub-sequence at 3, playing at x2, with start offset 20, but all contained inside a higher offset of 100.
 		FMovieSceneSequenceTransform Transform;
 		Transform.LinearTransform.Offset = FFrameTime(-100);
-		Transform.NestedTransforms.Emplace(FMovieSceneTimeTransform(-6 + 20, 2.f), FMovieSceneTimeWarping(20, 50));
+		Transform.NestedTransforms.Emplace(FMovieSceneTimeTransform(-6 + 20, 2.f));
+		Transform.AddLoop(20, 50);
+
 		TestEqual("Transform time 19", FFrameNumber(113) * Transform, FFrameTime(40));
 		TestEqual("Transform time 20", FFrameNumber(121) * Transform, FFrameTime(26));
 
-		FMovieSceneSequenceTransform Inv = Transform.InverseFromLoop(TArray<uint32> { 0 });
-		TestEqual("Inverse time 21", FFrameNumber(40) * Inv, FFrameTime(113));
-		Inv = Transform.InverseFromLoop(TArray<uint32> { 1 });
-		TestEqual("Inverse time 22", FFrameNumber(40) * Inv, FFrameTime(128));
+		FMovieSceneTransformBreadcrumbs Breadcrumbs;
+		FMovieSceneInverseSequenceTransform Inv = Transform.Inverse();
+
+		Breadcrumbs.AddBreadcrumb(35); // Loop 0
+		TestEqual("Inverse time 21", Inv.TryTransformTime(FFrameNumber(40), Breadcrumbs), TOptional<FFrameTime>(FFrameTime(113)));
+
+		Breadcrumbs.Reset();
+		Breadcrumbs.AddBreadcrumb(65); // Loop 1
+		TestEqual("Inverse time 22", Inv.TryTransformTime(FFrameNumber(40), Breadcrumbs), TOptional<FFrameTime>(FFrameTime(128)));
 	}
 
 	{
 		// Zero-timescale transform on a sub-sequence. Any frame numbers transformed in should be equal to the frame offset
 		FMovieSceneSequenceTransform Transform;
 		Transform.NestedTransforms.Add(FMovieSceneTimeTransform(0)); // no outer offset
-		Transform.NestedTransforms.Add(FMovieSceneTimeTransform(0, 0)); // 0 timescale
+		Transform.NestedTransforms.Add(FMovieSceneTimeWarpVariant(0.0)); // 0 timescale
 		Transform.NestedTransforms.Add(FMovieSceneTimeTransform(30)); // 30 inner frame offset
 
 		TestEqual("Outer time 40 through 0 timescale with 30 offset", FFrameNumber(40) * Transform, FFrameTime(30));
@@ -447,29 +424,32 @@ bool FMovieSceneSubSectionCoreWarpingAndScalingTransformsTest::RunTest(const FSt
 		// and any transforms by that infinite transform. Anything transformed out should just be equal to the outer offset
 		FMovieSceneSequenceTransform Transform;
 		Transform.NestedTransforms.Add(FMovieSceneTimeTransform(-10)); // 10 outer offset
-		Transform.NestedTransforms.Add(FMovieSceneTimeTransform(0, 0)); // 0 timescale
+		Transform.NestedTransforms.Add(FMovieSceneTimeWarpVariant(0.0)); // 0 timescale
 		Transform.NestedTransforms.Add(FMovieSceneTimeTransform(30)); // 30 inner frame offset
 
-		FMovieSceneSequenceTransform Inv = Transform.InverseNoLooping();
-		if (FMath::IsFinite(Inv.GetTimeScale()))
+		FMovieSceneInverseSequenceTransform Inv = Transform.Inverse();
+		if (Inv.IsLinear())
 		{
-			AddError(FString::Printf(TEXT("Inverse of a transform with zero timescale is not correctly infinite")));
+			AddError(FString::Printf(TEXT("Inverse of a transform with zero timescale is not correctly warping")));
 		}
 
-		TestEqual("Inner time 40 through inf timescale with 10 outer offset", FFrameNumber(40) * Inv, FFrameTime(10));
-		TestEqual("Inner time 0 through inf timescale with 10 outer offset", FFrameNumber(0)* Inv, FFrameTime(10));
-		TestEqual("Inner time 173 through inf timescale with 10 outer offset", FFrameNumber(173)* Inv, FFrameTime(10));
+		TestEqual("Inner time 40 through inf timescale with 10 outer offset",  Inv.TryTransformTime(FFrameNumber(40)),  TOptional<FFrameTime>());
+		TestEqual("Inner time 0 through inf timescale with 10 outer offset",   Inv.TryTransformTime(FFrameNumber(0)),   TOptional<FFrameTime>());
+		TestEqual("Inner time 173 through inf timescale with 10 outer offset", Inv.TryTransformTime(FFrameNumber(173)), TOptional<FFrameTime>());
+		TestEqual("Inner time 30 through inf timescale with 10 outer offset",  Inv.TryTransformTime(FFrameNumber(30)),  TOptional<FFrameTime>(FFrameTime(10)));
 
 		// Re-invert the inverse transform. This should be equivalent to the original transform and we shouldn't have lost anything.
-		FMovieSceneSequenceTransform InvInv = Inv.InverseNoLooping();
-		TestEqual("Doubly inverted zero-timescale transform should be equal to original", InvInv, Transform);
+
+		// @todo: is this necessary? We haven't needed an Inverse for an Inverse anywhere else in the codebase
+		// FMovieSceneSequenceTransform InvInv = Inv.Inverse();
+		// TestEqual("Doubly inverted zero-timescale transform should be equal to original", InvInv, Transform);
 	}
 
 	{
 		// Multiple levels of sub sequences with zero-timescale thrown in
 		FMovieSceneSequenceTransform OuterTransform;
 		OuterTransform.NestedTransforms.Add(FMovieSceneTimeTransform(-10)); // 10 outer offset
-		OuterTransform.NestedTransforms.Add(FMovieSceneTimeTransform(0, 0)); // 0 timescale
+		OuterTransform.NestedTransforms.Add(FMovieSceneTimeWarpVariant(0.0)); // 0 timescale
 		OuterTransform.NestedTransforms.Add(FMovieSceneTimeTransform(30)); // 30 inner frame offset
 
 		FMovieSceneSequenceTransform InnerTransform;
@@ -481,15 +461,17 @@ bool FMovieSceneSubSectionCoreWarpingAndScalingTransformsTest::RunTest(const FSt
 		TestEqual("Subsequence frame through zero timescale transform", FFrameNumber(0)* CompleteTransform, FFrameTime(35));
 		TestEqual("Subsequence frame through zero timescale transform", FFrameNumber(173)* CompleteTransform, FFrameTime(35));
 
-		FMovieSceneSequenceTransform InvCompleteTransform = CompleteTransform.InverseNoLooping();
+		FMovieSceneInverseSequenceTransform InvCompleteTransform = CompleteTransform.Inverse();
 
-		TestEqual("Inner time 40 through inf timescale with 10 outer offset", FFrameNumber(40)* InvCompleteTransform, FFrameTime(10));
-		TestEqual("Inner time 0 through inf timescale with 10 outer offset", FFrameNumber(0)* InvCompleteTransform, FFrameTime(10));
-		TestEqual("Inner time 173 through inf timescale with 10 outer offset", FFrameNumber(173)* InvCompleteTransform, FFrameTime(10));
+		TestEqual("Inner time 40 through inf timescale with 10 outer offset",  InvCompleteTransform.TryTransformTime(FFrameNumber(40)),  TOptional<FFrameTime>());
+		TestEqual("Inner time 0 through inf timescale with 10 outer offset",   InvCompleteTransform.TryTransformTime(FFrameNumber(0)),   TOptional<FFrameTime>());
+		TestEqual("Inner time 173 through inf timescale with 10 outer offset", InvCompleteTransform.TryTransformTime(FFrameNumber(173)), TOptional<FFrameTime>());
+		TestEqual("Inner time 35 through inf timescale with 10 outer offset",  InvCompleteTransform.TryTransformTime(FFrameNumber(35)),  TOptional<FFrameTime>(FFrameTime(10)));
 
 		// Re-invert the inverse transform. This should be equivalent to the original transform and we shouldn't have lost anything.
-		FMovieSceneSequenceTransform InvInv = InvCompleteTransform.InverseNoLooping();
-		TestEqual("Doubly inverted zero-timescale transform should be equal to original", InvInv, CompleteTransform);
+		// @todo: is this necessary? We haven't needed an Inverse for an Inverse anywhere else in the codebase
+		//FMovieSceneSequenceTransform InvInv = InvCompleteTransform.InverseNoLooping();
+		//TestEqual("Doubly inverted zero-timescale transform should be equal to original", InvInv, CompleteTransform);
 	}
 
 	return true;
diff --git a/Engine/Source/Runtime/MovieScene/Public/Channels/MovieSceneChannelEditorData.h b/Engine/Source/Runtime/MovieScene/Public/Channels/MovieSceneChannelEditorData.h
--- a/Engine/Source/Runtime/MovieScene/Public/Channels/MovieSceneChannelEditorData.h
+++ b/Engine/Source/Runtime/MovieScene/Public/Channels/MovieSceneChannelEditorData.h
@@ -2,7 +2,6 @@
 
 #pragma once
 
-#include "Evaluation/MovieSceneEvaluationTemplateInstance.h"
 #include "HAL/Platform.h"
 #include "Internationalization/Text.h"
 #include "Math/Color.h"
@@ -10,12 +9,17 @@
 #include "MovieSceneCommonHelpers.h"
 #include "Templates/Function.h"
 #include "UObject/NameTypes.h"
+#include "Misc/FrameNumber.h"
+#include "Misc/Attribute.h"
 
 class FTrackInstancePropertyBindings;
 class UMovieSceneSection;
+class IMovieScenePlayer;
 class UObject;
+struct FMovieSceneSequenceID;
 struct FFrameNumber;
 struct FFrameRate;
+struct FGuid;
 struct FMovieSceneRootEvaluationTemplateInstance;
 
 #if WITH_EDITOR
@@ -57,10 +61,17 @@ struct FMovieSceneChannelMetaData
 	 */
 	MOVIESCENE_API FString GetPropertyMetaData(const FName& InKey) const;
 
+	/*
+	 * Get the amount that all of this channel's keys are offset by
+	 */
+	MOVIESCENE_API FFrameNumber GetOffsetTime(const UMovieSceneSection* InSection) const;
+
 	/** Whether this channel is enabled or not */
 	uint8 bEnabled : 1;
 	/** True if this channel can be collapsed onto the top level track node */
 	uint8 bCanCollapseToTrack : 1;
+	/** True if this channel's times are defined relative to the section start time */
+	uint8 bRelativeToSection : 1;
 	/** A sort order for this channel. Channels are sorted by this order, then by name. Groups are sorted by the channel with the lowest sort order. */
 	uint32 SortOrder;
 	/** This channel's unique name */
@@ -92,6 +103,10 @@ struct FMovieSceneChannelMetaData
 	TOptional<FLinearColor> Color;
 	/** Property meta data */
 	TMap<FName, FString> PropertyMetaData;
+	/** Optional object that owns this channel. By default the section owns the channel but sometimes channels live inside other nested objects instead. */
+	TWeakObjectPtr<> WeakOwningObject;
+	/** Key offset */
+	TAttribute<FFrameNumber> KeyOffset;
 };
 
 
diff --git a/Engine/Source/Runtime/MovieScene/Public/Channels/MovieSceneChannelTraits.h b/Engine/Source/Runtime/MovieScene/Public/Channels/MovieSceneChannelTraits.h
--- a/Engine/Source/Runtime/MovieScene/Public/Channels/MovieSceneChannelTraits.h
+++ b/Engine/Source/Runtime/MovieScene/Public/Channels/MovieSceneChannelTraits.h
@@ -8,6 +8,8 @@
 #include "Misc/FrameRate.h"
 #include "Misc/Optional.h"
 #include "MovieSceneChannelEditorData.h"
+#include "MovieSceneSection.h"
+
 
 enum class EMovieSceneKeyInterpolation : uint8;
 
diff --git a/Engine/Source/Runtime/MovieScene/Public/Channels/MovieSceneDoubleChannel.h b/Engine/Source/Runtime/MovieScene/Public/Channels/MovieSceneDoubleChannel.h
--- a/Engine/Source/Runtime/MovieScene/Public/Channels/MovieSceneDoubleChannel.h
+++ b/Engine/Source/Runtime/MovieScene/Public/Channels/MovieSceneDoubleChannel.h
@@ -2,7 +2,6 @@
 
 #pragma once
 
-#include "Channels/MovieSceneChannelEditorData.h"
 #include "Channels/MovieSceneCurveChannelCommon.h"
 #include "Containers/Array.h"
 #include "Containers/ArrayView.h"
@@ -26,6 +25,7 @@
 #include "Templates/UnrealTemplate.h"
 #include "UObject/Class.h"
 #include "UObject/ObjectMacros.h"
+#include "Misc/EnumClassFlags.h"
 
 #include "MovieSceneDoubleChannel.generated.h"
 
@@ -34,9 +34,15 @@ struct FPropertyTag;
 template<typename> struct TMovieSceneCurveChannelImpl;
 template <typename T> struct TIsPODType;
 
+namespace UE::MovieScene
+{
+	enum class EInverseEvaluateFlags : uint8;
+	struct FPiecewiseCurve;
+}
 namespace UE::MovieScene::Interpolation
 {
 	struct FCachedInterpolation;
+	struct FInterpolationExtents;
 }
 
 
@@ -130,6 +136,8 @@ struct FMovieSceneDoubleChannel : public FMovieSceneChannel
 #endif
 	{}
 
+	UE::MovieScene::FPiecewiseCurve AsPiecewiseCurve() const;
+
 	/**
 	 * Access a mutable interface for this channel's data
 	 *
@@ -186,6 +194,53 @@ struct FMovieSceneDoubleChannel : public FMovieSceneChannel
 	 */
 	MOVIESCENE_API UE::MovieScene::Interpolation::FCachedInterpolation GetInterpolationForTime(FFrameTime InTime) const;
 
+
+	/**
+	 * Compute the value extents of this curve within the specified limits
+	 * 
+	 * @param StartTime     The start of the range within which to compute extents
+	 * @param EndTime       The end of the range within which to compute extents
+	 * @return The extents of this channel
+	 */
+	MOVIESCENE_API UE::MovieScene::Interpolation::FInterpolationExtents ComputeExtents(FFrameTime StartTime, FFrameTime EndTime) const;
+
+
+	/**
+	 * Solve this curve for a given (y). Where more than one solution exists, TimeHint will be used to find the solution closest to the hint.
+	 * 
+	 * @param Value         The result to solve for
+	 * @param TimeHint      Predicate time to use for locating the most relevant solution
+	 * @param Flags         Flag structure used to control how to solve the curve
+	 * @return The solution if one exists, false otherwise
+	 */
+	MOVIESCENE_API TOptional<FFrameTime> InverseEvaluate(double Value, FFrameTime TimeHint, UE::MovieScene::EInverseEvaluateFlags Flags) const;
+
+
+	/**
+	 * Solve this curve for a given (y), only considering solutions that lie within a certain range.
+	 * 
+	 * @param Value         The result to solve for
+	 * @param StartTime     Start time before which solutions will not be considered
+	 * @param EndTime       End time after which solutions will not be considered
+	 * @param Visitor       Callback that is invoked for each solution. Returning true allows the algorithm to continue, false will terminate the algorithm.
+	 * @return False if any invocation of Visitor returned false, true otherwise.
+	 */
+	MOVIESCENE_API bool InverseEvaluateBetween(double Value, FFrameTime StartTime, FFrameTime EndTime, const TFunctionRef<bool(FFrameTime)>& Visitor) const;
+
+
+	/**
+	 * Retrieve the index of the cycle that the specified time falls within according to this channel's extrapolation modes.
+	 * @note Negative cycle counts are returned for cycles that lie before the first key
+	 */
+	MOVIESCENE_API int32 GetCycleCount(FFrameTime InTime) const;
+
+
+	/**
+	 * Retrieve the time range of a single cycle of this channel
+	 */
+	MOVIESCENE_API TRange<FFrameNumber> GetCycleRange(int32 InCycleCount) const;
+
+public:
 	/**
 	 * Set the channel's times and values to the requested values
 	 */
diff --git a/Engine/Source/Runtime/MovieScene/Public/Channels/MovieSceneInterpolation.h b/Engine/Source/Runtime/MovieScene/Public/Channels/MovieSceneInterpolation.h
--- a/Engine/Source/Runtime/MovieScene/Public/Channels/MovieSceneInterpolation.h
+++ b/Engine/Source/Runtime/MovieScene/Public/Channels/MovieSceneInterpolation.h
@@ -20,6 +20,95 @@
 namespace UE::MovieScene::Interpolation
 {
 
+struct FInvalidValue;
+struct FConstantValue;
+struct FLinearInterpolation;
+struct FQuadraticInterpolation;
+struct FCubicInterpolation;
+struct FQuarticInterpolation;
+struct FCubicBezierInterpolation;
+struct FWeightedCubicInterpolation;
+
+
+/**
+ * Template structure used for representing an array of solutions with a minimum size
+ * 
+ * For example:
+ * // An array with at least 3 elements, varifyable at compile-time
+ * int32 Solve3(TInterpSolutions<double, 3> Solutions)
+ * {
+ *     Solutions[0] = 0.0;
+ *     Solutions[1] = 1.0;
+ *     Solutions[2] = 2.0;
+ *     return 3;
+ * }
+ * int32 Solve4(TInterpSolutions<double, 4> Solutions);
+ * 
+ * double Solutions[8];
+ * Solve3(Solutions);
+ * Solve4(Solutions);
+ */
+template<typename T, int MinN>
+struct TInterpSolutions
+{
+	/**
+	 * Construction from a c-style array with a size of at least MinN
+	 */
+	template<int N>
+	TInterpSolutions(T (&In)[N]) requires (N >= MinN)
+		: Array(In)
+	{}
+
+	/**
+	 * Conversion to another array of size <= this
+	 */
+	template<int OtherN>
+	operator TInterpSolutions<T, OtherN>() const requires (OtherN <= MinN)
+	{
+		return TInterpSolutions<T, OtherN>(Array);
+	}
+
+	/**
+	 * Access the nth element of this array
+	 */
+	T& operator[](int Index) const
+	{
+		check(Index<MinN);
+		return Array[Index];
+	}
+
+private:
+
+	template<typename U, int OtherN>
+	friend struct TInterpSolutions;
+
+	explicit TInterpSolutions(T* InArray)
+		: Array(InArray)
+	{}
+
+	T* Array;
+};
+
+
+
+/**
+ * Structure that represents the extents of a curve in y
+ */
+struct FInterpolationExtents
+{
+	double MinValue = std::numeric_limits<double>::max();
+	double MaxValue = std::numeric_limits<double>::lowest();
+
+	FFrameTime MinValueTime;
+	FFrameTime MaxValueTime;
+
+	MOVIESCENE_API bool IsValid() const;
+
+	MOVIESCENE_API void AddPoint(double Value, FFrameTime Time);
+
+	MOVIESCENE_API void Combine(const FInterpolationExtents& Other);
+};
+
 /**
  * Sentinal type that represents an invalid interpolation value.
  * Only used when a curve has no data whatsoever, and therefore cannot be evaluated
@@ -38,9 +127,20 @@ struct FConstantValue
 	/** The constant value */
 	double Value;
 
+	FFrameNumber Origin;
+
+	UE_DEPRECATED(5.5, "Please provide an origin. This is required for Integral() to work correctly.")
 	FConstantValue(double InValue)
+		: FConstantValue(0, InValue)
+	{}
+
+	FConstantValue(FFrameNumber InOrigin, double InValue)
 		: Value(InValue)
+		, Origin(InOrigin)
 	{}
+
+	MOVIESCENE_API FConstantValue Derivative() const;
+	MOVIESCENE_API FLinearInterpolation Integral(double ConstantOffset = 0.0) const;
 };
 
 /**
@@ -50,7 +150,7 @@ struct FLinearInterpolation
 {
 	/** The coeffient 'a' in f(t) = a(t-o) + b */
 	double Coefficient;
-	/** The constant 'c' in f(t) = a(t-o) + b */
+	/** The constant 'b' in f(t) = a(t-o) + b */
 	double Constant;
 	/** The origin 'o' in f(t) = a(t-o) + b */
 	FFrameNumber Origin;
@@ -64,16 +164,168 @@ struct FLinearInterpolation
 	/**
 	 * Evaluate the expression
 	 */
-	double Evaluate(FFrameTime InTime) const
-	{
-		return Coefficient*(InTime - Origin).AsDecimal() + Constant;
-	}
+	MOVIESCENE_API double Evaluate(FFrameTime InTime) const;
+
+	/**
+	 * Attempt to solve this interpolation for x
+	 */
+	MOVIESCENE_API int32 Solve(double Value, TInterpSolutions<FFrameTime, 1> OutResults) const;
+
+	/**
+	 * Attempt to solve this interpolation for x within limits
+	 */
+	MOVIESCENE_API int32 SolveWithin(FFrameTime Start, FFrameTime End, double Value, TInterpSolutions<FFrameTime, 1> OutResults) const;
+
+	/**
+	 * Compute this expression's derivative
+	 */
+	MOVIESCENE_API FConstantValue Derivative() const;
+
+	/**
+	 * Compute this expression's integral with an optional constant offset
+	 */
+	MOVIESCENE_API FQuadraticInterpolation Integral(double ConstantOffset = 0.0) const;
 };
 
 /**
- * A cubic bezier interpolation between 2 control points with tangents, represented as 4 control points on a Bezier curve
+ * Structure representing a quadratic interpolation of the form f(x) = g(x-o) and g(x) = ax^2 + bx + c.
+ */
+struct FQuadraticInterpolation
+{
+	/** The coeffients a and b in g(x) = ax^2 + bx + c */
+	double A, B;
+	/** The constant 'c' in g(x) = ax^2 + bx + c */
+	double Constant;
+	/** The origin 'o' in f(x) = g(x-o) */
+	FFrameNumber Origin;
+
+	FQuadraticInterpolation(FFrameNumber InOrigin, double InA, double InB, double InConstant)
+		: A(InA), B(InB)
+		, Constant(InConstant)
+		, Origin(InOrigin)
+	{}
+
+	/**
+	 * Evaluate the expression
+	 */
+	MOVIESCENE_API double Evaluate(FFrameTime InTime) const;
+
+	/**
+	 * Attempt to solve this interpolation for x
+	 */
+	MOVIESCENE_API int32 Solve(double Value, TInterpSolutions<FFrameTime, 2> OutResults) const;
+
+	/**
+	 * Attempt to solve this interpolation for x within limits
+	 */
+	MOVIESCENE_API int32 SolveWithin(FFrameTime Start, FFrameTime End, double Value, TInterpSolutions<FFrameTime, 2> OutResults) const;
+
+	/**
+	 * Compute this expression's derivative
+	 */
+	MOVIESCENE_API FLinearInterpolation Derivative() const;
+
+	/**
+	 * Compute this expression's integral with an optional constant offset
+	 */
+	MOVIESCENE_API FCubicInterpolation Integral(double ConstantOffset = 0.0) const;
+};
+
+/**
+ * Structure representing a cubic interpolation of the form f(x) = g(x-o) and g(x) = ax^3 + bx^2 + cx + d.
  */
 struct FCubicInterpolation
+{
+	/** The coeffients a, b, and c in g(x) = ax^3 + bx^2 + cx + d */
+	double A, B, C;
+	/** The constant 'd' in g(x) = ax^3 + bx^2 + cx + d */
+	double Constant;
+
+	double DX;
+
+	/** The origin 'o' in f(x) = g(x-o) */
+	FFrameNumber Origin;
+
+	FCubicInterpolation(FFrameNumber InOrigin, double InA, double InB, double InC, double InConstant, double InDX = 1.0)
+		: A(InA), B(InB), C(InC)
+		, Constant(InConstant)
+		, DX(InDX)
+		, Origin(InOrigin)
+	{}
+
+	/**
+	 * Evaluate the expression
+	 */
+	MOVIESCENE_API double Evaluate(FFrameTime InTime) const;
+
+	/**
+	 * Attempt to solve this interpolation for x
+	 */
+	MOVIESCENE_API int32 Solve(double Value, TInterpSolutions<FFrameTime, 3> OutResults) const;
+
+	/**
+	 * Attempt to solve this interpolation for x within limits
+	 */
+	MOVIESCENE_API int32 SolveWithin(FFrameTime Start, FFrameTime End, double Value, TInterpSolutions<FFrameTime, 3> OutResults) const;
+
+	/**
+	 * Compute this expression's derivative
+	 */
+	MOVIESCENE_API FQuadraticInterpolation Derivative() const;
+
+	/**
+	 * Compute this expression's integral with an optional constant offset
+	 */
+	MOVIESCENE_API FQuarticInterpolation Integral(double ConstantOffset = 0.0) const;
+};
+
+/**
+ * Structure representing a quartic interpolation of the form f(x) = g(x-o) and g(x) = ax^4 + bx^3 + cx^2 + dx + e.
+ */
+struct FQuarticInterpolation
+{
+	/** The coeffients a, b, c, and d in g(x) = ax^4 + bx^3 + cx^2 + dx + e */
+	double A, B, C, D;
+	/** The constant 'e' in g(x) = ax^4 + bx^3 + cx^2 + dx + e */
+	double Constant;
+
+	double DX;
+
+	/** The origin 'o' in f(x) = g(x-o) */
+	FFrameNumber Origin;
+
+	FQuarticInterpolation(FFrameNumber InOrigin, double InA, double InB, double InC, double InD, double InConstant, double InDX = 1.0)
+		: A(InA), B(InB), C(InC), D(InD)
+		, Constant(InConstant)
+		, DX(InDX)
+		, Origin(InOrigin)
+	{}
+
+	/**
+	 * Evaluate the expression
+	 */
+	MOVIESCENE_API double Evaluate(FFrameTime InTime) const;
+
+	/**
+	 * Attempt to solve this interpolation for x
+	 */
+	MOVIESCENE_API int32 Solve(double Value, TInterpSolutions<FFrameTime, 4> OutResults) const;
+
+	/**
+	 * Attempt to solve this interpolation for x within limits
+	 */
+	MOVIESCENE_API int32 SolveWithin(FFrameTime Start, FFrameTime End, double Value, TInterpSolutions<FFrameTime, 4> OutResults) const;
+
+	/**
+	 * Compute this expression's derivative
+	 */
+	MOVIESCENE_API FCubicInterpolation Derivative() const;
+};
+
+/**
+ * A cubic bezier interpolation between 2 control points with tangents, represented as 4 control points on a Bezier curve
+ */
+struct FCubicBezierInterpolation
 {
 	/** The delta value between the two control points in the time-domain */
 	double DX;
@@ -82,7 +334,7 @@ struct FCubicInterpolation
 	/** The origin time of the first control point */
 	FFrameNumber Origin;
 
-	MOVIESCENE_API FCubicInterpolation(
+	MOVIESCENE_API FCubicBezierInterpolation(
 		FFrameNumber InOrigin,
 		double InDX,
 		double InStartValue,
@@ -90,10 +342,27 @@ struct FCubicInterpolation
 		double InStartTangent,
 		double InEndTangent);
 
+	MOVIESCENE_API FCubicInterpolation AsCubic() const;
+
 	/**
 	 * Evaluate the expression
 	 */
 	MOVIESCENE_API double Evaluate(FFrameTime InTime) const;
+
+	/**
+	 * Attempt to solve this interpolation for x
+	 */
+	MOVIESCENE_API int32 Solve(double InValue, TInterpSolutions<FFrameTime, 4> OutResults) const;
+
+	/**
+	 * Compute this expression's derivative
+	 */
+	MOVIESCENE_API FQuadraticInterpolation Derivative() const;
+
+	/**
+	 * Compute this expression's integral with an optional constant offset
+	 */
+	MOVIESCENE_API FQuarticInterpolation Integral(double ConstantOffset = 0.0) const;
 };
 
 /**
@@ -131,6 +400,12 @@ struct FWeightedCubicInterpolation
 		bool bEndIsWeighted);
 
 	double Evaluate(FFrameTime InTime) const;
+
+
+	/**
+	 * Attempt to solve this interpolation for x
+	 */
+	MOVIESCENE_API int32 Solve(double Value, TInterpSolutions<FFrameTime, 4> OutResults) const;
 };
 
 /**
@@ -152,8 +427,25 @@ struct FCachedInterpolationRange
 	/** Make a range that covers all times up to (but not including) the specified end */
 	static MOVIESCENE_API FCachedInterpolationRange Until(FFrameNumber InEnd);
 
+	FFrameNumber Clamp(FFrameNumber In) const
+	{
+		return FMath::Clamp(
+			In.Value,
+			Start.Value,
+			End.Value);
+	}
+	FFrameTime Clamp(FFrameTime In) const
+	{
+		return FMath::Clamp(
+			In,
+			FFrameTime(Start),
+			FFrameTime(End));
+	}
+
 	MOVIESCENE_API bool Contains(FFrameNumber FrameNumber) const;
 
+	MOVIESCENE_API bool IsEmpty() const;
+
 	/** Inclusive start frame */
 	FFrameNumber Start;
 	/** Exclusive end frame (unless End == Max()) */
@@ -172,17 +464,33 @@ struct FCachedInterpolation
 	MOVIESCENE_API FCachedInterpolation(const FCachedInterpolationRange& InRange, const FConstantValue& Constant);
 	/** Construction as a linear interpolation */
 	MOVIESCENE_API FCachedInterpolation(const FCachedInterpolationRange& InRange, const FLinearInterpolation& Linear);
+	/** Construction as a quadratic interpolation */
+	MOVIESCENE_API FCachedInterpolation(const FCachedInterpolationRange& InRange, const FQuadraticInterpolation& Quadratic);
 	/** Construction as a cubic interpolation */
 	MOVIESCENE_API FCachedInterpolation(const FCachedInterpolationRange& InRange, const FCubicInterpolation& Cubic);
+	/** Construction as a quartic interpolation */
+	MOVIESCENE_API FCachedInterpolation(const FCachedInterpolationRange& InRange, const FQuarticInterpolation& Quartic);
+	/** Construction as a cubic interpolation */
+	MOVIESCENE_API FCachedInterpolation(const FCachedInterpolationRange& InRange, const FCubicBezierInterpolation& CubicBezier);
 	/** Construction as a weighted cubic interpolation */
 	MOVIESCENE_API FCachedInterpolation(const FCachedInterpolationRange& InRange, const FWeightedCubicInterpolation& WeightedCubic);
 
+	/**
+	 * Check whether this cache is valid (ie, has a valid range and interpolation)
+	 */
+	MOVIESCENE_API bool IsValid() const;
+
 	/**
 	 * Check whether this cache is still valid for the specified frame number.
 	 * @return true if this interpolation is relevant to the frame, or false if it should be re-generated
 	 */
 	MOVIESCENE_API bool IsCacheValidForTime(FFrameNumber FrameNumber) const;
 
+	/**
+	 * Retrieve the range that this interpolation applies to
+	 */
+	MOVIESCENE_API FCachedInterpolationRange GetRange() const;
+
 	/**
 	 * Evaluate this interpolation for the specified frame time
 	 *
@@ -192,13 +500,27 @@ struct FCachedInterpolation
 	 */
 	MOVIESCENE_API bool Evaluate(FFrameTime FrameTime, double& OutResult) const;
 
+	MOVIESCENE_API TOptional<FCachedInterpolation> ComputeIntegral(double ConstantOffset = 0.0) const;
+
+	MOVIESCENE_API TOptional<FCachedInterpolation> ComputeDerivative() const;
+
+	MOVIESCENE_API int32 InverseEvaluate(double InValue, TInterpSolutions<FFrameTime, 4> OutResults) const;
+
+	MOVIESCENE_API FInterpolationExtents ComputeExtents(FFrameTime From, FFrameTime To) const;
+	MOVIESCENE_API FInterpolationExtents ComputeExtents() const;
+
+	MOVIESCENE_API void Offset(double Amount);
+
 private:
 
 	/** Variant containing the actual interpolation implementation */
 	TVariant<FInvalidValue,
 		FConstantValue,
 		FLinearInterpolation,
+		FQuadraticInterpolation,
 		FCubicInterpolation,
+		FQuarticInterpolation,
+		FCubicBezierInterpolation,
 		FWeightedCubicInterpolation> Data;
 
 	/** Structure representint the range of times this interpolation applies to */
