Sequencer: Added time-warp capabilities to sequences, sub-sections and skeletal animation sections

diff --git a/Engine/Source/Runtime/MovieScene/Public/Evaluation/MovieScenePlayback.h b/Engine/Source/Runtime/MovieScene/Public/Evaluation/MovieScenePlayback.h
--- a/Engine/Source/Runtime/MovieScene/Public/Evaluation/MovieScenePlayback.h
+++ b/Engine/Source/Runtime/MovieScene/Public/Evaluation/MovieScenePlayback.h
@@ -274,24 +274,12 @@ struct FMovieSceneContext : FMovieSceneEvaluationRange
 	 * Get the inverse transform of the current sub sequence, to transform local times back to root times.
 	 */
 	UE_DEPRECATED(5.4, "Please use GetSequenceToRootSequenceTransform instead.")
-	FORCEINLINE FMovieSceneTimeTransform GetSequenceToRootTransform() const
-	{
-		FMovieSceneSequenceTransform SequenceTransform = RootToSequenceTransform.InverseFromLoop(RootToSequenceWarpCounter);
-		FMovieSceneTimeTransform ReturnTransform = SequenceTransform.LinearTransform;
-		for (int i = 0; i < SequenceTransform.NestedTransforms.Num(); ++i)
-		{
-			ReturnTransform = ReturnTransform * SequenceTransform.NestedTransforms[i].LinearTransform;
-		}
-		return ReturnTransform;
-	}
+	MOVIESCENE_API FMovieSceneTimeTransform GetSequenceToRootTransform() const;
 
 	/**
 	 * Get the inverse sequence transform of the current sub sequence, to transform local times back to root times.
 	 */
-	FORCEINLINE FMovieSceneSequenceTransform GetSequenceToRootSequenceTransform() const
-	{
-		return RootToSequenceTransform.InverseFromLoop(RootToSequenceWarpCounter);
-	}
+	MOVIESCENE_API FMovieSceneInverseSequenceTransform GetSequenceToRootSequenceTransform() const;
 
 	/**
 	 * Apply section pre and post roll based on whether we're in the leading (preroll), or trailing (postroll) region for the section, and the current play direction
@@ -346,58 +334,7 @@ public:
 	/**
 	 * Transform this context to a different sub sequence space
 	 */
-	FMovieSceneContext Transform(const FMovieSceneSequenceTransform& InTransform, FFrameRate NewFrameRate) const
-	{
-		FMovieSceneContext NewContext = *this;
-		NewContext.RootToSequenceTransform = NewContext.RootToSequenceTransform * InTransform;
-		NewContext.CurrentFrameRate = NewFrameRate;
-
-		NewContext.EvaluationRange = InTransform.TransformRangeUnwarped(EvaluationRange);
-
-		// Transform the current time so we get an idea in what loop(s) we are relative to the root sequence.
-		FFrameTime TransformedTime_Unused;
-		FMovieSceneWarpCounter WarpCounter;
-		NewContext.RootToSequenceTransform.TransformTime(GetTime(), TransformedTime_Unused, WarpCounter);
-		NewContext.RootToSequenceWarpCounter = WarpCounter;
-
-		if (InTransform.IsLooping())
-		{
-			// If we have some looping, the transformed range might extend past the end of a loop and into
-			// the beginning of another. In that case, technically, the evaluation range ends up being a
-			// discontinuous range... for instance: [X, LoopEnd) + [LoopStart, Y)
-			// It can be even more than 2 ranges if the loop is short enough and the original evaluation
-			// range spans more than one loop!
-			// TODO: For now, we just take the "last part" of this discontinuous range but we should really 
-			//       actually make contexts have an array of evaluation ranges.
-			TRangeBound<FFrameTime> UpperEvalutionRangeBound = NewContext.EvaluationRange.GetUpperBound();
-			const FMovieSceneNestedSequenceTransform& LeafTransform = InTransform.NestedTransforms.Last();
-			if (UpperEvalutionRangeBound.IsClosed() && LeafTransform.IsLooping())
-			{
-				const FFrameNumber LeafWarpLength = LeafTransform.Warping.Length();
-				// Below: use strictly greater than comparison so that if the evalution range ends on the
-				//        end of a loop, we correctly evaluate up to there, instead of doing a 0-width
-				//        evaluation of the first frame of the loop.
-				while (UpperEvalutionRangeBound.GetValue() > LeafTransform.Warping.End)
-				{
-					UpperEvalutionRangeBound.SetValue(UpperEvalutionRangeBound.GetValue() - LeafWarpLength);
-				}
-
-				// We can end up with the upper bound being greater than the lower bound in the case outlined
-				// in the comment above (where we overlap a loop boundary). This is where we would split things
-				// up in multiple discontinous ranges, but as previously mentioned, we only take the last part
-				// for now.
-				TRangeBound<FFrameTime> LowerEvaluationRangeBound = NewContext.EvaluationRange.GetLowerBound();
-				if (LowerEvaluationRangeBound.IsClosed() && UpperEvalutionRangeBound.GetValue() < LowerEvaluationRangeBound.GetValue())
-				{
-					LowerEvaluationRangeBound.SetValue(LeafTransform.Warping.Start);
-				}
-
-				NewContext.EvaluationRange = TRange<FFrameTime>(LowerEvaluationRangeBound, UpperEvalutionRangeBound);
-			}
-		}
-
-		return NewContext;
-	}
+	MOVIESCENE_API FMovieSceneContext Transform(const FMovieSceneSequenceTransform& InTransform, FFrameRate NewFrameRate) const;
 
 	/**
 	 * Get the hierarchical bias for the current context
diff --git a/Engine/Source/Runtime/MovieScene/Public/Evaluation/MovieSceneSectionParameters.h b/Engine/Source/Runtime/MovieScene/Public/Evaluation/MovieSceneSectionParameters.h
--- a/Engine/Source/Runtime/MovieScene/Public/Evaluation/MovieSceneSectionParameters.h
+++ b/Engine/Source/Runtime/MovieScene/Public/Evaluation/MovieSceneSectionParameters.h
@@ -5,6 +5,7 @@
 #include "CoreMinimal.h"
 #include "UObject/ObjectMacros.h"
 #include "Misc/EnumClassFlags.h"
+#include "Variants/MovieSceneTimeWarpVariant.h"
 #include "MovieSceneSectionParameters.generated.h"
 
 /**
@@ -88,7 +89,7 @@ public:
 
 	/** Playback time scaling factor. */
 	UPROPERTY(BlueprintReadWrite, EditAnywhere, Category="Timing")
-	float TimeScale;
+	FMovieSceneTimeWarpVariant TimeScale;
 
 	/** Hierachical bias. Higher bias will take precedence. */
 	UPROPERTY(config, BlueprintReadWrite, EditAnywhere, Category="Sequence")
diff --git a/Engine/Source/Runtime/MovieScene/Public/Evaluation/MovieSceneSequenceHierarchy.h b/Engine/Source/Runtime/MovieScene/Public/Evaluation/MovieSceneSequenceHierarchy.h
--- a/Engine/Source/Runtime/MovieScene/Public/Evaluation/MovieSceneSequenceHierarchy.h
+++ b/Engine/Source/Runtime/MovieScene/Public/Evaluation/MovieSceneSequenceHierarchy.h
@@ -52,6 +52,13 @@ struct FMovieSceneSubSequenceData
 	 */
 	MOVIESCENE_API FMovieSceneSubSequenceData(const UMovieSceneSubSection& InSubSection);
 
+	/*~ Defined constructors, operators and destructor to work around deprecation warnings on linux */
+	MOVIESCENE_API FMovieSceneSubSequenceData(const FMovieSceneSubSequenceData&);
+	MOVIESCENE_API FMovieSceneSubSequenceData& operator=(const FMovieSceneSubSequenceData&);
+	MOVIESCENE_API FMovieSceneSubSequenceData(FMovieSceneSubSequenceData&&);
+	MOVIESCENE_API FMovieSceneSubSequenceData& operator=(FMovieSceneSubSequenceData&&);
+	MOVIESCENE_API ~FMovieSceneSubSequenceData();
+
 	/**
 	 * Get this sub sequence's sequence asset, potentially loading it through its soft object path
 	 */
@@ -75,6 +82,7 @@ struct FMovieSceneSubSequenceData
 	/**
 	 * Re-creates a sub-section parameter struct.
 	 */
+	UE_DEPRECATED(5.5, "This function is no longer supported")
 	MOVIESCENE_API FMovieSceneSectionParameters ToSubSectionParameters() const;
 
 	/** The sequence that the sub section references */
@@ -89,50 +97,31 @@ struct FMovieSceneSubSequenceData
 	UPROPERTY()
 	FMovieSceneSequenceTransform RootToSequenceTransform;
 
-	/** The tick resolution of the inner sequence. */
-	UPROPERTY()
-	FFrameRate TickResolution;
-
-	/** This sequence's deterministic sequence ID. Used in editor to reduce the risk of collisions on recompilation. */ 
-	UPROPERTY()
-	FMovieSceneSequenceID DeterministicSequenceID;
+#if WITH_EDITORONLY_DATA
 
-	/** The play range of the parent section */
 	UPROPERTY()
-	FMovieSceneFrameRange ParentPlayRange;
+	FMovieSceneTransformBreadcrumbs StartTimeBreadcrumbs;
 
-	/** The start frame offset of the parent section */
 	UPROPERTY()
-	FFrameNumber ParentStartFrameOffset;
+	FMovieSceneTransformBreadcrumbs EndTimeBreadcrumbs;
 
-	/** The end frame offset of the parent section */
-	UPROPERTY()
-	FFrameNumber ParentEndFrameOffset;
+#endif // WITH_EDITORONLY_DATA
 
-	/** The offset for the first loop of the sub-sequence */
+	/** The tick resolution of the inner sequence. */
 	UPROPERTY()
-	FFrameNumber ParentFirstLoopStartFrameOffset;
+	FFrameRate TickResolution;
 
-	/** Whether this sub-sequence can loop */
+	/** This sequence's deterministic sequence ID. Used in editor to reduce the risk of collisions on recompilation. */ 
 	UPROPERTY()
-	bool bCanLoop = false;
+	FMovieSceneSequenceID DeterministicSequenceID;
 	
 	/** This sub sequence's playback range according to its parent sub section. Clamped recursively during template generation */
 	UPROPERTY()
 	FMovieSceneFrameRange PlayRange;
 
-	/** The sub-sequence's full playback range, in its own local time space. */
-	UPROPERTY()
-	FMovieSceneFrameRange FullPlayRange;
-
-	/**
-	 * The play range of the parent section, without any warping involved.
-	 * That means that, for a sub-sequence playing with an initial offset of 50 and looping 3 times,
-	 * this play range will start 50 frames after PlayRange's lower bound, and extend much past PlayRange's 
-	 * upper bound (3 times longer).
-	 */
+	/** The range of the sub-sequence in its parent space. Unclamped. */
 	UPROPERTY()
-	FMovieSceneFrameRange UnwarpedPlayRange;
+	FMovieSceneFrameRange ParentPlayRange;
 
 	/** The sequence preroll range considering the start offset */
 	UPROPERTY()
@@ -150,6 +139,10 @@ struct FMovieSceneSubSequenceData
 	UPROPERTY()
 	EMovieSceneSubSectionFlags AccumulatedFlags;
 
+	/** Whether this sub-sequence can loop */
+	UPROPERTY()
+	bool bCanLoop = false;
+
 	/** Instance data that should be used for any tracks contained immediately within this sub sequence */
 	UPROPERTY()
 	FMovieSceneSequenceInstanceDataPtr InstanceData;
@@ -160,6 +153,18 @@ struct FMovieSceneSubSequenceData
 	UPROPERTY()
 	FName SectionPath;
 
+	/** The sub-sequence's full playback range, in its own local time space. */
+	UE_DEPRECATED(5.5, "Please use GetSequence()->GetMovieScene()->GetPlaybackRange().")
+	FMovieSceneFrameRange FullPlayRange;
+	UE_DEPRECATED(5.5, "This field is no longer supported.")
+	FFrameNumber ParentStartFrameOffset;
+	UE_DEPRECATED(5.5, "This field is no longer supported.")
+	FFrameNumber ParentEndFrameOffset;
+	UE_DEPRECATED(5.5, "This field is no longer supported.")
+	FFrameNumber ParentFirstLoopStartFrameOffset;
+	UE_DEPRECATED(5.5, "This field is no longer used")
+	FMovieSceneFrameRange UnwarpedPlayRange;
+
 #endif
 
 private:
@@ -207,11 +212,28 @@ struct FMovieSceneSubSequenceTreeEntry
 {
 	GENERATED_BODY()
 
+	FMovieSceneSubSequenceTreeEntry(FMovieSceneSequenceID InSequenceID, ESectionEvaluationFlags InFlags)
+		: SequenceID(InSequenceID)
+		, Flags(InFlags)
+	{}
+
+	MOVIESCENE_API FMovieSceneSubSequenceTreeEntry();
+
+	MOVIESCENE_API FMovieSceneSubSequenceTreeEntry(const FMovieSceneSubSequenceTreeEntry&);
+	MOVIESCENE_API FMovieSceneSubSequenceTreeEntry& operator=(const FMovieSceneSubSequenceTreeEntry&);
+
+	MOVIESCENE_API FMovieSceneSubSequenceTreeEntry(FMovieSceneSubSequenceTreeEntry&&);
+	MOVIESCENE_API FMovieSceneSubSequenceTreeEntry& operator=(FMovieSceneSubSequenceTreeEntry&&);
+
+	MOVIESCENE_API ~FMovieSceneSubSequenceTreeEntry();
+
 	friend FArchive& operator<<(FArchive& Ar, FMovieSceneSubSequenceTreeEntry& InOutEntry);
 	friend bool operator==(const FMovieSceneSubSequenceTreeEntry& A, const FMovieSceneSubSequenceTreeEntry& B);
 
 	FMovieSceneSequenceID SequenceID;
 	ESectionEvaluationFlags Flags;
+
+	UE_DEPRECATED(5.5, "This member is no longer supported")
 	FMovieSceneWarpCounter RootToSequenceWarpCounter;
 };
 
@@ -319,8 +341,24 @@ struct FMovieSceneSequenceHierarchy
 	/**
 	 * Add an entry for the given sub sequence with the given root time range
 	 */
+	void AddRange(const TRange<FFrameNumber>& RootSpaceRange, FMovieSceneSequenceIDRef InSequenceID, ESectionEvaluationFlags InFlags);
+
+	UE_DEPRECATED(5.5, "The RootToSequenceWarpCounter parameter is no longer required or supported")
 	void AddRange(const TRange<FFrameNumber>& RootSpaceRange, FMovieSceneSequenceIDRef InSequenceID, ESectionEvaluationFlags InFlags, FMovieSceneWarpCounter RootToSequenceWarpCounter);
-	
+
+	/**
+	 * Mutate the sub-data for the specified sequence ID - only to be used during compilation
+	 */
+	FMovieSceneSubSequenceData* MutateSubData(FMovieSceneSequenceIDRef SequenceID)
+	{
+		return SequenceID == MovieSceneSequenceID::Root ? nullptr : SubSequences.Find(SequenceID);
+	}
+
+	void SetRootTransform(FMovieSceneSequenceTransform&& InTransform)
+	{
+		RootTransform = MoveTemp(InTransform);
+	}
+
 	/** Get all sub-sequence IDs */
 	void AllSubSequenceIDs(TArray<FMovieSceneSequenceID>& OutSequenceIDs) const
 	{
@@ -354,6 +392,11 @@ struct FMovieSceneSequenceHierarchy
 		AccumulatedNetworkMask &= Mask;
 	}
 
+	const FMovieSceneSequenceTransform& GetRootTransform() const
+	{
+		return RootTransform;
+	}
+
 #if !NO_LOGGING
 	void LogHierarchy() const;
 	void LogSubSequenceTree() const;
@@ -368,6 +411,9 @@ private:
 	UPROPERTY()
 	FMovieSceneSubSequenceTree Tree;
 
+	UPROPERTY()
+	FMovieSceneSequenceTransform RootTransform;
+
 	/** Map of all (recursive) sub sequences found in this template, keyed on sequence ID */
 	UPROPERTY()
 	TMap<FMovieSceneSequenceID, FMovieSceneSubSequenceData> SubSequences;
diff --git a/Engine/Source/Runtime/MovieScene/Public/Evaluation/MovieSceneSequenceTransform.h b/Engine/Source/Runtime/MovieScene/Public/Evaluation/MovieSceneSequenceTransform.h
--- a/Engine/Source/Runtime/MovieScene/Public/Evaluation/MovieSceneSequenceTransform.h
+++ b/Engine/Source/Runtime/MovieScene/Public/Evaluation/MovieSceneSequenceTransform.h
@@ -2,205 +2,607 @@
 
 #pragma once
 
-#include "Algo/AllOf.h"
-#include "Algo/AnyOf.h"
+#include "Containers/ContainersFwd.h"
 #include "Containers/Array.h"
-#include "Containers/ArrayView.h"
-#include "Containers/UnrealString.h"
-#include "CoreMinimal.h"
-#include "HAL/PlatformCrt.h"
-#include "HAL/PlatformMath.h"
-#include "Math/Range.h"
-#include "Math/RangeBound.h"
+#include "MovieSceneFwd.h"
 #include "Misc/AssertionMacros.h"
 #include "Misc/FrameNumber.h"
 #include "Misc/FrameTime.h"
 #include "MovieSceneTimeTransform.h"
 #include "MovieSceneTimeWarping.h"
-#include "Templates/UnrealTemplate.h"
-#include "UObject/ObjectMacros.h"
+#include "Variants/MovieSceneTimeWarpVariant.h"
 
 #include "MovieSceneSequenceTransform.generated.h"
 
+template<typename> class TRange;
 
-/**
- * Converts a range from one type of bounds to another. The output bounds type must be implicitely 
- * constructable from the input bounds type.
- */
-template<typename InBoundType, typename OutBoundType>
-inline TRange<OutBoundType> ConvertRange(const TRange<InBoundType>& Range)
+struct FMovieSceneTransformBreadcrumbs;
+struct FMovieSceneNestedSequenceTransform;
+struct FMovieSceneInverseNestedSequenceTransform;
+
+namespace UE::MovieScene
 {
-	const TRangeBound<InBoundType> SourceLower = Range.GetLowerBound();
-	TRangeBound<OutBoundType> DestLower = SourceLower.IsOpen() ?
-		TRangeBound<OutBoundType>() :
-		SourceLower.IsInclusive() ?
-			TRangeBound<OutBoundType>::Inclusive(SourceLower.GetValue()) :
-			TRangeBound<OutBoundType>::Exclusive(SourceLower.GetValue());
-
-	const TRangeBound<InBoundType> SourceUpper = Range.GetUpperBound();
-	TRangeBound<OutBoundType> DestUpper = SourceUpper.IsOpen() ?
-		TRangeBound<OutBoundType>() :
-		SourceUpper.IsInclusive() ?
-			TRangeBound<OutBoundType>::Inclusive(SourceUpper.GetValue()) :
-			TRangeBound<OutBoundType>::Exclusive(SourceUpper.GetValue());
-
-	return TRange<OutBoundType>(DestLower, DestUpper);
+	struct FInverseTransformTimeParams;
+	struct FTransformTimeParams;
 }
 
-// Specialization of ConvertRange for round down FFrameTime to FFrameNumber.
-template<>
-inline TRange<FFrameNumber> ConvertRange(const TRange<FFrameTime>& Range)
+
+/*~
+ * This file contains the structures necessary for representing transformations of time between various different
+ *      'spaces' within the Sequencer/MovieScene codebase. Typically these spaces denote root and sub-sequences,
+ *      or sub-sub-suquences. 'Lossy' transformations (ie, transformations that are mathematically non-linear)
+ *      are supported and typically referred to as 'warping' or 'non-linear' transforms. Such transforms are always
+ *      stored as a stack of nested FMovieSceneNestedSequenceTransform structures and cannot be combined mathematically
+ *      like FMovieSceneTimeTransform can.
+ *
+ * When dealing with time transformations, it is common to use the following nomenclature interchangably:
+ *      - Outer, Untransformed, Unwarped space; and,
+ *      - Inner, Transformed, Warped space.
+ * 
+ * Inverse transformations are also supported, but via a different, more restrictive API. Where outer->inner transforms
+ *      must always map an outer time to an inner time, inverse transformations can fail. This is clearly the case when 
+ *      attempting to inverse-transform a time of 20 from a loop that only spans [0, 10). The requested time is not present
+ *      in the set of times that can result from the outer->inner transform, and so the operation fails.
+ *      Similarly, inverse transformations may have zero or more real solutions. Consider the previous example, but with a
+ *      time of 5; the operation has an infinite number of solutions (one for each loop, or every t where t%5==0).
+ *
+ *                         ________ << FMovieSceneInverseSequenceTransform << ___________
+ *                        /                                                              \
+ * ----------------------                                                                -----------------------
+ * | Untransformed space |                                                                |  Transformed space  |
+ * -----------------------                                                                -----------------------
+ *                        \_____________ >> FMovieSceneSequenceTransform >> _____________/
+ */
+
+
+/**
+ * Enumeration defining how to store breadcrumb trails
+ */
+UENUM()
+enum class EMovieSceneBreadcrumbMode : uint8
 {
-	const TRangeBound<FFrameTime> SourceLower = Range.GetLowerBound();
-	TRangeBound<FFrameNumber> DestLower = SourceLower.IsOpen() ?
-		TRangeBound<FFrameNumber>() :
-		SourceLower.IsInclusive() ?
-			TRangeBound<FFrameNumber>::Inclusive(SourceLower.GetValue().FloorToFrame()) :
-			TRangeBound<FFrameNumber>::Exclusive(SourceLower.GetValue().FloorToFrame());
-
-	const TRangeBound<FFrameTime> SourceUpper = Range.GetUpperBound();
-	TRangeBound<FFrameNumber> DestUpper = SourceUpper.IsOpen() ?
-		TRangeBound<FFrameNumber>() :
-		SourceUpper.IsInclusive() ?
-			TRangeBound<FFrameNumber>::Inclusive(SourceUpper.GetValue().FloorToFrame()) :
-			TRangeBound<FFrameNumber>::Exclusive(SourceUpper.GetValue().FloorToFrame());
-
-	return TRange<FFrameNumber>(DestLower, DestUpper);
-}
+	/** Default: Only store breadcrumbs for non-linear transformations */
+	Sparse,
+	/** Store breadcrumbs for every nested time transformation */
+	Dense
+};
 
 /**
- * Struct that tracks warp counts in a way that works with the FMovieSceneTimeWarping struct.
+ * Struct that tracks a breadcumb trail when transformiung a time through FMovieSceneSequenceTransform
  */
 USTRUCT()
-struct FMovieSceneWarpCounter
+struct FMovieSceneTransformBreadcrumbs
 {
 	GENERATED_BODY()
 
-	void AddWarpingLevel(uint32 WarpCount)
+	/**
+	 * Default constructor, optionally taking a mode
+	 */
+	FMovieSceneTransformBreadcrumbs(EMovieSceneBreadcrumbMode InMode = EMovieSceneBreadcrumbMode::Sparse)
+		: Mode(InMode)
+	{}
+
+
+	/**
+	 * Return the breadcrumb at the specified index
+	 */
+	FFrameTime operator[](int32 Index) const
 	{
-		if (WarpCount != FMovieSceneTimeWarping::InvalidWarpCount)
-		{
-			WarpCounts.Add(WarpCount);
-		}
-		else
+		return Breadcrumbs[Index];
+	}
+
+
+	/**
+	 * Retrieve the length of this breadcrumb trail
+	 */
+	int32 Num() const
+	{
+		return Breadcrumbs.Num();
+	}
+
+
+	/**
+	 * Check if the specified index is valid
+	 */
+	bool IsValidIndex(int32 Index) const
+	{
+		return Breadcrumbs.IsValidIndex(Index);
+	}
+
+
+	/**
+	 * Check whether this breadcrumb trail only contains breadcrumbs for non-linear transformations (true) or for everything (false)
+	 */
+	bool IsSparse() const
+	{
+		return Mode == EMovieSceneBreadcrumbMode::Sparse;
+	}
+
+
+	/**
+	 * Retrieve this breadcrumb trail's capture mode
+	 */
+	EMovieSceneBreadcrumbMode GetMode() const
+	{
+		return Mode;
+	}
+
+
+	/**
+	 * Restore this trail to its original (empty) state
+	 */
+	void Reset()
+	{
+		Breadcrumbs.Reset();
+	}
+
+
+	/**
+	 * Add a breadcrumb to this trail
+	 */
+	void AddBreadcrumb(const FFrameTime& Breadcrumb)
+	{
+		Breadcrumbs.Add(Breadcrumb);
+	}
+
+
+	/**
+	 * Pop the most recently added breadcrumb off this trail
+	 */
+	FFrameTime PopBreadcrumb()
+	{
+		return Breadcrumbs.Pop();
+	}
+
+
+	/**
+	 * Prepend the specified breadcrumb trail to this one, resulting in a path from the start of Outer to the end of this
+	 */
+	void CombineWithOuterBreadcrumbs(const FMovieSceneTransformBreadcrumbs& OuterBreadcrumbs)
+	{
+		// Append the outer breadcrumbs to the head of this list
+		Breadcrumbs.Insert(OuterBreadcrumbs.Breadcrumbs, 0);
+	}
+
+
+	/**
+	 * Create a new breadcrumb trail of the same size as this one, but with all times set to a specific time
+	 */
+	FMovieSceneTransformBreadcrumbs OverwriteWith(FFrameTime InTime) const
+	{
+		FMovieSceneTransformBreadcrumbs Result;
+		Result.Breadcrumbs.SetNum(Breadcrumbs.Num());
+		for (FFrameTime& X : Result.Breadcrumbs)
 		{
-			WarpCounts.Add(FMovieSceneTimeWarping::InvalidWarpCount);
+			X = InTime;
 		}
+		return Result;
 	}
 
+public:
+
+	TArray<FFrameTime>::RangedForConstIteratorType begin () const { return Breadcrumbs.begin(); }
+	TArray<FFrameTime>::RangedForConstIteratorType end   () const { return Breadcrumbs.end();   }
+
+	MOVIESCENE_API operator TArrayView<const FFrameTime>() const;
+
+private:
+
+	UPROPERTY()
+	TArray<FFrameTime> Breadcrumbs;
+
+	UPROPERTY()
+	EMovieSceneBreadcrumbMode Mode;
+};
+
+
+USTRUCT()
+struct FMovieSceneWarpCounter : public FMovieSceneTransformBreadcrumbs
+{
+	GENERATED_BODY()
+
+public:
+
+	MOVIESCENE_API FMovieSceneWarpCounter();
+
+	MOVIESCENE_API FMovieSceneWarpCounter(const FMovieSceneWarpCounter&);
+	MOVIESCENE_API FMovieSceneWarpCounter& operator=(const FMovieSceneWarpCounter&);
+
+	MOVIESCENE_API FMovieSceneWarpCounter(FMovieSceneWarpCounter&&);
+	MOVIESCENE_API FMovieSceneWarpCounter& operator=(FMovieSceneWarpCounter&&);
+
+	MOVIESCENE_API ~FMovieSceneWarpCounter();
+
+	UE_DEPRECATED(5.5, "This function is no longer used. Please update your code to use time-based breadcrumbs instead.")
+	void AddWarpingLevel(uint32 WarpCount)
+	{
+	}
+
+	UE_DEPRECATED(5.5, "This function is no longer used. Please update your code to use time-based breadcrumbs instead.")
 	void AddNonWarpingLevel()
 	{
-		AddWarpingLevel(FMovieSceneTimeWarping::InvalidWarpCount);
 	}
 
+	UE_DEPRECATED(5.5, "This function is no longer used. Please update your code to use time-based breadcrumbs instead.")
 	int32 NumWarpCounts() const
 	{
-		return WarpCounts.Num();
+		return 0;
 	}
 
+	UE_DEPRECATED(5.5, "This function is no longer used. Please update your code to use time-based breadcrumbs instead.")
 	uint32 LastWarpCount() const
 	{
-		return WarpCounts.Num() > 0 ? WarpCounts[WarpCounts.Num() - 1] : FMovieSceneTimeWarping::InvalidWarpCount;
+		return (uint32)(-1);
+	}
+
+	UE_DEPRECATED(5.5, "Warp counts are no longer supported.")
+	TArray<uint32> WarpCounts;
+};
+
+
+/**
+ * Structure used to represent a specific inverse transformation (ie from transformed to untransformed space) that cannot be combined with another.
+ * Stored as a stack inside FMovieSceneInverseSequenceTransform to represent a complete transformation
+ * from inner time-space to outer time-space.
+ */
+USTRUCT()
+struct FMovieSceneInverseNestedSequenceTransform
+{
+	GENERATED_BODY()
+
+	/**
+	 * Default construction to an identity linear transform. Should only be used by serialization.
+	 */
+	FMovieSceneInverseNestedSequenceTransform()
+	{}
+
+
+	/**
+	 * Construction from a linear transform
+	 */
+	FMovieSceneInverseNestedSequenceTransform(const FMovieSceneTimeTransform& InLinearTransform)
+		: TimeScale(InLinearTransform.TimeScale), Offset(InLinearTransform.Offset)
+	{}
+
+
+
+	/**
+	 * Construction from a linear transform represented as an offset and scale (scale is applied first)
+	 */
+	FMovieSceneInverseNestedSequenceTransform(const FFrameTime& InOffset, double InTimeScale)
+		: TimeScale(InTimeScale), Offset(InOffset)
+	{}
+
+
+public:
+
+
+	/**
+	 * Check whether this transform is linear (true) or not (false)
+	 */
+	bool IsLinear() const
+	{
+		return TimeScale.GetType() == EMovieSceneTimeWarpType::FixedPlayRate;
 	}
 
-	friend bool operator==(const FMovieSceneWarpCounter& A, const FMovieSceneWarpCounter& B)
+
+	/**
+	 * Check whether this transformation requires a breadcrumb trail (true) or not (false)
+	 * @note Breadcrumbs may still be added for dense trails
+	 */
+	bool NeedsBreadcrumb() const
 	{
-		return A.WarpCounts == B.WarpCounts;
+		return !IsLinear();
 	}
 
+
+	/**
+	 * Convert this transform to its linear form.
+	 * @note This is only valid to call where IsLinear() is true
+	 */
+	MOVIESCENE_API FMovieSceneTimeTransform AsLinear() const;
+
+
+	/**
+	 * Convert this transform to its inverse (ie a transform that converts from untransformed space to transformed space)
+	 */
+	MOVIESCENE_API FMovieSceneNestedSequenceTransform Inverse() const;
+
+
+	/**
+	 * Attempt to transform the specified time by this inverse transform. The operation can fail if the time does not map to the un-transformed time space (for example, when looping).
+	 * 
+	 * @param Time        The time to transform
+	 * @param Breadcrumb  A breadcrumb in the un-transformed space to assist in resolving ambiguous transformations (ie, to transform into the correct time based on the breadcrumb's loop)
+	 *
+	 * @return The transformed time if successful, or an epty optional otherwise
+	 */
+	MOVIESCENE_API TOptional<FFrameTime> TryTransformTime(FFrameTime Time, FFrameTime Breadcrumb) const;
+
+
+	/**
+	 * Attempt to transform the specified time by this inverse transform. The operation can fail if the time does not map to the un-transformed time space (for example, when looping).
+	 * 
+	 * @param Time        The time to transform
+	 * @param Breadcrumb  A breadcrumb in the un-transformed space to assist in resolving ambiguous transformations (ie, to transform into the correct time based on the breadcrumb's loop)
+	 * @param Params      Parameter structure that controls how to perform the inverse operation (ie to disallow cycling outside of the current loop; to ignore clamps etc)
+	 * 
+	 * @return The transformed time if successful, or an epty optional otherwise
+	 */
+	MOVIESCENE_API TOptional<FFrameTime> TryTransformTime(FFrameTime Time, FFrameTime Breadcrumb, const UE::MovieScene::FInverseTransformTimeParams& Params) const;
+
+
+	/**
+	 * Attempt to transform the specified time by this inverse transform within a specified un-transformed range, calling a functor for every instance of the specified time in the un-transformed space.
+	 * 
+	 * @param Time        The time to transform
+	 * @param Visitor     A callback to invoke for every instance of Time in the untransformed range. The result of the functor determines whether to continue iteration (true) or not (false).
+	 * @param UntransformedRangeStart      The inclusive start time in untransformed space within which to allow Visitor to be invoked
+	 * @param UntransformedRangeEnd        The inclusive end time in untransformed space within which to allow Visitor to be invoked
+	 * 
+	 * @return True if every invocation of Visitor returned true, or there were no invocations; false if any invocation of Visitor returned false.
+	 */
+	MOVIESCENE_API bool TransformTimeWithinRange(FFrameTime Time, const TFunctionRef<bool(FFrameTime)>& Visitor, FFrameTime UntransformedRangeStart, FFrameTime UntransformedRangeEnd) const;
+
+private:
+
+	friend FMovieSceneNestedSequenceTransform;
+
+	/** Time scale as either a fixed play rate, or as an external implementation */
 	UPROPERTY()
-	TArray<uint32> WarpCounts;
+	FMovieSceneTimeWarpVariant TimeScale;
+
+	/** Constant time offset. Offset is applied differently for different internal implementations of TimeScale. */
+	UPROPERTY()
+	FFrameTime Offset;
 };
 
+
+
 /**
- * Time transform information for a nested sequence.
+ * Structure used to represent a specific transformation (ie from untransformed to transformed space) that cannot be combined with another.
+ *
+ * Stored as a stack inside FMovieSceneSequenceTransform to represent a complete transformation from inner time-space to outer time-space.
  */
 USTRUCT()
 struct FMovieSceneNestedSequenceTransform
 {
 	GENERATED_BODY()
 
-	FMovieSceneNestedSequenceTransform()
+
+	/**
+	 * Default construction to an identity linear transform
+	 */
+	MOVIESCENE_API FMovieSceneNestedSequenceTransform();
+
+	MOVIESCENE_API FMovieSceneNestedSequenceTransform(const FMovieSceneNestedSequenceTransform&);
+	MOVIESCENE_API FMovieSceneNestedSequenceTransform& operator=(const FMovieSceneNestedSequenceTransform&);
+
+	MOVIESCENE_API FMovieSceneNestedSequenceTransform(FMovieSceneNestedSequenceTransform&&);
+	MOVIESCENE_API FMovieSceneNestedSequenceTransform& operator=(FMovieSceneNestedSequenceTransform&&);
+
+	MOVIESCENE_API ~FMovieSceneNestedSequenceTransform();
+
+	/**
+	 * Construction from a time warp variant
+	 */
+	FMovieSceneNestedSequenceTransform(FMovieSceneTimeWarpVariant&& InVariant)
+		: TimeScale(MoveTemp(InVariant))
 	{}
 
-	FMovieSceneNestedSequenceTransform(FMovieSceneTimeTransform InLinearTransform)
-		: LinearTransform(InLinearTransform)
+
+	/**
+	 * Construction from a linear time transformation
+	 */
+	FMovieSceneNestedSequenceTransform(const FMovieSceneTimeTransform& InLinearTransform)
+		: TimeScale(InLinearTransform.TimeScale), Offset(InLinearTransform.Offset)
 	{}
 
-	FMovieSceneNestedSequenceTransform(FMovieSceneTimeWarping InWarping)
-		: Warping(InWarping)
+
+	/**
+	 * Construction from a linear time transformation (scale applies first)
+	 */
+	FMovieSceneNestedSequenceTransform(FFrameTime InOffset, double InTimeScale)
+		: TimeScale(InTimeScale), Offset(InOffset)
 	{}
 
-	FMovieSceneNestedSequenceTransform(FMovieSceneTimeTransform InLinearTransform, FMovieSceneTimeWarping InWarping)
-		: LinearTransform(InLinearTransform)
-		, Warping(InWarping)
+
+	/**
+	 * Default construction to an identity linear transform
+	 */
+	FMovieSceneNestedSequenceTransform(FFrameTime InOffset, FMovieSceneTimeWarpVariant&& InTimeScale)
+		: TimeScale(MoveTemp(InTimeScale)), Offset(InOffset)
 	{}
 
+
+public:
+
+
 	/**
-	 * Returns whether this transform is identity.
+	 * Check whether this transform is linear (true) or not (false)
 	 */
-	bool IsIdentity() const { return LinearTransform.IsIdentity() && !IsWarping(); }
+	bool IsLinear() const
+	{
+		return TimeScale.GetType() == EMovieSceneTimeWarpType::FixedPlayRate;
+	}
+
 
 	/**
-	 * Returns whether this transform is warping. This includes both looping as well as zero-timescale.
+	 * Convert this transform to its linear form.
+	 * @note This is only valid to call where IsLinear() is true
 	 */
-	bool IsWarping() const { return Warping.IsValid() || FMath::IsNearlyZero(LinearTransform.TimeScale); }
+	FMovieSceneTimeTransform AsLinear() const
+	{
+		return FMovieSceneTimeTransform(Offset, TimeScale.AsFixedPlayRateFloat());
+	}
+
 
 	/**
-	 * Returns whether this transform is warping specifically from looping. This doesn't include zero-timescale cases.
+	 * Check whether this transformation requires a breadcrumb trail (true) or not (false)
+	 * @note Breadcrumbs may still be added for dense trails
 	 */
-	bool IsLooping() const { return Warping.IsValid(); }
+	bool NeedsBreadcrumb() const
+	{
+		return !IsLinear();
+	}
+
 
 	/**
-	 * Linear time transform for this sub-sequence.
+	 * Returns whether this transform is an identity transformation (ie, a*T = a)
 	 */
-	UPROPERTY()
-	FMovieSceneTimeTransform LinearTransform;
+	bool IsIdentity() const
+	{
+		return Offset == 0 && TimeScale.GetType() == EMovieSceneTimeWarpType::FixedPlayRate && TimeScale.AsFixedPlayRate() == 1.0;
+	}
+
 
 	/**
-	 * Time warping information for this sub-sequence used for looping.
+	 * Transform the specified time from untransformed to transformed space
+	 * 
+	 * @param Time    The untransformed time to transform
+	 * 
+	 * @return the transformed time
 	 */
-	UPROPERTY()
-	FMovieSceneTimeWarping Warping;
+	MOVIESCENE_API FFrameTime TransformTime(FFrameTime Time) const;
+
+
+	/**
+	 * Transform the specified time from untransformed to transformed space
+	 * 
+	 * @param Time    The untransformed time to transform
+	 * @param Params  Transformation parameters, allowing control over the transformation, and gathering of breadcrumbs
+	 *
+	 * @return the transformed time
+	 */
+	MOVIESCENE_API FFrameTime TransformTime(FFrameTime Time, const UE::MovieScene::FTransformTimeParams& Params) const;
+
 
+	/**
+	 * Generate the inverse of this transformation (ie a transform from transformed to untransformed space)
+	 */
+	MOVIESCENE_API FMovieSceneInverseNestedSequenceTransform Inverse() const;
+
+
+	/**
+	 * Given a range in untransformed space, compute the hull of times that this range encompass when transformed.
+	 * For instance, if this transform represents a loop of [0, 10) an input range of:
+	 *     - [15,19) would yield [5,9)
+	 *     - [5,25) would yield [0,10)
+	 * 
+	 * @note Disjoint ranges are not supported by this function. An input of [5, 19) under the previous example would yield
+	 *       [0, 9) because we can only return a single range.
+	 */
+	MOVIESCENE_API TRange<FFrameTime> ComputeTraversedHull(const TRange<FFrameTime>& Range) const;
+
+
+	/**
+	 * Check whether this transformation supports 'boundaries'. A boundary is implementation defined,
+	 *     but is generally a loop or cycle point.
+	 */
+	MOVIESCENE_API bool SupportsBoundaries() const;
+
+
+	/**
+	 * Extract all the boundaries for this transform within the specified untransformed start and end time,
+	 * invoking Visitor for every boundary that is present within the range. Iteration will continue until
+	 * Visitor returns false, at which point this function itself will return false.
+	 *
+	 * @param UntransformedStart       The inclusive start of the range within which Visitor can be invoked
+	 * @param UntransformedEnd         The inclusive end of the range within which Visitor can be invoked
+	 * @param Visitor                  A functor to invoke for every boundary that is found within the range. Return value signifies whether iteration continues (true) or not (false)
+	 * 
+	 * @return (true) if there were no boundaries, or Visitor returned true for all encountered boundaries; false otherwise
+	 */
+	MOVIESCENE_API bool ExtractBoundariesWithinRange(FFrameTime UntransformedStart, FFrameTime UntransformedEnd, const TFunctionRef<bool(FFrameTime)>& Visitor) const;
+
+
+	/**
+	 * Convert this transformation to a string representation
+	 */
+	MOVIESCENE_API void ToString(TStringBuilderBase<TCHAR>& OutBuilder) const;
+
+
+	/**
+	 * Equality comparison
+	 */
 	friend bool operator==(const FMovieSceneNestedSequenceTransform& A, const FMovieSceneNestedSequenceTransform& B)
 	{
-		return A.LinearTransform == B.LinearTransform && A.Warping == B.Warping;
+		return A.TimeScale == B.TimeScale && A.Offset == B.Offset;
 	}
 
+
+	/**
+	 * Inequality comparison
+	 */
 	friend bool operator!=(const FMovieSceneNestedSequenceTransform& A, const FMovieSceneNestedSequenceTransform& B)
 	{
-		return A.LinearTransform != B.LinearTransform || A.Warping != B.Warping;
+		return A.TimeScale != B.TimeScale || A.Offset != B.Offset;
+	}
+
+private:
+
+	/**
+	 * Time scale implemented as an optionally-warping variant
+	 */
+	UPROPERTY()
+	FMovieSceneTimeWarpVariant TimeScale;
+
+	/**
+	 * Linear time transform for this sub-sequence.
+	 */
+	UPROPERTY()
+	FFrameTime Offset;
+
+public:
+
+	/*~~~~~~~~~~~~~~~~~~~~~~~~~~~~~~~~~~~~~ Begin API deprecation ~~~~~~~~~~~~~~~~~~~~~~~~~~~~~~~~~~~~~*/
+	UE_DEPRECATED(5.5, "Please update your code to use a FMovieSceneTimeWarpVariant")
+	FMovieSceneNestedSequenceTransform(const FMovieSceneTimeWarping& InWarping)
+	{
+	}
+
+	UE_DEPRECATED(5.5, "Please update your code to use a FMovieSceneTimeWarpVariant")
+	FMovieSceneNestedSequenceTransform(FMovieSceneTimeTransform InLinearTransform, const FMovieSceneTimeWarping& InWarping)
+		: TimeScale(InLinearTransform.TimeScale), Offset(InLinearTransform.Offset)
+	{
 	}
 
-	FMovieSceneTimeTransform InverseLinearOnly() const
+	UE_DEPRECATED(5.5, "Please update your code to check for IsLinear()")
+	bool IsLooping() const
 	{
-		// We don't support storing both an offset and a zero or infinite timescale in a single nested transform.
-		// These should be stored as separate nested transforms.
-		ensureMsgf(FMath::IsNearlyZero(LinearTransform.Offset.AsDecimal()) || (FMath::IsFinite(LinearTransform.TimeScale) && !FMath::IsNearlyZero(LinearTransform.TimeScale)), TEXT("Error- nested transform with both an offset and a non-deterministic timescale."));
+		return TimeScale.GetType() == EMovieSceneTimeWarpType::Loop;
+	}
 
-		return LinearTransform.Inverse();
+	UE_DEPRECATED(5.5, "Please use Inverse()")
+	FMovieSceneNestedSequenceTransform InverseLinearOnly() const
+	{
+		return IsLinear() ? AsLinear().Inverse() : FMovieSceneNestedSequenceTransform();
 	}
 
-	FMovieSceneTimeTransform InverseFromWarp(uint32 WarpCount) const
+	UE_DEPRECATED(5.5, "Please use Inverse()")
+	FMovieSceneNestedSequenceTransform InverseFromWarp(uint32 WarpCount) const
 	{
-		// We don't support storing both an offset and a zero or infinite timescale in a single nested transform.
-		// We also don't support storing a zero or infinite timescale in addition to looping information.
-		// These should be stored as separate nested transforms.
-		ensureMsgf(FMath::IsNearlyZero(LinearTransform.Offset.AsDecimal()) || ((FMath::IsFinite(LinearTransform.TimeScale) && !FMath::IsNearlyZero(LinearTransform.TimeScale))), TEXT("Error- nested transform with both an offset and a non-deterministic timescale."));
-		ensureMsgf(!IsLooping() || (FMath::IsFinite(LinearTransform.TimeScale) && !FMath::IsNearlyZero(LinearTransform.TimeScale)), TEXT("Error- nested transform with both looping and non-deterministic timescale."));
-		return LinearTransform.Inverse() * Warping.InverseFromWarp(WarpCount);
+		return IsLinear() ? AsLinear().Inverse() : FMovieSceneNestedSequenceTransform();
 	}
+
+	UE_DEPRECATED(5.5, "Warping is now implemented as a variant within TimeScale")
+	FMovieSceneTimeWarping Warping;
+	/*~~~~~~~~~~~~~~~~~~~~~~~~~~~~~~~~~~~~~  End API deprecation  ~~~~~~~~~~~~~~~~~~~~~~~~~~~~~~~~~~~~~*/
 };
 
+
+
+
 /**
- * Movie scene sequence transform class that transforms from one time-space to another.
+ * Movie scene sequence transform class that transforms from one time-space to another, represented as
+ * a linear transformation plus zero or more complex, non-linear transformations.
  */
 USTRUCT()
 struct FMovieSceneSequenceTransform
 {
 	GENERATED_BODY()
 
+
 	/**
 	 * Default construction to the identity transform
 	 */
@@ -217,56 +619,62 @@ struct FMovieSceneSequenceTransform
 		: LinearTransform(InOffset, InTimeScale)
 	{}
 
+
 	/**
 	 * Construction from a linear time transform.
 	 *
 	 * @param InLinearTransform	The linear transform
 	 */
-	FMovieSceneSequenceTransform(FMovieSceneTimeTransform InLinearTransform)
+	explicit FMovieSceneSequenceTransform(FMovieSceneTimeTransform InLinearTransform)
 		: LinearTransform(InLinearTransform)
 	{}
 
-	friend bool operator==(const FMovieSceneSequenceTransform& A, const FMovieSceneSequenceTransform& B)
-	{
-		return A.LinearTransform == B.LinearTransform && A.NestedTransforms == B.NestedTransforms;
-	}
 
-	friend bool operator!=(const FMovieSceneSequenceTransform& A, const FMovieSceneSequenceTransform& B)
+	/**
+	 * Construction from a single nested sequence transform structure
+	 *
+	 * @param InNestedTransform	The transform
+	 */
+	explicit FMovieSceneSequenceTransform(FMovieSceneNestedSequenceTransform&& InNestedTransform)
 	{
-		return A.LinearTransform != B.LinearTransform || A.NestedTransforms != B.NestedTransforms;
+		if (InNestedTransform.IsLinear())
+		{
+			LinearTransform = InNestedTransform.AsLinear();
+		}
+		else
+		{
+			NestedTransforms.Emplace(InNestedTransform);
+		}
 	}
 
+
 	/**
-	 * Returns whether this sequence transform is an identity transform (i.e. it doesn't change anything).
+	 * Equality comparison
 	 */
-	bool IsIdentity() const
+	friend bool operator==(const FMovieSceneSequenceTransform& A, const FMovieSceneSequenceTransform& B)
 	{
-		return LinearTransform.IsIdentity() && 
-			(NestedTransforms.Num() == 0 || 
-			 Algo::AllOf(NestedTransforms,
-				 [](const FMovieSceneNestedSequenceTransform& A) -> bool
-				 {
-					return A.IsIdentity();
-				 }));
+		return A.LinearTransform == B.LinearTransform && A.NestedTransforms == B.NestedTransforms;
 	}
 
+
 	/**
-	 * Returns whether this sequence transform includes any time warping.
+	 * Inequality comparison
 	 */
-	bool IsWarping() const
+	friend bool operator!=(const FMovieSceneSequenceTransform& A, const FMovieSceneSequenceTransform& B)
 	{
-		// Since we compress any linear transforms, we know that there will be warping as soon as we
-		// use the nested transform array.
-		return NestedTransforms.Num() > 0;
+		return A.LinearTransform != B.LinearTransform || A.NestedTransforms != B.NestedTransforms;
 	}
 
 
-	bool IsLooping() const
+	/**
+	 * Returns whether this sequence transform includes any time warping.
+	 */
+	bool NeedsBreadcrumbs() const
 	{
-		return NestedTransforms.Num() > 0 &&
-			Algo::AnyOf(NestedTransforms, &FMovieSceneNestedSequenceTransform::IsLooping);
+		return !IsLinear();
 	}
 
+
 	/**
 	 * Returns whether this sequence transform is purely linear (i.e. doesn't involve time warping).
 	 */
@@ -275,409 +683,332 @@ struct FMovieSceneSequenceTransform
 		return NestedTransforms.Num() == 0;
 	}
 
+
 	/**
-	 * Returns the total time-scale of this transform.
+	 * Convert this transform to its linear representation.
+	 * @note It is invalid to call this function unless IsLinear() is true
 	 */
-	float GetTimeScale() const
+	FMovieSceneTimeTransform AsLinear() const
 	{
-		float TimeScale = LinearTransform.TimeScale;
-		for (const FMovieSceneNestedSequenceTransform& NestedTransform : NestedTransforms)
-		{
-			if (NestedTransform.LinearTransform.TimeScale != 1.f)
-			{
-				TimeScale *= NestedTransform.LinearTransform.TimeScale;
-			}
-		}
-		return TimeScale;
+		return LinearTransform;
 	}
 
+
 	/**
-	 * Transforms the given time, returning the transformed time.
+	 * Returns whether this sequence transform is an identity transform (i.e. it doesn't change anything).
 	 */
-	FFrameTime TransformTime(FFrameTime InTime) const
-	{
-		FFrameTime OutTime;
-		FMovieSceneWarpCounter WarpCounter;
-		TransformTime(InTime, OutTime, WarpCounter);
-		return OutTime;
-	}
+	MOVIESCENE_API bool IsIdentity() const;
+
 
 	/**
-	 * Transforms the given time, returning the transformed time along with a series of loop indices that
-	 * indicate which loops each nested level is in.
+	 * Transform the specified time into the inner-most (transformed) space
+	 * 
+	 * @param Time       The input time to transform
+	 * 
+	 * @return The resulting time in transformed space
 	 */
-	void TransformTime(FFrameTime InTime, FFrameTime& OutTime, FMovieSceneWarpCounter& OutWarpCounter) const
-	{
-		OutWarpCounter.WarpCounts.Reset();
+	MOVIESCENE_API FFrameTime TransformTime(FFrameTime Time) const;
 
-		OutTime = InTime * LinearTransform;
-		
-		if (NestedTransforms.Num() > 0)
-		{
-			for (const FMovieSceneNestedSequenceTransform& NestedTransform : NestedTransforms)
-			{
-				OutTime = OutTime * NestedTransform.LinearTransform;
-
-				if (NestedTransform.IsLooping())
-				{
-					uint32 WarpIndex;
-					NestedTransform.Warping.TransformTime(OutTime, OutTime, WarpIndex);
-					OutWarpCounter.WarpCounts.Add(WarpIndex);
-				}
-				else
-				{
-					OutWarpCounter.WarpCounts.Add(FMovieSceneTimeWarping::InvalidWarpCount);
-				}
-			}
-		}
-	}
 
 	/**
-	 * Transforms the given range using the "pure" warping transformation.
-	 * See FMovieSceneTimeWarping for more information.
+	 * Transform the specified time into the inner-most (transformed) space
+	 * 
+	 * @param Time       The input time to transform
+	 * @param Params     Parameters for controlling the transform operation, and harvesting breadcrumbs
+	 * 
+	 * @return The resulting time in transformed space
 	 */
-	TRange<FFrameTime> TransformRangePure(const TRange<FFrameTime>& Range) const
-	{
-		TRange<FFrameTime> Result = Range * LinearTransform;
-		for (const FMovieSceneNestedSequenceTransform& NestedTransform : NestedTransforms)
-		{
-			Result = Result * NestedTransform.LinearTransform;
-			if (NestedTransform.IsLooping())
-			{
-				Result = NestedTransform.Warping.TransformRangePure(Result);
-			}
-		}
-		return Result;
-	}
+	MOVIESCENE_API FFrameTime TransformTime(FFrameTime InTime, const UE::MovieScene::FTransformTimeParams& Params) const;
+
 
 	/**
-	 * Transforms the given range using the "unwarped" warping transformation. 
-	 * See FMovieSceneTimeWarping for more information.
+	 * Given a range in untransformed space, compute the hull of times that this range encompass when transformed.
+	 * For instance, if this transform represents a loop of [0, 10) an input range of:
+	 *     - [15,19) would yield [5,9)
+	 *     - [5,25) would yield [0,10)
+	 * 
+	 * @note Disjoint ranges are not supported by this function. An input of [5, 19) under the previous example would yield
+	 *       [0, 9) because we can only return a single range.
 	 */
-	TRange<FFrameTime> TransformRangeUnwarped(const TRange<FFrameTime>& Range) const
-	{
-		TRange<FFrameTime> Result = Range * LinearTransform;
-		for (const FMovieSceneNestedSequenceTransform& NestedTransform : NestedTransforms)
-		{
-			Result = Result * NestedTransform.LinearTransform;
-			if (NestedTransform.IsLooping())
-			{
-				Result = NestedTransform.Warping.TransformRangeUnwarped(Result);
-			}
-		}
-		return Result;
-	}
+	MOVIESCENE_API TRange<FFrameTime> ComputeTraversedHull(const TRange<FFrameTime>& Range) const;
+
 
 	/**
-	 * Transforms the given range using the "constrained" warping transformation. 
-	 * See FMovieSceneTimeWarping for more information.
+	 * Given a range in untransformed space, compute the hull of times that this range encompass when transformed.
+	 * See above overload for additional details.
 	 */
-	TRange<FFrameTime> TransformRangeConstrained(const TRange<FFrameTime>& Range) const
-	{
-		TRange<FFrameTime> Result = Range * LinearTransform;
-		for (const FMovieSceneNestedSequenceTransform& NestedTransform : NestedTransforms)
-		{
-			Result = Result * NestedTransform.LinearTransform;
-			if (NestedTransform.IsLooping())
-			{
-				Result = NestedTransform.Warping.TransformRangeConstrained(Result);
-			}
-		}
-		return Result;
-	}
+	MOVIESCENE_API TRange<FFrameTime> ComputeTraversedHull(const TRange<FFrameNumber>& Range) const;
+
+
+public:
+
 
 	/**
-	 * Transforms the given range using the "pure" warping transformation.
-	 * See FMovieSceneTimeWarping for more information.
+	 * Add the specified linear transform to the end of this transform stack (ie, applying it last)
 	 */
-	TRange<FFrameNumber> TransformRangePure(const TRange<FFrameNumber>& Range) const
-	{
-		TRange<FFrameTime> TimeRange = ConvertRange<FFrameNumber, FFrameTime>(Range);
-		TimeRange = TransformRangePure(TimeRange);
-		return ConvertRange<FFrameTime, FFrameNumber>(TimeRange);
-	}
+	MOVIESCENE_API void Add(FMovieSceneTimeTransform InTransform);
+
 
 	/**
-	 * Transforms the given range using the "unwarped" warping transformation. 
-	 * See FMovieSceneTimeWarping for more information.
+	 * Add the specified warping transform to the end of this transform stack with an offset (ie, applying it last)
 	 */
-	TRange<FFrameNumber> TransformRangeUnwarped(const TRange<FFrameNumber>& Range) const
-	{
-		TRange<FFrameTime> TimeRange = ConvertRange<FFrameNumber, FFrameTime>(Range);
-		TimeRange = TransformRangeUnwarped(TimeRange);
-		return ConvertRange<FFrameTime, FFrameNumber>(TimeRange);
-	}
+	MOVIESCENE_API void Add(FFrameTime InOffset, FMovieSceneTimeWarpVariant&& InTimeWarp);
+
 
 	/**
-	 * Transforms the given range using the "constrained" warping transformation. 
-	 * See FMovieSceneTimeWarping for more information.
+	 * Add an entry to this transform denoting it should loop between the specified start and end point
+	 * @note Start and End are interpreted in the transformed space of the current stack of transforms in this class.
+	 * 
+	 * @param Start     The start frame number of the loop
+	 * @param End       The ending frame number of the loop
 	 */
-	TRange<FFrameNumber> TransformRangeConstrained(const TRange<FFrameNumber>& Range) const
-	{
-		TRange<FFrameTime> TimeRange = ConvertRange<FFrameNumber, FFrameTime>(Range);
-		TimeRange = TransformRangeConstrained(TimeRange);
-		return ConvertRange<FFrameTime, FFrameNumber>(TimeRange);
-	}
+	MOVIESCENE_API void AddLoop(FFrameNumber Start, FFrameNumber End);
+
 
 	/**
-	 * Retrieve the inverse of the linear part of this transform.
+	 * Multiply this tranmsform with another transform, resulting in a single transform that gets from RHS parent to LHS space
+	 * @note Transforms apply from right to left
 	 */
+	MOVIESCENE_API FMovieSceneSequenceTransform operator*(const FMovieSceneSequenceTransform& RHS) const;
+
+
+	/**
+	 * Compute the inverse of this transform, that is: the transform that goes from transformed to untransformed space.
+	 */
+	MOVIESCENE_API FMovieSceneInverseSequenceTransform Inverse() const;
+
+
+	/**
+	 * Extract all the boundaries for this transform within the specified untransformed start and end time,
+	 * invoking Visitor for every boundary that is present within the range. Iteration will continue until
+	 * Visitor returns false, at which point this function itself will return false.
+	 *
+	 * @param UntransformedStart       The inclusive start of the range within which Visitor can be invoked
+	 * @param UntransformedEnd         The inclusive end of the range within which Visitor can be invoked
+	 * @param Visitor                  A functor to invoke for every boundary that is found within the range. Return value signifies whether iteration continues (true) or not (false)
+	 * 
+	 * @return (true) if there were no boundaries, or Visitor returned true for all encountered boundaries; false otherwise
+	 */
+	MOVIESCENE_API bool ExtractBoundariesWithinRange(FFrameTime UntransformedStart, FFrameTime UntransformedEnd, const TFunctionRef<bool(FFrameTime)>& Visitor) const;
+
+	/*~~~~~~~~~~~~~~~~~~~~~~~~~~~~~~~~~~~~~ Begin API deprecation ~~~~~~~~~~~~~~~~~~~~~~~~~~~~~~~~~~~~~*/
+	UE_DEPRECATED(5.5, "This function is no longer supported. Please use !IsLinear()")
+	bool IsLooping() const;
+	UE_DEPRECATED(5.5, "Please upgrade your code to use TransformTime that takes breadcrumbs")
+	MOVIESCENE_API void TransformTime(FFrameTime InTime, FFrameTime& OutTime, FMovieSceneWarpCounter& OutWarpCounter) const;
+	UE_DEPRECATED(5.5, "Transforms no longer have a constant time scale.")
+	MOVIESCENE_API float GetTimeScale() const;
+	UE_DEPRECATED(5.5, "Please use ComputeTraversedHull or TransformTime")
+	MOVIESCENE_API TRange<FFrameTime> TransformRangePure(const TRange<FFrameTime>& Range) const;
+	UE_DEPRECATED(5.5, "Please use ComputeTraversedHull or TransformTime")
+	MOVIESCENE_API TRange<FFrameTime> TransformRangeUnwarped(const TRange<FFrameTime>& Range) const;
+	UE_DEPRECATED(5.5, "Please use ComputeTraversedHull or TransformTime")
+	MOVIESCENE_API TRange<FFrameTime> TransformRangeConstrained(const TRange<FFrameTime>& Range) const;
+	UE_DEPRECATED(5.5, "Please use ComputeTraversedHull or TransformTime")
+	MOVIESCENE_API TRange<FFrameNumber> TransformRangePure(const TRange<FFrameNumber>& Range) const;
+	UE_DEPRECATED(5.5, "Please use ComputeTraversedHull or TransformTime")
+	MOVIESCENE_API TRange<FFrameNumber> TransformRangeUnwarped(const TRange<FFrameNumber>& Range) const;
+	UE_DEPRECATED(5.5, "Please use ComputeTraversedHull or TransformTime")
+	MOVIESCENE_API TRange<FFrameNumber> TransformRangeConstrained(const TRange<FFrameNumber>& Range) const;
+	UE_DEPRECATED(5.5, "Please use Inverse()")
+	MOVIESCENE_API FMovieSceneSequenceTransform InverseNoLooping() const;
 	UE_DEPRECATED(5.4, "Please use InverseNoLooping instead.")
-	FMovieSceneTimeTransform InverseLinearOnly() const
-	{
-		ensureMsgf(!FMath::IsNearlyZero(LinearTransform.TimeScale), TEXT("Inverse of a zero timescale transform is undefined in a FMovieSceneTimeTransform. Please use InverseNoLooping for proper behavior."));
-		return LinearTransform.Inverse();
-	}
+	MOVIESCENE_API FMovieSceneTimeTransform InverseLinearOnly() const;
+	UE_DEPRECATED(5.4, "Please use InverseFromAllFirstLoops instead.")
+	MOVIESCENE_API FMovieSceneTimeTransform InverseFromAllFirstWarps() const;
+	UE_DEPRECATED(5.4, "Please use InverseFromLoop instead.")
+	MOVIESCENE_API FMovieSceneTimeTransform InverseFromWarp(const FMovieSceneWarpCounter& WarpCounter) const;
+	UE_DEPRECATED(5.4, "Please use InverseFromLoop instead.")
+	MOVIESCENE_API FMovieSceneTimeTransform InverseFromWarp(const TArrayView<const uint32>& WarpCounts) const;
+	UE_DEPRECATED(5.5, "Please use Inverse()")
+	MOVIESCENE_API FMovieSceneSequenceTransform InverseFromAllFirstLoops() const;
+	UE_DEPRECATED(5.5, "Please use Inverse()")
+	MOVIESCENE_API FMovieSceneSequenceTransform InverseFromLoop(const FMovieSceneWarpCounter& LoopCounter) const;
+	UE_DEPRECATED(5.5, "Please use Inverse()")
+	FMovieSceneSequenceTransform InverseFromLoop(const TArrayView<const FFrameTime>& Breadcrumbs) const;
+	UE_DEPRECATED(5.5, "Please use Inverse()")
+	MOVIESCENE_API FMovieSceneSequenceTransform InverseFromLoop(const TArrayView<const uint32>& LoopCounts) const;
+	/*~~~~~~~~~~~~~~~~~~~~~~~~~~~~~~~~~~~~~  End API deprecation  ~~~~~~~~~~~~~~~~~~~~~~~~~~~~~~~~~~~~~*/
 
 	/**
-	 * Retrieve the inverse of the linear part of this transform without any looping.
-	 * Note that if nested transforms exist, we will return a sequence transform with nested transforms inverted, but looping information will be lost.
+	 * The initial linear transformation represented as a 2D matrix. Always applied first.
 	 */
-	FMovieSceneSequenceTransform InverseNoLooping() const
-	{
-		ensureMsgf(FMath::IsFinite(LinearTransform.TimeScale) && !FMath::IsNearlyZero(LinearTransform.TimeScale), TEXT("Warning: Sequence LinearTransform has a zero or infinite timescale."));
-		if (NestedTransforms.Num() > 0)
-		{
-			// Start accumulating the inverse transforms in reverse order.
-			FMovieSceneSequenceTransform Result;
+	UPROPERTY()
+	FMovieSceneTimeTransform LinearTransform;
 
-			for(int i = NestedTransforms.Num() - 1; i >= 0; --i)
-			{
-				Result.NestedTransforms.Add(NestedTransforms[i].LinearTransform.Inverse());
-			}
+	/**
+	 * Additional transformations that should be applied after LinearTransform.
+	 * This array is populated whenever a non-linear transform is encountered.
+	 * */
+	UPROPERTY()
+	TArray<FMovieSceneNestedSequenceTransform> NestedTransforms;
+};
 
-			// Add the inverse of the main linear transform as a further nested transform if not identity
-			if (!LinearTransform.IsIdentity())
-			{
-				Result.NestedTransforms.Add(LinearTransform.Inverse());
-			}
 
-			return Result;
-		}
-		else
-		{
-			return LinearTransform.Inverse();
-		}
-	}
+
+/**
+ * The inverse of a FMovieSceneSequenceTransform representing a transformation from transformed, to untransformed space.
+ * This uses a different class and API because the algorithms for computing the inverse of non-linear are different,
+ * often more complex, and can fail. Whereas an FMovieSceneSequenceTransform can only represent a 1:1 mapping from outer
+ * to inner space, its inverse is a (sometimes empty) many:many mapping.
+ * 
+ * Consider a looping transform with a duration of 10 frames: [0, 10). Every time in the outer space maps to a time in the
+ * inner space, but the opposite is not true. Only frames 0-10 exist in the inner space, and each frame in that time maps
+ * to an infinite number of solutions in the outer space. Conversely, any inner time outside the loop range, ie, [-inf, 0)..(10, +inf]
+ * cannot be transformed into the outer space.
+ * 
+ * For this reason, the API only has functions for attempting such computations (TryTransformTime), and iterating the solutions
+ * for any given time within a range.
+ * 
+ * The inverse of an inverse transform is the original transform such that T*(1/T)=I theoretically holds true, although transform
+ * multiplication is not actually supported by the API.
+ */
+USTRUCT()
+struct FMovieSceneInverseSequenceTransform
+{
+	GENERATED_BODY()
+
 
 	/**
-	 * Retrieve the inverse of this transform assuming the local times would all belong to
-	 * the first warp inside each nested time range.
+	 * Returns whether this is a linear transform involving no non-linear components
 	 */
-	UE_DEPRECATED(5.4, "Please use InverseFromAllFirstLoops instead.")
-	FMovieSceneTimeTransform InverseFromAllFirstWarps() const
+	bool IsLinear() const
 	{
-		const size_t NestedTransformsSize = NestedTransforms.Num();
-		if (NestedTransformsSize == 0)
-		{
-			return LinearTransform.Inverse();
-		}
-		else
-		{
-			FMovieSceneSequenceTransform SequenceTransform = InverseFromAllFirstLoops();
-			FMovieSceneTimeTransform ReturnTransform = SequenceTransform.LinearTransform;
-			for (int i = 0; i < SequenceTransform.NestedTransforms.Num(); ++i)
-			{
-				ReturnTransform = ReturnTransform * SequenceTransform.NestedTransforms[i].LinearTransform;
-			}
-			return ReturnTransform;
-		}
+		return NestedTransforms.Num() == 0;
 	}
 
+
 	/**
-	 * Retrieve the inverse of this transform assuming the local times would all belong to the first loop inside each nested time range.
+	 * Cast this transform to a linear transformation, provided IsLinear() is true.
 	 */
-	FMovieSceneSequenceTransform InverseFromAllFirstLoops() const 
+	FMovieSceneTimeTransform AsLinear() const
 	{
-		const size_t NestedTransformsSize = NestedTransforms.Num();
-		if (NestedTransformsSize == 0)
-		{
-			ensureMsgf(FMath::IsFinite(LinearTransform.TimeScale) && !FMath::IsNearlyZero(LinearTransform.TimeScale), TEXT("Warning: Sequence LinearTransform has a zero or infinite timescale."));
-			return LinearTransform.Inverse();
-		}
-		else
-		{
-			TArray<uint32> LoopCounts;
-			LoopCounts.Reserve(NestedTransforms.Num());
-			for (const FMovieSceneNestedSequenceTransform& NestedTransform : NestedTransforms)
-			{
-				LoopCounts.Add(NestedTransform.IsLooping() ? 0 : FMovieSceneTimeWarping::InvalidWarpCount);
-			}
-			return InverseFromLoop(LoopCounts);
-		}
+		check(IsLinear());
+		return LinearTransform;
 	}
 
+
 	/**
-	 * Retrieve the inverse of this transform assuming the local times would belong to the
-	 * n'th warp inside the time range. As such, this method requires that the warp counts
-	 * in the counter object is equal to the number of nested loops in this transform.
+	 * Attempt to transform the specified transformed time into its untransformed space.
+	 * This function can fail if the time does not map to any other times in the untransformed space.
+	 * 
+	 * @param Time         The transformed time to convert to untransformed space
+	 * @param Breadcrumbs  A trail of breadcrumbs that lead us from untransformed to transformed space.
+	 *                     The algorithm will attempt to find the closest solution to each breadcrumb as it goes.
+	 * 
+	 * @return The time in untransformed space if succesful, or an empty optional otherwise.
 	 */
-	UE_DEPRECATED(5.4, "Please use InverseFromLoop instead.")
-	FMovieSceneTimeTransform InverseFromWarp(const FMovieSceneWarpCounter& WarpCounter) const
-	{
-		FMovieSceneSequenceTransform SequenceTransform = InverseFromLoop(WarpCounter.WarpCounts);
-		FMovieSceneTimeTransform ReturnTransform = SequenceTransform.LinearTransform;
-		for (int i = 0; i < SequenceTransform.NestedTransforms.Num(); ++i)
-		{
-			ReturnTransform = ReturnTransform * SequenceTransform.NestedTransforms[i].LinearTransform;
-		}
-		return ReturnTransform;
-	}
+	MOVIESCENE_API TOptional<FFrameTime> TryTransformTime(FFrameTime Time, const FMovieSceneTransformBreadcrumbs& Breadcrumbs) const;
+
 
 	/**
-	 * Retrieve the inverse of this transform assuming the local times would belong to the
-	 * n'th warp inside the time range. As such, this method takes as parameter an array with
-	 * as many items as there currently have nested loops in this transform.
+	 * Attempt to transform the specified transformed time into its untransformed space.
+	 * This function can fail if the time does not map to any other times in the untransformed space.
+	 * 
+	 * @param Time         The transformed time to convert to untransformed space
+	 * @param Breadcrumbs  A trail of breadcrumbs that lead us from untransformed to transformed space.
+	 *                     The algorithm will attempt to find the closest solution to each breadcrumb as it goes.
+	 * @param Params       Additional parameters that control the algotihm, such as whether it can look forwards, backwards, or across cycles/loops.
+	 * 
+	 * @return The time in untransformed space if succesful, or an empty optional otherwise.
 	 */
-	UE_DEPRECATED(5.4, "Please use InverseFromLoop instead.")
-	FMovieSceneTimeTransform InverseFromWarp(const TArrayView<const uint32>& WarpCounts) const
-	{
-		const size_t NestedTransformsSize = NestedTransforms.Num();
-		check(WarpCounts.Num() == NestedTransformsSize);
+	MOVIESCENE_API TOptional<FFrameTime> TryTransformTime(FFrameTime Time, const FMovieSceneTransformBreadcrumbs& Breadcrumbs, const UE::MovieScene::FInverseTransformTimeParams& Params) const;
 
-		if (NestedTransformsSize > 0)
-		{
-			// Start accumulating the inverse transforms in reverse order.
-			const FMovieSceneNestedSequenceTransform& LastNesting = NestedTransforms.Last();
-			FMovieSceneTimeTransform Result = LastNesting.InverseFromWarp(WarpCounts.Last());
-
-			size_t Index = NestedTransformsSize - 1;
-			while (Index > 0)
-			{
-				--Index;
-				
-				const uint32 CurWarpCount = WarpCounts[Index];
-				const FMovieSceneNestedSequenceTransform& CurNesting = NestedTransforms[Index];
-				Result = CurNesting.InverseFromWarp(CurWarpCount) * Result;
-			}
-
-			// Add the inverse of the main linear transform.
-			Result = LinearTransform.Inverse() * Result;
-
-			return Result;
-		}
-		else
-		{
-			return FMovieSceneTimeTransform(LinearTransform.Inverse());
-		}
-	}
 
 	/**
-	 * Retrieve the inverse of this transform assuming the local times would belong to the
-	 * n'th warp inside the time range. As such, this method requires that the warp counts
-	 * in the counter object is equal to the number of nested loops in this transform.
-	 * Correctly handles zero and infinite timescales.
+	 * Fallback overload that does not require a breadcrumb trail if one is not available.
+	 * 
+	 * Attempt to transform the specified transformed time into its untransformed space.
+	 * This function can fail if the time does not map to any other times in the untransformed space.
+	 * 
+	 * @note This overload may yield unexpected times due to the lack of breadcrumb context. Breadcrumbs should be used if possible.
+	 *
+	 * @param Time         The transformed time to convert to untransformed space
+	 * 
+	 * @return The time in untransformed space if succesful, or an empty optional otherwise.
 	 */
-	FMovieSceneSequenceTransform InverseFromLoop(const FMovieSceneWarpCounter& LoopCounter) const
-	{
-		return InverseFromLoop(LoopCounter.WarpCounts);
-	}
+	MOVIESCENE_API TOptional<FFrameTime> TryTransformTime(FFrameTime Time) const;
+
 
 	/**
-	 * Retrieve the inverse of this transform assuming the local times would belong to the
-	 * n'th warp inside the time range. As such, this method takes as parameter an array with
-	 * as many items as there currently have nested loops in this transform.
-	 * Correctly handles zero and infinite timescales.
+	 * Fallback overload that does not require a breadcrumb trail if one is not available.
+	 * 
+	 * Attempt to transform the specified transformed time into its untransformed space.
+	 * This function can fail if the time does not map to any other times in the untransformed space.
+	 * 
+	 * @note This overload may yield unexpected times due to the lack of breadcrumb context. Breadcrumbs should be used if possible.
+	 *
+	 * @param Time         The transformed time to convert to untransformed space
+	 * @param Params       Additional parameters that control the algotihm, such as whether it can look forwards, backwards, or across cycles/loops.
+	 * 
+	 * @return The time in untransformed space if succesful, or an empty optional otherwise.
 	 */
-	FMovieSceneSequenceTransform InverseFromLoop(const TArrayView<const uint32>& LoopCounts) const
-	{
-		const size_t NestedTransformsSize = NestedTransforms.Num();
-		check(LoopCounts.Num() == NestedTransformsSize);
+	MOVIESCENE_API TOptional<FFrameTime> TryTransformTime(FFrameTime Time, const UE::MovieScene::FInverseTransformTimeParams& Params) const;
+
+	/**
+	 * Transforms a time from transformed to untransformed space within a finite range specified by two breadcrumb trails.
+	 * For each solution that exists in the untransformed space, Visitor will be invoked and its result defines whether the
+	 * algorithm continues (true) or terminates early (false).
+	 *
+	 * @param Time              The transformed time to convert to untransformed space
+	 * @param Visitor           A functor to invoke for every solution that exists within the range.
+	 *                          Returning true will allow the algorithm to continue, false will cause it to terminate early.
+	 * @param StartBreadCrumbs  A breadcrumb trail for the start of the allowable range to solve within.
+	 * @param EndBreadCrumbs    A breadcrumb trail for the end of the allowable range to solve within.
+	 * 
+	 * @return False if any invocation of Visitor returned false, true otherwise (if there were no solutions, or every invocation of Visitor returned true).
+	 */
+	MOVIESCENE_API bool TransformTimeWithinRange(FFrameTime Time, const TFunctionRef<bool(FFrameTime)>& Visitor, const FMovieSceneTransformBreadcrumbs& StartBreadcrumbs, const FMovieSceneTransformBreadcrumbs& EndBreadcrumbs) const;
 
-		if (NestedTransformsSize > 0)
-		{
-			FMovieSceneSequenceTransform Result;
-			for(int i = NestedTransformsSize - 1; i >=0; --i)
-			{
-				const uint32 CurLoopCount = LoopCounts[i];
-				Result.NestedTransforms.Add(NestedTransforms[i].InverseFromWarp(CurLoopCount));
-			}
-
-			// Add the inverse of the main linear transform as a further nested transform if not identity
-			if (!LinearTransform.IsIdentity())
-			{
-				Result.NestedTransforms.Add(LinearTransform.Inverse());
-			}
-
-			return Result;
-		}
-		else
-		{
-			ensureMsgf(FMath::IsFinite(LinearTransform.TimeScale) && !FMath::IsNearlyZero(LinearTransform.TimeScale), TEXT("Warning: Sequence LinearTransform has a zero or infinite timescale."));
-			return LinearTransform.Inverse();
-		}
-	}
 
 	/**
-	 * Multiply 2 transforms together, resulting in a single transform that gets from RHS parent to LHS space
-	 * @note Transforms apply from right to left
+	 * Transforms a finite range in the transformed space, to non-empty ranges in untransformed space.
+	 * For each solution that exists in the untransformed space, Visitor will be invoked and its result defines whether the
+	 * algorithm continues (true) or terminates early (false).
+	 *
+	 * @param Range             The range in transformed space to convert to untransformed space.
+	 * @param Visitor           A functor to invoke for every solution that exists within the range.
+	 *                          Returning true will allow the algorithm to continue, false will cause it to terminate early.
+	 * @param StartBreadCrumbs  A breadcrumb trail for the start of the allowable range to solve within.
+	 * @param EndBreadCrumbs    A breadcrumb trail for the end of the allowable range to solve within.
+	 * 
+	 * @return False if any invocation of Visitor returned false, true otherwise (if there were no solutions, or every invocation of Visitor returned true).
 	 */
-	FMovieSceneSequenceTransform operator*(const FMovieSceneSequenceTransform& RHS) const
-	{
-		if (!IsWarping())
-		{
-			if (!RHS.IsWarping())
-			{
-				// None of the transforms are warping... we can combine them into another linear transform.
-				return FMovieSceneSequenceTransform(LinearTransform * RHS.LinearTransform);
-			}
-			else
-			{
-				// LHS is linear, but RHS is warping. Since transforms are supposed to apply from right to left,
-				// we need to append LHS at the "bottom" of RHS, i.e. add a new nested transform that's LHS. However
-				// if LHS is identity, we have nothing to do, and if both LHS and RHS' deeper transform are linear,
-				// we can combine both.
-				FMovieSceneSequenceTransform Result(RHS);
-				if (!LinearTransform.IsIdentity())
-				{
-					const int32 NumNestedTransforms = RHS.NestedTransforms.Num();
-					check(NumNestedTransforms > 0);
-					const FMovieSceneNestedSequenceTransform& LastNested = RHS.NestedTransforms[NumNestedTransforms - 1];
-					if (LastNested.IsWarping())
-					{
-						Result.NestedTransforms.Add(FMovieSceneNestedSequenceTransform(LinearTransform));
-					}
-					else
-					{
-						Result.NestedTransforms[NumNestedTransforms - 1] = LinearTransform * LastNested.LinearTransform;
-					}
-				}
-				return Result;
-			}
-		}
-		else // LHS is warping
-		{
-			if (!RHS.IsWarping())
-			{
-				// RHS isn't warping, but LHS is, so we combine the linear transform parts and start looping
-				// from there.
-				FMovieSceneSequenceTransform Result;
-				Result.LinearTransform = LinearTransform * RHS.LinearTransform;
-				Result.NestedTransforms = NestedTransforms;
-				return Result;
-			}
-			else
-			{
-				// Both are looping, we need to combine them. Usually, a warping transform doesn't use its linear part,
-				// because whatever linear placement/scaling it has would be in the linear part of the nested transform
-				// struct.
-				FMovieSceneSequenceTransform Result(RHS);
-				const bool bHasOnlyNested = LinearTransform.IsIdentity();
-				if (!bHasOnlyNested)
-				{
-					Result.NestedTransforms.Add(FMovieSceneNestedSequenceTransform(LinearTransform));
-				}
-				Result.NestedTransforms.Append(NestedTransforms);
-				return Result;
-			}
-		}
-	}
+	MOVIESCENE_API bool TransformFiniteRangeWithinRange(const TRange<FFrameTime>& Range, TFunctionRef<bool(TRange<FFrameTime>)> Visitor, const FMovieSceneTransformBreadcrumbs& StartBreadcrumbs, const FMovieSceneTransformBreadcrumbs& EndBreadcrumbs) const;
+
+
+	/**
+	 * Legacy function that folds all linear transforms together. Does not account for non-linear transforms and should not be used any more.
+	 */
+	UE_DEPRECATED(5.5, "This function is no longer supported. Please use FMovieSceneInverseSequenceTransform directly.")
+	MOVIESCENE_API FMovieSceneTimeTransform AsLegacyLinearTimeTransform() const;
+
+private:
 
+
+	/**
+	 * Implementation function for TransformTimeWithinRange that recursively solves 1:many transformations
+	 */
+	bool RecursiveTransformTimeWithinRange(int32 NestingIndex, FFrameTime Time, const TFunctionRef<bool(FFrameTime)>& FinalVisitor, TArrayView<const FFrameTime> StartBreadcrumbs, TArrayView<const FFrameTime> EndBreadcrumbs) const;
+
+private:
+
+	friend FMovieSceneSequenceTransform;
+
+
+	/**
+	 * The final linear transformation represented as a 2D matrix. Always applied last.
+	 */
 	UPROPERTY()
 	FMovieSceneTimeTransform LinearTransform;
 
+
+	/**
+	 * Additional transformations that should be applied before LinearTransform.
+	 * This array is populated whenever a non-linear transform is encountered.
+	 * */
 	UPROPERTY()
-	TArray<FMovieSceneNestedSequenceTransform> NestedTransforms;
+	TArray<FMovieSceneInverseNestedSequenceTransform> NestedTransforms;
 };
 
 /**
@@ -686,30 +1017,7 @@ struct FMovieSceneSequenceTransform
  * @param InTime 			The time to transform
  * @param RHS 				The transform
  */
-inline FFrameTime operator*(FFrameTime InTime, const FMovieSceneSequenceTransform& RHS)
-{
-	const uint32 NestedTransformsSize = RHS.NestedTransforms.Num();
-	if (NestedTransformsSize == 0)
-	{
-		return InTime * RHS.LinearTransform;
-	}
-	else
-	{
-		FFrameTime OutTime = InTime * RHS.LinearTransform;
-		for (const FMovieSceneNestedSequenceTransform& NestedTransform : RHS.NestedTransforms)
-		{
-			if (NestedTransform.IsLooping())
-			{
-				OutTime = OutTime * NestedTransform.LinearTransform * NestedTransform.Warping;
-			}
-			else
-			{
-				OutTime = OutTime * NestedTransform.LinearTransform;
-			}
-		}
-		return OutTime;
-	}
-}
+MOVIESCENE_API FFrameTime operator*(FFrameTime InTime, const FMovieSceneSequenceTransform& RHS);
 
 /**
  * Transform a time by a sequence transform
@@ -723,22 +1031,8 @@ inline FFrameTime& operator*=(FFrameTime& InTime, const FMovieSceneSequenceTrans
 	return InTime;
 }
 
-/**
- * Transform a time range by a sequence transform
- *
- * @param InTime 			The time range to transform
- * @param RHS 				The transform
- */
-template<typename T>
-TRange<T>& operator*=(TRange<T>& LHS, const FMovieSceneSequenceTransform& RHS)
-{
-	LHS = LHS * RHS;
-	return LHS;
-}
-
 /** Convert a FMovieSceneSequenceTransform into a string */
 FString LexToString(const FMovieSceneSequenceTransform& InTransform);
 
 /** Convert a FMovieSceneWarpCounter into a string */
-FString LexToString(const FMovieSceneWarpCounter& InCounter);
-
+FString LexToString(const FMovieSceneWarpCounter& InCounter);
\ No newline at end of file
diff --git a/Engine/Source/Runtime/MovieScene/Public/Evaluation/MovieSceneTimeArray.h b/Engine/Source/Runtime/MovieScene/Public/Evaluation/MovieSceneTimeArray.h
--- a/Engine/Source/Runtime/MovieScene/Public/Evaluation/MovieSceneTimeArray.h
+++ b/Engine/Source/Runtime/MovieScene/Public/Evaluation/MovieSceneTimeArray.h
@@ -2,17 +2,10 @@
 
 #pragma once
 
-#include "Containers/Array.h"
-#include "Containers/ArrayView.h"
-#include "CoreMinimal.h"
 #include "CoreTypes.h"
-#include "HAL/PlatformCrt.h"
-#include "Misc/AssertionMacros.h"
-#include "Misc/FrameNumber.h"
 #include "Misc/FrameTime.h"
-#include "MovieSceneSequenceTransform.h"
-#include "MovieSceneTimeTransform.h"
-#include "MovieSceneTimeWarping.h"
+#include "Containers/Array.h"
+#include "Containers/ArrayView.h"
 
 /** Entry struct for the TMovieSceneTimeArray class */
 template<typename DataType>
@@ -33,86 +26,9 @@ struct TMovieSceneTimeArrayEntry
 template<typename DataType>
 struct TMovieSceneTimeArray
 {
-	/** Starts a scope where local times are offset and scaled as specified */
-	void PushTransform(FFrameNumber FrameOffset, float TimeScale = 1.f)
+	void Add(FFrameTime RootTime, const DataType& Datum)
 	{
-		PushTransform(FMovieSceneTimeTransform(FrameOffset, TimeScale));
-	}
-
-	/** Starts a scope where local times are transformed as specified */
-	void PushTransform(const FMovieSceneTimeTransform& InTransform)
-	{
-		FTransformStep Step;
-		Step.LinearTransform = InTransform;
-		TransformStack.Push(Step);
-
-		UpdateCachedInverseTransform();
-	}
-
-	/** Starts a scope where local times are looped and offset as specified */
-	void PushTransform(FFrameNumber FirstLoopStartFrameOffset, FFrameNumber LoopStart, FFrameNumber LoopEnd, float TimeScale = 1.f)
-	{
-		PushTransform(
-				FMovieSceneTimeTransform(FirstLoopStartFrameOffset, TimeScale),
-				FMovieSceneTimeWarping(LoopStart, LoopEnd));
-	}
-
-	/** Starts a scope where local times are transformed as specified */
-	void PushTransform(const FMovieSceneTimeWarping& InWarping)
-	{
-		PushTransform(FMovieSceneTimeTransform(), InWarping);
-	}
-
-	/** Starts a scope where local times are transformed as specified */
-	void PushTransform(const FMovieSceneTimeTransform& InTransform, const FMovieSceneTimeWarping& InWarping)
-	{
-		FTransformStep Step;
-		Step.LinearTransform = InTransform;
-		Step.Warping = InWarping;
-		Step.WarpCounter = InWarping.IsValid() ? 0 : FMovieSceneTimeWarping::InvalidWarpCount;
-		TransformStack.Push(Step);
-
-		UpdateCachedInverseTransform();
-	}
-
-	/** Incremenets the warp counter for the local time space, i.e. all local times will be considered as being in a further loop */
-	void IncrementWarpCounter()
-	{
-		if (ensure(TransformStack.Num() > 0))
-		{
-			FTransformStep& LastStep = TransformStack.Last();
-			if (ensure(LastStep.Warping.IsValid() && LastStep.WarpCounter != FMovieSceneTimeWarping::InvalidWarpCount))
-			{
-				++LastStep.WarpCounter;
-			}
-		}
-
-		UpdateCachedInverseTransform();
-	}
-
-	/** Removes the last pushed local time transform */
-	void PopTransform()
-	{
-		TransformStack.Pop();
-
-		UpdateCachedInverseTransform();
-	}
-
-	/** Adds a new piece of data, timestamped to the given local time */
-	void Add(FFrameNumber LocalTime, const DataType& Datum)
-	{
-		const FFrameTime RootTime = LocalTime * CachedInverseTransform;
-
-		// Insert the new entry by keeping the array sorted by time.
-		int32 NextEntry = Entries.IndexOfByPredicate([RootTime](const FEntry& Entry) { return RootTime < Entry.RootTime; });
-		if (NextEntry != INDEX_NONE)
-		{
-			Entries.Insert(FEntry{ RootTime, Datum }, NextEntry);
-		}
-		else
-		{
-			Entries.Add(FEntry{ RootTime, Datum });
-		}
+		Entries.Emplace(FEntry{ RootTime, Datum });
 	}
 
 	/** Clears all entries */
@@ -127,91 +43,9 @@ struct TMovieSceneTimeArray
 		return TArrayView<const FEntry>(Entries);
 	}
 
-private:
-
-	void UpdateCachedInverseTransform()
-	{
-		// Build the total transform from root to local, and keep track of the loop indices.
-		FMovieSceneWarpCounter LoopCounter;
-		FMovieSceneSequenceTransform RootToLocalTransform;
-
-		// TODO: It seems a bit silly to store this as a TransformStack when we need to do the same work to inverse it.
-
-		for (int32 Index = 0; Index < TransformStack.Num(); ++Index)
-		{
-			RootToLocalTransform.NestedTransforms.Add(FMovieSceneNestedSequenceTransform(TransformStack[Index].LinearTransform, TransformStack[Index].Warping));
-			if (TransformStack[Index].Warping.IsValid())
-			{
-				LoopCounter.AddWarpingLevel(TransformStack[Index].WarpCounter);
-			}
-			else
-			{
-				LoopCounter.AddNonWarpingLevel();
-			}
-		}
-
-		CachedInverseTransform = RootToLocalTransform.InverseFromLoop(LoopCounter);
-	}
-
 private:
 
 	using FEntry = TMovieSceneTimeArrayEntry<DataType>;
 
 	TArray<FEntry> Entries;
-
-	struct FTransformStep
-	{
-		FMovieSceneTimeTransform LinearTransform;
-		FMovieSceneTimeWarping Warping;
-		uint32 WarpCounter = FMovieSceneTimeWarping::InvalidWarpCount;
-	};
-	TArray<FTransformStep> TransformStack;
-
-	FMovieSceneSequenceTransform CachedInverseTransform;
-};
-
-template<typename DataType>
-struct TMovieSceneTimeArrayTransformScope
-{
-	TMovieSceneTimeArrayTransformScope(TMovieSceneTimeArray<DataType>& InArray, FFrameNumber Offset, float TimeScale = 1.f)
-		: Array(&InArray)
-	{
-		Array->PushTransform(Offset, TimeScale);
-	}
-
-	TMovieSceneTimeArrayTransformScope(TMovieSceneTimeArray<DataType>& InArray, FMovieSceneTimeTransform&& InTransform)
-		: Array(&InArray)
-	{
-		Array->PushTransform(InTransform);
-	}
-
-	TMovieSceneTimeArrayTransformScope(TMovieSceneTimeArray<DataType>& InArray, FFrameNumber FirstLoopStartFrameOffset, FFrameNumber LoopStart, FFrameNumber LoopEnd, float TimeScale = 1.f)
-		: Array(&InArray)
-	{
-		Array->PushTransform(FirstLoopStartFrameOffset, LoopStart, LoopEnd, TimeScale);
-	}
-
-	TMovieSceneTimeArrayTransformScope(TMovieSceneTimeArray<DataType>& InArray, FMovieSceneTimeTransform&& InTransform, FMovieSceneTimeWarping&& InWarping)
-		: Array(&InArray)
-	{
-		Array->PushTransform(InTransform, InWarping);
-	}
-
-	TMovieSceneTimeArrayTransformScope(TMovieSceneTimeArrayTransformScope&& Other)
-		: Array(Other.Array)
-	{
-		Other.Array = nullptr;
-	}
-
-	~TMovieSceneTimeArrayTransformScope()
-	{
-		if (Array)
-		{
-			Array->PopTransform();
-		}
-	}
-
-private:
-	TMovieSceneTimeArray<DataType>* Array = nullptr;
 };
-
diff --git a/Engine/Source/Runtime/MovieScene/Public/Evaluation/MovieSceneTimeTransform.h b/Engine/Source/Runtime/MovieScene/Public/Evaluation/MovieSceneTimeTransform.h
--- a/Engine/Source/Runtime/MovieScene/Public/Evaluation/MovieSceneTimeTransform.h
+++ b/Engine/Source/Runtime/MovieScene/Public/Evaluation/MovieSceneTimeTransform.h
@@ -40,7 +40,9 @@ struct FMovieSceneTimeTransform
 	explicit FMovieSceneTimeTransform(FFrameTime InOffset, float InTimeScale = 1.f)
 		: TimeScale(InTimeScale)
 		, Offset(InOffset)
-	{}
+	{
+		checkf(!FMath::IsNearlyZero(TimeScale), TEXT("It is invalid to create a linear transformation with a zero timescale. This case should be handled externally by FMovieSceneTimeWarpVariant and FLAG_Zero."))
+	}
 
 	friend bool operator==(const FMovieSceneTimeTransform& A, const FMovieSceneTimeTransform& B)
 	{
@@ -121,56 +123,34 @@ inline FFrameTime& operator*=(FFrameTime& InTime, const FMovieSceneTimeTransform
 template<typename T>
 TRange<T> operator*(const TRange<T>& LHS, const FMovieSceneTimeTransform& RHS)
 {
-	//  Special case large (inf) timescale to return open range.
-	if (!FMath::IsFinite(RHS.TimeScale))
-		return TRange<T>(TRangeBound<T>::Open(), TRangeBound<T>::Open());
-
-	// Special case 0-timescale below. We force multiplies with zero timescale to always have inclusive boundaries to prevent no frames from being present.
-	TRangeBound<T> SourceLower = LHS.GetLowerBound();
-	TRangeBound<T> TransformedLower =
-		SourceLower.IsOpen() ? 
-			TRangeBound<T>() : 
-			SourceLower.IsInclusive() || FMath::IsNearlyZero(RHS.TimeScale)  ?
-				TRangeBound<T>::Inclusive(SourceLower.GetValue() * RHS) :
-				TRangeBound<T>::Exclusive(SourceLower.GetValue() * RHS);
-
-	TRangeBound<T> SourceUpper = LHS.GetUpperBound();
-	TRangeBound<T> TransformedUpper =
-		SourceUpper.IsOpen() ? 
-			TRangeBound<T>() : 
-			SourceUpper.IsInclusive() || FMath::IsNearlyZero(RHS.TimeScale) ?
-				TRangeBound<T>::Inclusive(SourceUpper.GetValue() * RHS) :
-				TRangeBound<T>::Exclusive(SourceUpper.GetValue() * RHS);
-
-	return TRange<T>(TransformedLower, TransformedUpper);
+	TRange<T> Result = LHS;
+
+	if (!Result.GetLowerBound().IsOpen())
+	{
+		Result.SetLowerBoundValue(Result.GetLowerBoundValue() * RHS);
+	}
+	if (!Result.GetUpperBound().IsOpen())
+	{
+		Result.SetUpperBoundValue(Result.GetUpperBoundValue() * RHS);
+	}
+
+	return Result;
 }
 
 inline TRange<FFrameNumber> operator*(const TRange<FFrameNumber>& LHS, const FMovieSceneTimeTransform& RHS)
 {
-	// Special case large (inf) timescale to return open range.
-	if (!FMath::IsFinite(RHS.TimeScale))
+	TRange<FFrameNumber> Result = LHS;
+
+	if (!Result.GetLowerBound().IsOpen())
 	{
-		return TRange<FFrameNumber>(TRangeBound<FFrameNumber>::Open(), TRangeBound<FFrameNumber>::Open());
+		Result.SetLowerBoundValue((Result.GetLowerBoundValue() * RHS).FloorToFrame());
 	}
-	
-	// Special case 0-timescale below. We force multiplies with zero timescale to always have inclusive boundaries to prevent no frames from being present.
-	TRangeBound<FFrameNumber> SourceLower = LHS.GetLowerBound();
-	TRangeBound<FFrameNumber> TransformedLower =
-		SourceLower.IsOpen() ? 
-			TRangeBound<FFrameNumber>() : 
-			SourceLower.IsInclusive() || FMath::IsNearlyZero(RHS.TimeScale) ?
-				TRangeBound<FFrameNumber>::Inclusive((SourceLower.GetValue() * RHS).FloorToFrame()) :
-				TRangeBound<FFrameNumber>::Exclusive((SourceLower.GetValue() * RHS).FloorToFrame());
-
-	TRangeBound<FFrameNumber> SourceUpper = LHS.GetUpperBound();
-	TRangeBound<FFrameNumber> TransformedUpper =
-		SourceUpper.IsOpen() ? 
-			TRangeBound<FFrameNumber>() : 
-			SourceUpper.IsInclusive() || FMath::IsNearlyZero(RHS.TimeScale) ?
-				TRangeBound<FFrameNumber>::Inclusive((SourceUpper.GetValue() * RHS).FloorToFrame()) :
-				TRangeBound<FFrameNumber>::Exclusive((SourceUpper.GetValue() * RHS).FloorToFrame());
-
-	return TRange<FFrameNumber>(TransformedLower, TransformedUpper);
+	if (!Result.GetUpperBound().IsOpen())
+	{
+		Result.SetUpperBoundValue((Result.GetUpperBoundValue() * RHS).FloorToFrame());
+	}
+
+	return Result;
 }
 
 /**
diff --git a/Engine/Source/Runtime/MovieScene/Public/Evaluation/MovieSceneTimeWarping.h b/Engine/Source/Runtime/MovieScene/Public/Evaluation/MovieSceneTimeWarping.h
--- a/Engine/Source/Runtime/MovieScene/Public/Evaluation/MovieSceneTimeWarping.h
+++ b/Engine/Source/Runtime/MovieScene/Public/Evaluation/MovieSceneTimeWarping.h
@@ -22,6 +22,7 @@
 USTRUCT()
 struct FMovieSceneTimeWarping
 {
+	UE_DEPRECATED(5.5, "Please update your code to the new FMovieSceneSequenceTransform API. The closest analogue is FMovieSceneTimeWarpLoop.")
 	MOVIESCENE_API static const uint32 InvalidWarpCount;
 
 	GENERATED_BODY()
@@ -41,16 +42,19 @@ struct FMovieSceneTimeWarping
 	/**
 	 * Returns the length of the warping.
 	 */
+	UE_DEPRECATED(5.5, "Please update your code to the new FMovieSceneSequenceTransform API. The closest analogue is FMovieSceneTimeWarpLoop.")
 	FFrameNumber Length() const { return End - Start; }
 
 	/**
 	 * Returns whether this warping transform is doing anything.
 	 */
+	UE_DEPRECATED(5.5, "Please update your code to the new FMovieSceneSequenceTransform API. The closest analogue is FMovieSceneTimeWarpLoop.")
 	bool IsValid() const { return End > Start; }
 
 	/**
 	 * Returns a range that encompasses the whole warping time span.
 	 */
+	UE_DEPRECATED(5.5, "Please update your code to the new FMovieSceneSequenceTransform API. The closest analogue is FMovieSceneTimeWarpLoop.")
 	TRange<FFrameTime> GetRange() const
 	{
 		const TRangeBound<FFrameTime> OutLower(TRangeBound<FFrameTime>::Inclusive(Start));
@@ -61,71 +65,27 @@ struct FMovieSceneTimeWarping
 	/**
 	 * Returns a transformation that takes us from a local time into a given loop back to the root time.
 	 */
-	FMovieSceneTimeTransform InverseFromWarp(uint32 WarpCount) const
-	{
-		check(IsValid() || WarpCount == InvalidWarpCount);
-
-		if (WarpCount == InvalidWarpCount || WarpCount == 0)
-		{
-			return FMovieSceneTimeTransform();
-		}
-		else
-		{
-			const FFrameNumber WarpLength = Length();
-			const FFrameNumber WarpsOffset = WarpLength * (float)WarpCount;
-			return FMovieSceneTimeTransform(WarpsOffset);
-		}
-	}
+	UE_DEPRECATED(5.5, "Please update your code to the new FMovieSceneSequenceTransform API. The closest analogue is FMovieSceneTimeWarpLoop.")
+	FMovieSceneTimeTransform InverseFromWarp(uint32 WarpCount) const;
 
 	/**
 	 * Transforms the given frame and returns the warped frame time along with the warp index we ended up in.
 	 */
-	void TransformFrame(FFrameNumber InFrame, FFrameNumber& OutFrame, uint32& OutWarpIndex) const
-	{
-		checkSlow(IsValid());
-		const FFrameNumber WarpLength = Length();
-		OutWarpIndex = 0;
-		FFrameNumber TempFrame = InFrame - Start;
-		while (TempFrame >= WarpLength)
-		{
-			TempFrame -= WarpLength;
-			++OutWarpIndex;
-		}
-		OutFrame = TempFrame + Start;
-	}
+	UE_DEPRECATED(5.5, "Please update your code to the new FMovieSceneSequenceTransform API. The closest analogue is FMovieSceneTimeWarpLoop.")
+	void TransformFrame(FFrameNumber InFrame, FFrameNumber& OutFrame, uint32& OutWarpIndex) const;
 
 	/**
 	 * Transforms the given time and returns the warped time along with the warp index we ended up in.
 	 */
-	void TransformTime(FFrameTime InTime, FFrameTime& OutTime, uint32& OutWarpIndex) const
-	{
-		checkSlow(IsValid());
-		const FFrameTime WarpLength = Length();
-		OutWarpIndex = 0;
-		FFrameTime TempTime = InTime - Start;
-		while (TempTime >= WarpLength)
-		{
-			TempTime = TempTime - WarpLength;
-			++OutWarpIndex;
-		}
-		OutTime = TempTime + Start;
-	}
+	UE_DEPRECATED(5.5, "Please update your code to the new FMovieSceneSequenceTransform API. The closest analogue is FMovieSceneTimeWarpLoop.")
+	void TransformTime(FFrameTime InTime, FFrameTime& OutTime, uint32& OutWarpIndex) const;
 
 	/**
 	 * Transforms the given time by warping it by a specific warp count, regardless of how many warps are
 	 * needed in theory to stay in the warping range.
 	 */
-	void TransformTimeSpecific(FFrameTime InTime, uint32 WarpCount, FFrameTime& OutTime) const
-	{
-		checkSlow(IsValid());
-		const FFrameTime WarpLength = Length();
-		FFrameTime TempTime = InTime - Start;
-		for (uint32 WarpIndex = 0; WarpIndex < WarpCount; ++WarpIndex)
-		{
-			TempTime = TempTime - WarpLength;
-		}
-		OutTime = TempTime + Start;
-	}
+	UE_DEPRECATED(5.5, "Please update your code to the new FMovieSceneSequenceTransform API. The closest analogue is FMovieSceneTimeWarpLoop.")
+	void TransformTimeSpecific(FFrameTime InTime, uint32 WarpCount, FFrameTime& OutTime) const;
 
 	/** 
 	 * Transforms the given range in a "naive" way, i.e. its lower and upper bounds are transformed
@@ -133,6 +93,7 @@ struct FMovieSceneTimeWarping
 	 * the output lower bound could be greater than the output upper bound, like, for instance, in the case
 	 * of the input range starting near the end of a loop, and ending near the beginning of another loop.
 	 */
+	UE_DEPRECATED(5.5, "Please update your code to the new FMovieSceneSequenceTransform API. The closest analogue is FMovieSceneTimeWarpLoop.")
 	TRange<FFrameTime> TransformRangePure(const TRange<FFrameTime>& Range) const;
 
 	/**
@@ -141,6 +102,7 @@ struct FMovieSceneTimeWarping
 	 * 3 loops, the output range will start in the middle of the first loop, and last 3 times as long as
 	 * the Length of this time warp.
 	 */
+	UE_DEPRECATED(5.5, "Please update your code to the new FMovieSceneSequenceTransform API. The closest analogue is FMovieSceneTimeWarpLoop.")
 	TRange<FFrameTime> TransformRangeUnwarped(const TRange<FFrameTime>& Range) const;
 
 	/**
@@ -150,13 +112,16 @@ struct FMovieSceneTimeWarping
 	 * If it didn't cover a full loop, return the transformed range, which should be a subset of the full
 	 * loop range mentioned above.
 	 */
+	UE_DEPRECATED(5.5, "Please update your code to the new FMovieSceneSequenceTransform API. The closest analogue is FMovieSceneTimeWarpLoop.")
 	TRange<FFrameTime> TransformRangeConstrained(const TRange<FFrameTime>& Range) const;
 
+	UE_DEPRECATED(5.5, "Please update your code to the new FMovieSceneSequenceTransform API. The closest analogue is FMovieSceneTimeWarpLoop.")
 	friend bool operator==(const FMovieSceneTimeWarping& A, const FMovieSceneTimeWarping& B)
 	{
 		return A.Start == B.Start && A.End == B.End;
 	}
 
+	UE_DEPRECATED(5.5, "Please update your code to the new FMovieSceneSequenceTransform API. The closest analogue is FMovieSceneTimeWarpLoop.")
 	friend bool operator!=(const FMovieSceneTimeWarping& A, const FMovieSceneTimeWarping& B)
 	{
 		return A.Start != B.Start || A.End != B.End;
@@ -169,6 +134,8 @@ struct FMovieSceneTimeWarping
 	FFrameNumber End;
 };
 
+PRAGMA_DISABLE_DEPRECATION_WARNINGS
+
 inline FFrameNumber operator*(FFrameNumber InFrame, const FMovieSceneTimeWarping& RHS)
 {
 	checkSlow(RHS.IsValid());
@@ -195,6 +162,62 @@ inline FFrameTime& operator*=(FFrameTime& InTime, const FMovieSceneTimeWarping&
 	return InTime;
 }
 
+inline FMovieSceneTimeTransform FMovieSceneTimeWarping::InverseFromWarp(uint32 WarpCount) const
+{
+	check(IsValid() || WarpCount == InvalidWarpCount);
+
+	if (WarpCount == InvalidWarpCount || WarpCount == 0)
+	{
+		return FMovieSceneTimeTransform();
+	}
+	else
+	{
+		const FFrameNumber WarpLength = Length();
+		const FFrameNumber WarpsOffset = WarpLength * (float)WarpCount;
+		return FMovieSceneTimeTransform(WarpsOffset);
+	}
+}
+
+inline void FMovieSceneTimeWarping::TransformFrame(FFrameNumber InFrame, FFrameNumber& OutFrame, uint32& OutWarpIndex) const
+{
+	checkSlow(IsValid());
+	const FFrameNumber WarpLength = Length();
+	OutWarpIndex = 0;
+	FFrameNumber TempFrame = InFrame - Start;
+	while (TempFrame >= WarpLength)
+	{
+		TempFrame -= WarpLength;
+		++OutWarpIndex;
+	}
+	OutFrame = TempFrame + Start;
+}
+
+inline void FMovieSceneTimeWarping::TransformTime(FFrameTime InTime, FFrameTime& OutTime, uint32& OutWarpIndex) const
+{
+	checkSlow(IsValid());
+	const FFrameTime WarpLength = Length();
+	OutWarpIndex = 0;
+	FFrameTime TempTime = InTime - Start;
+	while (TempTime >= WarpLength)
+	{
+		TempTime = TempTime - WarpLength;
+		++OutWarpIndex;
+	}
+	OutTime = TempTime + Start;
+}
+
+inline void FMovieSceneTimeWarping::TransformTimeSpecific(FFrameTime InTime, uint32 WarpCount, FFrameTime& OutTime) const
+{
+	checkSlow(IsValid());
+	const FFrameTime WarpLength = Length();
+	FFrameTime TempTime = InTime - Start;
+	for (uint32 WarpIndex = 0; WarpIndex < WarpCount; ++WarpIndex)
+	{
+		TempTime = TempTime - WarpLength;
+	}
+	OutTime = TempTime + Start;
+}
+
 inline TRange<FFrameTime> FMovieSceneTimeWarping::TransformRangePure(const TRange<FFrameTime>& Range) const
 {
 	checkSlow(IsValid());
@@ -285,3 +308,5 @@ inline FString LexToString(const FMovieSceneTimeWarping& InWarping)
 {
 	return *FString::Printf(TEXT("[ %+i ⟳ %+i ]"), InWarping.Start.Value, InWarping.End.Value);
 }
+
+PRAGMA_ENABLE_DEPRECATION_WARNINGS
\ No newline at end of file
diff --git a/Engine/Source/Runtime/MovieScene/Public/MovieSceneClipboard.h b/Engine/Source/Runtime/MovieScene/Public/MovieSceneClipboard.h
--- a/Engine/Source/Runtime/MovieScene/Public/MovieSceneClipboard.h
+++ b/Engine/Source/Runtime/MovieScene/Public/MovieSceneClipboard.h
@@ -191,11 +191,12 @@ public:
 
 	/** Create a key track that wraps the specified key type */
 	template<typename KeyType>
-	static FMovieSceneClipboardKeyTrack Create(FName InName)
+	static FMovieSceneClipboardKeyTrack Create(FName InName, FFrameNumber KeyOffset = 0)
 	{
 		FMovieSceneClipboardKeyTrack Track;
 		Track.TypeName = MovieSceneClipboard::GetKeyTypeName<KeyType>();
 		Track.Name = MoveTemp(InName);
+		Track.KeyOffset = KeyOffset;
 		return Track;
 	}
 
@@ -220,6 +221,8 @@ public:
 	template<typename KeyType>
 	void AddKey(FFrameNumber Time, KeyType Value)
 	{
+		Time -= KeyOffset;
+
 		checkf(IsKeyOfType<KeyType>(), TEXT("Unable to add a key of a different value type to the track"));
 		Keys.Add(FMovieSceneClipboardKey(Time, MoveTemp(Value)));
 	}
@@ -268,6 +271,9 @@ private:
 
 	/** Generic name of this track (generally the name of a key area within a track e.g. Location.X) */
 	FName Name;
+
+	/** Optional key offset that is applied to all keys added to the clipboard */
+	FFrameNumber KeyOffset;
 };
 
 /** Structure representing an environment a clipboard applies to */
@@ -330,6 +336,9 @@ class FMovieSceneClipboardBuilder
 {
 public:
 
+	/** Optional key offset that is applied to all keys added to the clipboard */
+	FFrameNumber KeyOffset;
+
 	/** Generate a clipboard for the current state of this builder, resetting the builder back to its default state */
 	MOVIESCENE_API FMovieSceneClipboard Commit(TOptional<FFrameNumber> CopyRelativeTo);
 
diff --git a/Engine/Source/Runtime/MovieScene/Public/MovieSceneFwd.h b/Engine/Source/Runtime/MovieScene/Public/MovieSceneFwd.h
--- a/Engine/Source/Runtime/MovieScene/Public/MovieSceneFwd.h
+++ b/Engine/Source/Runtime/MovieScene/Public/MovieSceneFwd.h
@@ -17,10 +17,6 @@
 	#define UE_MOVIESCENE_EXPENSIVE_CONSISTENCY_CHECKS 0
 #endif
 
-#ifndef UE_MOVIESCENE_TIME_WARP_DEPRECATED
-	#define UE_MOVIESCENE_TIME_WARP_DEPRECATED 0
-#endif
-
 #ifndef UE_MOVIESCENE_ENTITY_DEBUG
 	#define UE_MOVIESCENE_ENTITY_DEBUG !UE_BUILD_SHIPPING
 #endif
diff --git a/Engine/Source/Runtime/MovieScene/Public/MovieSceneKeyProxy.h b/Engine/Source/Runtime/MovieScene/Public/MovieSceneKeyProxy.h
--- a/Engine/Source/Runtime/MovieScene/Public/MovieSceneKeyProxy.h
+++ b/Engine/Source/Runtime/MovieScene/Public/MovieSceneKeyProxy.h
@@ -47,7 +47,7 @@ protected:
 	 * @param InOutTime            Time to move the underlying key to. If the section is locked, this will receive the existing key's time without changing the underlying time.
 	 */
 	template<typename ChannelType, typename ValueType>
-	void OnProxyValueChanged(TMovieSceneChannelHandle<ChannelType> InChannelHandle, UMovieSceneSection* InSection, FKeyHandle InKeyHandle, ValueType& InOutValue, FFrameNumber& InOutTime);
+	void OnProxyValueChanged(TMovieSceneChannelHandle<ChannelType> InChannelHandle, UMovieSceneSignedObject* InSignedObject, FKeyHandle InKeyHandle, ValueType& InOutValue, FFrameNumber& InOutTime);
 
 	/**
 	 * Implementation function that retrieves the underlying key time/value and applies then to the specified value and time parameters. Normally called once per tick.
@@ -64,31 +64,43 @@ protected:
 
 
 template<typename ChannelType, typename ValueType>
-void IMovieSceneKeyProxy::OnProxyValueChanged(TMovieSceneChannelHandle<ChannelType> InChannelHandle, UMovieSceneSection* InSection, FKeyHandle InKeyHandle, ValueType& InOutValue, FFrameNumber& InOutTime)
+void IMovieSceneKeyProxy::OnProxyValueChanged(TMovieSceneChannelHandle<ChannelType> InChannelHandle, UMovieSceneSignedObject* InSignedObject, FKeyHandle InKeyHandle, ValueType& InOutValue, FFrameNumber& InOutTime)
 {
 	auto* Channel = InChannelHandle.Get();
-	if (!Channel || !InSection)
+	if (!Channel || !InSignedObject)
 	{
 		return;
 	}
 
+	UMovieSceneSection* Section = Cast<UMovieSceneSection>(InSignedObject);
+	if (!Section)
+	{
+		Section = InSignedObject->GetTypedOuter<UMovieSceneSection>();
+	}
+
 	auto ChannelData = Channel->GetData();
 
 	int32 KeyIndex = ChannelData.GetIndex(InKeyHandle);
 	if (KeyIndex != INDEX_NONE)
 	{
-		// If we have no section, or it's locked, don't let the user change the value
-		if (!InSection || !InSection->TryModify())
+		// If we have no signed object, or it's locked, don't let the user change the value
+		if (!InSignedObject || (Section && !Section->TryModify()))
 		{
 			InOutTime  = ChannelData.GetTimes()[KeyIndex];
 			InOutValue = ChannelData.GetValues()[KeyIndex];
 		}
 		else
 		{
+			InSignedObject->Modify();
+
 			ChannelData.GetValues()[KeyIndex] = InOutValue;
 
 			ChannelData.MoveKey(KeyIndex, InOutTime);
-			InSection->ExpandToFrame(InOutTime);
+
+			if (Section)
+			{
+				Section->ExpandToFrame(InOutTime);
+			}
 		}
 		Channel->PostEditChange();
 	}
diff --git a/Engine/Source/Runtime/MovieScene/Public/MovieSceneSection.h b/Engine/Source/Runtime/MovieScene/Public/MovieSceneSection.h
--- a/Engine/Source/Runtime/MovieScene/Public/MovieSceneSection.h
+++ b/Engine/Source/Runtime/MovieScene/Public/MovieSceneSection.h
@@ -294,7 +294,13 @@ public:
 	 */
 	void ExpandToFrame(FFrameNumber InFrame)
 	{
+		FFrameNumber StartOffset = HasStartFrame() ? FMath::Min(GetInclusiveStartFrame(), InFrame) - GetInclusiveStartFrame() : 0;
+
 		SetRange(TRange<FFrameNumber>::Hull(GetRange(), TRange<FFrameNumber>::Inclusive(InFrame, InFrame)));
+		if (StartOffset != 0)
+		{
+			FixupRelativeKeyframes(StartOffset);
+		}
 	}
 
 	/**
@@ -693,6 +699,8 @@ protected:
 	virtual void OnMoved(int32 DeltaTime) {}
 	virtual void OnDilated(float DilationFactor, FFrameNumber Origin) {}
 
+	MOVIESCENE_API void FixupRelativeKeyframes(FFrameNumber Offset);
+
 	MOVIESCENE_API bool ShouldUpgradeEntityData(FArchive& Ar, FMovieSceneEvaluationCustomVersion::Type UpgradeVersion) const;
 
 private:
