Fix static destruction crash unregistering console variable sink

It's fine to leak these since the memory is going away anyway.

diff --git a/Engine/Source/Runtime/MovieScene/Private/MovieSceneModule.cpp b/Engine/Source/Runtime/MovieScene/Private/MovieSceneModule.cpp
--- a/Engine/Source/Runtime/MovieScene/Private/MovieSceneModule.cpp
+++ b/Engine/Source/Runtime/MovieScene/Private/MovieSceneModule.cpp
@@ -24,14 +24,9 @@ struct FCachedLegacyConversionFrameRate
 	FCachedLegacyConversionFrameRate()
 		: FrameRate(60000, 1)
 	{
-		DelegateHandle = IConsoleManager::Get().RegisterConsoleVariableSink_Handle(FConsoleCommandDelegate::CreateRaw(this, &FCachedLegacyConversionFrameRate::OnChanged));
-
+		IConsoleManager::Get().RegisterConsoleVariableSink_Handle(FConsoleCommandDelegate::CreateRaw(this, &FCachedLegacyConversionFrameRate::OnChanged));
 		OnChanged();
 	}
-	~FCachedLegacyConversionFrameRate()
-	{
-		IConsoleManager::Get().UnregisterConsoleVariableSink_Handle(DelegateHandle);
-	}
 
 	void OnChanged()
 	{
@@ -39,7 +34,6 @@ struct FCachedLegacyConversionFrameRate
 	}
 
 	FFrameRate FrameRate;
-	FConsoleVariableSinkHandle DelegateHandle;
 };
 
 FFrameRate GetLegacyConversionFrameRate()
