Sequencer: Fixed inner playback range start not being accounted for in transforms

diff --git a/Engine/Source/Runtime/MovieScene/Private/Sections/MovieSceneSectionTimingParameters.cpp b/Engine/Source/Runtime/MovieScene/Private/Sections/MovieSceneSectionTimingParameters.cpp
--- a/Engine/Source/Runtime/MovieScene/Private/Sections/MovieSceneSectionTimingParameters.cpp
+++ b/Engine/Source/Runtime/MovieScene/Private/Sections/MovieSceneSectionTimingParameters.cpp
@@ -96,7 +96,7 @@ FMovieSceneSequenceTransform FMovieSceneSectionTimingParametersFrames::MakeTrans
 	FFrameNumber LoopOffset(bLoop ? FirstLoopStartOffset.Value : 0);
 
 	// Start offset
-	Result.Add(FMovieSceneTimeTransform(InnerStartOffset + LoopOffset));
+	Result.Add(FMovieSceneTimeTransform(StartTime + LoopOffset));
 
 	// ----------------------------------------------------------------------------
 	// Looping or clamping
