Sequencer: attempt to fix check for leaking shared playback states

diff --git a/Engine/Source/Runtime/MovieScene/Private/EntitySystem/MovieSceneSequenceInstance.cpp b/Engine/Source/Runtime/MovieScene/Private/EntitySystem/MovieSceneSequenceInstance.cpp
--- a/Engine/Source/Runtime/MovieScene/Private/EntitySystem/MovieSceneSequenceInstance.cpp
+++ b/Engine/Source/Runtime/MovieScene/Private/EntitySystem/MovieSceneSequenceInstance.cpp
@@ -118,9 +118,11 @@ void FSequenceInstance::Initialize()
 
 FSequenceInstance::~FSequenceInstance()
 {
+	const int32 RefCount = SharedPlaybackState.GetSharedReferenceCount();
 	ensureAlwaysMsgf(
-			RootInstanceHandle != InstanceHandle || SharedPlaybackState.IsUnique(),
-			TEXT("References to SharedPlaybackState should not be held past the lifetime of its root sequence instance"));
+			RootInstanceHandle != InstanceHandle || RefCount == 1,
+			TEXT("References to SharedPlaybackState should not be held past the lifetime of its root sequence instance. SharedPlaybackState has %d references."),
+			RefCount);
 }
 
 FSequenceInstance::FSequenceInstance(FSequenceInstance&&) = default;
