Sequencer: Fixed non-looping skeletal animation clips not respecting first loop frame offsets

diff --git a/Engine/Source/Runtime/MovieScene/Private/Sections/MovieSceneSectionTimingParameters.cpp b/Engine/Source/Runtime/MovieScene/Private/Sections/MovieSceneSectionTimingParameters.cpp
--- a/Engine/Source/Runtime/MovieScene/Private/Sections/MovieSceneSectionTimingParameters.cpp
+++ b/Engine/Source/Runtime/MovieScene/Private/Sections/MovieSceneSectionTimingParameters.cpp
@@ -30,7 +30,7 @@ FMovieSceneSequenceTransform FMovieSceneSectionTimingParametersSeconds::MakeTran
 	const double StartTime   = InnerStartOffset;
 	const double EndTime     = SourceDuration - InnerEndOffset;
 	const double Duration    = EndTime - StartTime;
-	const double StartOffset = InnerStartOffset + (bLoop ? FirstLoopStartOffset : 0.f);
+	const double StartOffset = InnerStartOffset + FirstLoopStartOffset;
 
 	// Start offset
 	if (!FMath::IsNearlyZero(StartOffset))
