Sequencer: Fixed custom serialization of timewarp variants asserting on load due to IsLiteral being true for custom types

diff --git a/Engine/Source/Runtime/MovieScene/Private/Variants/MovieSceneNumericVariant.cpp b/Engine/Source/Runtime/MovieScene/Private/Variants/MovieSceneNumericVariant.cpp
--- a/Engine/Source/Runtime/MovieScene/Private/Variants/MovieSceneNumericVariant.cpp
+++ b/Engine/Source/Runtime/MovieScene/Private/Variants/MovieSceneNumericVariant.cpp
@@ -164,7 +164,7 @@ bool FMovieSceneNumericVariant::Serialize(FArchive& Ar)
 {
 	if (Ar.IsLoading())
 	{
-		SerializeCustom(Ar, [this](FArchive& Ar, void* InOutData){
+		SerializeCustom(Ar, [this](FArchive& Ar, uint8& TypeBits, void* InOutData){
 
 			uint8 Type = 0;
 			Ar << Type;
@@ -181,13 +181,13 @@ bool FMovieSceneNumericVariant::Serialize(FArchive& Ar)
 				uint8 Payload[6];
 				Ar.Serialize(Payload, sizeof(Payload));
 				FMemory::Memcpy(InOutData, Payload, sizeof(Payload));
-				SetTypeBits(Type);
+				TypeBits = Type;
 			}
 		});
 	}
 	else
 	{
-		SerializeCustom(Ar, [this](FArchive& Ar, void* InOutData){
+		SerializeCustom(Ar, [this](FArchive& Ar, uint8& TypeBits, void* InOutData){
 			uint8 Type = GetTypeBits();
 
 			if (Ar.GetArchiveState().IsSaving())
@@ -212,7 +212,7 @@ bool FMovieSceneNumericVariant::Serialize(FArchive& Ar)
 	return true;
 }
 
-bool FMovieSceneNumericVariant::SerializeCustom(FArchive& Ar, TFunctionRef<void(FArchive&, void*)> InCustomSerializer)
+bool FMovieSceneNumericVariant::SerializeCustom(FArchive& Ar, TFunctionRef<void(FArchive&, uint8&, void*)> InCustomSerializer)
 {
 	if (Ar.IsLoading())
 	{
@@ -225,7 +225,16 @@ bool FMovieSceneNumericVariant::SerializeCustom(FArchive& Ar, TFunctionRef<void(
 		}
 		else
 		{
-			InCustomSerializer(Ar, Data);
+			// First off, initialize this type to be a custom type specified by the tagged bits
+			uint64* Value = reinterpret_cast<uint64*>(Data);
+			*Value = TAGGED_Bits;
+
+			// Pass the type bits through to the serializer
+			uint8 TypeBits = 0;
+			InCustomSerializer(Ar, TypeBits, Data);
+
+			// Assign type bits
+			SetTypeBits(TypeBits);
 		}
 	}
 	else
@@ -243,7 +252,8 @@ bool FMovieSceneNumericVariant::SerializeCustom(FArchive& Ar, TFunctionRef<void(
 		}
 		else
 		{
-			InCustomSerializer(Ar, Data);
+			uint8 TypeBits = GetTypeBits();
+			InCustomSerializer(Ar, TypeBits, Data);
 		}
 	}
 	return true;
diff --git a/Engine/Source/Runtime/MovieScene/Private/Variants/MovieSceneTimeWarpVariant.cpp b/Engine/Source/Runtime/MovieScene/Private/Variants/MovieSceneTimeWarpVariant.cpp
--- a/Engine/Source/Runtime/MovieScene/Private/Variants/MovieSceneTimeWarpVariant.cpp
+++ b/Engine/Source/Runtime/MovieScene/Private/Variants/MovieSceneTimeWarpVariant.cpp
@@ -165,13 +165,18 @@ bool FMovieSceneTimeWarpVariant::Serialize(FArchive& Ar)
 {
 	Ar.UsingCustomVersion(FFortniteMainBranchObjectVersion::GUID);
 
-	return Variant.SerializeCustom(Ar, [this](FArchive& InAr, void* DataPtr){
+	return Variant.SerializeCustom(Ar, [this](FArchive& InAr, uint8& TypeBits, void* DataPtr){
 
 		const bool bIsLoading = InAr.GetArchiveState().IsLoading();
 
 		EMovieSceneTimeWarpType Type = this->GetType();
 		InAr << Type;
 
+		if (bIsLoading)
+		{
+			TypeBits = (uint8)Type - 1;
+		}
+
 		switch(Type)
 		{
 		case EMovieSceneTimeWarpType::Custom:
@@ -227,11 +232,6 @@ bool FMovieSceneTimeWarpVariant::Serialize(FArchive& Ar)
 			FMovieSceneTimeWarpClampFloat::StaticStruct()->SerializeItem(InAr, DataPtr, nullptr);
 			break;
 		}
-
-		if (bIsLoading)
-		{
-			this->Variant.SetTypeBits((uint8)Type - 1);
-		}
 	});
 }
 
diff --git a/Engine/Source/Runtime/MovieScene/Public/Variants/MovieSceneNumericVariant.h b/Engine/Source/Runtime/MovieScene/Public/Variants/MovieSceneNumericVariant.h
--- a/Engine/Source/Runtime/MovieScene/Public/Variants/MovieSceneNumericVariant.h
+++ b/Engine/Source/Runtime/MovieScene/Public/Variants/MovieSceneNumericVariant.h
@@ -163,7 +163,7 @@ public:
 	bool ImportTextItem(const TCHAR*& Buffer, int32 PortFlags, UObject* Parent, FOutputDevice* ErrorText, FArchive* InSerializingArchive = nullptr);
 	void GetPreloadDependencies(TArray<UObject*>& OutDeps);
 
-	bool SerializeCustom(FArchive& Ar, TFunctionRef<void(FArchive&, void*)>);
+	bool SerializeCustom(FArchive& Ar, TFunctionRef<void(FArchive&, uint8&, void*)>);
 
 
 public:
