Sequencer: Fixed time warp looping and boundary display for negative play rates, and time-warp that goes switches between positive and negative

diff --git a/Engine/Source/Runtime/MovieScene/Private/Evaluation/MovieSceneSequenceTransform.cpp b/Engine/Source/Runtime/MovieScene/Private/Evaluation/MovieSceneSequenceTransform.cpp
--- a/Engine/Source/Runtime/MovieScene/Private/Evaluation/MovieSceneSequenceTransform.cpp
+++ b/Engine/Source/Runtime/MovieScene/Private/Evaluation/MovieSceneSequenceTransform.cpp
@@ -28,6 +28,18 @@ namespace UE::MovieScene
 		return Result;
 	}
 
+	void CorrectInsideOutRange(TRange<FFrameTime>& InOutRange)
+	{
+		if (InOutRange.HasUpperBound() && InOutRange.HasLowerBound())
+		{
+			FFrameTime LowerBoundValue = InOutRange.GetLowerBoundValue();
+			if (LowerBoundValue > InOutRange.GetUpperBoundValue())
+			{
+				InOutRange.SetLowerBoundValue(InOutRange.GetUpperBoundValue());
+				InOutRange.SetUpperBoundValue(LowerBoundValue);
+			}
+		}
+	}
 
 	FTransformTimeParams& FTransformTimeParams::HarvestBreadcrumbs(FMovieSceneTransformBreadcrumbs& OutBreadcrumbs)
 	{
@@ -291,15 +303,8 @@ TRange<FFrameTime> FMovieSceneNestedSequenceTransform::ComputeTraversedHull(cons
 	{
 		FMovieSceneTimeTransform LinearTransform(Offset, TimeScale.AsFixedPlayRate());
 
-		TRange<FFrameTime> Result = InRange;
-		if (Result.HasLowerBound())
-		{
-			Result.SetLowerBoundValue(Result.GetLowerBoundValue() * LinearTransform);
-		}
-		if (Result.HasUpperBound())
-		{
-			Result.SetUpperBoundValue(Result.GetUpperBoundValue() * LinearTransform);
-		}
+		TRange<FFrameTime> Result = InRange * LinearTransform;
+		CorrectInsideOutRange(Result);
 		return Result;
 	}
 	else if (Type == EMovieSceneTimeWarpType::FixedTime)
@@ -365,15 +370,15 @@ TRange<FFrameTime> FMovieSceneNestedSequenceTransform::ComputeTraversedHull(cons
 	return OffsetRange;
 }
 
-bool FMovieSceneNestedSequenceTransform::ExtractBoundariesWithinRange(FFrameTime Start, FFrameTime End, const TFunctionRef<bool(FFrameTime)>& InVisitor) const
+bool FMovieSceneNestedSequenceTransform::ExtractBoundariesWithinRange(const TRange<FFrameTime>& Range, const TFunctionRef<bool(FFrameTime)>& InVisitor) const
 {
 	EMovieSceneTimeWarpType WarpType = TimeScale.GetType();
 	switch (WarpType)
 	{
 	case EMovieSceneTimeWarpType::Loop:
-		return TimeScale.AsLoop().ExtractBoundariesWithinRange(Start, End, InVisitor);
+		return TimeScale.AsLoop().ExtractBoundariesWithinRange(Range, InVisitor);
 	case EMovieSceneTimeWarpType::LoopFloat:
-		return TimeScale.AsLoopFloat().ExtractBoundariesWithinRange(Start, End, InVisitor);
+		return TimeScale.AsLoopFloat().ExtractBoundariesWithinRange(Range, InVisitor);
 	default:
 		return false;
 	}
@@ -437,7 +442,7 @@ bool FMovieSceneInverseNestedSequenceTransform::TransformTimeWithinRange(FFrameT
 		break;
 
 	case EMovieSceneTimeWarpType::FrameRate:
-		TransformedTime = ConvertFrameTime(InTime + Offset, TimeScale.AsFrameRate().GetFrameRate(), FFrameRate(1, 1));
+		TransformedTime = ConvertFrameTime(InTime + Offset, FFrameRate(1, 1), TimeScale.AsFrameRate().GetFrameRate());
 		break;
 
 	case EMovieSceneTimeWarpType::Loop:
@@ -886,7 +891,10 @@ FFrameTime FMovieSceneSequenceTransform::TransformTime(FFrameTime InTime, const
 
 TRange<FFrameTime> FMovieSceneSequenceTransform::ComputeTraversedHull(const TRange<FFrameTime>& Range) const
 {
+	using namespace UE::MovieScene;
+
 	TRange<FFrameTime> Result = Range * LinearTransform;
+	CorrectInsideOutRange(Result);
 
 	for (const FMovieSceneNestedSequenceTransform& NestedTransform : NestedTransforms)
 	{
@@ -906,21 +914,23 @@ TRange<FFrameTime> FMovieSceneSequenceTransform::ComputeTraversedHull(const TRan
 
 bool FMovieSceneSequenceTransform::ExtractBoundariesWithinRange(FFrameTime Start, FFrameTime End, const TFunctionRef<bool(FFrameTime)>& InVisitor) const
 {
+	using namespace UE::MovieScene;
+
 	FMovieSceneInverseSequenceTransform Inverse;
 
-	FMovieSceneTransformBreadcrumbs Breadcrumbs;
+	FMovieSceneTransformBreadcrumbs StartBreadcrumbs, EndBreadcrumbs;
 
-	const bool bHasStart = Start != MIN_int32;
-	const bool bHasEnd   = End != MAX_int32;
+	TRange<FFrameTime> TraversedHull = TRange<FFrameTime>::All();
 
-	if (bHasStart)
+	if (Start != MIN_int32)
 	{
-		Start = Start * LinearTransform;
+		TraversedHull.SetLowerBound(Start * LinearTransform);
 	}
-	if (bHasEnd)
+	if (End != MAX_int32)
 	{
-		End = End * LinearTransform;
+		TraversedHull.SetUpperBound(End * LinearTransform);
 	}
+	CorrectInsideOutRange(TraversedHull);
 
 	for (int32 NestedIndex = 0; NestedIndex < NestedTransforms.Num(); ++NestedIndex)
 	{
@@ -929,15 +939,12 @@ bool FMovieSceneSequenceTransform::ExtractBoundariesWithinRange(FFrameTime Start
 		// Find the first transform that has any boundaries
 		if (!NestedTransform.SupportsBoundaries())
 		{
-			Breadcrumbs.AddBreadcrumb(Start + (End - Start) / 2);
-			if (bHasStart)
+			if (NestedTransform.NeedsBreadcrumb())
 			{
-				Start = NestedTransform.TransformTime(Start);
-			}
-			if (bHasEnd)
-			{
-				End = NestedTransform.TransformTime(End);
+				StartBreadcrumbs.AddBreadcrumb(TraversedHull.HasLowerBound() ? TraversedHull.GetLowerBoundValue() : FFrameTime(MIN_int32));
+				EndBreadcrumbs.AddBreadcrumb(TraversedHull.HasUpperBound() ? TraversedHull.GetUpperBoundValue() : FFrameTime(MAX_int32));
 			}
+			TraversedHull = NestedTransform.ComputeTraversedHull(TraversedHull);
 			continue;
 		}
 
@@ -947,17 +954,12 @@ bool FMovieSceneSequenceTransform::ExtractBoundariesWithinRange(FFrameTime Start
 
 		FMovieSceneInverseSequenceTransform ParentToRootTransform = RootToParentTransform.Inverse();
 
-
-		auto VisitWrapper = [&InVisitor, &ParentToRootTransform, &Breadcrumbs](FFrameTime InBoundary)
+		auto VisitWrapper = [&InVisitor, &ParentToRootTransform, &StartBreadcrumbs, &EndBreadcrumbs](FFrameTime InBoundary)
 		{
-			if (TOptional<FFrameTime> RootTime = ParentToRootTransform.TryTransformTime(InBoundary, Breadcrumbs))
-			{
-				return InVisitor(RootTime.GetValue());
-			}
-			return true;
+			return ParentToRootTransform.TransformTimeWithinRange(InBoundary, InVisitor, StartBreadcrumbs, EndBreadcrumbs);
 		};
 
-		return NestedTransform.ExtractBoundariesWithinRange(Start, End, VisitWrapper);
+		return NestedTransform.ExtractBoundariesWithinRange(TraversedHull, VisitWrapper);
 	}
 
 	return false;
diff --git a/Engine/Source/Runtime/MovieScene/Private/Sections/MovieSceneSectionTimingParameters.cpp b/Engine/Source/Runtime/MovieScene/Private/Sections/MovieSceneSectionTimingParameters.cpp
--- a/Engine/Source/Runtime/MovieScene/Private/Sections/MovieSceneSectionTimingParameters.cpp
+++ b/Engine/Source/Runtime/MovieScene/Private/Sections/MovieSceneSectionTimingParameters.cpp
@@ -30,7 +30,13 @@ FMovieSceneSequenceTransform FMovieSceneSectionTimingParametersSeconds::MakeTran
 	const double StartTime   = InnerStartOffset;
 	const double EndTime     = SourceDuration - InnerEndOffset;
 	const double Duration    = EndTime - StartTime;
-	const double StartOffset = InnerStartOffset + FirstLoopStartOffset;
+	double StartOffset = InnerStartOffset + FirstLoopStartOffset;
+
+	// Accomodate negative play rates by playing from the end of the clip
+	if (PlayRate.GetType() == EMovieSceneTimeWarpType::FixedPlayRate && PlayRate.AsFixedPlayRate() < 0.0)
+	{
+		StartOffset += Duration;
+	}
 
 	// Start offset
 	if (!FMath::IsNearlyZero(StartOffset))
@@ -95,8 +101,14 @@ FMovieSceneSequenceTransform FMovieSceneSectionTimingParametersFrames::MakeTrans
 
 	FFrameNumber LoopOffset(bLoop ? FirstLoopStartOffset.Value : 0);
 
+	FFrameNumber NegativeRateOffset = 0;
+	if (PlayRate.GetType() == EMovieSceneTimeWarpType::FixedPlayRate && PlayRate.AsFixedPlayRate() < 0.0)
+	{
+		NegativeRateOffset = Duration;
+	}
+
 	// Start offset
-	Result.Add(FMovieSceneTimeTransform(StartTime + LoopOffset));
+	Result.Add(FMovieSceneTimeTransform(StartTime + LoopOffset + NegativeRateOffset));
 
 	// ----------------------------------------------------------------------------
 	// Looping or clamping
diff --git a/Engine/Source/Runtime/MovieScene/Private/Variants/MovieSceneTimeWarpVariantPayloads.cpp b/Engine/Source/Runtime/MovieScene/Private/Variants/MovieSceneTimeWarpVariantPayloads.cpp
--- a/Engine/Source/Runtime/MovieScene/Private/Variants/MovieSceneTimeWarpVariantPayloads.cpp
+++ b/Engine/Source/Runtime/MovieScene/Private/Variants/MovieSceneTimeWarpVariantPayloads.cpp
@@ -15,17 +15,25 @@ static_assert(std::is_trivially_copyable_v<FMovieSceneTimeWarpClampFloat>, "FMov
 
 FFrameTime FMovieSceneTimeWarpLoop::LoopTime(FFrameTime InTime) const
 {
-	// Maintain subframe
-	InTime.FrameNumber = (InTime.FrameNumber % Duration);
-	return InTime;
+	int32 Unused;
+	return LoopTime(InTime, Unused);
 }
 
 FFrameTime FMovieSceneTimeWarpLoop::LoopTime(FFrameTime InTime, int32& OutLoop) const
 {
-	OutLoop = InTime.FrameNumber.Value / Duration.Value;
+	const int32 Frame = InTime.FrameNumber.Value;
+	const int32 Dur   = Duration.Value;
+
+	// Make sure to compute negative loops correctly by subtracting 1 for any negative time
+	//    This results in the equivalent of floor( double(time) / double(duration) )
+	const uint32 FrameAsBits = *reinterpret_cast<const uint32*>(&Frame);
+	const uint32 SignAsBits  = (FrameAsBits&0x80000000)>>31;    // Yields 0 for +ve or 1 for -ve times by shifting the sign bit
+	const int32  Sign        = -static_cast<int32>(SignAsBits); // Yields 0 for +ve or -1 for -ve times
+
+	OutLoop = Frame / Dur + Sign;
 
 	// Maintain subframe
-	InTime.FrameNumber = InTime.FrameNumber % Duration;
+	InTime.FrameNumber = Frame - Dur*OutLoop;
 	return InTime;
 }
 
@@ -34,23 +42,35 @@ TRange<FFrameTime> FMovieSceneTimeWarpLoop::ComputeTraversedHull(const TRange<FF
 	int32 StartLoop = 0;
 	int32 EndLoop   = 0;
 
-	FFrameTime LoopStart = 0;
-	FFrameTime LoopEnd = Duration;
+	TRangeBound<FFrameTime> LoopStart = TRangeBound<FFrameTime>::Inclusive(0);
+	TRangeBound<FFrameTime> LoopEnd   = TRangeBound<FFrameTime>::Exclusive(Duration);
 
-	if (Range.GetLowerBound().IsOpen() || Range.GetUpperBound().IsOpen())
+	if (Range.IsEmpty())
+	{
+		// Empty range of 0
+		return TRange<FFrameTime>(0, 0);
+	}
+	else if (Range.GetLowerBound().IsOpen() || Range.GetUpperBound().IsOpen())
 	{
 		return TRange<FFrameTime>(LoopStart, LoopEnd);
 	}
 
-	const FFrameTime WarpedStart = LoopTime(Range.GetLowerBoundValue(), StartLoop);
-	const FFrameTime WarpedEnd   = LoopTime(Range.GetUpperBoundValue(), EndLoop);
+	TRangeBound<FFrameTime> WarpedStart = Range.GetLowerBound();
+	TRangeBound<FFrameTime> WarpedEnd   = Range.GetUpperBound();
+
+	WarpedStart.SetValue(LoopTime(WarpedStart.GetValue(), StartLoop));
+	WarpedEnd.SetValue(LoopTime(WarpedEnd.GetValue(), EndLoop));
+
+	// Do not loop exlusive end frames
+	if (WarpedEnd.GetValue() == 0 && WarpedEnd.IsExclusive())
+	{
+		--EndLoop;
+		WarpedEnd = LoopEnd;
+	}
 
 	if (StartLoop == EndLoop)
 	{
-		TRange<FFrameTime> Result = Range;
-		Result.SetLowerBoundValue(WarpedStart);
-		Result.SetUpperBoundValue(WarpedEnd);
-		return Result;
+		return TRange<FFrameTime>(WarpedStart, WarpedEnd);
 	}
 
 	const int32 NumCompleteLoops = EndLoop - StartLoop - 1;
@@ -60,7 +80,7 @@ TRange<FFrameTime> FMovieSceneTimeWarpLoop::ComputeTraversedHull(const TRange<FF
 	}
 
 	// If the range crosses a loop boundary and the end time is > the start time, we have traversed a full loop
-	if (WarpedEnd > WarpedStart)
+	if (WarpedEnd.GetValue() > WarpedStart.GetValue())
 	{
 		return TRange<FFrameTime>(LoopStart, LoopEnd);
 	}
@@ -76,7 +96,6 @@ TOptional<FFrameTime> FMovieSceneTimeWarpLoop::InverseRemapTimeCycled(FFrameTime
 	{
 		int32 HintCycle = 0;
 		FFrameTime LoopedHint = LoopTime(InTimeHint, HintCycle);
-		//FFrameTime Result = LoopTime(InValue);
 
 		FFrameTime Difference(InValue - LoopedHint);
 		int32 DifferenceCycle = 0;
@@ -91,8 +110,6 @@ TOptional<FFrameTime> FMovieSceneTimeWarpLoop::InverseRemapTimeCycled(FFrameTime
 
 bool FMovieSceneTimeWarpLoop::InverseRemapTimeWithinRange(FFrameTime InTime, FFrameTime RangeStart, FFrameTime RangeEnd, const TFunctionRef<bool(FFrameTime)>& VisitorCallback) const
 {
-	ensure(RangeStart < RangeEnd);
-
 	int32 Length = Duration.Value;
 
 	int32 InputLoop = 0;
@@ -103,6 +120,12 @@ bool FMovieSceneTimeWarpLoop::InverseRemapTimeWithinRange(FFrameTime InTime, FFr
 	FFrameTime StartTime   = LoopTime(RangeStart, StartLoop);
 	FFrameTime EndTime     = LoopTime(RangeEnd,   EndLoop);
 
+	if (StartLoop > EndLoop || (StartLoop == EndLoop && EndTime < StartTime) )
+	{
+		Swap(StartLoop, EndLoop);
+		Swap(StartTime, EndTime);
+	}
+
 	int32 LoopIndex = InputLoop;
 	FFrameTime Result  = LoopedInput + FFrameTime(Length*LoopIndex);
 
@@ -138,51 +161,26 @@ bool FMovieSceneTimeWarpLoop::InverseRemapTimeWithinRange(FFrameTime InTime, FFr
 	return true;
 }
 
-bool FMovieSceneTimeWarpLoop::ExtractBoundariesWithinRange(FFrameTime RangeStart, FFrameTime RangeEnd, const TFunctionRef<bool(FFrameTime)>& InVisitor) const
+bool FMovieSceneTimeWarpLoop::ExtractBoundariesWithinRange(const TRange<FFrameTime>& Range, const TFunctionRef<bool(FFrameTime)>& InVisitor) const
 {
-	ensure(!(RangeStart == MIN_int32 && RangeEnd == MAX_int32));
-
-	const int32 Length = Duration.Value;
+	const int32 Start = Range.GetLowerBound().IsClosed() ? Range.GetLowerBoundValue().FrameNumber.Value : MIN_int32;
+	const int32 End   = Range.GetUpperBound().IsClosed() ? Range.GetUpperBoundValue().FrameNumber.Value : MAX_int32;
 
-	if (RangeStart == MIN_int32)
-	{
-		// Start at the end and go back
-		int32 LoopIndex = 0;
-		int32 EndLoop = 0;
-
-		LoopTime(RangeStart, EndLoop);
-		LoopTime(RangeEnd, LoopIndex);
+	int32 LoopIndex = 0;
+	int32 EndLoop   = 0;
 
-		for (; LoopIndex >= EndLoop; --LoopIndex)
-		{
-			FFrameTime Result(Length * LoopIndex);
+	LoopTime(Start, LoopIndex);
+	LoopTime(End, EndLoop);
 
-			if (!InVisitor(Result))
-			{
-				return false;
-			}
-		}
-	}
-	else
+	for ( ; LoopIndex <= EndLoop; ++LoopIndex)
 	{
-		ensure(RangeStart < RangeEnd);
-
-		int32 LoopIndex = 0;
-		int32 EndLoop   = 0;
+		FFrameTime StartResult = FFrameTime(Duration*LoopIndex);
 
-		LoopTime(RangeStart, LoopIndex);
-		LoopTime(RangeEnd,   EndLoop);
-
-		for ( ; LoopIndex < EndLoop + 1; ++LoopIndex)
+		if (StartResult.FrameNumber.Value >= Start)
 		{
-			FFrameTime StartResult = FFrameTime(Length*LoopIndex);
-
-			if (StartResult >= RangeStart)
+			if (!InVisitor(StartResult))
 			{
-				if (!InVisitor(StartResult))
-				{
-					return false;
-				}
+				return false;
 			}
 		}
 	}
@@ -220,16 +218,16 @@ TRange<FFrameTime> FMovieSceneTimeWarpClamp::ComputeTraversedHull(const TRange<F
 
 FFrameTime FMovieSceneTimeWarpLoopFloat::LoopTime(FFrameTime InTime) const
 {
-	// Maintain subframe
-	InTime = FFrameTime::FromDecimal(FMath::Fmod(InTime.AsDecimal(), Duration));
-	return InTime;
+	int32 Unused;
+	return LoopTime(InTime, Unused);
 }
 
 FFrameTime FMovieSceneTimeWarpLoopFloat::LoopTime(FFrameTime InTime, int32& OutLoop) const
 {
-	OutLoop = FMath::FloorToInt(InTime.AsDecimal() / Duration);
-	InTime = FFrameTime::FromDecimal(FMath::Fmod(InTime.AsDecimal(), Duration));
-	return InTime;
+	const double Time = InTime.AsDecimal();
+
+	OutLoop = FMath::FloorToInt(Time / Duration);
+	return FFrameTime::FromDecimal(Time - Duration*OutLoop);
 }
 
 TRange<FFrameTime> FMovieSceneTimeWarpLoopFloat::ComputeTraversedHull(const TRange<FFrameTime>& Range) const
@@ -335,54 +333,25 @@ bool FMovieSceneTimeWarpLoopFloat::InverseRemapTimeWithinRange(FFrameTime InTime
 	return true;
 }
 
-bool FMovieSceneTimeWarpLoopFloat::ExtractBoundariesWithinRange(FFrameTime RangeStart, FFrameTime RangeEnd, const TFunctionRef<bool(FFrameTime)>& InVisitor) const
+bool FMovieSceneTimeWarpLoopFloat::ExtractBoundariesWithinRange(const TRange<FFrameTime>& Range, const TFunctionRef<bool(FFrameTime)>& InVisitor) const
 {
-	ensure(!(RangeStart == MIN_int32 && RangeEnd == MAX_int32));
+	const double Start = Range.GetLowerBound().IsClosed() ? Range.GetLowerBoundValue().AsDecimal() : double(MIN_int32);
+	const double End   = Range.GetUpperBound().IsClosed() ? Range.GetUpperBoundValue().AsDecimal() : double(MAX_int32);
 
-	FFrameTime Length = FFrameTime::FromDecimal(Duration);
+	int32 LoopIndex = FMath::FloorToInt(Start / Duration);
+	int32 EndLoop = FMath::FloorToInt(End / Duration);
 
-	if (RangeStart == MIN_int32)
+	for (; LoopIndex <= EndLoop; ++LoopIndex)
 	{
-		// Start at the end and go back
-		int32 LoopIndex = 0;
-		int32 EndLoop = 0;
-
-		LoopTime(RangeStart, EndLoop);
-		LoopTime(RangeEnd, LoopIndex);
-
-		for (; LoopIndex >= EndLoop; --LoopIndex)
+		const double Result = Duration*LoopIndex;
+		if (Result >= Start)
 		{
-			FFrameTime Result(Length * LoopIndex);
-
-			if (!InVisitor(Result))
+			if (!InVisitor(FFrameTime::FromDecimal(Result)))
 			{
 				return false;
 			}
 		}
 	}
-	else
-	{
-		ensure(RangeStart < RangeEnd);
-
-		int32 LoopIndex = 0;
-		int32 EndLoop   = 0;
-
-		LoopTime(RangeStart, LoopIndex);
-		LoopTime(RangeEnd,   EndLoop);
-
-		for ( ; LoopIndex < EndLoop + 1; ++LoopIndex)
-		{
-			FFrameTime StartResult = FFrameTime(Length*LoopIndex);
-
-			if (StartResult >= RangeStart)
-			{
-				if (!InVisitor(StartResult))
-				{
-					return false;
-				}
-			}
-		}
-	}
 
 	return true;
 }
diff --git a/Engine/Source/Runtime/MovieScene/Public/Evaluation/MovieSceneSequenceTransform.h b/Engine/Source/Runtime/MovieScene/Public/Evaluation/MovieSceneSequenceTransform.h
--- a/Engine/Source/Runtime/MovieScene/Public/Evaluation/MovieSceneSequenceTransform.h
+++ b/Engine/Source/Runtime/MovieScene/Public/Evaluation/MovieSceneSequenceTransform.h
@@ -507,13 +507,12 @@ public:
 	 * invoking Visitor for every boundary that is present within the range. Iteration will continue until
 	 * Visitor returns false, at which point this function itself will return false.
 	 *
-	 * @param UntransformedStart       The inclusive start of the range within which Visitor can be invoked
-	 * @param UntransformedEnd         The inclusive end of the range within which Visitor can be invoked
+	 * @param UntransformedRange       The range within which Visitor can be invoked in untransformed space
 	 * @param Visitor                  A functor to invoke for every boundary that is found within the range. Return value signifies whether iteration continues (true) or not (false)
 	 * 
 	 * @return (true) if there were no boundaries, or Visitor returned true for all encountered boundaries; false otherwise
 	 */
-	MOVIESCENE_API bool ExtractBoundariesWithinRange(FFrameTime UntransformedStart, FFrameTime UntransformedEnd, const TFunctionRef<bool(FFrameTime)>& Visitor) const;
+	MOVIESCENE_API bool ExtractBoundariesWithinRange(const TRange<FFrameTime>& UntransformedRange, const TFunctionRef<bool(FFrameTime)>& Visitor) const;
 
 
 	/**
diff --git a/Engine/Source/Runtime/MovieScene/Public/Variants/MovieSceneTimeWarpVariantPayloads.h b/Engine/Source/Runtime/MovieScene/Public/Variants/MovieSceneTimeWarpVariantPayloads.h
--- a/Engine/Source/Runtime/MovieScene/Public/Variants/MovieSceneTimeWarpVariantPayloads.h
+++ b/Engine/Source/Runtime/MovieScene/Public/Variants/MovieSceneTimeWarpVariantPayloads.h
@@ -78,7 +78,7 @@ struct FMovieSceneTimeWarpLoop
 	TRange<FFrameTime> ComputeTraversedHull(const TRange<FFrameTime>& Range) const;
 	TOptional<FFrameTime> InverseRemapTimeCycled(FFrameTime InValue, FFrameTime InTimeHint, const UE::MovieScene::FInverseTransformTimeParams& Params) const;
 	bool InverseRemapTimeWithinRange(FFrameTime InTime, FFrameTime RangeStart, FFrameTime RangeEnd, const TFunctionRef<bool(FFrameTime)>& VisitorCallback) const;
-	bool ExtractBoundariesWithinRange(FFrameTime RangeStart, FFrameTime RangeEnd, const TFunctionRef<bool(FFrameTime)>& InVisitor) const;
+	bool ExtractBoundariesWithinRange(const TRange<FFrameTime>& Range, const TFunctionRef<bool(FFrameTime)>& InVisitor) const;
 };
 
 
@@ -118,7 +118,7 @@ struct FMovieSceneTimeWarpLoopFloat
 	TRange<FFrameTime> ComputeTraversedHull(const TRange<FFrameTime>& Range) const;
 	TOptional<FFrameTime> InverseRemapTimeCycled(FFrameTime InValue, FFrameTime InTimeHint, const UE::MovieScene::FInverseTransformTimeParams& Params) const;
 	bool InverseRemapTimeWithinRange(FFrameTime InTime, FFrameTime RangeStart, FFrameTime RangeEnd, const TFunctionRef<bool(FFrameTime)>& VisitorCallback) const;
-	bool ExtractBoundariesWithinRange(FFrameTime RangeStart, FFrameTime RangeEnd, const TFunctionRef<bool(FFrameTime)>& InVisitor) const;
+	bool ExtractBoundariesWithinRange(const TRange<FFrameTime>& Range, const TFunctionRef<bool(FFrameTime)>& InVisitor) const;
 };
 
 
diff --git a/Engine/Source/Runtime/MovieSceneTracks/Private/Sections/MovieSceneSkeletalAnimationSection.cpp b/Engine/Source/Runtime/MovieSceneTracks/Private/Sections/MovieSceneSkeletalAnimationSection.cpp
--- a/Engine/Source/Runtime/MovieSceneTracks/Private/Sections/MovieSceneSkeletalAnimationSection.cpp
+++ b/Engine/Source/Runtime/MovieSceneTracks/Private/Sections/MovieSceneSkeletalAnimationSection.cpp
@@ -529,7 +529,7 @@ void UMovieSceneSkeletalAnimationSection::PostEditChangeProperty(FPropertyChange
 		{
 			float NewPlayRate = Params.PlayRate.AsFixedPlayRateFloat();
 
-			if (!FMath::IsNearlyZero(NewPlayRate))
+			if (NewPlayRate > KINDA_SMALL_NUMBER)
 			{
 				float CurrentDuration = UE::MovieScene::DiscreteSize(GetRange());
 				float NewDuration = CurrentDuration * (PreviousPlayRate / NewPlayRate);
