Sequencer: Fix initialization

#rnx

diff --git a/Engine/Source/Runtime/MovieScene/Public/Variants/MovieSceneTimeWarpVariantPayloads.h b/Engine/Source/Runtime/MovieScene/Public/Variants/MovieSceneTimeWarpVariantPayloads.h
--- a/Engine/Source/Runtime/MovieScene/Public/Variants/MovieSceneTimeWarpVariantPayloads.h
+++ b/Engine/Source/Runtime/MovieScene/Public/Variants/MovieSceneTimeWarpVariantPayloads.h
@@ -43,7 +43,7 @@ struct FMovieSceneFixedPlayRateStruct
 
 	/** The play rate */
 	UPROPERTY()
-	double PlayRate;
+	double PlayRate = 1.0;
 };
 
 
