Sequencer: Fixed determinism frames and dissection for looping sequences and sequences with time-warping

diff --git a/Engine/Source/Runtime/MovieScene/Private/Compilation/MovieSceneCompiledDataManager.cpp b/Engine/Source/Runtime/MovieScene/Private/Compilation/MovieSceneCompiledDataManager.cpp
--- a/Engine/Source/Runtime/MovieScene/Private/Compilation/MovieSceneCompiledDataManager.cpp
+++ b/Engine/Source/Runtime/MovieScene/Private/Compilation/MovieSceneCompiledDataManager.cpp
@@ -1101,6 +1101,16 @@ void UMovieSceneCompiledDataManager::CompileSubSequences(const FMovieSceneSequen
 					MetaData->TrailingFence = DiscreteExclusiveUpper(SubSequenceIt.Range());
 				}
 
+				// Add determinism fences for boundary conditions
+				if (!SubData->OuterToInnerTransform.IsLinear() && (SubEntry.CompiledFlags.bParentSequenceRequiresUpperFence || SubEntry.CompiledFlags.bParentSequenceRequiresLowerFence) )
+				{
+					SubData->OuterToInnerTransform.ExtractBoundariesWithinRange(SubSequenceIt.Range().GetLowerBoundValue(), SubSequenceIt.Range().GetUpperBoundValue(), [OutCompilerData](FFrameTime FrameTime)
+					{
+						OutCompilerData->DeterminismData.Fences.Add(FrameTime.FrameNumber);
+						return true;
+					});
+				}
+
 				MetaData->LastIterIndex = ItIndex;
 			}
 		}
@@ -1803,7 +1813,6 @@ void UMovieSceneCompiledDataManager::PopulateSubSequenceTree(UMovieSceneSubTrack
 
 				const ESectionEvaluationFlags SubEntryFlags = Entry.Flags | Params.Flags;
 
-				// The section isn't looping, so we can just add it to the tree.
 				InOutHierarchy->AddRange(FrameRange, SubSequenceID, SubEntryFlags);
 
 				// Recurse into the sub sequence
diff --git a/Engine/Source/Runtime/MovieScene/Private/EntitySystem/MovieSceneSequenceUpdaters.cpp b/Engine/Source/Runtime/MovieScene/Private/EntitySystem/MovieSceneSequenceUpdaters.cpp
--- a/Engine/Source/Runtime/MovieScene/Private/EntitySystem/MovieSceneSequenceUpdaters.cpp
+++ b/Engine/Source/Runtime/MovieScene/Private/EntitySystem/MovieSceneSequenceUpdaters.cpp
@@ -24,6 +24,7 @@
 
 #include "Algo/IndexOf.h"
 #include "Algo/Transform.h"
+#include "Algo/Unique.h"
 
 namespace UE
 {
@@ -276,6 +277,13 @@ void FSequenceUpdater_Flat::PopulateUpdateFlags(TSharedRef<const FSharedPlayback
 	{
 		OutUpdateFlags |= ESequenceInstanceUpdateFlags::NeedsDissection;
 	}
+
+	const FMovieSceneSequenceHierarchy* Hierarchy = SharedPlaybackState->GetCompiledDataManager()->FindHierarchy(CompiledDataID);
+	if (Hierarchy && !Hierarchy->GetRootTransform().IsIdentity())
+	{
+		// Time-warped root transforms require dissection to manipulate the evaluation range
+		OutUpdateFlags |= ESequenceInstanceUpdateFlags::NeedsDissection;
+	}
 }
 
 void FSequenceUpdater_Flat::DissectContext(TSharedRef<const FSharedPlaybackState> SharedPlaybackState, const FMovieSceneContext& Context, TArray<TRange<FFrameTime>>& OutDissections)
@@ -443,7 +451,11 @@ void FSequenceUpdater_Hierarchical::PopulateUpdateFlags(TSharedRef<const FShared
 
 	if (const FMovieSceneSequenceHierarchy* Hierarchy = CompiledDataManager->FindHierarchy(CompiledDataID))
 	{
-		for (const TPair<FMovieSceneSequenceID, FMovieSceneSubSequenceData>& Pair : Hierarchy->AllSubSequenceData())
+		if (!Hierarchy->GetRootTransform().IsLinear())
+		{
+			OutUpdateFlags |= ESequenceInstanceUpdateFlags::NeedsDissection;
+		}
+		else for (const TPair<FMovieSceneSequenceID, FMovieSceneSubSequenceData>& Pair : Hierarchy->AllSubSequenceData())
 		{
 			UMovieSceneSequence*      SubSequence = Pair.Value.GetSequence();
 			FMovieSceneCompiledDataID SubDataID   = SubSequence ? CompiledDataManager->GetDataID(SubSequence) : FMovieSceneCompiledDataID();
@@ -461,25 +473,52 @@ void FSequenceUpdater_Hierarchical::DissectContext(TSharedRef<const FSharedPlayb
 {
 	UMovieSceneCompiledDataManager* CompiledDataManager = SharedPlaybackState->GetCompiledDataManager();
 
-	TRange<FFrameNumber>                TraversedRange = Context.GetFrameNumberRange();
+	FMovieSceneCompiledDataID   RootCompiledDataID = CompiledDataID;
+	FMovieSceneContext          RootContext        = Context;
+
+	const FMovieSceneSequenceHierarchy* RootHierarchy = CompiledDataManager->FindHierarchy(CompiledDataID);
+
+	if (!RootHierarchy)
+	{
+		return;
+	}
+
+	if (RootOverrideSequenceID != MovieSceneSequenceID::Root)
+	{
+		const FMovieSceneSubSequenceData* SubData = RootHierarchy->FindSubData(RootOverrideSequenceID);
+		if (SubData)
+		{
+			RootCompiledDataID = CompiledDataManager->GetDataID(SubData->GetSequence());
+			RootContext = Context.Transform(SubData->RootToSequenceTransform, SubData->TickResolution);
+		}
+	}
+	else if (!RootHierarchy->GetRootTransform().IsIdentity())
+	{
+		RootContext = Context.Transform(RootHierarchy->GetRootTransform(), Context.GetFrameRate());
+	}
+
+	TRange<FFrameNumber> TraversedRange = RootContext.GetFrameNumberRange();
 	TArray<FMovieSceneDeterminismFenceWithSubframe> RootDissectionTimes;
 
 	{
-		const FMovieSceneCompiledDataEntry&           DataEntry       = CompiledDataManager->GetEntryRef(CompiledDataID);
-		TArrayView<const FMovieSceneDeterminismFence> TraversedFences = GetFencesWithinRange(DataEntry.DeterminismFences, Context.GetRange());
+		const FMovieSceneCompiledDataEntry&           DataEntry       = CompiledDataManager->GetEntryRef(RootCompiledDataID);
+		TArrayView<const FMovieSceneDeterminismFence> TraversedFences = GetFencesWithinRange(DataEntry.DeterminismFences, RootContext.GetRange());
 
-		UE::MovieScene::DissectRange(TraversedFences, Context.GetRange(), OutDissections);
+		for (const FMovieSceneDeterminismFence& Fence : TraversedFences)
+		{
+			RootDissectionTimes.Add(FMovieSceneDeterminismFenceWithSubframe{ Fence.FrameNumber, Fence.bInclusive });
+		}
 	}
 
 	// @todo: should this all just be compiled into the root hierarchy?
-	if (const FMovieSceneSequenceHierarchy* Hierarchy = CompiledDataManager->FindHierarchy(CompiledDataID))
+	if (const FMovieSceneSequenceHierarchy* Hierarchy = CompiledDataManager->FindHierarchy(RootCompiledDataID))
 	{
 		FMovieSceneEvaluationTreeRangeIterator SubSequenceIt = Hierarchy->GetTree().IterateFromLowerBound(TraversedRange.GetLowerBound());
 		for ( ; SubSequenceIt && SubSequenceIt.Range().Overlaps(TraversedRange); ++SubSequenceIt)
 		{
-			TRange<FFrameTime> RootClampRange = TRange<FFrameTime>::Intersection(ConvertToFrameTimeRange(SubSequenceIt.Range()), Context.GetRange());
+			TRange<FFrameTime> RootClampRange = TRange<FFrameTime>::Intersection(ConvertToFrameTimeRange(SubSequenceIt.Range()), RootContext.GetRange());
 
-			// When Context.GetRange() does not fall on whole frame boundaries, we can sometimes end up with a range that clamps to being empty, even though the range overlapped
+			// When RootContext.GetRange() does not fall on whole frame boundaries, we can sometimes end up with a range that clamps to being empty, even though the range overlapped
 			// the traversed range. ie if we evaluated range (1.5, 10], our traversed range would be [2, 11). If we have a sub sequence range of (10, 20), it would still be iterated here
 			// because [2, 11) overlaps (10, 20), but when clamped to the evaluated range, the range is (10, 10], which is empty.
 			if (RootClampRange.IsEmpty())
@@ -525,7 +564,7 @@ void FSequenceUpdater_Hierarchical::DissectContext(TSharedRef<const FSharedPlayb
 						for (FMovieSceneDeterminismFence Fence : TraversedFences)
 						{
 							TOptional<FFrameTime> RootTime = InverseTransform.TryTransformTime(Fence.FrameNumber, Breadcrumbs);
-							if (RootTime)
+							if (RootTime && TraversedRange.Contains(RootTime->FrameNumber))
 							{
 								RootDissectionTimes.Emplace(FMovieSceneDeterminismFenceWithSubframe{ RootTime.GetValue(), Fence.bInclusive });
 							}
@@ -539,7 +578,16 @@ void FSequenceUpdater_Hierarchical::DissectContext(TSharedRef<const FSharedPlayb
 	if (RootDissectionTimes.Num() > 0)
 	{
 		Algo::SortBy(RootDissectionTimes, &FMovieSceneDeterminismFenceWithSubframe::FrameTime);
-		UE::MovieScene::DissectRange(RootDissectionTimes, Context.GetRange(), OutDissections);
+		int32 Index = Algo::UniqueBy(RootDissectionTimes, &FMovieSceneDeterminismFenceWithSubframe::FrameTime);
+		if (Index < RootDissectionTimes.Num())
+		{
+			RootDissectionTimes.SetNum(Index);
+		}
+		UE::MovieScene::DissectRange(RootDissectionTimes, RootContext.GetRange(), OutDissections);
+	}
+	else if (!RootHierarchy->GetRootTransform().IsIdentity())
+	{
+		OutDissections.Add(RootContext.GetRange());
 	}
 }
 
@@ -634,10 +682,6 @@ void FSequenceUpdater_Hierarchical::Update(TSharedRef<const FSharedPlaybackState
 			ActiveSequences.Add(RootOverrideSequenceID);
 		}
 	}
-	else if (RootHierarchy && !RootHierarchy->GetRootTransform().IsIdentity())
-	{
-		RootContext = Context.Transform(RootHierarchy->GetRootTransform(), Context.GetFrameRate());
-	}
 
 	FFrameNumber ImportTime = RootContext.GetEvaluationFieldTime();
 	const bool bGatherEntities = !CachedEntityRange.Contains(ImportTime);
