Sequencer: Fixed ComputePiecewiseExtents not working correctly when the StartTime and EndTime were outside the range of the curve

diff --git a/Engine/Source/Runtime/MovieScene/Public/Channels/MovieScenePiecewiseCurveUtils.inl b/Engine/Source/Runtime/MovieScene/Public/Channels/MovieScenePiecewiseCurveUtils.inl
--- a/Engine/Source/Runtime/MovieScene/Public/Channels/MovieScenePiecewiseCurveUtils.inl
+++ b/Engine/Source/Runtime/MovieScene/Public/Channels/MovieScenePiecewiseCurveUtils.inl
@@ -73,21 +73,11 @@ UE::MovieScene::Interpolation::FInterpolationExtents ComputePiecewiseExtents(con
 		return Extents;
 	}
 
-	if (StartTime == EndTime)
-	{
-		// Easy - just evaluate the curve
-		double Value = 0.0;
-		PiecewiseData.GetPieceByTime(StartTime).Evaluate(StartTime, Value);
-		FInterpolationExtents Extents;
-		Extents.AddPoint(Value, StartTime);
-		return Extents;
-	}
-
-	FInterpolationExtents FinalExtents;
-
 	FFrameNumber FiniteSplineStart = PiecewiseData.GetFiniteStart();
 	FFrameNumber FiniteSplineEnd   = PiecewiseData.GetFiniteEnd();
 
+	FInterpolationExtents FinalExtents;
+
 	// Deal with linear pre-post extrapolation
 	if (PiecewiseData.GetPreExtrapolation() == RCCE_Linear && StartTime.FrameNumber.Value < FiniteSplineStart)
 	{
@@ -252,6 +242,27 @@ UE::MovieScene::Interpolation::FInterpolationExtents ComputePiecewiseExtents(con
 			FinalExtents.Combine(FullExtents);
 		}
 	}
+	else if (EndCycled.Time == StartCycled.Time)
+	{
+		double Value = 0.0;
+
+		if (StartCycled.Time < FiniteSplineStart)
+		{
+			Value = PiecewiseData.PreExtrapolate(StartCycled.Time);
+		}
+		else if (StartCycled.Time >= FiniteSplineEnd)
+		{
+			Value = PiecewiseData.PostExtrapolate(StartCycled.Time);
+		}
+		else
+		{
+			PiecewiseData.GetPieceByTime(StartCycled.Time).Evaluate(StartCycled.Time, Value);
+		}
+
+		FInterpolationExtents Extents;
+		Extents.AddPoint(Value, StartTime);
+		return Extents;
+	}
 	else
 	{
 		// Within the same cycle is easy - just compute the extents between the two bounds, and offset by the oscillation amount
