Sequencer- Deleting a subsequence with spawnables in it will leave the spawnables around.

diff --git a/Engine/Source/Runtime/MovieScene/Private/EntitySystem/MovieSceneSpawnablesSystem.cpp b/Engine/Source/Runtime/MovieScene/Private/EntitySystem/MovieSceneSpawnablesSystem.cpp
--- a/Engine/Source/Runtime/MovieScene/Private/EntitySystem/MovieSceneSpawnablesSystem.cpp
+++ b/Engine/Source/Runtime/MovieScene/Private/EntitySystem/MovieSceneSpawnablesSystem.cpp
@@ -136,10 +136,20 @@ void UMovieSceneSpawnablesSystem::OnRun(FSystemTaskPrerequisites& InPrerequisite
 		{
 			const FSequenceInstance& Instance = InstanceRegistry->GetInstance(InstanceHandle);
 			TSharedRef<FSharedPlaybackState> SharedPlaybackState = Instance.GetSharedPlaybackState();
-
-			const UMovieSceneSequence* Sequence = SharedPlaybackState->GetSequence(Instance.GetSequenceID());
-			if (Sequence)
+			
+			const UMovieSceneSequence* Sequence = nullptr;
+			// Prefer getting the Sequence from the evaluation state if possible.
+			// This is because if a sub section has just been destroyed, the sequence will be recompiled and the hierarchy (used by SharedPlaybackState) is not reliable for finding the Sequence asset.
+			if (FMovieSceneEvaluationState* EvaluationState = SharedPlaybackState->FindCapability<FMovieSceneEvaluationState>())
+			{
+				Sequence = EvaluationState->FindSequence(Instance.GetSequenceID());
+			}
+			else
 			{
+				SharedPlaybackState->GetSequence(Instance.GetSequenceID());
+			}
+			if (Sequence)
+			{ 
 				if (Linker->EntityManager.HasComponent(EntityID, FBuiltInComponentTypes::Get()->Tags.OldStyleSpawnable))
 				{
 					if (FMovieSceneSpawnable* Spawnable = Sequence->GetMovieScene()->FindSpawnable(SpawnableObjectID))
