Sequencer: Changed defaults for time-warp curves to have a key at the start and end, and constant pre/post extrapolation

diff --git a/Engine/Source/Runtime/MovieScene/Private/Variants/MovieSceneTimeWarpCurve.cpp b/Engine/Source/Runtime/MovieScene/Private/Variants/MovieSceneTimeWarpCurve.cpp
--- a/Engine/Source/Runtime/MovieScene/Private/Variants/MovieSceneTimeWarpCurve.cpp
+++ b/Engine/Source/Runtime/MovieScene/Private/Variants/MovieSceneTimeWarpCurve.cpp
@@ -24,13 +24,24 @@ void UMovieSceneTimeWarpCurve::InitializeDefaults()
 	if (Channel.Owner)
 	{
 		FFrameNumber StartFrame = DiscreteInclusiveLower(Channel.Owner->GetPlaybackRange());
+		FFrameNumber EndFrame   = DiscreteExclusiveUpper(Channel.Owner->GetPlaybackRange());
 
 		UMovieSceneSection* OwningSection = GetTypedOuter<UMovieSceneSection>();
-		if (OwningSection && OwningSection->HasStartFrame())
+		if (OwningSection)
 		{
-			StartFrame = OwningSection->GetInclusiveStartFrame();
+			if (OwningSection->HasStartFrame())
+			{
+				StartFrame = 0;
+				if (OwningSection->HasEndFrame())
+				{
+					EndFrame = OwningSection->GetExclusiveEndFrame() - OwningSection->GetInclusiveStartFrame();
+				}
+			}
+			else if (OwningSection->HasEndFrame())
+			{
+				EndFrame = OwningSection->GetExclusiveEndFrame();
+			}
 		}
-		FFrameRate TickResolution = Channel.Owner->GetTickResolution();
 
 		TMovieSceneChannelData<FMovieSceneDoubleValue> ChannelData = Channel.GetData();
 
@@ -40,11 +51,11 @@ void UMovieSceneTimeWarpCurve::InitializeDefaults()
 
 		ChannelData.AddKey(StartFrame, Value);
 
-		Value.Value = (StartFrame + TickResolution.Numerator).Value;
-		ChannelData.AddKey(StartFrame + TickResolution.Numerator, Value);
+		Value.Value = EndFrame.Value;
+		ChannelData.AddKey(EndFrame, Value);
 
-		Channel.PreInfinityExtrap  = RCCE_Linear;
-		Channel.PostInfinityExtrap = RCCE_Linear;
+		Channel.PreInfinityExtrap  = RCCE_Constant;
+		Channel.PostInfinityExtrap = RCCE_Constant;
 	}
 }
 
