Sequencer: Constrain evaluation hook evaluation to the main thread

#12181

diff --git a/Engine/Source/Runtime/MovieScene/Private/EntitySystem/MovieSceneEvaluationHookSystem.cpp b/Engine/Source/Runtime/MovieScene/Private/EntitySystem/MovieSceneEvaluationHookSystem.cpp
--- a/Engine/Source/Runtime/MovieScene/Private/EntitySystem/MovieSceneEvaluationHookSystem.cpp
+++ b/Engine/Source/Runtime/MovieScene/Private/EntitySystem/MovieSceneEvaluationHookSystem.cpp
@@ -103,6 +103,7 @@ void UMovieSceneEvaluationHookSystem::OnSchedulePersistentTasks(UE::MovieScene::
 	.Read(Components->EvaluationHook)
 	.Read(Components->EvalTime)
 	.Write(Components->EvaluationHookFlags)
+	.SetDesiredThread(Linker->EntityManager.GetDispatchThread())
 	.Schedule_PerAllocation<FEvaluationHookUpdater>(&Linker->EntityManager, TaskScheduler, this, Linker->GetInstanceRegistry());
 }
 
@@ -126,6 +127,7 @@ void UMovieSceneEvaluationHookSystem::OnRun(FSystemTaskPrerequisites& InPrerequi
 		.Read(Components->EvaluationHook)
 		.Read(Components->EvalTime)
 		.Write(Components->EvaluationHookFlags)
+		.SetDesiredThread(Linker->EntityManager.GetDispatchThread())
 		.Dispatch_PerAllocation<FEvaluationHookUpdater>(&Linker->EntityManager, InPrerequisites, &Subsequents, this, Linker->GetInstanceRegistry());
 	}
 	else if (HasEvents())
