Sequencer: Fix ensure/crash in PlayReverse.

Compare FrameTimes rather than FrameNumbers, which was truncating the subframes.

Remove erroneous ensure which compares PositionRelativeToStart and DurationWithSubFrames since on PlayReverse, PositionRelativeToStart is going to be negative.

#tests Tested that playing a sequence with PlayReverse completes successfully and can be restarted. Ran Sequencer EngineTests.

diff --git a/Engine/Source/Runtime/MovieScene/Private/MovieSceneSequencePlayer.cpp b/Engine/Source/Runtime/MovieScene/Private/MovieSceneSequencePlayer.cpp
--- a/Engine/Source/Runtime/MovieScene/Private/MovieSceneSequencePlayer.cpp
+++ b/Engine/Source/Runtime/MovieScene/Private/MovieSceneSequencePlayer.cpp
@@ -1191,17 +1191,15 @@ void UMovieSceneSequencePlayer::UpdateTimeCursorPosition_Internal(FFrameTime New
 	if (!PauseRange.IsSet() && Method == EUpdatePositionMethod::Play && ShouldStopOrLoop(NewPosition))
 	{
 		// The actual start time taking into account reverse playback
-		FFrameNumber StartTimeWithReversed = bReversePlayback ? GetLastValidTime().FrameNumber : StartTime;
+		FFrameTime StartTimeWithReversed = bReversePlayback ? GetLastValidTime() : StartTime;
 
 		// The actual end time taking into account reverse playback
-		FFrameTime EndTimeWithReversed = bReversePlayback ? StartTime : GetLastValidTime().FrameNumber;
+		FFrameTime EndTimeWithReversed = bReversePlayback ? StartTime : GetLastValidTime();
 
 		// Operate in tick resolution (for subframes)
 		const double DurationWithSubFrames   = FMath::Max<double>(UE_SMALL_NUMBER, GetDuration().Time.AsDecimal());
 		const double PositionRelativeToStart = (NewPosition - StartTimeWithReversed).AsDecimal();
 
-		ensure(PositionRelativeToStart >= DurationWithSubFrames);
-
 		const int32 NumTimesLooped    = FMath::Abs(FMath::TruncToInt32(PositionRelativeToStart / DurationWithSubFrames));
 		const bool  bLoopIndefinitely = PlaybackSettings.LoopCount.Value < 0;
 
