Sequencer: Remove ShowOnlyInnerProperties on TickInterval

The inner properties on the TickInterval struct aren't editable from the Details panel because there is an EditCondition on them - !bInheritTickIntervalFromOwner. But the TickInterval itself isn't toggleable because ShowOnlyInnerProperties makes it so that there is no row.

diff --git a/Engine/Source/Runtime/MovieScene/Public/MovieSceneSequencePlaybackSettings.h b/Engine/Source/Runtime/MovieScene/Public/MovieSceneSequencePlaybackSettings.h
--- a/Engine/Source/Runtime/MovieScene/Public/MovieSceneSequencePlaybackSettings.h
+++ b/Engine/Source/Runtime/MovieScene/Public/MovieSceneSequencePlaybackSettings.h
@@ -77,7 +77,7 @@ struct FMovieSceneSequencePlaybackSettings
 	FMovieSceneSequenceLoopCount LoopCount;
 
 	/** Overridable tick interval for this sequence to update at. When not overridden, the owning actor or component's tick interval will be used */
-	UPROPERTY(EditAnywhere, BlueprintReadWrite, Category="Playback", meta=(ShowOnlyInnerProperties, Units=s, EditCondition="!bInheritTickIntervalFromOwner"))
+	UPROPERTY(EditAnywhere, BlueprintReadWrite, Category="Playback", meta=(Units=s, EditCondition="!bInheritTickIntervalFromOwner"))
 	FMovieSceneSequenceTickInterval TickInterval;
 
 	/** The rate at which to playback the animation */
