Sequencer: Handle <= 0 SourceDurations inside FMovieSceneSectionTimingParametersSeconds::MakeTransform

diff --git a/Engine/Source/Runtime/MovieScene/Private/Sections/MovieSceneSectionTimingParameters.cpp b/Engine/Source/Runtime/MovieScene/Private/Sections/MovieSceneSectionTimingParameters.cpp
--- a/Engine/Source/Runtime/MovieScene/Private/Sections/MovieSceneSectionTimingParameters.cpp
+++ b/Engine/Source/Runtime/MovieScene/Private/Sections/MovieSceneSectionTimingParameters.cpp
@@ -11,7 +11,13 @@ FMovieSceneSequenceTransform FMovieSceneSectionTimingParametersSeconds::MakeTran
 	FMovieSceneSequenceTransform Result;
 
 	check(OuterRange.HasLowerBound());
-	check(SourceDuration > 0.0);
+
+	if (SourceDuration <= 0)
+	{
+		// Zero source duration is handled by zero play rate (always evaluate time zero)
+		Result.Add(0, FMovieSceneTimeWarpVariant(0.0));
+		return Result;
+	}
 
 	// ----------------------------------------------------------------------------
 	// First things first, subtract the section start bound
