Sequencer: Evaluate v1==v2 as constant only if tangents are also 0.

#rnx

diff --git a/Engine/Source/Runtime/MovieScene/Private/Channels/MovieSceneCurveChannelImpl.cpp b/Engine/Source/Runtime/MovieScene/Private/Channels/MovieSceneCurveChannelImpl.cpp
--- a/Engine/Source/Runtime/MovieScene/Private/Channels/MovieSceneCurveChannelImpl.cpp
+++ b/Engine/Source/Runtime/MovieScene/Private/Channels/MovieSceneCurveChannelImpl.cpp
@@ -501,7 +501,7 @@ UE::MovieScene::Interpolation::FCachedInterpolation TMovieSceneCurveChannelImpl<
 	{
 	case RCIM_Cubic:
 		{
-			if (V1 == V2)
+			if (V1 == V2 && T1 == 0 && T2 == 0 && W1 == 0 && W2 == 0)
 			{
 				return FCachedInterpolation(Range, FConstantValue(Time1, V1));
 			}
