Sequencer: Add UMovieSceneCompiledData::AppendToClassSchema to inform iterative cooking of the compiler version

diff --git a/Engine/Source/Runtime/MovieScene/Private/Compilation/MovieSceneCompiledDataManager.cpp b/Engine/Source/Runtime/MovieScene/Private/Compilation/MovieSceneCompiledDataManager.cpp
--- a/Engine/Source/Runtime/MovieScene/Private/Compilation/MovieSceneCompiledDataManager.cpp
+++ b/Engine/Source/Runtime/MovieScene/Private/Compilation/MovieSceneCompiledDataManager.cpp
@@ -413,10 +413,27 @@ void UMovieSceneCompiledData::Reset()
 	AccumulatedFlags = EMovieSceneSequenceFlags::None;
 }
 
+#if WITH_EDITORONLY_DATA
+void UMovieSceneCompiledData::AppendToClassSchema(FAppendToClassSchemaContext& Context)
+{
+	Super::AppendToClassSchema(Context);
+
+	// Specify the compiler version to the iterative cooker. Any changes to the schema of
+	//    compiled data should update the version to ensure that compiled data is invalidated
+	//    for the purposes of iterative cooking.
+
+	FGuid ParsedCompilerVersion;
+	if (FGuid::Parse(GMovieSceneCompilerVersion, ParsedCompilerVersion))
+	{
+		Context.Update(&ParsedCompilerVersion, sizeof(ParsedCompilerVersion));
+	}
+}
+#endif
+
 UMovieSceneCompiledDataManager::UMovieSceneCompiledDataManager()
 {
 	const bool bParsed = FGuid::Parse(GMovieSceneCompilerVersion, CompilerVersion);
-	ensureMsgf(bParsed, TEXT("Invalid compiler version specific - this will break any persistent compiled data"));
+	ensureMsgf(bParsed, TEXT("Invalid compiler version specified - this will break any persistent compiled data"));
 
 	IConsoleManager::Get().RegisterConsoleVariableSink_Handle(FConsoleCommandDelegate::CreateUObject(this, &UMovieSceneCompiledDataManager::ConsoleVariableSink));
 
diff --git a/Engine/Source/Runtime/MovieScene/Public/Compilation/MovieSceneCompiledDataManager.h b/Engine/Source/Runtime/MovieScene/Public/Compilation/MovieSceneCompiledDataManager.h
--- a/Engine/Source/Runtime/MovieScene/Public/Compilation/MovieSceneCompiledDataManager.h
+++ b/Engine/Source/Runtime/MovieScene/Public/Compilation/MovieSceneCompiledDataManager.h
@@ -92,6 +92,10 @@ public:
 
 	void Reset();
 
+#if WITH_EDITORONLY_DATA
+	static void AppendToClassSchema(FAppendToClassSchemaContext& Context);
+#endif
+
 private:
 	friend class UMovieSceneCompiledDataManager;
 
