Sequencer: Fixed potential compile errors with bound object preanimated storage where KeyType was not an FObjectKey

diff --git a/Engine/Source/Runtime/MovieScene/Public/Evaluation/PreAnimatedState/MovieScenePreAnimatedObjectStorage.h b/Engine/Source/Runtime/MovieScene/Public/Evaluation/PreAnimatedState/MovieScenePreAnimatedObjectStorage.h
--- a/Engine/Source/Runtime/MovieScene/Public/Evaluation/PreAnimatedState/MovieScenePreAnimatedObjectStorage.h
+++ b/Engine/Source/Runtime/MovieScene/Public/Evaluation/PreAnimatedState/MovieScenePreAnimatedObjectStorage.h
@@ -121,7 +121,7 @@ public:
 		for (int32 Index = 0; Index < Num; ++Index)
 		{
 			UObject* BoundObject = BoundObjects[Index];
-			FObjectKey Key{ BoundObject };
+			KeyType Key{ BoundObject };
 
 			FPreAnimatedStorageGroupHandle GroupHandle  = this->Traits.MakeGroup(BoundObject);
 			FPreAnimatedStorageIndex       StorageIndex = this->GetOrCreateStorageIndex(Key);
@@ -133,7 +133,7 @@ public:
 
 	FPreAnimatedStateEntry MakeEntry(UObject* BoundObject)
 	{
-		FObjectKey Key{ BoundObject };
+		KeyType Key{ BoundObject };
 	
 		FPreAnimatedStorageIndex       StorageIndex = this->GetOrCreateStorageIndex(Key);
 		FPreAnimatedStorageGroupHandle GroupHandle  = this->Traits.MakeGroup(BoundObject);
