Sequencer: Deprecated SceneComponentBinding components in favor of a more flexible and generic BoundObjectResolver component

This new component allows sections to specify any function ptr to map a boundobject to the object it actually needs to animate. For Transforms this is a USceneComponent, for skeletal animations this is a USkeletalMeshComponent. The methods by which these are resolved varies by track and is now specified through this additional component.

diff --git a/Engine/Plugins/MovieScene/TemplateSequence/Source/TemplateSequence/Private/CameraAnimationSequenceSubsystem.cpp b/Engine/Plugins/MovieScene/TemplateSequence/Source/TemplateSequence/Private/CameraAnimationSequenceSubsystem.cpp
--- a/Engine/Plugins/MovieScene/TemplateSequence/Source/TemplateSequence/Private/CameraAnimationSequenceSubsystem.cpp
+++ b/Engine/Plugins/MovieScene/TemplateSequence/Source/TemplateSequence/Private/CameraAnimationSequenceSubsystem.cpp
@@ -51,17 +51,12 @@ struct FCameraAnimationInstantiationMutation : IMovieSceneEntityMutation
 			{
 				BindObjectImpl(InstanceHandles[Index], ObjectBindings[Index], OutBoundObjects[Index]);
 			}
-		}
-		else if (AllocationType.Contains(BuiltInComponents->SceneComponentBinding))
-		{
-			// Initialize bound scene components.
-			TComponentReader<FGuid> SceneComponentBindings = Allocation->ReadComponents(BuiltInComponents->SceneComponentBinding);
-			for (int32 Index = 0; Index < Num; ++Index)
-			{
-				BindObjectImpl(InstanceHandles[Index], SceneComponentBindings[Index], OutBoundObjects[Index]);
-			}
+
+			// @NOTE: Camera animations intentionally do not use bound object resolvers because they explicitly
+			//        always use a camera stand in object.
 		}
 	}
+
 	void BindObjectImpl(const FInstanceHandle& InstanceHandle, const FGuid& ObjectBinding, UObject*& OutBoundObject) const
 	{
 		const FSequenceInstance& Instance = InstanceRegistry.GetInstance(InstanceHandle);
@@ -150,8 +145,7 @@ void UCameraAnimationBoundObjectInstantiator::OnRun(FSystemTaskPrerequisites& In
 	// Initialize all new allocations with bound objects and output components.
 	FCameraAnimationInstantiationMutation Mutation(*Linker->GetInstanceRegistry());
 	FEntityComponentFilter Filter = FEntityComponentFilter()
-		.Any({ BuiltInComponents->GenericObjectBinding, 
-				BuiltInComponents->SceneComponentBinding })
+		.Any({ BuiltInComponents->GenericObjectBinding })
 		.All({ BuiltInComponents->InstanceHandle, BuiltInComponents->Tags.NeedsLink })
 		.None({ BuiltInComponents->Tags.NeedsUnlink });
 	Linker->EntityManager.MutateAll(Filter, Mutation, EMutuallyInclusiveComponentType::All);
diff --git a/Engine/Plugins/MovieScene/TemplateSequence/Source/TemplateSequence/Private/Systems/TemplateSequenceSystem.cpp b/Engine/Plugins/MovieScene/TemplateSequence/Source/TemplateSequence/Private/Systems/TemplateSequenceSystem.cpp
--- a/Engine/Plugins/MovieScene/TemplateSequence/Source/TemplateSequence/Private/Systems/TemplateSequenceSystem.cpp
+++ b/Engine/Plugins/MovieScene/TemplateSequence/Source/TemplateSequence/Private/Systems/TemplateSequenceSystem.cpp
@@ -218,13 +218,7 @@ void UTemplateSequencePropertyScalingInstantiatorSystem::OnRun(FSystemTaskPrereq
 
 		if (ParentEntity.IsValid())
 		{
-			TOptionalComponentReader<FGuid> Reader;
-			Reader = Linker->EntityManager.ReadComponent(ParentEntity, BuiltInComponents->SceneComponentBinding);
-			if (!Reader.IsValid())
-			{
-				Reader = Linker->EntityManager.ReadComponent(ParentEntity, BuiltInComponents->GenericObjectBinding);
-			}
-
+			TOptionalComponentReader<FGuid> Reader = Linker->EntityManager.ReadComponent(ParentEntity, BuiltInComponents->GenericObjectBinding);
 			if (ensure(Reader.IsValid()))
 			{
 				const FGuid ObjectBindingID = *Reader.ComponentAtIndex(0);
diff --git a/Engine/Source/Editor/MovieSceneTools/Private/Channels/BuiltInChannelEditors.cpp b/Engine/Source/Editor/MovieSceneTools/Private/Channels/BuiltInChannelEditors.cpp
--- a/Engine/Source/Editor/MovieSceneTools/Private/Channels/BuiltInChannelEditors.cpp
+++ b/Engine/Source/Editor/MovieSceneTools/Private/Channels/BuiltInChannelEditors.cpp
@@ -182,9 +182,11 @@ FKeyHandle AddOrUpdateKeyImpl(ChannelType* Channel, UMovieSceneSection* SectionT
 					const FPropertyDefinition& BoundPropertyDefinition = PropertyDefinitions[BoundPropertyDefinitionIndex];
 
 					check(FirstBoundObject != nullptr);
-					if (Interrogator.GetLinker()->EntityManager.HasComponent(EntityID, BuiltInComponents->SceneComponentBinding))
+					TComponentLock<TReadOptional<FBoundObjectResolver>> Resolver =
+							Interrogator.GetLinker()->EntityManager.ReadComponent(EntityID, BuiltInComponents->BoundObjectResolver);
+					if (Resolver)
 					{
-						FirstBoundObject = MovieSceneHelpers::SceneComponentFromRuntimeObject(FirstBoundObject);
+						FirstBoundObject = (*Resolver)(FirstBoundObject);
 						check(FirstBoundObject != nullptr);
 					}
 
diff --git a/Engine/Source/Runtime/MovieScene/Private/EntitySystem/BuiltInComponentTypes.cpp b/Engine/Source/Runtime/MovieScene/Private/EntitySystem/BuiltInComponentTypes.cpp
--- a/Engine/Source/Runtime/MovieScene/Private/EntitySystem/BuiltInComponentTypes.cpp
+++ b/Engine/Source/Runtime/MovieScene/Private/EntitySystem/BuiltInComponentTypes.cpp
@@ -289,7 +289,10 @@ FBuiltInComponentTypes::FBuiltInComponentTypes()
 
 	ComponentRegistry->NewComponentType(&PropertyBinding,         TEXT("Property Binding"), EComponentTypeFlags::CopyToOutput);
 	ComponentRegistry->NewComponentType(&GenericObjectBinding,    TEXT("Generic Object Binding ID"));
-	ComponentRegistry->NewComponentType(&SceneComponentBinding,   TEXT("USceneComponent Binding ID"));
+	ComponentRegistry->NewComponentType(&BoundObjectResolver,     TEXT("Bound Object Resolver"));
+	PRAGMA_DISABLE_DEPRECATION_WARNINGS
+	ComponentRegistry->NewComponentType(&SceneComponentBinding,   TEXT("[DEPRECATED] USceneComponent Binding ID"));
+	PRAGMA_ENABLE_DEPRECATION_WARNINGS
 	ComponentRegistry->NewComponentType(&SpawnableBinding,        TEXT("Spawnable Binding"));
 	ComponentRegistry->NewComponentType(&TrackInstance,           TEXT("Track Instance"));
 	ComponentRegistry->NewComponentType(&BoolChannel,             TEXT("Bool Channel"));
diff --git a/Engine/Source/Runtime/MovieScene/Private/EntitySystem/MovieSceneBoundObjectInstantiator.cpp b/Engine/Source/Runtime/MovieScene/Private/EntitySystem/MovieSceneBoundObjectInstantiator.cpp
--- a/Engine/Source/Runtime/MovieScene/Private/EntitySystem/MovieSceneBoundObjectInstantiator.cpp
+++ b/Engine/Source/Runtime/MovieScene/Private/EntitySystem/MovieSceneBoundObjectInstantiator.cpp
@@ -37,41 +37,14 @@ void UMovieSceneGenericBoundObjectInstantiator::OnRun(FSystemTaskPrerequisites&
 
 	UnlinkStaleObjectBindings(Components->GenericObjectBinding);
 
-	struct FGenericBoundObjectBatch : FObjectFactoryBatch
-	{
-		virtual EResolveError ResolveObjects(FInstanceRegistry* InstanceRegistry, FInstanceHandle InstanceHandle, int32 InEntityIndex, const FGuid& ObjectBinding) override
-		{
-			EResolveError Error = EResolveError::UnresolvedBinding;
-
-			FSequenceInstance& SequenceInstance = InstanceRegistry->MutateInstance(InstanceHandle);
-			TSharedRef<const FSharedPlaybackState> SharedPlaybackState = SequenceInstance.GetSharedPlaybackState();
-			TArrayView<TWeakObjectPtr<>> BoundObjects = SharedPlaybackState->FindBoundObjects(ObjectBinding, SequenceInstance.GetSequenceID());
-			for (TWeakObjectPtr<> WeakObject : BoundObjects)
-			{
-				if (UObject* Object = WeakObject.Get())
-				{
-					if (!ensureMsgf(!FBuiltInComponentTypes::IsBoundObjectGarbage(Object), TEXT("Attempting to bind an object that is garbage or unreachable")))
-					{
-						continue;
-					}
-
-					// Make a child entity for this resolved binding
-					Add(InEntityIndex, Object);
-					Error = EResolveError::None;
-				}
-			}
-
-			return Error;
-		}
-	};
-
-	TBoundObjectTask<FGenericBoundObjectBatch> BoundObjectTask(Linker);
+	FBoundObjectTask BoundObjectTask(Linker);
 
 	// Gather all newly instanced entities with an object binding ID
 	FEntityTaskBuilder()
 	.ReadEntityIDs()
 	.Read(Components->InstanceHandle)
 	.Read(Components->GenericObjectBinding)
+	.ReadOptional(Components->BoundObjectResolver)
 	.FilterAny({ Components->Tags.NeedsLink, Components->Tags.HasUnresolvedBinding })
 	.FilterNone({ Components->Tags.NeedsUnlink })
 	.RunInline_PerAllocation(&Linker->EntityManager, BoundObjectTask);
diff --git a/Engine/Source/Runtime/MovieScene/Private/EntitySystem/MovieSceneBoundSceneComponentInstantiator.cpp b/Engine/Source/Runtime/MovieScene/Private/EntitySystem/MovieSceneBoundSceneComponentInstantiator.cpp
--- a/Engine/Source/Runtime/MovieScene/Private/EntitySystem/MovieSceneBoundSceneComponentInstantiator.cpp
+++ b/Engine/Source/Runtime/MovieScene/Private/EntitySystem/MovieSceneBoundSceneComponentInstantiator.cpp
@@ -1,6 +1,7 @@
 // Copyright Epic Games, Inc. All Rights Reserved.
 
 #include "EntitySystem/MovieSceneBoundSceneComponentInstantiator.h"
+#include "EntitySystem/MovieSceneBoundObjectInstantiator.h"
 #include "EntitySystem/MovieSceneEntityInstantiatorSystem.h"
 #include "EntitySystem/MovieSceneEntitySystemTask.h"
 #include "EntitySystem/MovieSceneEntityManager.h"
@@ -8,6 +9,7 @@
 #include "EntitySystem/BuiltInComponentTypes.h"
 #include "EntitySystem/MovieSceneInstanceRegistry.h"
 #include "EntitySystem/MovieSceneEntitySystemLinker.h"
+#include "EntitySystem/MovieSceneEntityMutations.h"
 
 #include "Evaluation/MovieSceneEvaluationState.h"
 #include "MovieSceneCommonHelpers.h"
@@ -17,6 +19,8 @@
 
 #include UE_INLINE_GENERATED_CPP_BY_NAME(MovieSceneBoundSceneComponentInstantiator)
 
+PRAGMA_DISABLE_DEPRECATION_WARNINGS
+
 UMovieSceneBoundSceneComponentInstantiator::UMovieSceneBoundSceneComponentInstantiator(const FObjectInitializer& ObjInit)
 	: Super(ObjInit)
 {
@@ -27,8 +31,7 @@ UMovieSceneBoundSceneComponentInstantiator::UMovieSceneBoundSceneComponentInstan
 
 	if (HasAnyFlags(RF_ClassDefaultObject))
 	{
-		DefineComponentProducer(GetClass(), Components->BoundObject);
-		DefineComponentProducer(GetClass(), Components->SymbolicTags.CreatesEntities);
+		DefineImplicitPrerequisite(GetClass(), UMovieSceneGenericBoundObjectInstantiator::StaticClass());
 	}
 }
 
@@ -38,58 +41,36 @@ void UMovieSceneBoundSceneComponentInstantiator::OnRun(FSystemTaskPrerequisites&
 
 	FBuiltInComponentTypes* Components = FBuiltInComponentTypes::Get();
 
-	UnlinkStaleObjectBindings(Components->SceneComponentBinding);
-
-	struct FBoundSceneComponentBatch : FObjectFactoryBatch
+	struct FSceneComponentBindingMutation : IMovieSceneEntityMutation
 	{
-		virtual EResolveError ResolveObjects(FInstanceRegistry* InstanceRegistry, FInstanceHandle InstanceHandle, int32 InEntityIndex, const FGuid& ObjectBinding) override
+		virtual void CreateMutation(FEntityManager* EntityManager, FComponentMask* InOutEntityComponentTypes) const override
+		{
+			FBuiltInComponentTypes* Components = FBuiltInComponentTypes::Get();
+			InOutEntityComponentTypes->SetAll({ Components->GenericObjectBinding, Components->BoundObjectResolver });
+		}
+		virtual void InitializeAllocation(FEntityAllocation* Allocation, const FComponentMask& AllocationType) const
 		{
-			EResolveError Result = EResolveError::UnresolvedBinding;
+			FBuiltInComponentTypes* Components = FBuiltInComponentTypes::Get();
 
-			FSequenceInstance& SequenceInstance = InstanceRegistry->MutateInstance(InstanceHandle);
-			TSharedRef<const FSharedPlaybackState> SharedPlaybackState = SequenceInstance.GetSharedPlaybackState();
-			TArrayView<TWeakObjectPtr<>> BoundObjects = SharedPlaybackState->FindBoundObjects(ObjectBinding, SequenceInstance.GetSequenceID());
-			if (BoundObjects.Num() == 0)
-			{
-				UE_LOG(LogMovieSceneECS, Verbose, TEXT("FBoundSceneComponentBatch::ResolveObjects: No bound objects returned for FGuid: %s"), *ObjectBinding.ToString());
-				return Result;
-			}
+			const FGuid*          SceneComponentBindings = Allocation->ReadComponents(Components->SceneComponentBinding).AsPtr();
+			FGuid*                GenericObjectBindings  = Allocation->WriteComponents(Components->GenericObjectBinding, FEntityAllocationWriteContext::NewAllocation()).AsPtr();
+			FBoundObjectResolver* BoundObjectResolvers   = Allocation->WriteComponents(Components->BoundObjectResolver, FEntityAllocationWriteContext::NewAllocation()).AsPtr();
 
-			for (TWeakObjectPtr<> WeakObject : BoundObjects)
+			const int32 Num = Allocation->Num();
+			FMemory::Memcpy(GenericObjectBindings, SceneComponentBindings, sizeof(FGuid)*Num);
+
+			for (int32 Index = 0; Index < Num; ++Index)
 			{
-				if (UObject* Object = WeakObject.Get())
-				{
-					if (USceneComponent* SceneComponent = MovieSceneHelpers::SceneComponentFromRuntimeObject(Object))
-					{
-						if (!ensureMsgf(!FBuiltInComponentTypes::IsBoundObjectGarbage(SceneComponent), TEXT("Attempting to bind an object that is garbage or unreachable")))
-						{
-							continue;
-						}
-
-						// Make a child entity for this resolved binding
-						Add(InEntityIndex, SceneComponent);
-						Result = EResolveError::None;
-					}
-				}
-				else
-				{
-					UE_LOG(LogMovieSceneECS, Verbose, TEXT("FBoundSceneComponentBatch::ResolveObjects: Invalid weak object returned for FGuid: %s"), *ObjectBinding.ToString());
-				}
+				BoundObjectResolvers[Index] = MovieSceneHelpers::ResolveSceneComponentBoundObject;
 			}
-
-			return Result;
 		}
-	};
-
-	TBoundObjectTask<FBoundSceneComponentBatch> ObjectBindingTask(Linker);
-
-	// Gather all newly instanced entities with an object binding ID
-	FEntityTaskBuilder()
-	.ReadEntityIDs()
-	.Read(Components->InstanceHandle)
-	.Read(Components->SceneComponentBinding)
-	.FilterAny({ Components->Tags.NeedsLink, Components->Tags.HasUnresolvedBinding })
-	.FilterNone({ Components->Tags.NeedsUnlink })
-	.RunInline_PerAllocation(&Linker->EntityManager, ObjectBindingTask);
+	} Mutation;
+
+
+	Linker->EntityManager.MutateAll(
+		FEntityComponentFilter().All({ Components->SceneComponentBinding }), Mutation);
 }
 
+
+PRAGMA_ENABLE_DEPRECATION_WARNINGS
+
diff --git a/Engine/Source/Runtime/MovieScene/Private/EntitySystem/MovieSceneEntityFactory.cpp b/Engine/Source/Runtime/MovieScene/Private/EntitySystem/MovieSceneEntityFactory.cpp
--- a/Engine/Source/Runtime/MovieScene/Private/EntitySystem/MovieSceneEntityFactory.cpp
+++ b/Engine/Source/Runtime/MovieScene/Private/EntitySystem/MovieSceneEntityFactory.cpp
@@ -153,7 +153,19 @@ FBoundObjectTask::FBoundObjectTask(UMovieSceneEntitySystemLinker* InLinker)
 	: Linker(InLinker)
 {}
 
-void FBoundObjectTask::ForEachAllocation(FEntityAllocationProxy AllocationProxy, FReadEntityIDs EntityIDs, TRead<FInstanceHandle> Instances, TRead<FGuid> ObjectBindings)
+void FBoundObjectTask::Apply()
+{
+	for (TTuple<FEntityAllocationProxy, FObjectFactoryBatch>& Pair : Batches)
+	{
+		// Determine the type for the new entities
+		if (Pair.Value.Num() != 0)
+		{
+			Pair.Value.Apply(Linker, Pair.Key);
+		}
+	}
+}
+
+void FBoundObjectTask::ForEachAllocation(FEntityAllocationProxy AllocationProxy, FReadEntityIDs EntityIDs, TRead<FInstanceHandle> Instances, TRead<FGuid> ObjectBindings, TReadOptional<FBoundObjectResolver> Resolvers)
 {
 	const FEntityAllocation* Allocation = AllocationProxy.GetAllocation();
 	const FComponentTypeID TagHasUnresolvedBinding = FBuiltInComponentTypes::Get()->Tags.HasUnresolvedBinding;
@@ -161,7 +173,7 @@ void FBoundObjectTask::ForEachAllocation(FEntityAllocationProxy AllocationProxy,
 	// Check whether every binding in this allocation is currently unresolved
 	const bool bWasUnresolvedBinding = Allocation->FindComponentHeader(TagHasUnresolvedBinding) != nullptr;
 
-	FObjectFactoryBatch& Batch = AddBatch(AllocationProxy);
+	FObjectFactoryBatch& Batch = Batches.Emplace(AllocationProxy);
 	Batch.StaleEntitiesToPreserve = &StaleEntitiesToPreserve;
 
 	const int32 Num = Allocation->Num();
@@ -171,6 +183,8 @@ void FBoundObjectTask::ForEachAllocation(FEntityAllocationProxy AllocationProxy,
 	// Keep track of existing bindings so we can preserve any components on them
 	TComponentTypeID<UObject*> BoundObjectComponent = FBuiltInComponentTypes::Get()->BoundObject;
 
+	const FBoundObjectResolver* ResolverPtr = Resolvers.AsPtr();
+
 	for (int32 Index = 0; Index < Num; ++Index)
 	{
 		FMovieSceneEntityID ParentID = EntityIDs[Index];
@@ -190,25 +204,47 @@ void FBoundObjectTask::ForEachAllocation(FEntityAllocationProxy AllocationProxy,
 			}
 		}
 
-		const FObjectFactoryBatch::EResolveError Error = Batch.ResolveObjects(InstanceRegistry, Instances[Index], Index, ObjectBindings[Index]);
-		if (Error == FObjectFactoryBatch::EResolveError::None)
+		bool bIsResolvedBinding = false;
+
+		const FSequenceInstance&     SequenceInstance = InstanceRegistry->GetInstance(Instances[Index]);
+		TArrayView<TWeakObjectPtr<>> BoundObjects     = SequenceInstance.GetSharedPlaybackState()->FindBoundObjects(ObjectBindings[Index], SequenceInstance.GetSequenceID());
+
+		for (TWeakObjectPtr<> WeakObject : BoundObjects)
 		{
-			// We have successfully resolved a binding, so remove the HasUnresolvedBinding tag
-			if (bWasUnresolvedBinding)
+			UObject* Object = WeakObject.Get();
+
+			// Pass the object through the resolver component if necessary
+			if (ResolverPtr && Object)
 			{
-				constexpr bool bAddComponent = false;
-				EntityMutations.Add(FEntityMutationData{ ParentID, TagHasUnresolvedBinding, bAddComponent });
+				Object = (ResolverPtr[Index])(Object);
 			}
-		}
-		else if (Error == FObjectFactoryBatch::EResolveError::UnresolvedBinding)
-		{
-			if (!bWasUnresolvedBinding)
+
+			if (Object)
 			{
-				// Only bother attempting to add the HasUnresolvedBindingTag if it is not already tagged in such a way
-				constexpr bool bAddComponent = true;
-				EntityMutations.Add(FEntityMutationData{ ParentID, TagHasUnresolvedBinding, bAddComponent });
+				if (!ensureMsgf(!FBuiltInComponentTypes::IsBoundObjectGarbage(Object), TEXT("Attempting to bind an object that is garbage or unreachable")))
+				{
+					continue;
+				}
+
+				// Make a child entity for this resolved binding
+				Batch.Add(Index, Object);
+				bIsResolvedBinding = true;
 			}
 		}
+
+
+		if (bIsResolvedBinding && bWasUnresolvedBinding)
+		{
+			// We have successfully resolved a binding, so remove the HasUnresolvedBinding tag
+			constexpr bool bAddComponent = false;
+			EntityMutations.Add(FEntityMutationData{ ParentID, TagHasUnresolvedBinding, bAddComponent });
+		}
+		else if (!bIsResolvedBinding && !bWasUnresolvedBinding)
+		{
+			// Only bother attempting to add the HasUnresolvedBindingTag if it is not already tagged in such a way
+			constexpr bool bAddComponent = true;
+			EntityMutations.Add(FEntityMutationData{ ParentID, TagHasUnresolvedBinding, bAddComponent });
+		}
 	}
 
 }
diff --git a/Engine/Source/Runtime/MovieScene/Private/MovieSceneCommonHelpers.cpp b/Engine/Source/Runtime/MovieScene/Private/MovieSceneCommonHelpers.cpp
--- a/Engine/Source/Runtime/MovieScene/Private/MovieSceneCommonHelpers.cpp
+++ b/Engine/Source/Runtime/MovieScene/Private/MovieSceneCommonHelpers.cpp
@@ -469,6 +469,11 @@ void MovieSceneHelpers::GetDescendantSubSections(const UMovieScene* InMovieScene
 	}
 }
 
+UObject* MovieSceneHelpers::ResolveSceneComponentBoundObject(UObject* Object)
+{
+	return SceneComponentFromRuntimeObject(Object);
+}
+
 USceneComponent* MovieSceneHelpers::SceneComponentFromRuntimeObject(UObject* Object)
 {
 	AActor* Actor = Cast<AActor>(Object);
diff --git a/Engine/Source/Runtime/MovieScene/Public/EntitySystem/BuiltInComponentTypes.h b/Engine/Source/Runtime/MovieScene/Public/EntitySystem/BuiltInComponentTypes.h
--- a/Engine/Source/Runtime/MovieScene/Public/EntitySystem/BuiltInComponentTypes.h
+++ b/Engine/Source/Runtime/MovieScene/Public/EntitySystem/BuiltInComponentTypes.h
@@ -342,6 +342,8 @@ struct FMovieSceneBindingLifetimeComponentData
 	EMovieSceneBindingLifetimeState BindingLifetimeState = EMovieSceneBindingLifetimeState::Active;
 };
 
+using FBoundObjectResolver = UObject* (*)(UObject*);
+
 /**
  * Specifies a unique, sorted path of hbiases that contribute to a blended output
  * Supports up to 8 unique HBiases in its path
@@ -610,7 +612,11 @@ public:
 	// An FGuid relating to a direct object binding in a sequence
 	TComponentTypeID<FGuid> GenericObjectBinding;
 
+	// A custom bound object resolver that defines a function to resolve a bound object (ie, from AActor -> RootComponent where a track must operate on a component)
+	TComponentTypeID<FBoundObjectResolver> BoundObjectResolver;
+
 	// An FGuid that is always resolved as a USceneComponent either directly or through the AActor that the GUID relates to
+	UE_DEPRECATED(5.5, "Please use GenericObjectBinding and BoundObjectResolver")
 	TComponentTypeID<FGuid> SceneComponentBinding;
 
 	// An FGuid relating to a spawnable binding in a sequence
@@ -743,7 +749,7 @@ public:
 		FComponentTypeID NeedsLink;
 		FComponentTypeID NeedsUnlink;
 
-		/** Tag that is added to imported entities with a GenericObjectBinding or SceneComponentBinding whose binding did not resolve */
+		/** Tag that is added to imported entities with a GenericObjectBinding whose binding did not resolve */
 		FComponentTypeID HasUnresolvedBinding;
 
 		FComponentTypeID HasAssignedInitialValue;
diff --git a/Engine/Source/Runtime/MovieScene/Public/EntitySystem/MovieSceneEntityFactoryTemplates.h b/Engine/Source/Runtime/MovieScene/Public/EntitySystem/MovieSceneEntityFactoryTemplates.h
--- a/Engine/Source/Runtime/MovieScene/Public/EntitySystem/MovieSceneEntityFactoryTemplates.h
+++ b/Engine/Source/Runtime/MovieScene/Public/EntitySystem/MovieSceneEntityFactoryTemplates.h
@@ -26,6 +26,7 @@ namespace MovieScene
 {
 
 struct FInstanceRegistry;
+using FBoundObjectResolver = UObject* (*)(UObject*);
 
 template<typename ParentComponentType, typename ChildComponentType>
 struct TChildEntityInitializer : FChildEntityInitializer
@@ -108,12 +109,6 @@ struct TDuplicateChildEntityInitializer : FChildEntityInitializer
 
 struct FObjectFactoryBatch : FChildEntityFactory
 {
-	enum class EResolveError
-	{
-		None              = 0x0,
-		UnresolvedBinding = 0x1,
-	};
-
 	void Add(int32 EntityIndex, UObject* BoundObject);
 
 	virtual void GenerateDerivedType(FComponentMask& OutNewEntityType) override;
@@ -122,25 +117,20 @@ struct FObjectFactoryBatch : FChildEntityFactory
 
 	virtual void PostInitialize(UMovieSceneEntitySystemLinker* InLinker) override;
 
-	virtual EResolveError ResolveObjects(FInstanceRegistry* InstanceRegistry, FInstanceHandle InstanceHandle, int32 InEntityIndex, const FGuid& ObjectBinding) = 0;
-
 	TMap<TTuple<UObject*, FMovieSceneEntityID>, FMovieSceneEntityID>* StaleEntitiesToPreserve;
 
 private:
 	TSortedMap<FMovieSceneEntityID, FMovieSceneEntityID> PreservedEntities;
 	TArray<UObject*> ObjectsToAssign;
 };
-ENUM_CLASS_FLAGS(FObjectFactoryBatch::EResolveError)
 
 struct FBoundObjectTask
 {
 	FBoundObjectTask(UMovieSceneEntitySystemLinker* InLinker);
-	virtual ~FBoundObjectTask(){}
 
-	virtual FObjectFactoryBatch& AddBatch(FEntityAllocationProxy ParentProxy) = 0;
-	virtual void Apply() = 0;
+	void Apply();
 
-	void ForEachAllocation(FEntityAllocationProxy AllocationProxy, FReadEntityIDs EntityIDs, TRead<FInstanceHandle> Instances, TRead<FGuid> ObjectBindings);
+	void ForEachAllocation(FEntityAllocationProxy AllocationProxy, FReadEntityIDs EntityIDs, TRead<FInstanceHandle> Instances, TRead<FGuid> ObjectBindings, TReadOptional<FBoundObjectResolver> Resolvers);
 
 	void PostTask();
 
@@ -154,6 +144,7 @@ private:
 	};
 
 	TMap<TTuple<UObject*, FMovieSceneEntityID>, FMovieSceneEntityID> StaleEntitiesToPreserve;
+	TMap<FEntityAllocationProxy, FObjectFactoryBatch> Batches;
 	TArray<FMovieSceneEntityID> EntitiesToDiscard;
 	TArray<FEntityMutationData> EntityMutations;
 
@@ -162,36 +153,6 @@ protected:
 	UMovieSceneEntitySystemLinker* Linker;
 };
 
-template<typename BatchType>
-struct TBoundObjectTask : FBoundObjectTask
-{
-	TBoundObjectTask(UMovieSceneEntitySystemLinker* InLinker)
-		: FBoundObjectTask(InLinker)
-	{}
-
-private:
-
-	virtual FObjectFactoryBatch& AddBatch(FEntityAllocationProxy ParentProxy) override
-	{
-		return Batches.Add(ParentProxy);
-	}
-
-	virtual void Apply() override
-	{
-		for (TTuple<FEntityAllocationProxy, BatchType>& Pair : Batches)
-		{
-			// Determine the type for the new entities
-			if (Pair.Value.Num() != 0)
-			{
-				Pair.Value.Apply(Linker, Pair.Key);
-			}
-		}
-	}
-
-	TMap<FEntityAllocationProxy, BatchType> Batches;
-};
-
-
 
 template<typename ComponentType>
 inline void FEntityFactories::DuplicateChildComponent(TComponentTypeID<ComponentType> InComponent)
diff --git a/Engine/Source/Runtime/MovieScene/Public/MovieSceneCommonHelpers.h b/Engine/Source/Runtime/MovieScene/Public/MovieSceneCommonHelpers.h
--- a/Engine/Source/Runtime/MovieScene/Public/MovieSceneCommonHelpers.h
+++ b/Engine/Source/Runtime/MovieScene/Public/MovieSceneCommonHelpers.h
@@ -124,6 +124,7 @@ public:
 	 * @return The found scene component
 	 */	
 	static MOVIESCENE_API USceneComponent* SceneComponentFromRuntimeObject(UObject* Object);
+	static MOVIESCENE_API UObject* ResolveSceneComponentBoundObject(UObject* Object);
 
 	/**
 	 * Get the active camera component from the actor 
diff --git a/Engine/Source/Runtime/MovieSceneTracks/Private/Sections/MovieScene3DAttachSection.cpp b/Engine/Source/Runtime/MovieSceneTracks/Private/Sections/MovieScene3DAttachSection.cpp
--- a/Engine/Source/Runtime/MovieSceneTracks/Private/Sections/MovieScene3DAttachSection.cpp
+++ b/Engine/Source/Runtime/MovieSceneTracks/Private/Sections/MovieScene3DAttachSection.cpp
@@ -56,7 +56,8 @@ void UMovieScene3DAttachSection::ImportEntityImpl(UMovieSceneEntitySystemLinker*
 
 	OutImportedEntity->AddBuilder(
 		FEntityBuilder()
-		.Add(BuiltInComponentTypes->SceneComponentBinding, Params.GetObjectBindingID())
+		.Add(BuiltInComponentTypes->GenericObjectBinding, Params.GetObjectBindingID())
+		.Add(BuiltInComponentTypes->BoundObjectResolver, MovieSceneHelpers::ResolveSceneComponentBoundObject)
 		.Add(TrackComponents->AttachParentBinding, ConstraintBindingID)
 		.Add(TrackComponents->AttachComponent, AttachComponent)
 	);
diff --git a/Engine/Source/Runtime/MovieSceneTracks/Private/Sections/MovieScene3DTransformSection.cpp b/Engine/Source/Runtime/MovieSceneTracks/Private/Sections/MovieScene3DTransformSection.cpp
--- a/Engine/Source/Runtime/MovieSceneTracks/Private/Sections/MovieScene3DTransformSection.cpp
+++ b/Engine/Source/Runtime/MovieSceneTracks/Private/Sections/MovieScene3DTransformSection.cpp
@@ -632,7 +632,8 @@ void UMovieScene3DTransformSection::ImportConstraintEntity(UMovieSceneEntitySyst
 			ComponentData.Section = this;
 			OutImportedEntity->AddBuilder(
 				FEntityBuilder()
-				.Add(BuiltInComponentTypes->SceneComponentBinding, ObjectBindingID)
+				.Add(BuiltInComponentTypes->BoundObjectResolver, MovieSceneHelpers::ResolveSceneComponentBoundObject)
+				.Add(BuiltInComponentTypes->GenericObjectBinding, ObjectBindingID)
 				.Add(TrackComponents->ConstraintChannel, ComponentData)
 			);
 		}
@@ -654,16 +655,13 @@ void UMovieScene3DTransformSection::ImportEntityImpl(UMovieSceneEntitySystemLink
 
 	check(Track);
 
-	// 3D Transform tracks use a scene component binding by default. Every other transform property track must be bound directly to the object.
-	const TComponentTypeID<FGuid>& ObjectBinding = Track->IsA<UMovieScene3DTransformTrack>()
-		? BuiltInComponentTypes->SceneComponentBinding
-		: BuiltInComponentTypes->GenericObjectBinding;
-
 	FGuid ObjectBindingID = Params.GetObjectBindingID();
 
 	auto BaseBuilder = FEntityBuilder()
 		.Add(BuiltInComponentTypes->PropertyBinding, Track->GetPropertyBinding())
-		.AddConditional(ObjectBinding,               ObjectBindingID, ObjectBindingID.IsValid());
+		// 3D Transform tracks use a scene component binding by default. Every other transform property track must be bound directly to the object.
+		.AddConditional(BuiltInComponentTypes->BoundObjectResolver, MovieSceneHelpers::ResolveSceneComponentBoundObject, Track->IsA<UMovieScene3DTransformTrack>())
+		.AddConditional(BuiltInComponentTypes->GenericObjectBinding, ObjectBindingID, ObjectBindingID.IsValid());
 
 	BuildEntity(BaseBuilder, EntityLinker, Params, OutImportedEntity);
 }
