Sequencer: Fixed 2 issues with converting between spawnable types

The first is that the spawn track is never modified when adding a new spawn section. This means it doesn't get added to the transaction, and thus means that undoing the operation will not undo the change to UMovieSceneSpawnTrack::Sections when the spawn section was added. This leaves the (now garbage) section in the array leading to the crash listed in this jira. The fix is to just call SpawnTrack->Modify().

The second issue is that none of this code should even be running because the spawn track already has a spawn section on it which should be preserved rather than an additional infinite section being added on top. The fix is to check whether the spawn track has any sections before creating this new infinite one.

diff --git a/Engine/Source/Runtime/MovieScene/Private/Bindings/MovieSceneSpawnableBinding.cpp b/Engine/Source/Runtime/MovieScene/Private/Bindings/MovieSceneSpawnableBinding.cpp
--- a/Engine/Source/Runtime/MovieScene/Private/Bindings/MovieSceneSpawnableBinding.cpp
+++ b/Engine/Source/Runtime/MovieScene/Private/Bindings/MovieSceneSpawnableBinding.cpp
@@ -182,8 +182,10 @@ void UMovieSceneSpawnableBindingBase::SetupDefaults(UObject* SpawnedObject, FGui
 		SpawnTrack = Cast<UMovieSceneSpawnTrack>(OwnerMovieScene.AddTrack(UMovieSceneSpawnTrack::StaticClass(), ObjectBindingId));
 	}
 
-	if (SpawnTrack)
+	if (SpawnTrack && SpawnTrack->GetAllSections().Num() == 0)
 	{
+		SpawnTrack->Modify();
+
 		UMovieSceneBoolSection* SpawnSection = Cast<UMovieSceneBoolSection>(SpawnTrack->CreateNewSection());
 		SpawnSection->GetChannel().SetDefault(true);
 		SpawnSection->SetRange(TRange<FFrameNumber>::All());
