Updated spawnables to have a more readable label after spawning, which is usually based off the actor name in the editor (rather than a label that has no relation to the actor name, like "StaticMeshActor0"). This has a slight potential to cause performance issues based on past behavior, so a cvar was added to disable this behavior if needed (LevelSequence.EnableReadableActorLabelsForSpawnables).

Having more predictable labels for spawnables will help users match actors in Movie Render Graph collections.

diff --git a/Engine/Plugins/Experimental/DaySequence/Source/DaySequence/Private/DaySequenceActorSpawner.cpp b/Engine/Plugins/Experimental/DaySequence/Source/DaySequence/Private/DaySequenceActorSpawner.cpp
--- a/Engine/Plugins/Experimental/DaySequence/Source/DaySequence/Private/DaySequenceActorSpawner.cpp
+++ b/Engine/Plugins/Experimental/DaySequence/Source/DaySequence/Private/DaySequenceActorSpawner.cpp
@@ -244,8 +244,12 @@ UObject* FDaySequenceActorSpawner::SpawnObject(FMovieSceneSpawnable& Spawnable,
 	SpawnedActor->FinishSpawning(SpawnTransform, bIsDefaultTransform);
 
 #if WITH_EDITOR
-	// Don't set the actor label in PIE as this requires flushing async loading.
-	if (WorldContext->WorldType == EWorldType::Editor)
+	const IConsoleVariable* AllowSetActorLabelCvar = IConsoleManager::Get().FindConsoleVariable(TEXT("LevelSequence.EnableReadableActorLabelsForSpawnables"));
+	const bool bAllowSetActorLabel = AllowSetActorLabelCvar && AllowSetActorLabelCvar->GetBool();
+	
+	// Historically, setting the actor label has caused performance issues in some scenarios (by causing async loading flushes); however, there's no
+	// evidence for this anymore, so the cvar is here to turn off this behavior if needed.
+	if ((WorldContext->WorldType == EWorldType::Editor) || bAllowSetActorLabel)
 	{
 		SpawnedActor->SetActorLabel(Spawnable.GetName());
 	}
diff --git a/Engine/Source/Runtime/LevelSequence/Private/LevelSequenceActorSpawner.cpp b/Engine/Source/Runtime/LevelSequence/Private/LevelSequenceActorSpawner.cpp
--- a/Engine/Source/Runtime/LevelSequence/Private/LevelSequenceActorSpawner.cpp
+++ b/Engine/Source/Runtime/LevelSequence/Private/LevelSequenceActorSpawner.cpp
@@ -242,8 +242,12 @@ UObject* FLevelSequenceActorSpawner::SpawnObject(FMovieSceneSpawnable& Spawnable
 	SpawnedActor->FinishSpawning(SpawnTransform, bIsDefaultTransform);
 
 #if WITH_EDITOR
-	// Don't set the actor label in PIE as this requires flushing async loading.
-	if (WorldContext->WorldType == EWorldType::Editor)
+	const IConsoleVariable* AllowSetActorLabelCvar = IConsoleManager::Get().FindConsoleVariable(TEXT("LevelSequence.EnableReadableActorLabelsForSpawnables"));
+	const bool bAllowSetActorLabel = AllowSetActorLabelCvar && AllowSetActorLabelCvar->GetBool();
+	
+	// Historically, setting the actor label has caused performance issues in some scenarios (by causing async loading flushes); however, there's no
+	// evidence for this anymore, so the cvar is here to turn off this behavior if needed.
+	if ((WorldContext->WorldType == EWorldType::Editor) || bAllowSetActorLabel)
 	{
 		SpawnedActor->SetActorLabel(Spawnable.GetName());
 	}
diff --git a/Engine/Source/Runtime/MovieScene/Private/Bindings/MovieSceneSpawnableBinding.cpp b/Engine/Source/Runtime/MovieScene/Private/Bindings/MovieSceneSpawnableBinding.cpp
--- a/Engine/Source/Runtime/MovieScene/Private/Bindings/MovieSceneSpawnableBinding.cpp
+++ b/Engine/Source/Runtime/MovieScene/Private/Bindings/MovieSceneSpawnableBinding.cpp
@@ -18,6 +18,13 @@
 
 #define LOCTEXT_NAMESPACE "FPossessableModel"
 
+static TAutoConsoleVariable<int32> CVarEnableReadableActorLabelsForSpawnables(
+	TEXT("LevelSequence.EnableReadableActorLabelsForSpawnables"),
+	1,
+	TEXT("If true, in the editor during PIE, Sequencer will set the DisplayName of spawned actors to match their Spawnable name in Sequencer, mimicking edit-time behavior. This helps with identifying spawnables more reliably, but isn't available in packaged builds. Try disabling this if you see async loads being flushed during actor spawning in PIE.\n")
+	TEXT("0: off, 1: on"),
+	ECVF_Default);
+
 UObject* UMovieSceneSpawnableBindingBase::SpawnObject(const FGuid& BindingId, int32 BindingIndex, UMovieScene& MovieScene, FMovieSceneSequenceIDRef TemplateID, TSharedRef<const UE::MovieScene::FSharedPlaybackState> SharedPlaybackState)
 {
 	UWorld* WorldContext = GetWorldContext(SharedPlaybackState);
@@ -88,8 +95,9 @@ UObject* UMovieSceneSpawnableBindingBase::SpawnObject(const FGuid& BindingId, in
 #endif
 
 #if WITH_EDITOR
-		// Don't set the actor label in PIE as this requires flushing async loading.
-		if (WorldContext->WorldType == EWorldType::Editor)
+		// Historically, setting the actor label has caused performance issues in some scenarios (by causing async loading flushes); however, there's no
+		// evidence for this anymore, so the cvar is here to turn off this behavior if needed.
+		if ((WorldContext->WorldType == EWorldType::Editor) || CVarEnableReadableActorLabelsForSpawnables->GetBool())
 		{
 			FString BindingName = GetDesiredBindingName();
 			FString ActorLabel = !BindingName.IsEmpty() ? BindingName : SpawnName.ToString();
