Sequencer: Automatically enable rich component debugging if there is a debugger attached

diff --git a/Engine/Source/Runtime/MovieScene/Private/EntitySystem/MovieSceneEntityManager.cpp b/Engine/Source/Runtime/MovieScene/Private/EntitySystem/MovieSceneEntityManager.cpp
--- a/Engine/Source/Runtime/MovieScene/Private/EntitySystem/MovieSceneEntityManager.cpp
+++ b/Engine/Source/Runtime/MovieScene/Private/EntitySystem/MovieSceneEntityManager.cpp
@@ -57,6 +57,7 @@ FAutoConsoleVariableRef CVarThreadedEvaluationEntityThreshold(
 );
 
 #if UE_MOVIESCENE_ENTITY_DEBUG
+bool GRichComponentDebuggingInitialized = false;
 bool GRichComponentDebugging = false;
 FAutoConsoleVariableRef CVarRichComponentDebugging(
 	TEXT("Sequencer.RichComponentDebugging"),
@@ -413,6 +414,14 @@ FEntityManager::FEntityManager()
 	SystemSerialNumber = 1;
 	StructureMutationSystemSerialNumber = 0;
 	ThreadingModel = EEntityThreadingModel::NoThreading;
+
+#if UE_MOVIESCENE_ENTITY_DEBUG
+	if (!GRichComponentDebuggingInitialized)
+	{
+		GRichComponentDebugging = FPlatformMisc::IsDebuggerPresent();
+		GRichComponentDebuggingInitialized = true;
+	}
+#endif
 }
 
 FEntityManager::~FEntityManager()
