Sequencer- Crash when using raw Actor reference as an argument to a 'Dynamic Binding'/Director Blueprint binding.

The 'raw actor' case adds a bunch of intermediate nodes, and the code changed above was assuming all child fields found in the FieldIterator were parameters, when in fact some were these intermediate nodes that got overwritten by this behavior.

diff --git a/Engine/Source/Runtime/MovieScene/Private/MovieSceneDynamicBindingInvoker.cpp b/Engine/Source/Runtime/MovieScene/Private/MovieSceneDynamicBindingInvoker.cpp
--- a/Engine/Source/Runtime/MovieScene/Private/MovieSceneDynamicBindingInvoker.cpp
+++ b/Engine/Source/Runtime/MovieScene/Private/MovieSceneDynamicBindingInvoker.cpp
@@ -102,7 +102,7 @@ FMovieSceneDynamicBindingResolveResult FMovieSceneDynamicBindingInvoker::InvokeD
 	{
 		FProperty* LocalProp = *It;
 		checkSlow(LocalProp);
-		if (!LocalProp->HasAnyPropertyFlags(CPF_ZeroConstructor))
+		if (!LocalProp->HasAnyPropertyFlags(CPF_ZeroConstructor) && LocalProp->HasAllPropertyFlags(CPF_Parm))
 		{
 			LocalProp->InitializeValue_InContainer(Parameters);
 		}
@@ -140,7 +140,13 @@ FMovieSceneDynamicBindingResolveResult FMovieSceneDynamicBindingInvoker::InvokeD
 	// Destroy parameters.
 	for (TFieldIterator<FProperty> It(DynamicBindingFunc); It; ++It)
 	{
-		It->DestroyValue_InContainer(Parameters);
+		FProperty* LocalProp = *It;
+		checkSlow(LocalProp);
+
+		if (LocalProp->HasAllPropertyFlags(CPF_Parm))
+		{
+			It->DestroyValue_InContainer(Parameters);
+		}
 	}
 
 	return Result;
