Sequencer: Fix post infinity linear extrapolation for IntegerChannel and BoolChannel

diff --git a/Engine/Source/Runtime/MovieScene/Private/Channels/MovieSceneBoolChannel.cpp b/Engine/Source/Runtime/MovieScene/Private/Channels/MovieSceneBoolChannel.cpp
--- a/Engine/Source/Runtime/MovieScene/Private/Channels/MovieSceneBoolChannel.cpp
+++ b/Engine/Source/Runtime/MovieScene/Private/Channels/MovieSceneBoolChannel.cpp
@@ -69,7 +69,7 @@ bool FMovieSceneBoolChannel::Evaluate(FFrameTime InTime, bool& OutValue) const
 				return false;
 			}
 
-			if (PostInfinityExtrap == RCCE_Constant || PreInfinityExtrap == RCCE_Linear)
+			if (PostInfinityExtrap == RCCE_Constant || PostInfinityExtrap == RCCE_Linear)
 			{
 				OutValue = Values.Last();
 				return true;
diff --git a/Engine/Source/Runtime/MovieScene/Private/Channels/MovieSceneIntegerChannel.cpp b/Engine/Source/Runtime/MovieScene/Private/Channels/MovieSceneIntegerChannel.cpp
--- a/Engine/Source/Runtime/MovieScene/Private/Channels/MovieSceneIntegerChannel.cpp
+++ b/Engine/Source/Runtime/MovieScene/Private/Channels/MovieSceneIntegerChannel.cpp
@@ -70,7 +70,7 @@ bool FMovieSceneIntegerChannel::Evaluate(FFrameTime InTime, int32& OutValue) con
 				return false;
 			}
 
-			if (PostInfinityExtrap == RCCE_Constant || PreInfinityExtrap == RCCE_Linear)
+			if (PostInfinityExtrap == RCCE_Constant || PostInfinityExtrap == RCCE_Linear)
 			{
 				OutValue = Values.Last();
 				return true;
