Sequencer- Fix issue where invalidated spawnable bindings don't unspawn.

diff --git a/Engine/Source/Runtime/MovieScene/Private/EntitySystem/MovieSceneSpawnablesSystem.cpp b/Engine/Source/Runtime/MovieScene/Private/EntitySystem/MovieSceneSpawnablesSystem.cpp
--- a/Engine/Source/Runtime/MovieScene/Private/EntitySystem/MovieSceneSpawnablesSystem.cpp
+++ b/Engine/Source/Runtime/MovieScene/Private/EntitySystem/MovieSceneSpawnablesSystem.cpp
@@ -205,18 +205,24 @@ void UMovieSceneSpawnablesSystem::OnRun(FSystemTaskPrerequisites& InPrerequisite
 		.FilterAll({ BuiltInComponents->Tags.NeedsUnlink })
 		.Iterate_PerEntity(&Linker->EntityManager, DestroyOldSpawnables);
 
-	for (TTuple<FGuid, int32, FMovieSceneSequenceID, FInstanceHandle> Tuple : DestroyedObjects)
-	{
-		// Have to check whether the player is still valid because there is a possibility it got cleaned up
-		if (InstanceRegistry->IsHandleValid(Tuple.Get<3>()))
+	auto DoDestroy = [this, &DestroyedObjects, InstanceRegistry]() {
+
+		for (TTuple<FGuid, int32, FMovieSceneSequenceID, FInstanceHandle> Tuple : DestroyedObjects)
 		{
-			TSharedRef<const FSharedPlaybackState> SharedPlaybackState = InstanceRegistry->GetInstance(Tuple.Get<3>()).GetSharedPlaybackState();
-			if (FMovieSceneSpawnRegister* SpawnRegister = SharedPlaybackState->FindCapability<FMovieSceneSpawnRegister>())
+			// Have to check whether the player is still valid because there is a possibility it got cleaned up
+			if (InstanceRegistry->IsHandleValid(Tuple.Get<3>()))
 			{
-				SpawnRegister->DestroySpawnedObject(Tuple.Get<0>(), Tuple.Get<2>(), SharedPlaybackState, Tuple.Get<1>());
+				TSharedRef<const FSharedPlaybackState> SharedPlaybackState = InstanceRegistry->GetInstance(Tuple.Get<3>()).GetSharedPlaybackState();
+				if (FMovieSceneSpawnRegister* SpawnRegister = SharedPlaybackState->FindCapability<FMovieSceneSpawnRegister>())
+				{
+					SpawnRegister->DestroySpawnedObject(Tuple.Get<0>(), Tuple.Get<2>(), SharedPlaybackState, Tuple.Get<1>());
+				}
 			}
 		}
-	}
+		DestroyedObjects.Empty();
+	};
+
+	DoDestroy();
 
 	// ----------------------------------------------------------------------------------------------------------------------------------------
 	// Step 2 - iterate all pending spawnables and spawn their objects if necessary
@@ -311,11 +317,15 @@ void UMovieSceneSpawnablesSystem::OnRun(FSystemTaskPrerequisites& InPrerequisite
 
 		InstanceRegistry->InvalidateObjectBinding(SpawnableBindingID, InstanceHandle);
 	};
+
 	FEntityTaskBuilder()
 	.ReadEntityIDs()
 	.Read(BuiltInComponents->InstanceHandle)
 	.Read(BuiltInComponents->SpawnableBinding)
 	.FilterAll({ BuiltInComponents->Tags.NeedsLink })
 	.Iterate_PerEntity(&Linker->EntityManager, SpawnNewObjects);
+
+	// Destroy any spawnables that have since been added to the DestroyList
+	DoDestroy();
 }
 
