Make the movie scene subsystems work under AutoRTFM.

diff --git a/Engine/Source/Runtime/MovieScene/Private/EntitySystem/MovieSceneEntityManager.cpp b/Engine/Source/Runtime/MovieScene/Private/EntitySystem/MovieSceneEntityManager.cpp
--- a/Engine/Source/Runtime/MovieScene/Private/EntitySystem/MovieSceneEntityManager.cpp
+++ b/Engine/Source/Runtime/MovieScene/Private/EntitySystem/MovieSceneEntityManager.cpp
@@ -1048,19 +1048,13 @@ void FEntityManager::AccumulateMask(const FEntityComponentFilter& InFilter, FCom
 
 void FEntityManager::EnterIteration() const
 {
-	UE_AUTORTFM_OPEN2 { ++IterationCount; };
-
-	AutoRTFM::PushOnAbortHandler(this, [this] { --(this->IterationCount); });
+	IterationCount.Increment(ThreadingModel);
 }
 
 void FEntityManager::ExitIteration() const
 {
-	checkSlow(static_cast<uint16>(IterationCount) > 0);
-
-	// We only ever call the exit after an enter so we should not have to handle the case where exit is called first in order to ++ on the iterator
-	AutoRTFM::PopOnAbortHandler(this);
-
-	UE_AUTORTFM_OPEN2 { --IterationCount; };
+	checkSlow(IterationCount.Load(ThreadingModel) > 0);
+	IterationCount.Decrement(ThreadingModel);
 }
 
 FEntityAllocation* FEntityManager::CreateEntityAllocation(const FComponentMask& EntityComponentMask, uint16 InitialCapacity, uint16 MaxCapacity, FEntityAllocation* MigrateComponentDataFrom)
diff --git a/Engine/Source/Runtime/MovieScene/Private/EntitySystem/MovieSceneEntitySystemGraphs.cpp b/Engine/Source/Runtime/MovieScene/Private/EntitySystem/MovieSceneEntitySystemGraphs.cpp
--- a/Engine/Source/Runtime/MovieScene/Private/EntitySystem/MovieSceneEntitySystemGraphs.cpp
+++ b/Engine/Source/Runtime/MovieScene/Private/EntitySystem/MovieSceneEntitySystemGraphs.cpp
@@ -424,6 +424,9 @@ void FMovieSceneEntitySystemGraph::ExecutePhase(UE::MovieScene::ESystemPhase Pha
 
 	const EEntityThreadingModel ThreadingModel = Linker->EntityManager.GetThreadingModel();
 
+	// In a transaction we can only support the no-threading mode.
+	check(!AutoRTFM::IsTransactional() || (EEntityThreadingModel::NoThreading == ThreadingModel));
+
 	FSystemSubsequentTasks DownstreamTasks(this, &OutTasks, ThreadingModel);
 
 	FSystemTaskPrerequisites NoPrerequisites;
@@ -437,7 +440,7 @@ void FMovieSceneEntitySystemGraph::ExecutePhase(UE::MovieScene::ESystemPhase Pha
 		UMovieSceneEntitySystem* System = Nodes.Array[NodeID].System;
 		checkSlow(System);
 
-		// Initilaize downstream task structure for this system
+		// Initialize downstream task structure for this system
 		DownstreamTasks.ResetNode(NodeID);
 
 		TSharedPtr<FSystemTaskPrerequisites> Prerequisites = Nodes.Array[NodeID].Prerequisites;
diff --git a/Engine/Source/Runtime/MovieScene/Private/EntitySystem/MovieSceneTaskScheduler.cpp b/Engine/Source/Runtime/MovieScene/Private/EntitySystem/MovieSceneTaskScheduler.cpp
--- a/Engine/Source/Runtime/MovieScene/Private/EntitySystem/MovieSceneTaskScheduler.cpp
+++ b/Engine/Source/Runtime/MovieScene/Private/EntitySystem/MovieSceneTaskScheduler.cpp
@@ -56,8 +56,8 @@ FScheduledTask::FScheduledTask(FScheduledTask&& InTask)
 	, WriteContextOffset(InTask.WriteContextOffset)
 	, LockedComponentData(InTask.LockedComponentData)
 	, NumPrerequisites(InTask.NumPrerequisites)
-	, WaitCount(InTask.WaitCount.load())
-	, ChildCompleteCount(InTask.ChildCompleteCount.load())
+	, WaitCount(InTask.WaitCount.Load(EEntityThreadingModel::NoThreading))
+	, ChildCompleteCount(InTask.ChildCompleteCount.Load(EEntityThreadingModel::NoThreading))
 	, Parent(InTask.Parent)
 	, TaskFunctionType(InTask.TaskFunctionType)
 {
@@ -138,7 +138,7 @@ void FScheduledTask::Run(const FEntitySystemScheduler* Scheduler, FTaskExecution
 		// and destroy or  otherwise mutate the contents of FEntitySystemScheduler resulting in a crash.
 		//
 		// Once our loop has finished we check the child complete count to see if this was the last one
-		ChildCompleteCount.fetch_add(1);
+		ChildCompleteCount.Add(Scheduler->GetEntityManager()->GetThreadingModel(), 1);
 
 		for (int32 Index : ChildTasks)
 		{
@@ -146,7 +146,7 @@ void FScheduledTask::Run(const FEntitySystemScheduler* Scheduler, FTaskExecution
 		}
 
 		// Subtract our count added on ln 205. If this is the last count, complete this task (all children have completed)
-		const int32 PreviousCompleteCount = ChildCompleteCount.fetch_sub(1);
+		const int32 PreviousCompleteCount = ChildCompleteCount.Sub(Scheduler->GetEntityManager()->GetThreadingModel(), 1);
 		if (PreviousCompleteCount == 1)
 		{
 			Scheduler->CompleteTask(this, InFlags);
@@ -154,7 +154,7 @@ void FScheduledTask::Run(const FEntitySystemScheduler* Scheduler, FTaskExecution
 	}
 	else
 	{
-		checkSlow(ChildCompleteCount.load() == 0);
+		checkSlow(ChildCompleteCount.Load(Scheduler->GetEntityManager()->GetThreadingModel()) == 0);
 		Scheduler->CompleteTask(this, InFlags);
 	}
 }
@@ -484,10 +484,15 @@ void FEntitySystemScheduler::ExecuteTasks()
 		return;
 	}
 
-	const int32 PreviousNumRemaining = NumTasksRemaining.exchange(Tasks.Num());
+	const EEntityThreadingModel ThreadingModelToUse = EntityManager->GetThreadingModel();
+
+	const int32 PreviousNumRemaining = NumTasksRemaining.Exchange(ThreadingModelToUse, Tasks.Num());
 	check(PreviousNumRemaining == 0);
 
-	ThreadingModel = EntityManager->GetThreadingModel();
+	ThreadingModel = ThreadingModelToUse;
+	// In a transaction we can only support the no-threading mode.
+	check(!AutoRTFM::IsTransactional() || (EEntityThreadingModel::NoThreading == ThreadingModel));
+
 	WriteContextBase = FEntityAllocationWriteContext(*EntityManager);
 
 	// Condition 1: No threading
@@ -500,7 +505,7 @@ void FEntitySystemScheduler::ExecuteTasks()
 			Tasks[Index].Run(this, FTaskExecutionFlags());
 		}
 
-		check(NumTasksRemaining.load() == 0);
+		check(NumTasksRemaining.Load(ThreadingModel) == 0);
 		EntityManager->IncrementSystemSerial(SystemSerialIncrement);
 		return;
 	}
@@ -545,7 +550,7 @@ void FEntitySystemScheduler::ExecuteTasks()
 			Task->Run(this, FTaskExecutionFlags());
 		}
 
-		if (NumTasksRemaining.load() == 0)
+		if (NumTasksRemaining.Load(ThreadingModel) == 0)
 		{
 			break;
 		}
@@ -553,15 +558,15 @@ void FEntitySystemScheduler::ExecuteTasks()
 		GameThreadSignal->Wait();
 	}
 
-	check(NumTasksRemaining.load() == 0);
+	check(NumTasksRemaining.Load(ThreadingModel) == 0);
 	EntityManager->IncrementSystemSerial(SystemSerialIncrement);
 }
 
 void FEntitySystemScheduler::CompleteTask(const FScheduledTask* Task, FTaskExecutionFlags InFlags) const
 {
 	// Reset the WaitCount ready for the next run
-	const int32 PreviousWaitCount = Task->WaitCount.exchange(Task->NumPrerequisites);
-	const int32 PreviousChildCount = Task->ChildCompleteCount.exchange(Task->ChildTasks.CountSetBits());
+	const int32 PreviousWaitCount = Task->WaitCount.Exchange(ThreadingModel, Task->NumPrerequisites);
+	const int32 PreviousChildCount = Task->ChildCompleteCount.Exchange(ThreadingModel, Task->ChildTasks.CountSetBits());
 
 	checkSlow(PreviousWaitCount == 0 && PreviousChildCount == 0);
 
@@ -575,7 +580,7 @@ void FEntitySystemScheduler::CompleteTask(const FScheduledTask* Task, FTaskExecu
 	if (Task->Parent)
 	{
 		const FScheduledTask* Parent = &Tasks[Task->Parent.Index];
-		const int32 PreviousCompleteCount = Parent->ChildCompleteCount.fetch_sub(1);
+		const int32 PreviousCompleteCount = Parent->ChildCompleteCount.Sub(ThreadingModel, 1);
 		if (PreviousCompleteCount == 1)
 		{
 			CompleteTask(Parent, InFlags);
@@ -587,7 +592,7 @@ void FEntitySystemScheduler::CompleteTask(const FScheduledTask* Task, FTaskExecu
 		Tasks[FirstInlineIndex].Run(this, FTaskExecutionFlags());
 	}
 
-	const int32 PreviousNumRemaining = NumTasksRemaining.fetch_sub(1);
+	const int32 PreviousNumRemaining = NumTasksRemaining.Sub(ThreadingModel, 1);
 	if (PreviousNumRemaining == 1)
 	{
 		OnAllTasksFinished();
@@ -596,7 +601,7 @@ void FEntitySystemScheduler::CompleteTask(const FScheduledTask* Task, FTaskExecu
 
 void FEntitySystemScheduler::PrerequisiteCompleted(FTaskID TaskID, int32* OptRunInlineIndex) const
 {
-	int32 PreviousWaitCount = Tasks[TaskID.Index].WaitCount.fetch_sub(1);
+	int32 PreviousWaitCount = Tasks[TaskID.Index].WaitCount.Sub(ThreadingModel, 1);
 	checkSlow(PreviousWaitCount >= 1);
 	if (PreviousWaitCount <= 1)
 	{
@@ -707,8 +712,8 @@ void FEntitySystemScheduler::EndConstruction()
 	for (int32 Index = 0; Index < Tasks.Num(); ++Index)
 	{
 		FScheduledTask& Task = Tasks[Index];
-		Task.WaitCount.exchange(Task.NumPrerequisites);
-		Task.ChildCompleteCount.exchange(Task.ChildTasks.CountSetBits());
+		Task.WaitCount.Exchange(ThreadingModel, Task.NumPrerequisites);
+		Task.ChildCompleteCount.Exchange(ThreadingModel, Task.ChildTasks.CountSetBits());
 
 		if (Task.NumPrerequisites == 0)
 		{
diff --git a/Engine/Source/Runtime/MovieScene/Private/EntitySystem/MovieSceneTaskScheduler.h b/Engine/Source/Runtime/MovieScene/Private/EntitySystem/MovieSceneTaskScheduler.h
--- a/Engine/Source/Runtime/MovieScene/Private/EntitySystem/MovieSceneTaskScheduler.h
+++ b/Engine/Source/Runtime/MovieScene/Private/EntitySystem/MovieSceneTaskScheduler.h
@@ -11,8 +11,8 @@
 #include "EntitySystem/RelativePtr.h"
 #include "EntitySystem/IMovieSceneTaskScheduler.h"
 #include "EntitySystem/MovieSceneEntityManager.h"
-
-#include <atomic>
+#include "EntitySystem/MovieSceneMaybeAtomic.h"
+#include "Misc/TransactionallySafeCriticalSection.h"
 
 namespace UE::MovieScene
 {
@@ -150,9 +150,9 @@ struct FScheduledTask
 	/** 4 Bytes - the total number of tasks that must complete before this one can begin */
 	int32 NumPrerequisites = 0;
 	/** 4 Bytes - the number of outstanding prerequisite tasks this task is waiting on. Reset to NumPrerequisites on completion */
-	mutable std::atomic<int32> WaitCount = std::atomic<int32>(0);
+	mutable FEntitySystemMaybeAtomicInt32 WaitCount = 0;
 	/** 4 Bytes - the number of child tasks that must be completed before this task is considered complete */
-	mutable std::atomic<int32> ChildCompleteCount = std::atomic<int32>(0);
+	mutable FEntitySystemMaybeAtomicInt32 ChildCompleteCount = 0;
 
 	/** 4 Bytes - This task's parent (or None() if it is not a child task) */
 	FTaskID Parent;
@@ -350,7 +350,7 @@ private:
 
 	FEntityManager* EntityManager;
 
-	mutable std::atomic<int32> NumTasksRemaining = 0;
+	mutable FEntitySystemMaybeAtomicInt32 NumTasksRemaining = 0;
 
 	FEvent* GameThreadSignal = nullptr;
 	mutable TLockFreePointerListFIFO<FScheduledTask, PLATFORM_CACHE_LINE_SIZE> GameThreadTaskList;
diff --git a/Engine/Source/Runtime/MovieScene/Private/Tests/AutoRTFM/AutoRTFMMovieSceneTests.cpp b/Engine/Source/Runtime/MovieScene/Private/Tests/AutoRTFM/AutoRTFMMovieSceneTests.cpp
new file mode 100644
--- /dev/null
+++ b/Engine/Source/Runtime/MovieScene/Private/Tests/AutoRTFM/AutoRTFMMovieSceneTests.cpp
@@ -0,0 +1,117 @@
+// Copyright Epic Games, Inc. All Rights Reserved.
+
+#include "Misc/AutomationTest.h"
+#include "Misc/App.h"
+#include "AutoRTFM/AutoRTFM.h"
+#include "EntitySystem/MovieSceneEntityBuilder.h"
+#include "EntitySystem/MovieSceneEntitySystemTypes.h"
+#include "EntitySystem/MovieSceneEntityManager.h"
+#include "EntitySystem/MovieSceneEntitySystemTask.h"
+#include "EntitySystem/MovieSceneTaskScheduler.h"
+#include "EntitySystem/MovieSceneEntityFactoryTemplates.h"
+
+#if WITH_DEV_AUTOMATION_TESTS
+
+#define TEST_CHECK_TRUE(b) do \
+{ \
+	if (!(b)) \
+	{ \
+		FString String(FString::Printf(TEXT("FAILED: %s:%u"), TEXT(__FILE__), __LINE__)); \
+		ExecutionInfo.AddEvent(FAutomationEvent(EAutomationEventType::Info, String)); \
+		return false; \
+	} \
+} while(false)
+
+IMPLEMENT_SIMPLE_AUTOMATION_TEST(FAutoRTFMMovieSceneTests, "AutoRTFM + Movie Scene", EAutomationTestFlags::EngineFilter | EAutomationTestFlags::ClientContext | EAutomationTestFlags::ServerContext | EAutomationTestFlags::CommandletContext)
+
+using namespace UE::MovieScene;
+
+struct FMyTask final : ITaskContext
+{
+	void Run(FEntityAllocationWriteContext& WriteContext) const { WasHit += 1; }
+
+	static uint32 WasHit;
+};
+
+uint32 FMyTask::WasHit = 0;
+
+bool FAutoRTFMMovieSceneTests::RunTest(const FString & Parameters)
+{
+	if (!AutoRTFM::ForTheRuntime::IsAutoRTFMRuntimeEnabled())
+	{
+		ExecutionInfo.AddEvent(FAutomationEvent(EAutomationEventType::Info, TEXT("SKIPPED 'FAutoRTFMMovieSceneTests' test. AutoRTFM disabled.")));
+		return true;
+	}
+	
+	{
+		// Reset this as we'll test against it later.
+		FMyTask::WasHit = 0;
+
+		FComponentRegistry ComponentRegistry;
+		FEntityManager EntityManager;
+		FEntitySystemScheduler Scheduler(&EntityManager);
+
+		TStatId StatId;
+
+		Scheduler.BeginConstruction();
+		Scheduler.AddTask<FMyTask>(FTaskParams(StatId));
+		Scheduler.EndConstruction();
+
+		AutoRTFM::ETransactionResult Result = AutoRTFM::Transact([&]
+			{
+				Scheduler.ExecuteTasks();
+				AutoRTFM::AbortTransaction();
+			});
+
+		TEST_CHECK_TRUE(AutoRTFM::ETransactionResult::AbortedByRequest == Result);
+		TEST_CHECK_TRUE(0 == FMyTask::WasHit);
+
+		Result = AutoRTFM::Transact([&]
+			{
+				Scheduler.ExecuteTasks();
+			});
+
+		TEST_CHECK_TRUE(AutoRTFM::ETransactionResult::Committed == Result);
+		TEST_CHECK_TRUE(1 == FMyTask::WasHit);
+	}
+
+	{
+		// Reset this as we'll test against it later.
+		FMyTask::WasHit = 0;
+
+		FComponentRegistry ComponentRegistry;
+		FEntityManager EntityManager;
+		FEntitySystemScheduler Scheduler(&EntityManager);
+
+		AutoRTFM::ETransactionResult Result = AutoRTFM::Transact([&]
+			{
+				TStatId StatId;
+				Scheduler.BeginConstruction();
+				Scheduler.AddTask<FMyTask>(FTaskParams(StatId));
+				Scheduler.EndConstruction();
+				Scheduler.ExecuteTasks();
+				AutoRTFM::AbortTransaction();
+			});
+
+		TEST_CHECK_TRUE(AutoRTFM::ETransactionResult::AbortedByRequest == Result);
+		TEST_CHECK_TRUE(0 == FMyTask::WasHit);
+
+		Result = AutoRTFM::Transact([&]
+			{
+				TStatId StatId;
+				Scheduler.BeginConstruction();
+				Scheduler.AddTask<FMyTask>(FTaskParams(StatId));
+				Scheduler.EndConstruction();
+				Scheduler.ExecuteTasks();
+			});
+
+		TEST_CHECK_TRUE(AutoRTFM::ETransactionResult::Committed == Result);
+		TEST_CHECK_TRUE(1 == FMyTask::WasHit);
+	}
+
+	return true;
+}
+
+#undef TEST_CHECK_TRUE
+
+#endif //WITH_DEV_AUTOMATION_TESTS
diff --git a/Engine/Source/Runtime/MovieScene/Public/EntitySystem/MovieSceneEntityManager.h b/Engine/Source/Runtime/MovieScene/Public/EntitySystem/MovieSceneEntityManager.h
--- a/Engine/Source/Runtime/MovieScene/Public/EntitySystem/MovieSceneEntityManager.h
+++ b/Engine/Source/Runtime/MovieScene/Public/EntitySystem/MovieSceneEntityManager.h
@@ -18,6 +18,7 @@
 #include "EntitySystem/MovieSceneEntityFactoryTypes.h"
 #include "EntitySystem/MovieSceneEntityIDs.h"
 #include "EntitySystem/MovieSceneEntitySystemTypes.h"
+#include "EntitySystem/MovieSceneMaybeAtomic.h"
 #include "Evaluation/MovieScenePlayback.h"
 #include "HAL/CriticalSection.h"
 #include "HAL/PlatformCrt.h"
@@ -25,7 +26,6 @@
 #include "Misc/EnumClassFlags.h"
 #include "Misc/InlineValue.h"
 #include "MovieSceneSequenceID.h"
-#include "Templates/Atomic.h"
 #include "Templates/UnrealTemplate.h"
 #include "UObject/StrongObjectPtr.h"
 #include "UObject/UObjectArray.h"
@@ -67,20 +67,6 @@ enum class EEntityRecursion : uint8
 };
 ENUM_CLASS_FLAGS(EEntityRecursion);
 
-
-/**
- * Enumeration that defines a threading model for this entity manager
- */
-enum class EEntityThreadingModel : uint8
-{
-	/** Specified when the data contained within an entity manager does not satisfy the requirements to justify using threaded evaluation */
-	NoThreading,
-
-	/** Specified when the data contained within an entity manager is large or complex enough to justify threaded evaluation  */
-	TaskGraph,
-};
-
-
 /**
  * Top-level manager class that is responsible for all entity data and interaction/
  *
@@ -777,7 +763,7 @@ public:
 
 	void CheckCanChangeStructure() const
 	{
-		checkf(static_cast<uint16>(IterationCount) == 0, TEXT("Mutation of entities is not permissible while entities are being iterated"));
+		checkf(IterationCount.Load(ThreadingModel) == 0, TEXT("Mutation of entities is not permissible while entities are being iterated"));
 		checkf(LockdownState == ELockdownState::Unlocked, TEXT("Structural changes to the entity manager are not permitted while it is locked down"));
 	}
 
@@ -812,7 +798,7 @@ public:
 	 */
 	FEntityComponentFilter& ModifyGlobalIterationFilter()
 	{
-		ensureMsgf(!IsLockedDown() && static_cast<uint16>(IterationCount) == 0, TEXT("Manipulating the global iteration filter while locked down or iterating is not recommended"));
+		ensureMsgf(!IsLockedDown() && IterationCount.Load(ThreadingModel) == 0, TEXT("Manipulating the global iteration filter while locked down or iterating is not recommended"));
 		return GlobalIterationFilter;
 	}
 
@@ -1034,7 +1020,7 @@ private:
 	mutable bool bAccumulatedMaskStale;
 
 	/** Atomic counter that is incremented when an iteration begins, and decremented when it finishes */
-	mutable TAtomic<uint16> IterationCount;
+	mutable FEntitySystemMaybeAtomicInt32 IterationCount;
 
 	ENamedThreads::Type GatherThread;
 	ENamedThreads::Type DispatchThread;
diff --git a/Engine/Source/Runtime/MovieScene/Public/EntitySystem/MovieSceneMaybeAtomic.h b/Engine/Source/Runtime/MovieScene/Public/EntitySystem/MovieSceneMaybeAtomic.h
new file mode 100644
--- /dev/null
+++ b/Engine/Source/Runtime/MovieScene/Public/EntitySystem/MovieSceneMaybeAtomic.h
@@ -0,0 +1,117 @@
+// Copyright Epic Games, Inc. All Rights Reserved.
+
+#pragma once
+
+#include "HAL/PlatformAtomics.h"
+#include "Templates/UnrealTemplate.h"
+
+namespace UE
+{
+namespace MovieScene
+{
+
+/**
+ * Enumeration that defines a threading model for this entity manager
+ */
+enum class EEntityThreadingModel : uint8
+{
+	/** Specified when the data contained within an entity manager does not satisfy the requirements to justify using threaded evaluation */
+	NoThreading,
+
+	/** Specified when the data contained within an entity manager is large or complex enough to justify threaded evaluation  */
+	TaskGraph,
+};
+
+/**
+ * A potentially atomic struct that will perform atomic operations if
+ * required on an underlying integer, depending on what
+ * `EEntityThreadingModel` was specified.
+ */
+struct FEntitySystemMaybeAtomicInt32 final
+{
+	FEntitySystemMaybeAtomicInt32() = default;
+	FEntitySystemMaybeAtomicInt32(int32 Payload) : Payload(Payload) {}
+
+	int32 Load(EEntityThreadingModel ThreadingModel) const
+	{
+		if (EEntityThreadingModel::NoThreading == ThreadingModel)
+		{
+			return Payload;
+		}
+		else
+		{
+			return FPlatformAtomics::AtomicRead(&Payload);
+		}
+	}
+
+	int32 Add(EEntityThreadingModel ThreadingModel, const int32 Value)
+	{
+		if (EEntityThreadingModel::NoThreading == ThreadingModel)
+		{
+			const int32 Old = Payload;
+			Payload += Value;
+			return Old;
+		}
+		else
+		{
+			return FPlatformAtomics::InterlockedAdd(&Payload, Value);
+		}
+	}
+
+	int32 Sub(EEntityThreadingModel ThreadingModel, const int32 Value)
+	{
+		if (EEntityThreadingModel::NoThreading == ThreadingModel)
+		{
+			const int32 Old = Payload;
+			Payload -= Value;
+			return Old;
+		}
+		else
+		{
+			return FPlatformAtomics::InterlockedAdd(&Payload, -Value);
+		}
+	}
+
+	int32 Exchange(EEntityThreadingModel ThreadingModel, int32 Other)
+	{
+		if (EEntityThreadingModel::NoThreading == ThreadingModel)
+		{
+			Swap(Payload, Other);
+			return Other;
+		}
+		else
+		{
+			return FPlatformAtomics::InterlockedExchange(&Payload, Other);
+		}
+	}
+
+	int32 Increment(EEntityThreadingModel ThreadingModel)
+	{
+		if (EEntityThreadingModel::NoThreading == ThreadingModel)
+		{
+			return Payload++;
+		}
+		else
+		{
+			return FPlatformAtomics::InterlockedIncrement(&Payload);
+		}
+	}
+
+	int32 Decrement(EEntityThreadingModel ThreadingModel)
+	{
+		if (EEntityThreadingModel::NoThreading == ThreadingModel)
+		{
+			return Payload--;
+		}
+		else
+		{
+			return FPlatformAtomics::InterlockedDecrement(&Payload);
+		}
+	}
+
+private:
+	int32 Payload = 0;
+};
+
+} // namespace MovieScene
+} // namespace UE
