Fix for static analyzer being dumb.

#rnx

diff --git a/Engine/Source/Runtime/MovieScene/Private/Tests/AutoRTFM/AutoRTFMMovieSceneTests.cpp b/Engine/Source/Runtime/MovieScene/Private/Tests/AutoRTFM/AutoRTFMMovieSceneTests.cpp
--- a/Engine/Source/Runtime/MovieScene/Private/Tests/AutoRTFM/AutoRTFMMovieSceneTests.cpp
+++ b/Engine/Source/Runtime/MovieScene/Private/Tests/AutoRTFM/AutoRTFMMovieSceneTests.cpp
@@ -64,7 +64,7 @@ bool FAutoRTFMMovieSceneTests::RunTest(const FString & Parameters)
 			});
 
 		TEST_CHECK_TRUE(AutoRTFM::ETransactionResult::AbortedByRequest == Result);
-		TEST_CHECK_TRUE(0 == FMyTask::WasHit);
+		TEST_CHECK_TRUE(0 == FMyTask::WasHit);//-V547
 
 		Result = AutoRTFM::Transact([&]
 			{
@@ -72,7 +72,7 @@ bool FAutoRTFMMovieSceneTests::RunTest(const FString & Parameters)
 			});
 
 		TEST_CHECK_TRUE(AutoRTFM::ETransactionResult::Committed == Result);
-		TEST_CHECK_TRUE(1 == FMyTask::WasHit);
+		TEST_CHECK_TRUE(1 == FMyTask::WasHit);//-V547
 	}
 
 	{
@@ -94,7 +94,7 @@ bool FAutoRTFMMovieSceneTests::RunTest(const FString & Parameters)
 			});
 
 		TEST_CHECK_TRUE(AutoRTFM::ETransactionResult::AbortedByRequest == Result);
-		TEST_CHECK_TRUE(0 == FMyTask::WasHit);
+		TEST_CHECK_TRUE(0 == FMyTask::WasHit);//-V547
 
 		Result = AutoRTFM::Transact([&]
 			{
@@ -106,7 +106,7 @@ bool FAutoRTFMMovieSceneTests::RunTest(const FString & Parameters)
 			});
 
 		TEST_CHECK_TRUE(AutoRTFM::ETransactionResult::Committed == Result);
-		TEST_CHECK_TRUE(1 == FMyTask::WasHit);
+		TEST_CHECK_TRUE(1 == FMyTask::WasHit);//-V547
 	}
 
 	return true;
