Make MovieScenes TComponentLock AutoRTFM safe.

diff --git a/Engine/Source/Runtime/MovieScene/Private/Tests/AutoRTFM/AutoRTFMMovieSceneTests.cpp b/Engine/Source/Runtime/MovieScene/Private/Tests/AutoRTFM/AutoRTFMMovieSceneTests.cpp
--- a/Engine/Source/Runtime/MovieScene/Private/Tests/AutoRTFM/AutoRTFMMovieSceneTests.cpp
+++ b/Engine/Source/Runtime/MovieScene/Private/Tests/AutoRTFM/AutoRTFMMovieSceneTests.cpp
@@ -109,6 +109,19 @@ bool FAutoRTFMMovieSceneTests::RunTest(const FString & Parameters)
 		TEST_CHECK_TRUE(1 == FMyTask::WasHit);//-V547
 	}
 
+	{
+		FComponentHeader Header;
+		FEntityManager EntityManager;
+		FEntityAllocationWriteContext Context(EntityManager);
+
+		AutoRTFM::ETransactionResult Result = AutoRTFM::Transact([&]
+			{
+				TComponentLock<UE::MovieScene::FWriteErased> Lock(&Header, EComponentHeaderLockMode::Mutex, Context);
+			});
+
+		TEST_CHECK_TRUE(AutoRTFM::ETransactionResult::Committed == Result);
+	}
+
 	return true;
 }
 
diff --git a/Engine/Source/Runtime/MovieScene/Public/EntitySystem/MovieSceneEntitySystemTypes.h b/Engine/Source/Runtime/MovieScene/Public/EntitySystem/MovieSceneEntitySystemTypes.h
--- a/Engine/Source/Runtime/MovieScene/Public/EntitySystem/MovieSceneEntitySystemTypes.h
+++ b/Engine/Source/Runtime/MovieScene/Public/EntitySystem/MovieSceneEntitySystemTypes.h
@@ -13,6 +13,7 @@
 #include "Math/UnrealMathSSE.h"
 #include "Misc/AssertionMacros.h"
 #include "Misc/EnumClassFlags.h"
+#include "Misc/TransactionallySafeRWLock.h"
 #include "Stats/Stats.h"
 #include "Stats/Stats2.h"
 #include "Templates/UnrealTemplate.h"
@@ -486,7 +487,7 @@ struct FComponentHeader
 	mutable uint8* Components;
 
 	/** Legacy mutex used for locked threaded evaluation */
-	mutable FRWLock ReadWriteLock;
+	mutable FTransactionallySafeRWLock ReadWriteLock;
 
 private:
 
@@ -495,8 +496,21 @@ private:
 
 public:
 
+	struct FAutoRTFMCompiledOutAtomicInt32 final
+	{
+		int32 fetch_add(const int32 Value, const std::memory_order) { return 0; }
+		int32 fetch_sub(const int32 Value, const std::memory_order) { return 0; }
+		int32 exchange(const int32 Value, const std::memory_order) { return 0; }
+	};
+
+#if UE_AUTORTFM
+	using FScheduledAccessCountType = FAutoRTFMCompiledOutAtomicInt32;
+#else
+	using FScheduledAccessCountType = std::atomic<int32>;
+#endif
+
 	/** Atomic access count used for verifying write exclusivity at runtime. */
-	mutable std::atomic<int32> ScheduledAccessCount;
+	mutable FScheduledAccessCountType ScheduledAccessCount;
 
 	/** sizeof(T) for the component type this header represents (ie: the stride of our component data). 0 if this header represents a tag. */
 	uint8 Sizeof;
