Sequencer: Added a method for externally invalidating a section's channel proxy

diff --git a/Engine/Source/Runtime/MovieScene/Private/MovieSceneSection.cpp b/Engine/Source/Runtime/MovieScene/Private/MovieSceneSection.cpp
--- a/Engine/Source/Runtime/MovieScene/Private/MovieSceneSection.cpp
+++ b/Engine/Source/Runtime/MovieScene/Private/MovieSceneSection.cpp
@@ -192,6 +192,12 @@ void UMovieSceneSection::SetEndFrame(TRangeBound<FFrameNumber> NewEndFrame)
 	}
 }
 
+void UMovieSceneSection::InvalidateChannelProxy()
+{
+	ChannelProxyType = EMovieSceneChannelProxyType::Dynamic;
+	ChannelProxy = nullptr;
+}
+
 FMovieSceneChannelProxy& UMovieSceneSection::GetChannelProxy() const
 {
 	if (!ChannelProxy.IsValid())
diff --git a/Engine/Source/Runtime/MovieScene/Public/MovieSceneSection.h b/Engine/Source/Runtime/MovieScene/Public/MovieSceneSection.h
--- a/Engine/Source/Runtime/MovieScene/Public/MovieSceneSection.h
+++ b/Engine/Source/Runtime/MovieScene/Public/MovieSceneSection.h
@@ -634,6 +634,11 @@ public:
 	 */
 	MOVIESCENE_API FMovieSceneChannelProxy& GetChannelProxy() const;
 
+	/**
+	 * Invalidate this section's channel proxy, causing it to be re-constructed when it is next used
+	 */
+	MOVIESCENE_API void InvalidateChannelProxy();
+
 	/** Does this movie section support infinite ranges for evaluation */
 	bool GetSupportsInfiniteRange() const { return bSupportsInfiniteRange; }
 
