Sequencer: Fix so User/Broken tangents get set to curve slope automatically wherever added in Curve Editor or Sequencer. Need to make sure we only do this when within the keys, if a key is added outside it stays flat.

diff --git a/Engine/Source/Editor/CurveEditor/Private/Views/SInteractiveCurveEditorView.cpp b/Engine/Source/Editor/CurveEditor/Private/Views/SInteractiveCurveEditorView.cpp
--- a/Engine/Source/Editor/CurveEditor/Private/Views/SInteractiveCurveEditorView.cpp
+++ b/Engine/Source/Editor/CurveEditor/Private/Views/SInteractiveCurveEditorView.cpp
@@ -1961,8 +1961,34 @@ void SInteractiveCurveEditorView::AddKeyAtTime(const TSet<FCurveModelID>& ToCurv
 			}
 			else
 			{
-				const FKeyAttributes& KeyAttributes = GetDefaultKeyAttributesForCurveTime(*CurveEditor, *CurveModel, EvalTime);
+				FKeyAttributes KeyAttributes = GetDefaultKeyAttributesForCurveTime(*CurveEditor, *CurveModel, EvalTime);
+				if (KeyAttributes.HasInterpMode() && 
+					KeyAttributes.GetInterpMode() == ERichCurveInterpMode::RCIM_Cubic&& 
+					KeyAttributes.HasTangentMode() && 
+					(KeyAttributes.GetTangentMode() == RCTM_User || KeyAttributes.GetTangentMode() == RCTM_Break))
+				{
+					//if we are within the range of existing keys set the tangent to be that of the slope of the curve, otherwise
+					//just set it as flat
+					double MinTime = 0., MaxTime = 0.;
+					CurveModel->GetTimeRange(MinTime, MaxTime);
+					if (EvalTime > MinTime && EvalTime < MaxTime)
+					{
+						const double DeltaTime = 0.1;
+
+						// Compute right tangent
+						double RightTangent = GetTangentValue(EvalTime, CurveValue, CurveModel, DeltaTime);
+						KeyAttributes.SetLeaveTangent(RightTangent);
 
+						// Left
+						double LeftTangent = GetTangentValue(EvalTime, CurveValue, CurveModel, -DeltaTime);
+						KeyAttributes.SetArriveTangent(LeftTangent);
+					}
+					else
+					{
+						KeyAttributes.SetLeaveTangent(0.0);
+						KeyAttributes.SetArriveTangent(0.0);
+					}
+				}
 				// Add a key on this curve
 				NewKey = CurveModel->AddKey(FKeyPosition(EvalTime, CurveValue), KeyAttributes);
 			}
diff --git a/Engine/Source/Runtime/MovieScene/Private/Channels/MovieSceneCurveChannelImpl.cpp b/Engine/Source/Runtime/MovieScene/Private/Channels/MovieSceneCurveChannelImpl.cpp
--- a/Engine/Source/Runtime/MovieScene/Private/Channels/MovieSceneCurveChannelImpl.cpp
+++ b/Engine/Source/Runtime/MovieScene/Private/Channels/MovieSceneCurveChannelImpl.cpp
@@ -949,16 +949,10 @@ EMovieSceneKeyInterpolation TMovieSceneCurveChannelImpl<ChannelType>::GetInterpo
 				return EMovieSceneKeyInterpolation::SmartAuto;
 				break;
 			case RCTM_Break:
+				return EMovieSceneKeyInterpolation::Break;
+				break;
 			case RCTM_User:
-				if (DefaultInterpolationMode == EMovieSceneKeyInterpolation::Auto ||
-					DefaultInterpolationMode == EMovieSceneKeyInterpolation::SmartAuto)
-				{
-					return DefaultInterpolationMode;
-				}
-				else
-				{
-					return EMovieSceneKeyInterpolation::SmartAuto;
-				}
+				return EMovieSceneKeyInterpolation::User;
 				break;
 			}
 			break;
@@ -1029,21 +1023,27 @@ FKeyHandle TMovieSceneCurveChannelImpl<ChannelType>::AddKeyToChannel(ChannelType
 		if ((Interpolation == EMovieSceneKeyInterpolation::User || Interpolation == EMovieSceneKeyInterpolation::Break)
 			&& ChannelData.GetTimes().Num() >= 2)
 		{
-			const double DeltaTime = 0.1;
-
-			// Left
-			TangentData.ArriveTangent = -GetTangentValue(InChannel, InFrameNumber, InValue, -DeltaTime);
+			//if we are within the range of existing keys set the tangent to be that of the slope of the curve, otherwise
+			//just set it as flat(leave as default)
+			TRange<FFrameNumber> Range = ChannelData.GetTotalRange();
+			if (Range.Contains(InFrameNumber))
+			{
+				const double DeltaTime = 0.1;
+				// Left
+				TangentData.ArriveTangent = -GetTangentValue(InChannel, InFrameNumber, InValue, -DeltaTime);
 
-			// Right
-			TangentData.LeaveTangent = GetTangentValue(InChannel, InFrameNumber, InValue, DeltaTime);
+				// Right
+				TangentData.LeaveTangent = GetTangentValue(InChannel, InFrameNumber, InValue, DeltaTime);
+			}
+			
 
 		}
 		switch (Interpolation)
 		{
 			case EMovieSceneKeyInterpolation::SmartAuto:     ExistingIndex = InChannel->AddCubicKey(InFrameNumber, InValue, RCTM_SmartAuto);  break;
 			case EMovieSceneKeyInterpolation::Auto:     ExistingIndex = InChannel->AddCubicKey(InFrameNumber, InValue, RCTM_Auto);  break;
-			case EMovieSceneKeyInterpolation::User:     ExistingIndex = InChannel->AddCubicKey(InFrameNumber, InValue, RCTM_User);  break;
-			case EMovieSceneKeyInterpolation::Break:    ExistingIndex = InChannel->AddCubicKey(InFrameNumber, InValue, RCTM_Break); break;
+			case EMovieSceneKeyInterpolation::User:     ExistingIndex = InChannel->AddCubicKey(InFrameNumber, InValue, RCTM_User,TangentData);  break;
+			case EMovieSceneKeyInterpolation::Break:    ExistingIndex = InChannel->AddCubicKey(InFrameNumber, InValue, RCTM_Break,TangentData); break;
 			case EMovieSceneKeyInterpolation::Linear:   ExistingIndex = InChannel->AddLinearKey(InFrameNumber, InValue);            break;
 			case EMovieSceneKeyInterpolation::Constant: ExistingIndex = InChannel->AddConstantKey(InFrameNumber, InValue);          break;
 		}
