Sequencer: Fixed potential for double-delete in FHierarchicalBlendTarget

Technically, each of the changes to FHierarchicalBlendTarget fix this bug on their own but I have included both to be extra sure.
Also added to the automation test to catch these errors in future.

diff --git a/Engine/Source/Runtime/MovieScene/Private/EntitySystem/BuiltInComponentTypes.cpp b/Engine/Source/Runtime/MovieScene/Private/EntitySystem/BuiltInComponentTypes.cpp
--- a/Engine/Source/Runtime/MovieScene/Private/EntitySystem/BuiltInComponentTypes.cpp
+++ b/Engine/Source/Runtime/MovieScene/Private/EntitySystem/BuiltInComponentTypes.cpp
@@ -58,10 +58,11 @@ FHierarchicalBlendTarget& FHierarchicalBlendTarget::operator=(const FHierarchica
 	if (RHS.Capacity != InlineCapacity)
 	{
 		int16* NewAllocation = new int16[RHS.Capacity];
-		Capacity = RHS.Capacity;
 		*reinterpret_cast<int16**>(Data) = NewAllocation;
 	}
 
+	Capacity = RHS.Capacity;
+
 	// Copy the data
 	FMemory::Memcpy(this->GetMemory(), RHS.GetMemory(), sizeof(int16)*Capacity);
 	return *this;
@@ -106,6 +107,10 @@ void FHierarchicalBlendTarget::FreeAllocation()
 	if (Capacity != InlineCapacity)
 	{
 		delete[] GetMemory();
+		// Reset the capacity to be on the safe side.
+		// Note: this does not reinitialize the inline array to the default state
+		//       since we leave that to the caller to do if necessary
+		Capacity = InlineCapacity;
 	}
 }
 
diff --git a/Engine/Source/Runtime/MovieScene/Private/Tests/MovieSceneEntitySystemTypesTests.cpp b/Engine/Source/Runtime/MovieScene/Private/Tests/MovieSceneEntitySystemTypesTests.cpp
--- a/Engine/Source/Runtime/MovieScene/Private/Tests/MovieSceneEntitySystemTypesTests.cpp
+++ b/Engine/Source/Runtime/MovieScene/Private/Tests/MovieSceneEntitySystemTypesTests.cpp
@@ -541,6 +541,86 @@ bool FMovieSceneBlendTargetComponentTest::RunTest(const FString& Parameters)
 		}
 	}
 
+	/** Now perform some operator tests on various allocation types */
+	FRandomStream Random(0xfeedbeef);
+
+	auto MakeBlendTarget = [&Random](int32 Size)
+	{
+		FHierarchicalBlendTarget Target;
+		while (--Size >= 0)
+		{
+			Target.Add(Random.RandHelper(MAX_int16));
+		}
+		return Target;
+	};
+
+	auto CopyConstruct = [&MakeBlendTarget](int32 Size1)
+	{
+		FHierarchicalBlendTarget A = MakeBlendTarget(Size1);
+		FHierarchicalBlendTarget B(A);
+
+		check(A.AsArray().GetData() != B.AsArray().GetData());
+	};
+	auto CopyAssign = [&MakeBlendTarget](int32 Size1, int32 Size2)
+	{
+		FHierarchicalBlendTarget A = MakeBlendTarget(Size1);
+		FHierarchicalBlendTarget B = MakeBlendTarget(Size2);
+		A = B;
+
+		check(A.AsArray().GetData() != B.AsArray().GetData());
+	};
+	auto MoveConstruct = [&MakeBlendTarget](int32 Size1)
+	{
+		FHierarchicalBlendTarget A = MakeBlendTarget(Size1);
+		FHierarchicalBlendTarget B(MoveTemp(A));
+
+		check(A.AsArray().GetData() != B.AsArray().GetData());
+	};
+	auto MoveAssign = [&MakeBlendTarget](int32 Size1, int32 Size2)
+	{
+		FHierarchicalBlendTarget A = MakeBlendTarget(Size1);
+		FHierarchicalBlendTarget B = MakeBlendTarget(Size2);
+		A = MoveTemp(B);
+
+		check(A.AsArray().GetData() != B.AsArray().GetData());
+	};
+
+	// Test Inline to Inline
+	CopyConstruct(5);
+	CopyAssign(5, 7);
+	MoveConstruct(5);
+	MoveAssign(5, 7);
+
+	// Test Heap to Inline
+	CopyConstruct(5);
+	CopyAssign(5, 15);
+	MoveConstruct(5);
+	MoveAssign(5, 15);
+
+	// Test Inline to Heap
+	CopyConstruct(15);
+	CopyAssign(15, 5);
+	MoveConstruct(15);
+	MoveAssign(15, 5);
+
+	// Test Copying Heap to Heap (Same Capacity)
+	CopyConstruct(9);
+	CopyAssign(9, 15);
+	MoveConstruct(9);
+	MoveAssign(9, 15);
+
+	// Test Copying Heap to Heap (Different Capacity)
+	CopyConstruct(35);
+	CopyAssign(35, 15);
+	MoveConstruct(35);
+	MoveAssign(35, 15);
+
+	// Test Copying Heap to Heap (Different Capacity)
+	CopyConstruct(15);
+	CopyAssign(15, 35);
+	MoveConstruct(15);
+	MoveAssign(15, 35);
+
 	return true;
 }
 
