Make MoviePlayer use TransactionallySafe RW Locks

diff --git a/Engine/Source/Runtime/MovieScene/Private/MovieScenePlayer.cpp b/Engine/Source/Runtime/MovieScene/Private/MovieScenePlayer.cpp
--- a/Engine/Source/Runtime/MovieScene/Private/MovieScenePlayer.cpp
+++ b/Engine/Source/Runtime/MovieScene/Private/MovieScenePlayer.cpp
@@ -15,13 +15,15 @@
 #include "MovieSceneSequenceID.h"
 #include "UniversalObjectLocatorResolveParameterBuffer.inl"
 #include "MovieSceneBindingReferences.h"
+#include "Misc/TransactionallySafeRWLock.h"
+#include "Misc/TransactionallySafeRWScopeLock.h"
 
 namespace UE
 {
 namespace MovieScene
 {
 
-static FRWLock                          GGlobalPlayerRegistryLock;
+static FTransactionallySafeRWLock       GGlobalPlayerRegistryLock;
 static TSparseArray<IMovieScenePlayer*> GGlobalPlayerRegistry;
 static TBitArray<> GGlobalPlayerUpdateFlags;
 
@@ -52,7 +54,7 @@ UE::MovieScene::TPlaybackCapabilityID<IMovieScenePlaybackClient> IMovieScenePlay
 
 IMovieScenePlayer::IMovieScenePlayer()
 {
-	FWriteScopeLock ScopeLock(UE::MovieScene::GGlobalPlayerRegistryLock);
+	FTransactionallySafeWriteScopeLock ScopeLock(UE::MovieScene::GGlobalPlayerRegistryLock);
 
 	UE::MovieScene::GGlobalPlayerRegistry.Shrink();
 	UniqueIndex = UE::MovieScene::GGlobalPlayerRegistry.Add(this);
@@ -63,7 +65,7 @@ IMovieScenePlayer::IMovieScenePlayer()
 
 IMovieScenePlayer::~IMovieScenePlayer()
 {	
-	FWriteScopeLock ScopeLock(UE::MovieScene::GGlobalPlayerRegistryLock);
+	FTransactionallySafeWriteScopeLock ScopeLock(UE::MovieScene::GGlobalPlayerRegistryLock);
 
 	UE::MovieScene::GGlobalPlayerUpdateFlags[UniqueIndex] = 0;
 	UE::MovieScene::GGlobalPlayerRegistry.RemoveAt(UniqueIndex, 1);
@@ -71,14 +73,14 @@ IMovieScenePlayer::~IMovieScenePlayer()
 
 IMovieScenePlayer* IMovieScenePlayer::Get(uint16 InUniqueIndex)
 {
-	FReadScopeLock ScopeLock(UE::MovieScene::GGlobalPlayerRegistryLock);
+	FTransactionallySafeReadScopeLock ScopeLock(UE::MovieScene::GGlobalPlayerRegistryLock);
 	check(UE::MovieScene::GGlobalPlayerRegistry.IsValidIndex(InUniqueIndex));
 	return UE::MovieScene::GGlobalPlayerRegistry[InUniqueIndex];
 }
 
 void IMovieScenePlayer::Get(TArray<IMovieScenePlayer*>& OutPlayers, bool bOnlyUnstoppedPlayers)
 {
-	FReadScopeLock ScopeLock(UE::MovieScene::GGlobalPlayerRegistryLock);
+	FTransactionallySafeReadScopeLock ScopeLock(UE::MovieScene::GGlobalPlayerRegistryLock);
 	for (auto It = UE::MovieScene::GGlobalPlayerRegistry.CreateIterator(); It; ++It)
 	{
 		if (IMovieScenePlayer* Player = *It)
