Sequencer Conditions: Conditions in editor will no longer cache results. This is because it's impossible to know what data conditions might rely on and how a user might change that data in editor and expect the conditions to refresh properly.

diff --git a/Engine/Source/Editor/MovieSceneTools/Private/Conditions/MovieSceneConditionCustomization.cpp b/Engine/Source/Editor/MovieSceneTools/Private/Conditions/MovieSceneConditionCustomization.cpp
--- a/Engine/Source/Editor/MovieSceneTools/Private/Conditions/MovieSceneConditionCustomization.cpp
+++ b/Engine/Source/Editor/MovieSceneTools/Private/Conditions/MovieSceneConditionCustomization.cpp
@@ -34,20 +34,27 @@ public:
 
 	virtual bool IsClassAllowed(const FClassViewerInitializationOptions& InInitOptions, const UClass* InClass, TSharedRef< FClassViewerFilterFuncs > InFilterFuncs) override
 	{
-		if (MovieScene.IsValid())
+		if (InClass && InClass->IsChildOf(UMovieSceneCondition::StaticClass()))
 		{
-			return MovieScene->IsConditionClassAllowed(InClass);
+			if (MovieScene.IsValid())
+			{
+				return MovieScene->IsConditionClassAllowed(InClass);
+			}
 		}
-		return true;
+		return false;
 	}
 
 	virtual bool IsUnloadedClassAllowed(const FClassViewerInitializationOptions& InInitOptions, const TSharedRef< const IUnloadedBlueprintData > InBlueprint, TSharedRef< FClassViewerFilterFuncs > InFilterFuncs) override
 	{
-		if (MovieScene.IsValid())
+		const UClass* NativeParent = InBlueprint->GetNativeParent();
+		if (NativeParent && NativeParent->IsChildOf(UMovieSceneCondition::StaticClass()))
 		{
-			return MovieScene->IsConditionClassAllowed(InBlueprint->GetNativeParent());
+			if (MovieScene.IsValid())
+			{
+				return MovieScene->IsConditionClassAllowed(NativeParent);
+			}
 		}
-		return true;
+		return false;
 	}
 };
 
diff --git a/Engine/Source/Runtime/MovieScene/Private/Conditions/MovieSceneCondition.cpp b/Engine/Source/Runtime/MovieScene/Private/Conditions/MovieSceneCondition.cpp
--- a/Engine/Source/Runtime/MovieScene/Private/Conditions/MovieSceneCondition.cpp
+++ b/Engine/Source/Runtime/MovieScene/Private/Conditions/MovieSceneCondition.cpp
@@ -69,4 +69,18 @@ EMovieSceneConditionCheckFrequency UMovieSceneCondition::BP_GetCheckFrequency_Im
 EMovieSceneConditionCheckFrequency UMovieSceneCondition::GetCheckFrequencyInternal() const
 {
 	return BP_GetCheckFrequency();
+}
+
+bool UMovieSceneCondition::CanCacheResult(TSharedRef<const UE::MovieScene::FSharedPlaybackState> SharedPlaybackState) const
+{ 
+#if WITH_EDITOR
+	// Specifically in editor worlds we don't cache condition results- this is because it's too difficult to know
+	// what sort of things the user might change to invalidate the cached results.
+	if (SharedPlaybackState->GetPlaybackContext() && SharedPlaybackState->GetPlaybackContext()->GetWorld()->IsEditorWorld())
+	{
+		return false;
+	}
+#endif
+	
+	return GetCheckFrequencyInternal() != EMovieSceneConditionCheckFrequency::OnTick; 
 }
\ No newline at end of file
diff --git a/Engine/Source/Runtime/MovieScene/Private/EntitySystem/MovieSceneEntityLedger.cpp b/Engine/Source/Runtime/MovieScene/Private/EntitySystem/MovieSceneEntityLedger.cpp
--- a/Engine/Source/Runtime/MovieScene/Private/EntitySystem/MovieSceneEntityLedger.cpp
+++ b/Engine/Source/Runtime/MovieScene/Private/EntitySystem/MovieSceneEntityLedger.cpp
@@ -252,7 +252,7 @@ bool FEntityLedger::CanImportEntity(UMovieSceneEntitySystemLinker* Linker, const
 	else
 	{
 		bool bResult = EntityMetadata->Condition->EvaluateCondition(BindingID, ImportParams.SequenceID, SequenceInstance.GetSharedPlaybackState());
-		if (bCacheAllResults || EntityMetadata->Condition->CanCacheResult())
+		if (bCacheAllResults || EntityMetadata->Condition->CanCacheResult(SequenceInstance.GetSharedPlaybackState()))
 		{
 			ConditionResultCache.Add(CacheKey, bResult);
 		}
diff --git a/Engine/Source/Runtime/MovieScene/Private/EntitySystem/MovieSceneSequenceUpdaters.cpp b/Engine/Source/Runtime/MovieScene/Private/EntitySystem/MovieSceneSequenceUpdaters.cpp
--- a/Engine/Source/Runtime/MovieScene/Private/EntitySystem/MovieSceneSequenceUpdaters.cpp
+++ b/Engine/Source/Runtime/MovieScene/Private/EntitySystem/MovieSceneSequenceUpdaters.cpp
@@ -447,7 +447,7 @@ bool FSequenceUpdater_Flat::EvaluateCondition(const FGuid& BindingID, const FMov
 {
 	if (Condition)
 	{
-		if (Condition->CanCacheResult())
+		if (Condition->CanCacheResult(SharedPlaybackState))
 		{
 			if (bool* ConditionResult = CachedConditionResults.Find(Condition->ComputeCacheKey(BindingID, SequenceID, SharedPlaybackState, ConditionOwnerObject)))
 			{
@@ -844,7 +844,7 @@ void FSequenceUpdater_Hierarchical::Update(TSharedRef<const FSharedPlaybackState
 			{
 				// If we're able to cache the condition result, then it should be cached above when its entity got processed- retrieve that value.
 				// Otherwise, test it again.
-				if (SubData->Condition->CanCacheResult())
+				if (SubData->Condition->CanCacheResult(SharedPlaybackState))
 				{
 					if (bool* ConditionResult = CachedConditionResults.Find(SubData->Condition->ComputeCacheKey(FGuid(), RootOverrideSequenceID, SharedPlaybackState, FindObject<UMovieSceneSubSection>(CompiledDataManager->GetEntryRef(RootCompiledDataID).GetSequence(), *SubData->SectionPath.ToString()))))
 					{
@@ -1110,7 +1110,7 @@ bool FSequenceUpdater_Hierarchical::EvaluateCondition(const FGuid& BindingID, co
 {
 	if (Condition)
 	{
-		if (Condition->CanCacheResult())
+		if (Condition->CanCacheResult(SharedPlaybackState))
 		{
 			if (bool* ConditionResult = CachedConditionResults.Find(Condition->ComputeCacheKey(BindingID, SequenceID, SharedPlaybackState, ConditionOwnerObject)))
 			{
diff --git a/Engine/Source/Runtime/MovieScene/Public/Conditions/MovieSceneCondition.h b/Engine/Source/Runtime/MovieScene/Public/Conditions/MovieSceneCondition.h
--- a/Engine/Source/Runtime/MovieScene/Public/Conditions/MovieSceneCondition.h
+++ b/Engine/Source/Runtime/MovieScene/Public/Conditions/MovieSceneCondition.h
@@ -99,7 +99,7 @@ public:
 	*/
 	virtual uint32 ComputeCacheKey(FGuid BindingGuid, FMovieSceneSequenceID SequenceID, TSharedRef<const UE::MovieScene::FSharedPlaybackState> SharedPlaybackState, UObject* EntityOwner) const;
 
-	bool CanCacheResult() const { return GetCheckFrequencyInternal() != EMovieSceneConditionCheckFrequency::OnTick; }
+	bool CanCacheResult(TSharedRef<const UE::MovieScene::FSharedPlaybackState> SharedPlaybackState) const;
 
 protected:
 
