Only create the game thread signal event when executing the movie scene stuff with threading.

diff --git a/Engine/Source/Runtime/MovieScene/Private/EntitySystem/MovieSceneTaskScheduler.cpp b/Engine/Source/Runtime/MovieScene/Private/EntitySystem/MovieSceneTaskScheduler.cpp
--- a/Engine/Source/Runtime/MovieScene/Private/EntitySystem/MovieSceneTaskScheduler.cpp
+++ b/Engine/Source/Runtime/MovieScene/Private/EntitySystem/MovieSceneTaskScheduler.cpp
@@ -162,12 +162,11 @@ void FScheduledTask::Run(const FEntitySystemScheduler* Scheduler, FTaskExecution
 FEntitySystemScheduler::FEntitySystemScheduler(FEntityManager* InEntityManager)
 	: EntityManager(InEntityManager)
 {
-	GameThreadSignal = FPlatformProcess::GetSynchEventFromPool();
 }
 
 FEntitySystemScheduler::~FEntitySystemScheduler()
 {
-	FPlatformProcess::ReturnSynchEventToPool(GameThreadSignal);
+	check(!GameThreadSignal);
 }
 
 bool FEntitySystemScheduler::IsCustomSchedulingEnabled()
@@ -512,6 +511,12 @@ void FEntitySystemScheduler::ExecuteTasks()
 
 	// Condition 2: Task graph threading
 	//              Schedule initial tasks tasks immediately. Gamethread tasks will be added to the GT queue to ensure that threaded work can be scheduled asap.
+
+	// We need to get a game thread signal from the event pool for the execution
+	// which we'll return to the pool after we've completed execution.
+	check(!GameThreadSignal);
+	GameThreadSignal = FPlatformProcess::GetSynchEventFromPool();
+
 	{
 		TRACE_CPUPROFILER_EVENT_SCOPE_STR("Dispatch Scheduled Tasks");
 
@@ -560,6 +565,12 @@ void FEntitySystemScheduler::ExecuteTasks()
 
 	check(NumTasksRemaining.Load(ThreadingModel) == 0);
 	EntityManager->IncrementSystemSerial(SystemSerialIncrement);
+
+	// Lastly return the game thread signal event to the pool as we are
+	// done executing.
+	check(GameThreadSignal);
+	FPlatformProcess::ReturnSynchEventToPool(GameThreadSignal);
+	GameThreadSignal = nullptr;
 }
 
 void FEntitySystemScheduler::CompleteTask(const FScheduledTask* Task, FTaskExecutionFlags InFlags) const
@@ -601,6 +612,9 @@ void FEntitySystemScheduler::CompleteTask(const FScheduledTask* Task, FTaskExecu
 
 void FEntitySystemScheduler::PrerequisiteCompleted(FTaskID TaskID, int32* OptRunInlineIndex) const
 {
+	// We either need to not be using threading, or have a valid game thread signal event!
+	check((ThreadingModel == EEntityThreadingModel::NoThreading) || (GameThreadSignal != nullptr));
+
 	int32 PreviousWaitCount = Tasks[TaskID.Index].WaitCount.Sub(ThreadingModel, 1);
 	checkSlow(PreviousWaitCount >= 1);
 	if (PreviousWaitCount <= 1)
diff --git a/Engine/Source/Runtime/MovieScene/Private/Tests/AutoRTFM/AutoRTFMMovieSceneTests.cpp b/Engine/Source/Runtime/MovieScene/Private/Tests/AutoRTFM/AutoRTFMMovieSceneTests.cpp
--- a/Engine/Source/Runtime/MovieScene/Private/Tests/AutoRTFM/AutoRTFMMovieSceneTests.cpp
+++ b/Engine/Source/Runtime/MovieScene/Private/Tests/AutoRTFM/AutoRTFMMovieSceneTests.cpp
@@ -122,6 +122,18 @@ bool FAutoRTFMMovieSceneTests::RunTest(const FString & Parameters)
 		TEST_CHECK_TRUE(AutoRTFM::ETransactionResult::Committed == Result);
 	}
 
+	// Test we can construct and destruct an entity system scheduler.
+	{
+		FEntityManager EntityManager;
+
+		AutoRTFM::ETransactionResult Result = AutoRTFM::Transact([&]
+			{
+				FEntitySystemScheduler Scheduler(&EntityManager);
+			});
+
+		TEST_CHECK_TRUE(AutoRTFM::ETransactionResult::Committed == Result);
+	}
+
 	return true;
 }
 
