[MovieScene]
* Changed template specialization to have dllexport on functions instead of type

diff --git a/Engine/Source/Runtime/MovieScene/Private/Evaluation/PreAnimatedState/MovieScenePreAnimatedCaptureSources.cpp b/Engine/Source/Runtime/MovieScene/Private/Evaluation/PreAnimatedState/MovieScenePreAnimatedCaptureSources.cpp
--- a/Engine/Source/Runtime/MovieScene/Private/Evaluation/PreAnimatedState/MovieScenePreAnimatedCaptureSources.cpp
+++ b/Engine/Source/Runtime/MovieScene/Private/Evaluation/PreAnimatedState/MovieScenePreAnimatedCaptureSources.cpp
@@ -14,10 +14,10 @@ namespace MovieScene
 
 
 // Explicit, exported template instantiation for these classes
-template struct MOVIESCENE_API TPreAnimatedCaptureSources<FObjectKey>;
-template struct MOVIESCENE_API TPreAnimatedCaptureSources<FPreAnimatedEvaluationKeyType>;
-template struct MOVIESCENE_API TPreAnimatedCaptureSources<FPreAnimatedEvalHookKeyType>;
-template struct MOVIESCENE_API TPreAnimatedCaptureSources<FMovieSceneEntityID>;
+template struct TPreAnimatedCaptureSources<FObjectKey>;
+template struct TPreAnimatedCaptureSources<FPreAnimatedEvaluationKeyType>;
+template struct TPreAnimatedCaptureSources<FPreAnimatedEvalHookKeyType>;
+template struct TPreAnimatedCaptureSources<FMovieSceneEntityID>;
 
 FPreAnimatedEntityCaptureSource::FPreAnimatedEntityCaptureSource(FPreAnimatedStateExtension* InOwner)
 	: TPreAnimatedCaptureSources<FMovieSceneEntityID>(InOwner)
diff --git a/Engine/Source/Runtime/MovieScene/Public/Evaluation/PreAnimatedState/MovieScenePreAnimatedCaptureSources.h b/Engine/Source/Runtime/MovieScene/Public/Evaluation/PreAnimatedState/MovieScenePreAnimatedCaptureSources.h
--- a/Engine/Source/Runtime/MovieScene/Public/Evaluation/PreAnimatedState/MovieScenePreAnimatedCaptureSources.h
+++ b/Engine/Source/Runtime/MovieScene/Public/Evaluation/PreAnimatedState/MovieScenePreAnimatedCaptureSources.h
@@ -80,26 +80,26 @@ struct TPreAnimatedCaptureSources : IPreAnimatedCaptureSource
 	 * Make an association for the specified evaluation key to the value specified by Entry,
 	 * along with an indication of whether it should be restored on completion
 	 */
-	EPreAnimatedCaptureSourceState BeginTrackingCaptureSource(const KeyType& InKey, const FPreAnimatedStateMetaData& MetaData);
+	MOVIESCENE_API EPreAnimatedCaptureSourceState BeginTrackingCaptureSource(const KeyType& InKey, const FPreAnimatedStateMetaData& MetaData);
 
 	/**
 	 * Stop tracking the evaluation key in question, where associated with the given storage value ID
 	 * This restores the value if it was captured with bWantsRestoreState
 	 */
-	void StopTrackingCaptureSource(const KeyType& InKey, FPreAnimatedStorageID InStorageID);
+	MOVIESCENE_API void StopTrackingCaptureSource(const KeyType& InKey, FPreAnimatedStorageID InStorageID);
 
 	/**
 	 * Stop tracking the evaluation key in question, restoring the value if it was captured with bWantsRestoreState
 	 */
-	void StopTrackingCaptureSource(const KeyType& InKey);
+	MOVIESCENE_API void StopTrackingCaptureSource(const KeyType& InKey);
 
 	// IPreAnimatedCaptureSource members
-	void Reset() override;
-	bool ContainsInstanceHandle(FRootInstanceHandle RootInstanceHandle) const override;
-	void GatherAndRemoveExpiredMetaData(const FRestoreStateParams& Params, TArray<FPreAnimatedStateMetaData>& OutExpiredMetaData) override;
-	void GatherAndRemoveMetaDataForGroup(FPreAnimatedStorageGroupHandle Group, TArray<FPreAnimatedStateMetaData>& OutExpiredMetaData) override;
-	void GatherAndRemoveMetaDataForStorage(FPreAnimatedStorageID StorageID, FPreAnimatedStorageIndex StorageIndex, TArray<FPreAnimatedStateMetaData>& OutExpiredMetaData) override;
-	void GatherAndRemoveMetaDataForRootInstance(FRootInstanceHandle InstanceHandle, TArray<FPreAnimatedStateMetaData>& OutExpiredMetaData) override;
+	MOVIESCENE_API void Reset() override;
+	MOVIESCENE_API bool ContainsInstanceHandle(FRootInstanceHandle RootInstanceHandle) const override;
+	MOVIESCENE_API void GatherAndRemoveExpiredMetaData(const FRestoreStateParams& Params, TArray<FPreAnimatedStateMetaData>& OutExpiredMetaData) override;
+	MOVIESCENE_API void GatherAndRemoveMetaDataForGroup(FPreAnimatedStorageGroupHandle Group, TArray<FPreAnimatedStateMetaData>& OutExpiredMetaData) override;
+	MOVIESCENE_API void GatherAndRemoveMetaDataForStorage(FPreAnimatedStorageID StorageID, FPreAnimatedStorageIndex StorageIndex, TArray<FPreAnimatedStateMetaData>& OutExpiredMetaData) override;
+	MOVIESCENE_API void GatherAndRemoveMetaDataForRootInstance(FRootInstanceHandle InstanceHandle, TArray<FPreAnimatedStateMetaData>& OutExpiredMetaData) override;
 
 private:
 
