Sequencer: fix false positive with shared playback state ref-count assert

diff --git a/Engine/Source/Runtime/MovieScene/Private/EntitySystem/MovieSceneInstanceRegistry.cpp b/Engine/Source/Runtime/MovieScene/Private/EntitySystem/MovieSceneInstanceRegistry.cpp
--- a/Engine/Source/Runtime/MovieScene/Private/EntitySystem/MovieSceneInstanceRegistry.cpp
+++ b/Engine/Source/Runtime/MovieScene/Private/EntitySystem/MovieSceneInstanceRegistry.cpp
@@ -36,6 +36,15 @@ FInstanceRegistry::FInstanceRegistry(UMovieSceneEntitySystemLinker* InLinker)
 
 FInstanceRegistry::~FInstanceRegistry()
 {
+	// Remove all sub-instances from the array so that they release their ref-count on their shared playback state.
+	// This prevents the root instances from triggering an assert about it.
+	for (auto It = Instances.CreateIterator(); It; ++It)
+	{
+		if (!It->IsRootSequence())
+		{
+			It.RemoveCurrent();
+		}
+	}
 }
 
 FInstanceHandle FInstanceRegistry::FindRelatedInstanceHandle(FInstanceHandle InstanceHandle, FMovieSceneSequenceID SequenceID) const
