Sequencer: Fix typo

#rnx

diff --git a/Engine/Source/Runtime/MovieScene/Private/Compilation/MovieSceneCompiledDataManager.cpp b/Engine/Source/Runtime/MovieScene/Private/Compilation/MovieSceneCompiledDataManager.cpp
--- a/Engine/Source/Runtime/MovieScene/Private/Compilation/MovieSceneCompiledDataManager.cpp
+++ b/Engine/Source/Runtime/MovieScene/Private/Compilation/MovieSceneCompiledDataManager.cpp
@@ -1606,7 +1606,7 @@ bool UMovieSceneCompiledDataManager::CompileHierarchy(UMovieSceneSequence* Seque
 					bContainsTimeWarp = true;
 				}
 
-				// Only 1 time-wrap track supported
+				// Only 1 time-warp track supported
 				break;
 			}
 		}
diff --git a/Engine/Source/Runtime/MovieScene/Public/Variants/MovieScenePlayRateCurve.h b/Engine/Source/Runtime/MovieScene/Public/Variants/MovieScenePlayRateCurve.h
--- a/Engine/Source/Runtime/MovieScene/Public/Variants/MovieScenePlayRateCurve.h
+++ b/Engine/Source/Runtime/MovieScene/Public/Variants/MovieScenePlayRateCurve.h
@@ -9,7 +9,7 @@
 
 
 /**
- * A time-wrap defined as a play rate curve.
+ * A time-warp defined as a play rate curve.
  * Time remapping is computed using the integral of the play rate curve.
  */
 UCLASS(MinimalAPI, DisplayName="Play Rate Curve")
