Sequencer: Deprecate GetData that takes an optional FMovieSceneChannel. It is now required

This allows OnKeyAddedEvent, OnKeyDeletedEvent, OnKeyMovedEvent to be bound for all channels.

diff --git a/Engine/Plugins/Animation/ControlRig/Source/ControlRig/Public/Sequencer/MovieSceneControlRigSpaceChannel.h b/Engine/Plugins/Animation/ControlRig/Source/ControlRig/Public/Sequencer/MovieSceneControlRigSpaceChannel.h
--- a/Engine/Plugins/Animation/ControlRig/Source/ControlRig/Public/Sequencer/MovieSceneControlRigSpaceChannel.h
+++ b/Engine/Plugins/Animation/ControlRig/Source/ControlRig/Public/Sequencer/MovieSceneControlRigSpaceChannel.h
@@ -75,7 +75,7 @@ struct CONTROLRIG_API FMovieSceneControlRigSpaceChannel : public FMovieSceneChan
 	*/
 	TMovieSceneChannelData<FMovieSceneControlRigSpaceBaseKey> GetData()
 	{
-		return TMovieSceneChannelData<FMovieSceneControlRigSpaceBaseKey>(&KeyTimes, &KeyValues, &KeyHandles, this);
+		return TMovieSceneChannelData<FMovieSceneControlRigSpaceBaseKey>(&KeyTimes, &KeyValues, this, &KeyHandles);
 	}
 
 	/**
diff --git a/Engine/Plugins/FX/Niagara/Source/NiagaraEditor/Private/Sequencer/NiagaraSequence/Sections/MovieSceneNiagaraEmitterSection.h b/Engine/Plugins/FX/Niagara/Source/NiagaraEditor/Private/Sequencer/NiagaraSequence/Sections/MovieSceneNiagaraEmitterSection.h
--- a/Engine/Plugins/FX/Niagara/Source/NiagaraEditor/Private/Sequencer/NiagaraSequence/Sections/MovieSceneNiagaraEmitterSection.h
+++ b/Engine/Plugins/FX/Niagara/Source/NiagaraEditor/Private/Sequencer/NiagaraSequence/Sections/MovieSceneNiagaraEmitterSection.h
@@ -50,7 +50,7 @@ struct FMovieSceneNiagaraEmitterChannel : public FMovieSceneChannel
 	*/
 	FORCEINLINE TMovieSceneChannelData<FNiagaraEmitterSectionKey> GetData()
 	{
-		return TMovieSceneChannelData<FNiagaraEmitterSectionKey>(&Times, &Values, &KeyHandles);
+		return TMovieSceneChannelData<FNiagaraEmitterSectionKey>(&Times, &Values, this, &KeyHandles);
 	}
 
 	/**
diff --git a/Engine/Plugins/MovieScene/MovieSceneTextTrack/Source/MovieSceneTextTrack/Public/MovieSceneTextChannel.h b/Engine/Plugins/MovieScene/MovieSceneTextTrack/Source/MovieSceneTextTrack/Public/MovieSceneTextChannel.h
--- a/Engine/Plugins/MovieScene/MovieSceneTextTrack/Source/MovieSceneTextTrack/Public/MovieSceneTextChannel.h
+++ b/Engine/Plugins/MovieScene/MovieSceneTextTrack/Source/MovieSceneTextTrack/Public/MovieSceneTextChannel.h
@@ -33,7 +33,7 @@ struct MOVIESCENETEXTTRACK_API FMovieSceneTextChannel : public FMovieSceneChanne
 	 */
 	FORCEINLINE TMovieSceneChannelData<FText> GetData()
 	{
-		return TMovieSceneChannelData<FText>(&Times, &Values, &KeyHandles);
+		return TMovieSceneChannelData<FText>(&Times, &Values, this, &KeyHandles);
 	}
 
 	/**
diff --git a/Engine/Plugins/Runtime/GameplayAbilities/Source/GameplayAbilities/Public/Sequencer/MovieSceneGameplayCueSections.h b/Engine/Plugins/Runtime/GameplayAbilities/Source/GameplayAbilities/Public/Sequencer/MovieSceneGameplayCueSections.h
--- a/Engine/Plugins/Runtime/GameplayAbilities/Source/GameplayAbilities/Public/Sequencer/MovieSceneGameplayCueSections.h
+++ b/Engine/Plugins/Runtime/GameplayAbilities/Source/GameplayAbilities/Public/Sequencer/MovieSceneGameplayCueSections.h
@@ -87,7 +87,7 @@ struct FMovieSceneGameplayCueChannel : public FMovieSceneChannel
 	 */
 	FORCEINLINE TMovieSceneChannelData<FMovieSceneGameplayCueKey> GetData()
 	{
-		return TMovieSceneChannelData<FMovieSceneGameplayCueKey>(&KeyTimes, &KeyValues, &KeyHandles);
+		return TMovieSceneChannelData<FMovieSceneGameplayCueKey>(&KeyTimes, &KeyValues, this, &KeyHandles);
 	}
 
 	/**
diff --git a/Engine/Source/Runtime/Experimental/Animation/Constraints/Public/ConstraintChannel.h b/Engine/Source/Runtime/Experimental/Animation/Constraints/Public/ConstraintChannel.h
--- a/Engine/Source/Runtime/Experimental/Animation/Constraints/Public/ConstraintChannel.h
+++ b/Engine/Source/Runtime/Experimental/Animation/Constraints/Public/ConstraintChannel.h
@@ -32,7 +32,7 @@ struct FMovieSceneConstraintChannel : public FMovieSceneBoolChannel
 	 */
 	virtual FORCEINLINE TMovieSceneChannelData<bool> GetData() override
 	{
-		return TMovieSceneChannelData<bool>(&Times, &Values, &KeyHandles, this);
+		return TMovieSceneChannelData<bool>(&Times, &Values, this, &KeyHandles);
 	}
 
 	/**
diff --git a/Engine/Source/Runtime/MovieScene/Private/Channels/MovieSceneChannelData.cpp b/Engine/Source/Runtime/MovieScene/Private/Channels/MovieSceneChannelData.cpp
--- a/Engine/Source/Runtime/MovieScene/Private/Channels/MovieSceneChannelData.cpp
+++ b/Engine/Source/Runtime/MovieScene/Private/Channels/MovieSceneChannelData.cpp
@@ -91,7 +91,7 @@ namespace MovieScene
 } // namespace MovieScene
 } // namespace UE
 
-FMovieSceneChannelData::FMovieSceneChannelData(TArray<FFrameNumber>* InTimes, FKeyHandleLookupTable* InKeyHandles, FMovieSceneChannel* InChannel)
+FMovieSceneChannelData::FMovieSceneChannelData(FMovieSceneChannel* InChannel, TArray<FFrameNumber>* InTimes, FKeyHandleLookupTable* InKeyHandles)
 	: Times(InTimes), KeyHandles(InKeyHandles), OwningChannel(InChannel)
 {}
 
diff --git a/Engine/Source/Runtime/MovieScene/Public/Channels/MovieSceneAudioTriggerChannel.h b/Engine/Source/Runtime/MovieScene/Public/Channels/MovieSceneAudioTriggerChannel.h
--- a/Engine/Source/Runtime/MovieScene/Public/Channels/MovieSceneAudioTriggerChannel.h
+++ b/Engine/Source/Runtime/MovieScene/Public/Channels/MovieSceneAudioTriggerChannel.h
@@ -42,7 +42,7 @@ struct FMovieSceneAudioTriggerChannel : public FMovieSceneChannel
 	 */
 	FORCEINLINE TMovieSceneChannelData<bool> GetData()
 	{
-		return TMovieSceneChannelData<bool>(&Times, &Values, &KeyHandles);
+		return TMovieSceneChannelData<bool>(&Times, &Values, this, &KeyHandles);
 	}
 
 	/**
diff --git a/Engine/Source/Runtime/MovieScene/Public/Channels/MovieSceneBoolChannel.h b/Engine/Source/Runtime/MovieScene/Public/Channels/MovieSceneBoolChannel.h
--- a/Engine/Source/Runtime/MovieScene/Public/Channels/MovieSceneBoolChannel.h
+++ b/Engine/Source/Runtime/MovieScene/Public/Channels/MovieSceneBoolChannel.h
@@ -48,7 +48,7 @@ struct FMovieSceneBoolChannel : public FMovieSceneChannel
 	 */
 	virtual FORCEINLINE TMovieSceneChannelData<bool> GetData()
 	{
-		return TMovieSceneChannelData<bool>(&Times, &Values, &KeyHandles);
+		return TMovieSceneChannelData<bool>(&Times, &Values, this, &KeyHandles);
 	}
 
 	/**
diff --git a/Engine/Source/Runtime/MovieScene/Public/Channels/MovieSceneByteChannel.h b/Engine/Source/Runtime/MovieScene/Public/Channels/MovieSceneByteChannel.h
--- a/Engine/Source/Runtime/MovieScene/Public/Channels/MovieSceneByteChannel.h
+++ b/Engine/Source/Runtime/MovieScene/Public/Channels/MovieSceneByteChannel.h
@@ -54,7 +54,7 @@ struct FMovieSceneByteChannel : public FMovieSceneChannel
 	 */
 	FORCEINLINE TMovieSceneChannelData<uint8> GetData()
 	{
-		return TMovieSceneChannelData<uint8>(&Times, &Values, &KeyHandles);
+		return TMovieSceneChannelData<uint8>(&Times, &Values, this, &KeyHandles);
 	}
 
 	/**
diff --git a/Engine/Source/Runtime/MovieScene/Public/Channels/MovieSceneChannelData.h b/Engine/Source/Runtime/MovieScene/Public/Channels/MovieSceneChannelData.h
--- a/Engine/Source/Runtime/MovieScene/Public/Channels/MovieSceneChannelData.h
+++ b/Engine/Source/Runtime/MovieScene/Public/Channels/MovieSceneChannelData.h
@@ -197,9 +197,12 @@ protected:
 	 *
 	 * @param InTimes        A pointer to an array that should be operated on by this class. Externally owned.
 	 * @param InKeyHandles   A key handle map used for persistent, order independent identification of keys
-	 * @param InChannel      A optional pointer to the owning channel.
+	 * @param InChannel      A pointer to the owning channel.
 	 */
-	MOVIESCENE_API FMovieSceneChannelData(TArray<FFrameNumber>* InTimes, FKeyHandleLookupTable* InKeyHandles, FMovieSceneChannel* InChannel = nullptr);
+	MOVIESCENE_API FMovieSceneChannelData(FMovieSceneChannel* InChannel, TArray<FFrameNumber>* InTimes, FKeyHandleLookupTable* InKeyHandles);
+
+	UE_DEPRECATED(5.5, "Constructor that takes an optional FMovieSceneChannel is now deprecated. FMovieSceneChannel is now required.")
+	MOVIESCENE_API FMovieSceneChannelData(TArray<FFrameNumber>* InTimes, FKeyHandleLookupTable* InKeyHandles, FMovieSceneChannel* InChannel = nullptr) { FMovieSceneChannelData(InChannel, InTimes, InKeyHandles); }
 
 	/**
 	 * Move the key at index KeyIndex to a new time
@@ -248,8 +251,9 @@ struct TMovieSceneChannelData : FMovieSceneChannelData
 	 * @param InKeyHandles   A key handle map used for persistent, order independent identification of keys
 	 * @param InChannel A option point to the owning channel, should be set if the move,add, delete delegates are utilizaed
 	 */
+	UE_DEPRECATED(5.5, "Constructor that takes an optional FMovieSceneChannel is now deprecated. FMovieSceneChannel is now required.")
 	TMovieSceneChannelData(TArray<FFrameNumber>* InTimes, TArray<ValueType>* InValues, FKeyHandleLookupTable* InKeyHandles, FMovieSceneChannel* InChannel = nullptr)
-		: FMovieSceneChannelData(InTimes, InKeyHandles, InChannel), Values(InValues) 
+		: FMovieSceneChannelData(InChannel, InTimes, InKeyHandles), Values(InValues) 
 	{
 		check(Times && Values);
 	}
@@ -260,10 +264,11 @@ struct TMovieSceneChannelData : FMovieSceneChannelData
 	 *
 	 * @param InTimes        A pointer to an array of times that should be operated on by this class. Externally owned.
 	 * @param InValues       A pointer to an array of values that should be operated on by this class. Externally owned.
+	 * @param InChannel      A pointer to the owning channel.
 	 * @param InKeyHandles   A key handle map used for persistent, order independent identification of keys
 	 */
 	TMovieSceneChannelData(TArray<FFrameNumber>* InTimes, TArray<ValueType>* InValues, FMovieSceneChannel* InChannel, FKeyHandleLookupTable* InKeyHandles)
-		: FMovieSceneChannelData(InTimes, InKeyHandles), Values(InValues)
+		: FMovieSceneChannelData(InChannel, InTimes, InKeyHandles), Values(InValues)
 	{
 		check(Times && Values);
 	}
diff --git a/Engine/Source/Runtime/MovieScene/Public/Channels/MovieSceneDoubleChannel.h b/Engine/Source/Runtime/MovieScene/Public/Channels/MovieSceneDoubleChannel.h
--- a/Engine/Source/Runtime/MovieScene/Public/Channels/MovieSceneDoubleChannel.h
+++ b/Engine/Source/Runtime/MovieScene/Public/Channels/MovieSceneDoubleChannel.h
@@ -145,7 +145,7 @@ struct FMovieSceneDoubleChannel : public FMovieSceneChannel
 	 */
 	FORCEINLINE TMovieSceneChannelData<FMovieSceneDoubleValue> GetData()
 	{
-		return TMovieSceneChannelData<FMovieSceneDoubleValue>(&Times, &Values, &KeyHandles);
+		return TMovieSceneChannelData<FMovieSceneDoubleValue>(&Times, &Values, this, &KeyHandles);
 	}
 
 	/**
diff --git a/Engine/Source/Runtime/MovieScene/Public/Channels/MovieSceneFloatChannel.h b/Engine/Source/Runtime/MovieScene/Public/Channels/MovieSceneFloatChannel.h
--- a/Engine/Source/Runtime/MovieScene/Public/Channels/MovieSceneFloatChannel.h
+++ b/Engine/Source/Runtime/MovieScene/Public/Channels/MovieSceneFloatChannel.h
@@ -136,7 +136,7 @@ struct FMovieSceneFloatChannel : public FMovieSceneChannel
 	 */
 	FORCEINLINE TMovieSceneChannelData<FMovieSceneFloatValue> GetData()
 	{
-		return TMovieSceneChannelData<FMovieSceneFloatValue>(&Times, &Values, &KeyHandles);
+		return TMovieSceneChannelData<FMovieSceneFloatValue>(&Times, &Values, this, &KeyHandles);
 	}
 
 	/**
diff --git a/Engine/Source/Runtime/MovieScene/Public/Channels/MovieSceneIntegerChannel.h b/Engine/Source/Runtime/MovieScene/Public/Channels/MovieSceneIntegerChannel.h
--- a/Engine/Source/Runtime/MovieScene/Public/Channels/MovieSceneIntegerChannel.h
+++ b/Engine/Source/Runtime/MovieScene/Public/Channels/MovieSceneIntegerChannel.h
@@ -49,7 +49,7 @@ struct FMovieSceneIntegerChannel : public FMovieSceneChannel
 	 */
 	FORCEINLINE TMovieSceneChannelData<int32> GetData()
 	{
-		return TMovieSceneChannelData<int32>(&Times, &Values, &KeyHandles);
+		return TMovieSceneChannelData<int32>(&Times, &Values, this, &KeyHandles);
 	}
 
 	/**
diff --git a/Engine/Source/Runtime/MovieScene/Public/Channels/MovieSceneObjectPathChannel.h b/Engine/Source/Runtime/MovieScene/Public/Channels/MovieSceneObjectPathChannel.h
--- a/Engine/Source/Runtime/MovieScene/Public/Channels/MovieSceneObjectPathChannel.h
+++ b/Engine/Source/Runtime/MovieScene/Public/Channels/MovieSceneObjectPathChannel.h
@@ -93,7 +93,7 @@ struct FMovieSceneObjectPathChannel : public FMovieSceneChannel
 	 */
 	FORCEINLINE TMovieSceneChannelData<FMovieSceneObjectPathChannelKeyValue> GetData()
 	{
-		return TMovieSceneChannelData<FMovieSceneObjectPathChannelKeyValue>(&Times, &Values, &KeyHandles);
+		return TMovieSceneChannelData<FMovieSceneObjectPathChannelKeyValue>(&Times, &Values, this, &KeyHandles);
 	}
 
 	/**
diff --git a/Engine/Source/Runtime/MovieSceneTracks/Public/Channels/MovieSceneCameraShakeSourceTriggerChannel.h b/Engine/Source/Runtime/MovieSceneTracks/Public/Channels/MovieSceneCameraShakeSourceTriggerChannel.h
--- a/Engine/Source/Runtime/MovieSceneTracks/Public/Channels/MovieSceneCameraShakeSourceTriggerChannel.h
+++ b/Engine/Source/Runtime/MovieSceneTracks/Public/Channels/MovieSceneCameraShakeSourceTriggerChannel.h
@@ -65,7 +65,7 @@ struct FMovieSceneCameraShakeSourceTriggerChannel : public FMovieSceneChannel
 
 	FORCEINLINE TMovieSceneChannelData<FMovieSceneCameraShakeSourceTrigger> GetData()
 	{
-		return TMovieSceneChannelData<FMovieSceneCameraShakeSourceTrigger>(&KeyTimes, &KeyValues, &KeyHandles);
+		return TMovieSceneChannelData<FMovieSceneCameraShakeSourceTrigger>(&KeyTimes, &KeyValues, this, &KeyHandles);
 	}
 
 	FORCEINLINE TMovieSceneChannelData<const FMovieSceneCameraShakeSourceTrigger> GetData() const
diff --git a/Engine/Source/Runtime/MovieSceneTracks/Public/Channels/MovieSceneEventChannel.h b/Engine/Source/Runtime/MovieSceneTracks/Public/Channels/MovieSceneEventChannel.h
--- a/Engine/Source/Runtime/MovieSceneTracks/Public/Channels/MovieSceneEventChannel.h
+++ b/Engine/Source/Runtime/MovieSceneTracks/Public/Channels/MovieSceneEventChannel.h
@@ -33,7 +33,7 @@ struct FMovieSceneEventChannel : public FMovieSceneChannel
 	 */
 	FORCEINLINE TMovieSceneChannelData<FMovieSceneEvent> GetData()
 	{
-		return TMovieSceneChannelData<FMovieSceneEvent>(&KeyTimes, &KeyValues, &KeyHandles);
+		return TMovieSceneChannelData<FMovieSceneEvent>(&KeyTimes, &KeyValues, this, &KeyHandles);
 	}
 
 	/**
diff --git a/Engine/Source/Runtime/MovieSceneTracks/Public/Channels/MovieSceneStringChannel.h b/Engine/Source/Runtime/MovieSceneTracks/Public/Channels/MovieSceneStringChannel.h
--- a/Engine/Source/Runtime/MovieSceneTracks/Public/Channels/MovieSceneStringChannel.h
+++ b/Engine/Source/Runtime/MovieSceneTracks/Public/Channels/MovieSceneStringChannel.h
@@ -47,7 +47,7 @@ struct FMovieSceneStringChannel : public FMovieSceneChannel
 	 */
 	FORCEINLINE TMovieSceneChannelData<FString> GetData()
 	{
-		return TMovieSceneChannelData<FString>(&Times, &Values, &KeyHandles);
+		return TMovieSceneChannelData<FString>(&Times, &Values, this, &KeyHandles);
 	}
 
 	/**
diff --git a/Engine/Source/Runtime/MovieSceneTracks/Public/Sections/MovieSceneActorReferenceSection.h b/Engine/Source/Runtime/MovieSceneTracks/Public/Sections/MovieSceneActorReferenceSection.h
--- a/Engine/Source/Runtime/MovieSceneTracks/Public/Sections/MovieSceneActorReferenceSection.h
+++ b/Engine/Source/Runtime/MovieSceneTracks/Public/Sections/MovieSceneActorReferenceSection.h
@@ -84,7 +84,7 @@ struct FMovieSceneActorReferenceData : public FMovieSceneChannel
 	 */
 	FORCEINLINE TMovieSceneChannelData<FMovieSceneActorReferenceKey> GetData()
 	{
-		return TMovieSceneChannelData<FMovieSceneActorReferenceKey>(&KeyTimes, &KeyValues, &KeyHandles);
+		return TMovieSceneChannelData<FMovieSceneActorReferenceKey>(&KeyTimes, &KeyValues, this, &KeyHandles);
 	}
 
 	/**
diff --git a/Engine/Source/Runtime/MovieSceneTracks/Public/Sections/MovieSceneEventSection.h b/Engine/Source/Runtime/MovieSceneTracks/Public/Sections/MovieSceneEventSection.h
--- a/Engine/Source/Runtime/MovieSceneTracks/Public/Sections/MovieSceneEventSection.h
+++ b/Engine/Source/Runtime/MovieSceneTracks/Public/Sections/MovieSceneEventSection.h
@@ -156,7 +156,7 @@ struct FMovieSceneEventSectionData : public FMovieSceneChannel
 	 */
 	FORCEINLINE TMovieSceneChannelData<FEventPayload> GetData()
 	{
-		return TMovieSceneChannelData<FEventPayload>(&Times, &KeyValues, &KeyHandles);
+		return TMovieSceneChannelData<FEventPayload>(&Times, &KeyValues, this, &KeyHandles);
 	}
 
 	/**
