Sequencer: Fix static analysis

#rnx

diff --git a/Engine/Source/Runtime/MovieScene/Private/Channels/MovieSceneChannelData.cpp b/Engine/Source/Runtime/MovieScene/Private/Channels/MovieSceneChannelData.cpp
--- a/Engine/Source/Runtime/MovieScene/Private/Channels/MovieSceneChannelData.cpp
+++ b/Engine/Source/Runtime/MovieScene/Private/Channels/MovieSceneChannelData.cpp
@@ -95,6 +95,10 @@ FMovieSceneChannelData::FMovieSceneChannelData(FMovieSceneChannel* InChannel, TA
 	: Times(InTimes), KeyHandles(InKeyHandles), OwningChannel(InChannel)
 {}
 
+FMovieSceneChannelData::FMovieSceneChannelData(TArray<FFrameNumber>* InTimes, FKeyHandleLookupTable* InKeyHandles, FMovieSceneChannel* InChannel)
+	: Times(InTimes), KeyHandles(InKeyHandles), OwningChannel(InChannel)
+{}
+
 FKeyHandle FMovieSceneChannelData::GetHandle(int32 Index)
 {
 	check(Times->IsValidIndex(Index));
diff --git a/Engine/Source/Runtime/MovieScene/Public/Channels/MovieSceneChannelData.h b/Engine/Source/Runtime/MovieScene/Public/Channels/MovieSceneChannelData.h
--- a/Engine/Source/Runtime/MovieScene/Public/Channels/MovieSceneChannelData.h
+++ b/Engine/Source/Runtime/MovieScene/Public/Channels/MovieSceneChannelData.h
@@ -202,7 +202,7 @@ protected:
 	MOVIESCENE_API FMovieSceneChannelData(FMovieSceneChannel* InChannel, TArray<FFrameNumber>* InTimes, FKeyHandleLookupTable* InKeyHandles);
 
 	UE_DEPRECATED(5.5, "Constructor that takes an optional FMovieSceneChannel is now deprecated. FMovieSceneChannel is now required.")
-	MOVIESCENE_API FMovieSceneChannelData(TArray<FFrameNumber>* InTimes, FKeyHandleLookupTable* InKeyHandles, FMovieSceneChannel* InChannel = nullptr) { FMovieSceneChannelData(InChannel, InTimes, InKeyHandles); }
+	MOVIESCENE_API FMovieSceneChannelData(TArray<FFrameNumber>* InTimes, FKeyHandleLookupTable* InKeyHandles, FMovieSceneChannel* InChannel = nullptr);
 
 	/**
 	 * Move the key at index KeyIndex to a new time
