Sequencer: Fixed dangling spawn register pointer existing during FSequenceInstance::Finish when reinitializing a player

diff --git a/Engine/Source/Runtime/MovieScene/Private/MovieScenePlayer.cpp b/Engine/Source/Runtime/MovieScene/Private/MovieScenePlayer.cpp
--- a/Engine/Source/Runtime/MovieScene/Private/MovieScenePlayer.cpp
+++ b/Engine/Source/Runtime/MovieScene/Private/MovieScenePlayer.cpp
@@ -272,7 +272,13 @@ void IMovieScenePlayer::InitializeRootInstance(TSharedRef<UE::MovieScene::FShare
 
 	NewSharedPlaybackState->AddCapability<FPlayerIndexPlaybackCapability>(UniqueIndex);
 	NewSharedPlaybackState->AddCapabilityRaw(&State);
-	NewSharedPlaybackState->AddCapabilityRaw(&GetSpawnRegister());
+
+	// Only add the spawnregister if it is different from the default 'null' register (which does nothing)
+	FMovieSceneSpawnRegister* SpawnRegister = &GetSpawnRegister();
+	if (SpawnRegister != &IMovieScenePlayer::GetSpawnRegister())
+	{
+		NewSharedPlaybackState->AddCapabilityRaw(SpawnRegister);
+	}
 	NewSharedPlaybackState->AddCapabilityRaw((IObjectBindingNotifyPlaybackCapability*)this);
 	NewSharedPlaybackState->AddCapabilityRaw((IStaticBindingOverridesPlaybackCapability*)this);
 
