For the MovieSceneSequencePlayer, adding a blueprint getter and setter for the hideHud playback setting.

#tests tested with a in game level sequence.

diff --git a/Engine/Source/Runtime/MovieScene/Private/MovieSceneSequencePlayer.cpp b/Engine/Source/Runtime/MovieScene/Private/MovieSceneSequencePlayer.cpp
--- a/Engine/Source/Runtime/MovieScene/Private/MovieSceneSequencePlayer.cpp
+++ b/Engine/Source/Runtime/MovieScene/Private/MovieSceneSequencePlayer.cpp
@@ -802,6 +802,16 @@ void UMovieSceneSequencePlayer::SetPlayRate(float PlayRate)
 	PlaybackSettings.PlayRate = PlayRate;
 }
 
+bool UMovieSceneSequencePlayer::GetHideHud() const
+{
+	return PlaybackSettings.bHideHud;
+}
+
+void UMovieSceneSequencePlayer::SetHideHud(bool HideHud)
+{
+	PlaybackSettings.bHideHud = HideHud;
+}
+
 FFrameTime UMovieSceneSequencePlayer::GetLastValidTime() const
 {
 	if (DurationFrames > 0)
diff --git a/Engine/Source/Runtime/MovieScene/Public/MovieSceneSequencePlayer.h b/Engine/Source/Runtime/MovieScene/Public/MovieSceneSequencePlayer.h
--- a/Engine/Source/Runtime/MovieScene/Public/MovieSceneSequencePlayer.h
+++ b/Engine/Source/Runtime/MovieScene/Public/MovieSceneSequencePlayer.h
@@ -382,6 +382,17 @@ public:
 	UFUNCTION(BlueprintCallable, Category = "Sequencer|Player")
 	MOVIESCENE_API void SetPlayRate(float PlayRate);
 
+	/** Get if the hud is hidden during play. */
+	UFUNCTION(BlueprintCallable, Category = "Sequencer|Player")
+	MOVIESCENE_API bool GetHideHud() const;
+
+	/**
+	 * Set if hiding the hud during play.
+	 * @param HideHud - The new value of Hide Hud during play.
+	 */
+	UFUNCTION(BlueprintCallable, Category = "Sequencer|Player")
+	MOVIESCENE_API void SetHideHud(bool HideHud);
+
 	/** Set whether to disable camera cuts */
 	UFUNCTION(BlueprintCallable, Category = "Sequencer|Player")
 	void SetDisableCameraCuts(bool bInDisableCameraCuts) { PlaybackSettings.bDisableCameraCuts = bInDisableCameraCuts; }
