Sequencer: Track color inline for the picker

#rnx

diff --git a/Engine/Source/Runtime/MovieScene/Public/MovieSceneTrack.h b/Engine/Source/Runtime/MovieScene/Public/MovieSceneTrack.h
--- a/Engine/Source/Runtime/MovieScene/Public/MovieSceneTrack.h
+++ b/Engine/Source/Runtime/MovieScene/Public/MovieSceneTrack.h
@@ -576,7 +576,7 @@ protected:
 	FGuid ObjectBindingID;
 
 	/** This track's tint color */
-	UPROPERTY(EditAnywhere, Category=General, DisplayName=Color)
+	UPROPERTY(EditAnywhere, Category=General, DisplayName=Color, meta = (InlineColorPicker))
 	FColor TrackTint;
 
 	/** This folder's desired sorting order */
