Sequencer: Conditions also need to check IsPlayInEditor when determining not to cache the result for just purely editor context. ie. PIE should cache the result.

#rnx

diff --git a/Engine/Source/Runtime/MovieScene/Private/Conditions/MovieSceneCondition.cpp b/Engine/Source/Runtime/MovieScene/Private/Conditions/MovieSceneCondition.cpp
--- a/Engine/Source/Runtime/MovieScene/Private/Conditions/MovieSceneCondition.cpp
+++ b/Engine/Source/Runtime/MovieScene/Private/Conditions/MovieSceneCondition.cpp
@@ -76,9 +76,15 @@ bool UMovieSceneCondition::CanCacheResult(TSharedRef<const UE::MovieScene::FShar
 #if WITH_EDITOR
 	// Specifically in editor worlds we don't cache condition results- this is because it's too difficult to know
 	// what sort of things the user might change to invalidate the cached results.
-	if (SharedPlaybackState->GetPlaybackContext() && SharedPlaybackState->GetPlaybackContext()->GetWorld()->IsEditorWorld())
+	if (UObject* Context = SharedPlaybackState->GetPlaybackContext())
 	{
-		return false;
+		if (UWorld* World = Context->GetWorld())
+		{
+			if (World->IsEditorWorld() && !World->IsPlayInEditor())
+			{
+				return false;
+			}
+		}
 	}
 #endif
 	
