Sequencer: When we set user tangents interactively using the mouse with auto key, need to average the tangents instead of making it broken since this gives more consistent results that the animator is expecting.

diff --git a/Engine/Source/Runtime/MovieScene/Private/Channels/MovieSceneCurveChannelImpl.cpp b/Engine/Source/Runtime/MovieScene/Private/Channels/MovieSceneCurveChannelImpl.cpp
--- a/Engine/Source/Runtime/MovieScene/Private/Channels/MovieSceneCurveChannelImpl.cpp
+++ b/Engine/Source/Runtime/MovieScene/Private/Channels/MovieSceneCurveChannelImpl.cpp
@@ -1030,10 +1030,12 @@ FKeyHandle TMovieSceneCurveChannelImpl<ChannelType>::AddKeyToChannel(ChannelType
 			{
 				const double DeltaTime = 0.1;
 				// Left
-				TangentData.ArriveTangent = -GetTangentValue(InChannel, InFrameNumber, InValue, -DeltaTime);
+				const float ArriveTangent = -GetTangentValue(InChannel, InFrameNumber, InValue, -DeltaTime);
 
 				// Right
-				TangentData.LeaveTangent = GetTangentValue(InChannel, InFrameNumber, InValue, DeltaTime);
+				const float LeaveTangent = GetTangentValue(InChannel, InFrameNumber, InValue, DeltaTime);
+
+				TangentData.ArriveTangent = TangentData.LeaveTangent = (ArriveTangent + LeaveTangent) * 0.5f;
 			}
 			
 
