Sequencer: Fix issue where conditions on tracks on spawned actors could potentially run before the actor was spawned, fail to evaluate properly if they relied on data from the bound actor, and then have their results be cached, not to be evaluated again.

To fix this, we've had to defer evaluating conditions that rely on the bound object until after the object is bound. We do this by checking these types of conditions in the bound object instantiator rather than during entity import.

diff --git a/Engine/Source/Runtime/MovieScene/Private/Conditions/MovieSceneCondition.cpp b/Engine/Source/Runtime/MovieScene/Private/Conditions/MovieSceneCondition.cpp
--- a/Engine/Source/Runtime/MovieScene/Private/Conditions/MovieSceneCondition.cpp
+++ b/Engine/Source/Runtime/MovieScene/Private/Conditions/MovieSceneCondition.cpp
@@ -89,4 +89,9 @@ bool UMovieSceneCondition::CanCacheResult(TSharedRef<const UE::MovieScene::FShar
 #endif
 	
 	return GetCheckFrequencyInternal() != EMovieSceneConditionCheckFrequency::OnTick; 
-}
\ No newline at end of file
+}
+
+EMovieSceneConditionScope UMovieSceneCondition::GetConditionScope() const
+{
+	return GetScopeInternal();
+}
diff --git a/Engine/Source/Runtime/MovieScene/Private/EntitySystem/BuiltInComponentTypes.cpp b/Engine/Source/Runtime/MovieScene/Private/EntitySystem/BuiltInComponentTypes.cpp
--- a/Engine/Source/Runtime/MovieScene/Private/EntitySystem/BuiltInComponentTypes.cpp
+++ b/Engine/Source/Runtime/MovieScene/Private/EntitySystem/BuiltInComponentTypes.cpp
@@ -390,6 +390,7 @@ FBuiltInComponentTypes::FBuiltInComponentTypes()
 
 	ComponentRegistry->NewComponentType(&EvaluationHook,        TEXT("Evaluation Hook"));
 	ComponentRegistry->NewComponentType(&EvaluationHookFlags,   TEXT("Evaluation Hook Flags"), EComponentTypeFlags::Preserved);
+	ComponentRegistry->NewComponentType(&Condition, TEXT("Condition"));
 
 	ComponentRegistry->NewComponentType(&Interrogation.InputKey,  TEXT("Interrogation Input"));
 	ComponentRegistry->NewComponentType(&Interrogation.Instance,  TEXT("Interrogation Instance"));
diff --git a/Engine/Source/Runtime/MovieScene/Private/EntitySystem/IMovieSceneEntityProvider.cpp b/Engine/Source/Runtime/MovieScene/Private/EntitySystem/IMovieSceneEntityProvider.cpp
--- a/Engine/Source/Runtime/MovieScene/Private/EntitySystem/IMovieSceneEntityProvider.cpp
+++ b/Engine/Source/Runtime/MovieScene/Private/EntitySystem/IMovieSceneEntityProvider.cpp
@@ -33,7 +33,8 @@ FMovieSceneEntityID FImportedEntity::Manufacture(const FEntityImportParams& Para
 	.AddTagConditional(Components->Tags.AlwaysCacheInitialValue, Params.Sequence.bDynamicWeighting)
 	.AddConditional(Components->SequenceID, Params.Sequence.SequenceID, Params.Sequence.SequenceID != MovieSceneSequenceID::Root)
 	.AddConditional(Components->RootInstanceHandle, Params.Sequence.RootInstanceHandle, Params.Sequence.RootInstanceHandle.IsValid())
-	.AddConditional(Components->InstanceHandle, Params.Sequence.InstanceHandle, Params.Sequence.InstanceHandle.IsValid());
+	.AddConditional(Components->InstanceHandle, Params.Sequence.InstanceHandle, Params.Sequence.InstanceHandle.IsValid())
+	.AddConditional(Components->Condition, Params.EntityMetaData ? Params.EntityMetaData->Condition : nullptr, Params.EntityMetaData && Params.EntityMetaData->Condition);
 
 	FEntityAllocationWriteContext WriteContext(*EntityManager);
 
diff --git a/Engine/Source/Runtime/MovieScene/Private/EntitySystem/MovieSceneEntityFactory.cpp b/Engine/Source/Runtime/MovieScene/Private/EntitySystem/MovieSceneEntityFactory.cpp
--- a/Engine/Source/Runtime/MovieScene/Private/EntitySystem/MovieSceneEntityFactory.cpp
+++ b/Engine/Source/Runtime/MovieScene/Private/EntitySystem/MovieSceneEntityFactory.cpp
@@ -8,6 +8,7 @@
 #include "EntitySystem/MovieSceneEntitySystemTask.h"
 #include "EntitySystem/MovieSceneEntityFactoryTemplates.h"
 #include "MovieSceneFwd.h"
+#include "Conditions/MovieSceneCondition.h"
 
 namespace UE
 {
@@ -177,6 +178,8 @@ void FBoundObjectTask::ForEachAllocation(FEntityAllocationProxy AllocationProxy,
 	Batch.StaleEntitiesToPreserve = &StaleEntitiesToPreserve;
 
 	const int32 Num = Allocation->Num();
+	
+	TOptionalComponentReader<TObjectPtr<const UMovieSceneCondition>> Conditions = Allocation->TryReadComponents(FBuiltInComponentTypes::Get()->Condition);
 
 	FInstanceRegistry* InstanceRegistry = Linker->GetInstanceRegistry();
 
@@ -208,7 +211,8 @@ void FBoundObjectTask::ForEachAllocation(FEntityAllocationProxy AllocationProxy,
 
 		const FSequenceInstance&     SequenceInstance = InstanceRegistry->GetInstance(Instances[Index]);
 		TArrayView<TWeakObjectPtr<>> BoundObjects     = SequenceInstance.GetSharedPlaybackState()->FindBoundObjects(ObjectBindings[Index], SequenceInstance.GetSequenceID());
-
+		
+		bool bCheckedCondition = false;
 		for (TWeakObjectPtr<> WeakObject : BoundObjects)
 		{
 			UObject* Object = WeakObject.Get();
@@ -226,6 +230,19 @@ void FBoundObjectTask::ForEachAllocation(FEntityAllocationProxy AllocationProxy,
 					continue;
 				}
 
+				if (!bCheckedCondition && Conditions && Conditions[Index] && Conditions[Index]->GetConditionScope() != EMovieSceneConditionScope::Global)
+				{
+					// If this entity has a condition that could depend on a bound object, then it hasn't yet been tested, and we must test it here. \
+					// Note that it will only be tested once here, and then the entity ledger will take care of testing it again if it needs to
+					// and it is a per-tick condition.
+					bCheckedCondition = true;
+					if (!SequenceInstance.EvaluateCondition(ObjectBindings[Index], SequenceInstance.GetSequenceID(), Conditions[Index], Conditions[Index]->GetTypedOuter<UMovieSceneSignedObject>()))
+					{
+						// Condition has failed, don't add this entity to the batch
+						break;
+					}
+				}
+
 				// Make a child entity for this resolved binding
 				Batch.Add(Index, Object);
 				bIsResolvedBinding = true;
diff --git a/Engine/Source/Runtime/MovieScene/Private/EntitySystem/MovieSceneEntityLedger.cpp b/Engine/Source/Runtime/MovieScene/Private/EntitySystem/MovieSceneEntityLedger.cpp
--- a/Engine/Source/Runtime/MovieScene/Private/EntitySystem/MovieSceneEntityLedger.cpp
+++ b/Engine/Source/Runtime/MovieScene/Private/EntitySystem/MovieSceneEntityLedger.cpp
@@ -224,7 +224,7 @@ void FEntityLedger::FindImportedEntities(TWeakObjectPtr<UObject> EntityOwner, TA
 	}
 }
 
-bool FEntityLedger::CanImportEntity(UMovieSceneEntitySystemLinker* Linker, const FEntityImportSequenceParams& ImportParams, const FMovieSceneEntityComponentField* EntityField, const FMovieSceneEvaluationFieldEntityQuery& Query, FMovieSceneEvaluationFieldEntitySet& OutPerTickConditionalEntities, TMap<uint32, bool>& ConditionResultCache, bool bCacheAllResults)
+bool FEntityLedger::CanImportEntity(UMovieSceneEntitySystemLinker* Linker, const FEntityImportSequenceParams& ImportParams, const FMovieSceneEntityComponentField* EntityField, const FMovieSceneEvaluationFieldEntityQuery& Query, FMovieSceneEvaluationFieldEntitySet& OutPerTickConditionalEntities, TMap<uint32, bool>& ConditionResultCache, bool bUpdatingPerTickEntities)
 {
 	// If we don't have a condition, just return true
 	const FMovieSceneEvaluationFieldEntityMetaData* EntityMetadata = EntityField->FindMetaData(Query);
@@ -237,14 +237,30 @@ bool FEntityLedger::CanImportEntity(UMovieSceneEntitySystemLinker* Linker, const
 	const FInstanceRegistry* InstanceRegistry = Linker->GetInstanceRegistry();
 	const FSequenceInstance& SequenceInstance = InstanceRegistry->GetInstance(ImportParams.InstanceHandle);
 
+	bool bCanCacheResult = EntityMetadata->Condition->CanCacheResult(SequenceInstance.GetSharedPlaybackState());
+	if (!bCanCacheResult)
+	{
+		// If we can't cache the result, it will need to be checked again next tick
+		OutPerTickConditionalEntities.Add(Query);
+	}
+
 	FGuid BindingID;
 
-	if (const FMovieSceneEvaluationFieldSharedEntityMetaData * SharedMetadata = EntityField->FindSharedMetaData(Query))
+	if (const FMovieSceneEvaluationFieldSharedEntityMetaData* SharedMetadata = EntityField->FindSharedMetaData(Query))
 	{
 		BindingID = SharedMetadata->ObjectBindingID;
 	}
 
+	// If we have a valid binding ID, and the condition depends on object binding, then we must ensure the object binding is resolved
+	// before evaluating the condition. To ensure this, we always defer checking the condition for non-global conditions on bound objects to the bound object resolver.
+	// We don't do this for updating per-tick entities as the bound object resolver is only run once.
+	if (!bUpdatingPerTickEntities && BindingID.IsValid() && EntityMetadata->Condition->GetConditionScope() != EMovieSceneConditionScope::Global)
+	{
+		return true;
+	}
+
 	uint32 CacheKey = EntityMetadata->Condition->ComputeCacheKey(BindingID, ImportParams.SequenceID, SequenceInstance.GetSharedPlaybackState(), Query.Entity.Key.EntityOwner.Get());
+	
 	if (bool* CachedResult = ConditionResultCache.Find(CacheKey))
 	{
 		return *CachedResult;
@@ -252,15 +268,11 @@ bool FEntityLedger::CanImportEntity(UMovieSceneEntitySystemLinker* Linker, const
 	else
 	{
 		bool bResult = EntityMetadata->Condition->EvaluateCondition(BindingID, ImportParams.SequenceID, SequenceInstance.GetSharedPlaybackState());
-		if (bCacheAllResults || EntityMetadata->Condition->CanCacheResult(SequenceInstance.GetSharedPlaybackState()))
+		// We always cache the results for per tick entities as they get thrown away after the tick, and we might as well prevent the same condition from getting re-evaluated multiple times per tick.
+		if (bCanCacheResult || bUpdatingPerTickEntities)
 		{
 			ConditionResultCache.Add(CacheKey, bResult);
 		}
-		else
-		{
-			// If we can't cache the result, it will need to be checked again next tick
-			OutPerTickConditionalEntities.Add(Query);
-		}
 		return bResult;
 	}
 }
diff --git a/Engine/Source/Runtime/MovieScene/Public/Conditions/MovieSceneCondition.h b/Engine/Source/Runtime/MovieScene/Public/Conditions/MovieSceneCondition.h
--- a/Engine/Source/Runtime/MovieScene/Public/Conditions/MovieSceneCondition.h
+++ b/Engine/Source/Runtime/MovieScene/Public/Conditions/MovieSceneCondition.h
@@ -101,6 +101,8 @@ public:
 
 	bool CanCacheResult(TSharedRef<const UE::MovieScene::FSharedPlaybackState> SharedPlaybackState) const;
 
+	EMovieSceneConditionScope GetConditionScope() const;
+
 #if WITH_EDITORONLY_DATA
 	/* If true, will skip evaluating the condition and always return true. Useful for authoring or debugging. */
 	UPROPERTY(EditAnywhere, BlueprintReadWrite, Category="Default")
diff --git a/Engine/Source/Runtime/MovieScene/Public/EntitySystem/BuiltInComponentTypes.h b/Engine/Source/Runtime/MovieScene/Public/EntitySystem/BuiltInComponentTypes.h
--- a/Engine/Source/Runtime/MovieScene/Public/EntitySystem/BuiltInComponentTypes.h
+++ b/Engine/Source/Runtime/MovieScene/Public/EntitySystem/BuiltInComponentTypes.h
@@ -724,6 +724,8 @@ public:
 
 	TComponentTypeID<FEvaluationHookFlags> EvaluationHookFlags;
 
+	TComponentTypeID<TObjectPtr<const UMovieSceneCondition>> Condition;
+
 public:
 
 	// 
diff --git a/Engine/Source/Runtime/MovieScene/Public/EntitySystem/MovieSceneEntityLedger.h b/Engine/Source/Runtime/MovieScene/Public/EntitySystem/MovieSceneEntityLedger.h
--- a/Engine/Source/Runtime/MovieScene/Public/EntitySystem/MovieSceneEntityLedger.h
+++ b/Engine/Source/Runtime/MovieScene/Public/EntitySystem/MovieSceneEntityLedger.h
@@ -111,7 +111,7 @@ public:
 	/*
 	* Checks to see whether we can currently import the requested entity, evaluating any conditions that may exist first.
 	*/
-	MOVIESCENE_API bool CanImportEntity(UMovieSceneEntitySystemLinker* Linker, const FEntityImportSequenceParams& ImportParams, const FMovieSceneEntityComponentField* EntityField, const FMovieSceneEvaluationFieldEntityQuery& Query, FMovieSceneEvaluationFieldEntitySet& OutPerTickConditionalEntities, TMap<uint32, bool>& ConditionResultCache, bool bCacheAllResults=false);
+	MOVIESCENE_API bool CanImportEntity(UMovieSceneEntitySystemLinker* Linker, const FEntityImportSequenceParams& ImportParams, const FMovieSceneEntityComponentField* EntityField, const FMovieSceneEvaluationFieldEntityQuery& Query, FMovieSceneEvaluationFieldEntitySet& OutPerTickConditionalEntities, TMap<uint32, bool>& ConditionResultCache, bool bUpdatingPerTickEntities=false);
 
 	
 	UE_DEPRECATED(5.5, "Please call the version that takes OutPerTickConditionalEntities and ConditionResultCache")
