Sequencer: Fix Decomposer test

diff --git a/Engine/Source/Runtime/MovieScene/Private/EntitySystem/MovieSceneDecompositionQuery.cpp b/Engine/Source/Runtime/MovieScene/Private/EntitySystem/MovieSceneDecompositionQuery.cpp
--- a/Engine/Source/Runtime/MovieScene/Private/EntitySystem/MovieSceneDecompositionQuery.cpp
+++ b/Engine/Source/Runtime/MovieScene/Private/EntitySystem/MovieSceneDecompositionQuery.cpp
@@ -100,7 +100,7 @@ double FDecomposedValue::Recompose(FMovieSceneEntityID RecomposeEntity, double C
 		BlendType == EDecomposedValueBlendType::AdditiveFromBase);
 	const bool bIsOverride = (BlendType == EDecomposedValueBlendType::Override);
 	float TotalAbsoluteWeight = ResultAbsolute.TotalWeight + OtherAbsolute.TotalWeight;
-	if (!bIsAdditive || !bIsOverride)
+	if (!bIsAdditive && !bIsOverride)
 	{
 		TotalAbsoluteWeight += Channel.Weight;
 	}
