Sequencer: Disable (at least for now) conditions from being supported on binding lifetime tracks and spawn tracks due to ordering issues.

diff --git a/Engine/Source/Runtime/MovieScene/Private/Tracks/MovieSceneBindingLifetimeTrack.cpp b/Engine/Source/Runtime/MovieScene/Private/Tracks/MovieSceneBindingLifetimeTrack.cpp
--- a/Engine/Source/Runtime/MovieScene/Private/Tracks/MovieSceneBindingLifetimeTrack.cpp
+++ b/Engine/Source/Runtime/MovieScene/Private/Tracks/MovieSceneBindingLifetimeTrack.cpp
@@ -22,7 +22,9 @@ UMovieSceneBindingLifetimeTrack::UMovieSceneBindingLifetimeTrack(const FObjectIn
 {
 #if WITH_EDITORONLY_DATA
 	TrackTint = FColor(26, 117, 49, 150);
+	bSupportsConditions = false;
 #endif
+
 }
 
 
diff --git a/Engine/Source/Runtime/MovieScene/Private/Tracks/MovieSceneSpawnTrack.cpp b/Engine/Source/Runtime/MovieScene/Private/Tracks/MovieSceneSpawnTrack.cpp
--- a/Engine/Source/Runtime/MovieScene/Private/Tracks/MovieSceneSpawnTrack.cpp
+++ b/Engine/Source/Runtime/MovieScene/Private/Tracks/MovieSceneSpawnTrack.cpp
@@ -24,7 +24,9 @@ UMovieSceneSpawnTrack::UMovieSceneSpawnTrack(const FObjectInitializer& Obj)
 {
 #if WITH_EDITORONLY_DATA
 	TrackTint = FColor(43, 43, 155, 65);
+	bSupportsConditions = false;
 #endif
+
 }
 
 void UMovieSceneSpawnTrack::PostLoad()
