Sequencer: Fix group conditions not respecting inner sub-conditions update frequency. This meant that if you used conditions requiring a tick in a sub-condition (or if you tried to combine conditions on track row/track/section), that they wouldn't work on a per-tick basis.

diff --git a/Engine/Source/Runtime/MovieScene/Private/Conditions/MovieSceneGroupCondition.cpp b/Engine/Source/Runtime/MovieScene/Private/Conditions/MovieSceneGroupCondition.cpp
--- a/Engine/Source/Runtime/MovieScene/Private/Conditions/MovieSceneGroupCondition.cpp
+++ b/Engine/Source/Runtime/MovieScene/Private/Conditions/MovieSceneGroupCondition.cpp
@@ -4,6 +4,7 @@
 #include "EntitySystem/MovieSceneSharedPlaybackState.h"
 #include "Engine/World.h"
 #include "MovieSceneCommonHelpers.h"
+#include "Algo/AllOf.h"
 
 
 uint32 UMovieSceneGroupCondition::ComputeCacheKey(FGuid BindingGuid, FMovieSceneSequenceID SequenceID, TSharedRef<const UE::MovieScene::FSharedPlaybackState> SharedPlaybackState, UObject* EntityOwner) const
@@ -61,3 +62,8 @@ bool UMovieSceneGroupCondition::EvaluateConditionInternal(FGuid BindingGuid, FMo
 	}
 	return bResult;
 }
+
+bool UMovieSceneGroupCondition::CanCacheResult(TSharedRef<const UE::MovieScene::FSharedPlaybackState> SharedPlaybackState) const
+{
+	return Algo::AllOf(SubConditions, [SharedPlaybackState](const FMovieSceneConditionContainer& SubCondition) { return !SubCondition.Condition || SubCondition.Condition->CanCacheResult(SharedPlaybackState);});
+}
\ No newline at end of file
diff --git a/Engine/Source/Runtime/MovieScene/Public/Conditions/MovieSceneCondition.h b/Engine/Source/Runtime/MovieScene/Public/Conditions/MovieSceneCondition.h
--- a/Engine/Source/Runtime/MovieScene/Public/Conditions/MovieSceneCondition.h
+++ b/Engine/Source/Runtime/MovieScene/Public/Conditions/MovieSceneCondition.h
@@ -99,7 +99,7 @@ public:
 	*/
 	virtual uint32 ComputeCacheKey(FGuid BindingGuid, FMovieSceneSequenceID SequenceID, TSharedRef<const UE::MovieScene::FSharedPlaybackState> SharedPlaybackState, UObject* EntityOwner) const;
 
-	bool CanCacheResult(TSharedRef<const UE::MovieScene::FSharedPlaybackState> SharedPlaybackState) const;
+	virtual bool CanCacheResult(TSharedRef<const UE::MovieScene::FSharedPlaybackState> SharedPlaybackState) const;
 
 	EMovieSceneConditionScope GetConditionScope() const;
 
diff --git a/Engine/Source/Runtime/MovieScene/Public/Conditions/MovieSceneGroupCondition.h b/Engine/Source/Runtime/MovieScene/Public/Conditions/MovieSceneGroupCondition.h
--- a/Engine/Source/Runtime/MovieScene/Public/Conditions/MovieSceneGroupCondition.h
+++ b/Engine/Source/Runtime/MovieScene/Public/Conditions/MovieSceneGroupCondition.h
@@ -52,5 +52,6 @@ protected:
 
 	/* UMovieSceneCondition overrides */
 	bool EvaluateConditionInternal(FGuid BindingGuid, FMovieSceneSequenceID SequenceID, TSharedRef<const UE::MovieScene::FSharedPlaybackState> SharedPlaybackState) const override;
-
+	
+	virtual bool CanCacheResult(TSharedRef<const UE::MovieScene::FSharedPlaybackState> SharedPlaybackState) const override;
 };
