Sequencer: Decomposition Test Fixes, Windows compiler is doing some messy optimizations that are returning bogus values, we are fixng this by making sure a few
functions aren't inlined.

diff --git a/Engine/Source/Runtime/MovieScene/Public/EntitySystem/MovieSceneDecompositionQuery.h b/Engine/Source/Runtime/MovieScene/Public/EntitySystem/MovieSceneDecompositionQuery.h
--- a/Engine/Source/Runtime/MovieScene/Public/EntitySystem/MovieSceneDecompositionQuery.h
+++ b/Engine/Source/Runtime/MovieScene/Public/EntitySystem/MovieSceneDecompositionQuery.h
@@ -73,32 +73,32 @@ namespace MovieScene
 		float TotalWeight = 0.f;
 
 		/** Normalizes the total value by the total weight */
-		double Normalize() const
+		FORCENOINLINE double Normalize() const
 		{
 			return TotalWeight != 0.f ? Total / TotalWeight : Total;
 		}
 
 		/** Accumulates more values into this value */
-		FAccumulatedWeightedValue& AccumulateThis(const FWeightedValue& Other)
+		FORCENOINLINE FAccumulatedWeightedValue& AccumulateThis(const FWeightedValue& Other)
 		{
 			Total += Other.Get();
 			TotalWeight += Other.Weight;
 			return *this;
 		}
 		/** Accumulates more values into this value */
-		FAccumulatedWeightedValue& AccumulateThis(const FAccumulatedWeightedValue& Other)
+		FORCENOINLINE FAccumulatedWeightedValue& AccumulateThis(const FAccumulatedWeightedValue& Other)
 		{
 			Total += Other.Total;
 			TotalWeight += Other.TotalWeight;
 			return *this;
 		}
 		/** Accumulates values */
-		FAccumulatedWeightedValue Accumulate(const FWeightedValue& Other) const
+		FORCENOINLINE FAccumulatedWeightedValue Accumulate(const FWeightedValue& Other) const
 		{
 			return FAccumulatedWeightedValue{ Total + Other.Get(), TotalWeight + Other.Weight };
 		}
 		/** Accumulates values */
-		FAccumulatedWeightedValue Accumulate(const FAccumulatedWeightedValue& Other) const
+		FORCENOINLINE FAccumulatedWeightedValue Accumulate(const FAccumulatedWeightedValue& Other) const
 		{
 			return FAccumulatedWeightedValue{ Total + Other.Total, TotalWeight + Other.TotalWeight };
 		}
