Sequencer: Fix Sequencer issue where playhead and other visual indicators would show incorrect local time on 3-layer deep Sequences with non-zero start times.
This was due to relying on the compiled data manager to fully concatenate multiple layers of transforms and not including the unwarped transform in that process, which is used in editor-only to display things like the playhead and local time.

diff --git a/Engine/Source/Runtime/MovieScene/Private/Compilation/MovieSceneCompiledDataManager.cpp b/Engine/Source/Runtime/MovieScene/Private/Compilation/MovieSceneCompiledDataManager.cpp
--- a/Engine/Source/Runtime/MovieScene/Private/Compilation/MovieSceneCompiledDataManager.cpp
+++ b/Engine/Source/Runtime/MovieScene/Private/Compilation/MovieSceneCompiledDataManager.cpp
@@ -161,8 +161,9 @@ struct FGatherParameters
 		SubParams.Flags                     = this->Flags;
 		SubParams.RootToSequenceTransform   = SubData.RootToSequenceTransform;
 #if WITH_EDITORONLY_DATA
-		SubParams.StartTimeBreadcrumbs      = SubData.StartTimeBreadcrumbs;
-		SubParams.EndTimeBreadcrumbs        = SubData.EndTimeBreadcrumbs;
+		SubParams.RootToUnwarpedLocalTransform = SubData.RootToUnwarpedLocalTransform;
+		SubParams.StartTimeBreadcrumbs				= SubData.StartTimeBreadcrumbs;
+		SubParams.EndTimeBreadcrumbs				= SubData.EndTimeBreadcrumbs;
 #else
 		SubParams.StartTimeBreadcrumbs      = this->StartTimeBreadcrumbs;
 		SubParams.EndTimeBreadcrumbs        = this->EndTimeBreadcrumbs;
@@ -296,6 +297,10 @@ struct FGatherParameters
 
 	/** Transform from the root time-space to the current sequence's time-space */
 	FMovieSceneSequenceTransform RootToSequenceTransform;
+#if WITH_EDITORONLY_DATA
+	/** The transform from root space to this sub-sequence's unwarped local space. */
+	FMovieSceneSequenceTransform RootToUnwarpedLocalTransform;
+#endif
 	FMovieSceneTransformBreadcrumbs StartTimeBreadcrumbs;
 	FMovieSceneTransformBreadcrumbs EndTimeBreadcrumbs;
 
@@ -1583,7 +1588,6 @@ bool UMovieSceneCompiledDataManager::CompileHierarchy(UMovieSceneSequence* Seque
 
 	UE::MovieScene::FSubSequencePath RootPath;
 
-	FGatherParameters TimeWarpedParams;
 	const FGatherParameters* ParamsToUse = &Params;
 
 	bool bContainsTimeWarp = false;
@@ -1718,6 +1722,9 @@ bool UMovieSceneCompiledDataManager::GenerateSubSequenceData(UMovieSceneSubTrack
 		// Put the root play range in the new root space
 		NewSubData.PlayRange               = TRange<FFrameNumber>::Intersection(InnerClampRange, NewSubData.PlayRange.Value);
 		NewSubData.RootToSequenceTransform = NewSubData.RootToSequenceTransform * Params.RootToSequenceTransform;
+#if WITH_EDITORONLY_DATA
+		NewSubData.RootToUnwarpedLocalTransform = NewSubData.RootToUnwarpedLocalTransform * Params.RootToUnwarpedLocalTransform;
+#endif
 		NewSubData.HierarchicalBias        = Params.HierarchicalBias + NewSubData.HierarchicalBias;
 		NewSubData.AccumulatedFlags        = UE::MovieScene::AccumulateChildSubSectionFlags(Params.AccumulatedFlags, NewSubData.AccumulatedFlags);
 
diff --git a/Engine/Source/Runtime/MovieScene/Public/MovieSceneSequenceVisitor.h b/Engine/Source/Runtime/MovieScene/Public/MovieSceneSequenceVisitor.h
--- a/Engine/Source/Runtime/MovieScene/Public/MovieSceneSequenceVisitor.h
+++ b/Engine/Source/Runtime/MovieScene/Public/MovieSceneSequenceVisitor.h
@@ -69,6 +69,10 @@ struct FSubSequenceSpace
 
 	/** Transform from the root time-space to the current sequence's time-space */
 	FMovieSceneSequenceTransform RootToSequenceTransform;
+#if WITH_EDITORONLY_DATA
+	/** Transform from the root time-space to the current sequence's time-space without time warp */
+	FMovieSceneSequenceTransform RootToUnwarpedLocalTransform;
+#endif
 	FMovieSceneTransformBreadcrumbs StartBreadcrumbs;
 	FMovieSceneTransformBreadcrumbs EndBreadcrumbs;
 	/** The ID of the sequence being compiled */
