Sequencer: Replaceables were not re-linking upon replaying the same sequence from a level sequence actor. They would stay bound as before, but without calling the dynamic binding code, which users would likely expect.

diff --git a/Engine/Source/Runtime/MovieScene/Private/EntitySystem/MovieSceneBindingLifetimeSystem.cpp b/Engine/Source/Runtime/MovieScene/Private/EntitySystem/MovieSceneBindingLifetimeSystem.cpp
--- a/Engine/Source/Runtime/MovieScene/Private/EntitySystem/MovieSceneBindingLifetimeSystem.cpp
+++ b/Engine/Source/Runtime/MovieScene/Private/EntitySystem/MovieSceneBindingLifetimeSystem.cpp
@@ -94,6 +94,11 @@ void UMovieSceneBindingLifetimeSystem::OnRun(FSystemTaskPrerequisites& InPrerequ
 						}
 					}
 				}
+				if (!bLink)
+				{
+					// Invalidate the binding, forcing it to be rebound
+					EvaluationState->Invalidate(BindingLifetime.BindingGuid, SequenceID);
+				}
 			}
 		}
 	};
