Sequencer: Check for valid instance handle

This fixes a crash in OnBlueprintPreCompile that gets triggered when Sequencer initializes

diff --git a/Engine/Source/Runtime/MovieScene/Private/Evaluation/MovieScenePreAnimatedState.cpp b/Engine/Source/Runtime/MovieScene/Private/Evaluation/MovieScenePreAnimatedState.cpp
--- a/Engine/Source/Runtime/MovieScene/Private/Evaluation/MovieScenePreAnimatedState.cpp
+++ b/Engine/Source/Runtime/MovieScene/Private/Evaluation/MovieScenePreAnimatedState.cpp
@@ -354,7 +354,7 @@ FMovieSceneInstancePreAnimatedState* FMovieScenePreAnimatedState::GetState() con
 	using namespace UE::MovieScene;
 
 	UMovieSceneEntitySystemLinker* Linker = WeakLinker.Get();
-	if (Linker && InstanceHandle.IsValid())
+	if (Linker && InstanceHandle.IsValid() && Linker->GetInstanceRegistry()->IsHandleValid(InstanceHandle))
 	{
 		const FSequenceInstance& Instance = Linker->GetInstanceRegistry()->GetInstance(InstanceHandle);
 		return &Instance.GetSharedPlaybackState()->GetPreAnimatedState();
