Sequencer: Fix sequence leaking caused by conditions being hard referenced by entity metadata which remains in the global compiled data manager.

To fix this, the reference in the entity metadata to conditions has been changed to a soft reference. Instead, we keep a TArray of hard references to generated conditions on the movie scene itself. This array gets cleared on compilation, and then regenerated.

diff --git a/Engine/Source/Runtime/MovieScene/Private/Compilation/MovieSceneCompiledDataManager.cpp b/Engine/Source/Runtime/MovieScene/Private/Compilation/MovieSceneCompiledDataManager.cpp
--- a/Engine/Source/Runtime/MovieScene/Private/Compilation/MovieSceneCompiledDataManager.cpp
+++ b/Engine/Source/Runtime/MovieScene/Private/Compilation/MovieSceneCompiledDataManager.cpp
@@ -879,6 +879,13 @@ void UMovieSceneCompiledDataManager::Compile(FMovieSceneCompiledDataID DataID, U
 	Params.TemplateGenerator.Reset(&Entry);
 	Params.NetworkMask = InNetworkMask;
 
+	// Clear list of generated conditions
+	UMovieScene* MovieScene = Sequence->GetMovieScene();
+	if (ensure(MovieScene))
+	{
+		MovieScene->ResetGeneratedConditions();
+	}
+
 	// ---------------------------------------------------------------------------------------------------
 	// Step 1 - Always ensure the hierarchy information is completely up to date first
 	FMovieSceneSequenceHierarchy NewHierarchy;
@@ -899,8 +906,6 @@ void UMovieSceneCompiledDataManager::Compile(FMovieSceneCompiledDataID DataID, U
 	TSet<FGuid> GatheredSignatures;
 
 	{
-		UMovieScene* MovieScene = Sequence->GetMovieScene();
-
 		if (ensure(MovieScene))
 		{
 			for (const FMovieSceneMarkedFrame& Mark : MovieScene->GetMarkedFrames())
@@ -1482,7 +1487,7 @@ void UMovieSceneCompiledDataManager::GatherTrack(const FMovieSceneBinding* Objec
 				MetaData.bEvaluateInSequencePreRoll  = Track->EvalOptions.bEvaluateInPreroll;
 				MetaData.bEvaluateInSequencePostRoll = Track->EvalOptions.bEvaluateInPostroll;
 				
-				MetaData.Condition = MovieSceneHelpers::GetSequenceCondition(Track, Entry.Section);
+				MetaData.Condition = MovieSceneHelpers::GetSequenceCondition(Track, Entry.Section, true);
 
 				if (!EntityProvider->PopulateEvaluationField(EffectiveRange, MetaData, &FieldBuilder))
 				{
diff --git a/Engine/Source/Runtime/MovieScene/Private/EntitySystem/IMovieSceneEntityProvider.cpp b/Engine/Source/Runtime/MovieScene/Private/EntitySystem/IMovieSceneEntityProvider.cpp
--- a/Engine/Source/Runtime/MovieScene/Private/EntitySystem/IMovieSceneEntityProvider.cpp
+++ b/Engine/Source/Runtime/MovieScene/Private/EntitySystem/IMovieSceneEntityProvider.cpp
@@ -34,7 +34,7 @@ FMovieSceneEntityID FImportedEntity::Manufacture(const FEntityImportParams& Para
 	.AddConditional(Components->SequenceID, Params.Sequence.SequenceID, Params.Sequence.SequenceID != MovieSceneSequenceID::Root)
 	.AddConditional(Components->RootInstanceHandle, Params.Sequence.RootInstanceHandle, Params.Sequence.RootInstanceHandle.IsValid())
 	.AddConditional(Components->InstanceHandle, Params.Sequence.InstanceHandle, Params.Sequence.InstanceHandle.IsValid())
-	.AddConditional(Components->Condition, Params.EntityMetaData ? Params.EntityMetaData->Condition : nullptr, Params.EntityMetaData && Params.EntityMetaData->Condition);
+	.AddConditional(Components->Condition, Params.EntityMetaData ? Params.EntityMetaData->Condition.Get() : nullptr, Params.EntityMetaData && Params.EntityMetaData->Condition.IsValid());
 
 	FEntityAllocationWriteContext WriteContext(*EntityManager);
 
diff --git a/Engine/Source/Runtime/MovieScene/Private/Evaluation/MovieSceneSequenceHierarchy.cpp b/Engine/Source/Runtime/MovieScene/Private/Evaluation/MovieSceneSequenceHierarchy.cpp
--- a/Engine/Source/Runtime/MovieScene/Private/Evaluation/MovieSceneSequenceHierarchy.cpp
+++ b/Engine/Source/Runtime/MovieScene/Private/Evaluation/MovieSceneSequenceHierarchy.cpp
@@ -101,7 +101,7 @@ FMovieSceneSubSequenceData::FMovieSceneSubSequenceData(const UMovieSceneSubSecti
 		PostRollRange = TRange<FFrameNumber>(PostRollStart, PostRollEnd);
 	}
 
-	Condition = MovieSceneHelpers::GetSequenceCondition(InSubSection.GetTypedOuter<UMovieSceneTrack>(), &InSubSection);
+	Condition = MovieSceneHelpers::GetSequenceCondition(InSubSection.GetTypedOuter<UMovieSceneTrack>(), &InSubSection, true);
 }
 
 UMovieSceneSequence* FMovieSceneSubSequenceData::GetSequence() const
diff --git a/Engine/Source/Runtime/MovieScene/Private/MovieSceneCommonHelpers.cpp b/Engine/Source/Runtime/MovieScene/Private/MovieSceneCommonHelpers.cpp
--- a/Engine/Source/Runtime/MovieScene/Private/MovieSceneCommonHelpers.cpp
+++ b/Engine/Source/Runtime/MovieScene/Private/MovieSceneCommonHelpers.cpp
@@ -1050,7 +1050,7 @@ UObject* MovieSceneHelpers::GetResolutionContext(UMovieSceneSequence* Sequence,
 	return ResolutionContext;
 }
 
-const UMovieSceneCondition* MovieSceneHelpers::GetSequenceCondition(const UMovieSceneTrack* Track, const UMovieSceneSection* Section)
+const UMovieSceneCondition* MovieSceneHelpers::GetSequenceCondition(const UMovieSceneTrack* Track, const UMovieSceneSection* Section, bool bFromCompilation)
 {
 	TArray<UMovieSceneCondition*, TInlineAllocator<1>> Conditions;
 
@@ -1093,14 +1093,18 @@ const UMovieSceneCondition* MovieSceneHelpers::GetSequenceCondition(const UMovie
 	{
 		// Generate a group condition. During compilation this will get referenced by the entity metadata, otherwise this is considered a temporary and the caller
 		// is responsible for holding a reference to this condition.
-		UObject* Outer = Section ? Section->GetTypedOuter<UMovieScene>() : Track ? Track->GetTypedOuter<UMovieScene>() : nullptr;
-		check(Outer);
-		UMovieSceneGroupCondition* GroupCondition = NewObject<UMovieSceneGroupCondition>(Outer);
+		UMovieScene* MovieScene = Section ? Section->GetTypedOuter<UMovieScene>() : Track ? Track->GetTypedOuter<UMovieScene>() : nullptr;
+		check(MovieScene);
+		UMovieSceneGroupCondition* GroupCondition = NewObject<UMovieSceneGroupCondition>(MovieScene);
 		for (UMovieSceneCondition* Condition : Conditions)
 		{
 			FMovieSceneConditionContainer& ConditionContainer = GroupCondition->SubConditions.AddDefaulted_GetRef();
 			ConditionContainer.Condition = Condition;
 		}
+		if (bFromCompilation)
+		{
+			MovieScene->AddGeneratedCondition(GroupCondition);
+		}
 		return GroupCondition;
 	}
 }
diff --git a/Engine/Source/Runtime/MovieScene/Private/Tracks/MovieSceneBindingLifetimeTrack.cpp b/Engine/Source/Runtime/MovieScene/Private/Tracks/MovieSceneBindingLifetimeTrack.cpp
--- a/Engine/Source/Runtime/MovieScene/Private/Tracks/MovieSceneBindingLifetimeTrack.cpp
+++ b/Engine/Source/Runtime/MovieScene/Private/Tracks/MovieSceneBindingLifetimeTrack.cpp
@@ -175,7 +175,7 @@ bool UMovieSceneBindingLifetimeTrack::PopulateEvaluationFieldImpl(const TRange<F
 				SectionMetaData.Flags = Entry.Flags;
 				if (Entry.Section)
 				{
-					SectionMetaData.Condition = MovieSceneHelpers::GetSequenceCondition(this, Entry.Section);
+					SectionMetaData.Condition = MovieSceneHelpers::GetSequenceCondition(this, Entry.Section, true);
 				}
 				BindingLifetimeSection->ExternalPopulateEvaluationField(SectionEffectiveRange, SectionMetaData, OutFieldBuilder);
 			}
diff --git a/Engine/Source/Runtime/MovieScene/Public/Evaluation/MovieSceneEvaluationField.h b/Engine/Source/Runtime/MovieScene/Public/Evaluation/MovieSceneEvaluationField.h
--- a/Engine/Source/Runtime/MovieScene/Public/Evaluation/MovieSceneEvaluationField.h
+++ b/Engine/Source/Runtime/MovieScene/Public/Evaluation/MovieSceneEvaluationField.h
@@ -132,10 +132,9 @@ struct FMovieSceneEvaluationFieldEntityMetaData
 	}
 
 	/* Optional condition that the entity can provide to be checked at runtime before import. If the condition fails, the entity will not be imported. 
-	*  This is specifically not a weak object ptr because it might be a group condition generated by compilation.
 	*/
 	UPROPERTY()
-	TObjectPtr<const UMovieSceneCondition> Condition;
+	TSoftObjectPtr<const UMovieSceneCondition> Condition;
 
 	/** Opt-in - when this value is set, the entity should use this property path instead of the one defined on its generating section */
 	UPROPERTY()
diff --git a/Engine/Source/Runtime/MovieScene/Public/MovieScene.h b/Engine/Source/Runtime/MovieScene/Public/MovieScene.h
--- a/Engine/Source/Runtime/MovieScene/Public/MovieScene.h
+++ b/Engine/Source/Runtime/MovieScene/Public/MovieScene.h
@@ -51,6 +51,7 @@ class UK2Node;
 class UMovieSceneFolder;
 class UMovieSceneSection;
 class UMovieSceneTrack;
+class UMovieSceneGroupCondition;
 namespace UE { namespace MovieScene { class ISequenceDataEventHandler; } }
 struct FMovieSceneChannelMetaData;
 struct FMovieSceneTimeController;
@@ -1194,6 +1195,12 @@ public:
 	 */
 	MOVIESCENE_API void RemoveTag(const FName& TagToRemove);
 
+	/* Called during compilation to add a new generated condition to store in the movie scene. */
+	void AddGeneratedCondition(UMovieSceneGroupCondition* InGeneratedCondition) { GeneratedConditions.Add(InGeneratedCondition); }
+
+	/* Called by the compiler to empty the list of generated conditions*/
+	void ResetGeneratedConditions() { GeneratedConditions.Reset(); }
+
 protected:
 
 	/**
@@ -1304,6 +1311,10 @@ private:
 	UPROPERTY()
 	TArray<FMovieSceneMarkedFrame> MarkedFrames;
 
+	/* List of compiler generated group conditions, stored here to prevent garbage collection. */
+	UPROPERTY()
+	TArray<TObjectPtr<UMovieSceneGroupCondition>> GeneratedConditions;
+
 #if WITH_EDITORONLY_DATA
 
 	/** Indicates whether this movie scene is read only */
diff --git a/Engine/Source/Runtime/MovieScene/Public/MovieSceneCommonHelpers.h b/Engine/Source/Runtime/MovieScene/Public/MovieSceneCommonHelpers.h
--- a/Engine/Source/Runtime/MovieScene/Public/MovieSceneCommonHelpers.h
+++ b/Engine/Source/Runtime/MovieScene/Public/MovieSceneCommonHelpers.h
@@ -284,8 +284,9 @@ public:
 	/* Given a movie scene track and an optional section inside it, returns an optional single condition that needs to be evaluated.
 	* If multiple conditions exist in the given scope (for example a track condition, a track row condition for the row the section is on, and a section),
 	* a UMovieSceneGroupCondition will be generated, and the caller is responsible for holding a reference to this new UObject.
+	* If bFromCompilation is true, then any generated conditions will be stored on the movie scene.
 	*/
-	static MOVIESCENE_API const UMovieSceneCondition* GetSequenceCondition(const UMovieSceneTrack* Track, const UMovieSceneSection* Section);
+	static MOVIESCENE_API const UMovieSceneCondition* GetSequenceCondition(const UMovieSceneTrack* Track, const UMovieSceneSection* Section, bool bFromCompilation=false);
 	
 	/* Helper function for evaluating a condition in a movie scene, taking advantage of any cacheing that may apply. */
 	static MOVIESCENE_API bool EvaluateSequenceCondition(const FGuid& BindingID, const FMovieSceneSequenceID& SequenceID, const UMovieSceneCondition* Condition, UObject* ConditionOwnerObject, TSharedRef<const UE::MovieScene::FSharedPlaybackState> SharedPlaybackState);
diff --git a/Engine/Source/Runtime/MovieSceneTracks/Private/Tracks/MovieSceneCustomPrimitiveDataTrack.cpp b/Engine/Source/Runtime/MovieSceneTracks/Private/Tracks/MovieSceneCustomPrimitiveDataTrack.cpp
--- a/Engine/Source/Runtime/MovieSceneTracks/Private/Tracks/MovieSceneCustomPrimitiveDataTrack.cpp
+++ b/Engine/Source/Runtime/MovieSceneTracks/Private/Tracks/MovieSceneCustomPrimitiveDataTrack.cpp
@@ -297,7 +297,7 @@ bool UMovieSceneCustomPrimitiveDataTrack::PopulateEvaluationFieldImpl(const TRan
 				SectionMetaData.Flags = Entry.Flags;
 				if (Entry.Section)
 				{
-					SectionMetaData.Condition = MovieSceneHelpers::GetSequenceCondition(this, Entry.Section);
+					SectionMetaData.Condition = MovieSceneHelpers::GetSequenceCondition(this, Entry.Section, true);
 				}
 				CPDSection->ExternalPopulateEvaluationField(SectionEffectiveRange, SectionMetaData, OutFieldBuilder);
 			}
diff --git a/Engine/Source/Runtime/MovieSceneTracks/Private/Tracks/MovieSceneMaterialParameterCollectionTrack.cpp b/Engine/Source/Runtime/MovieSceneTracks/Private/Tracks/MovieSceneMaterialParameterCollectionTrack.cpp
--- a/Engine/Source/Runtime/MovieSceneTracks/Private/Tracks/MovieSceneMaterialParameterCollectionTrack.cpp
+++ b/Engine/Source/Runtime/MovieSceneTracks/Private/Tracks/MovieSceneMaterialParameterCollectionTrack.cpp
@@ -75,7 +75,7 @@ bool UMovieSceneMaterialParameterCollectionTrack::PopulateEvaluationFieldImpl(co
 				SectionMetaData.Flags = Entry.Flags;
 				if (Entry.Section)
 				{
-					SectionMetaData.Condition = MovieSceneHelpers::GetSequenceCondition(this, Entry.Section);
+					SectionMetaData.Condition = MovieSceneHelpers::GetSequenceCondition(this, Entry.Section, true);
 				}
 				if (ParameterSection)
 				{
diff --git a/Engine/Source/Runtime/MovieSceneTracks/Private/Tracks/MovieSceneMaterialTrack.cpp b/Engine/Source/Runtime/MovieSceneTracks/Private/Tracks/MovieSceneMaterialTrack.cpp
--- a/Engine/Source/Runtime/MovieSceneTracks/Private/Tracks/MovieSceneMaterialTrack.cpp
+++ b/Engine/Source/Runtime/MovieSceneTracks/Private/Tracks/MovieSceneMaterialTrack.cpp
@@ -280,7 +280,7 @@ bool UMovieSceneComponentMaterialTrack::PopulateEvaluationFieldImpl(const TRange
 				SectionMetaData.Flags = Entry.Flags;
 				if (Entry.Section)
 				{
-					SectionMetaData.Condition = MovieSceneHelpers::GetSequenceCondition(this, Entry.Section);
+					SectionMetaData.Condition = MovieSceneHelpers::GetSequenceCondition(this, Entry.Section, true);
 				}
 				if (ParameterSection)
 				{
diff --git a/Engine/Source/Runtime/UMG/Private/Animation/MovieSceneWidgetMaterialTrack.cpp b/Engine/Source/Runtime/UMG/Private/Animation/MovieSceneWidgetMaterialTrack.cpp
--- a/Engine/Source/Runtime/UMG/Private/Animation/MovieSceneWidgetMaterialTrack.cpp
+++ b/Engine/Source/Runtime/UMG/Private/Animation/MovieSceneWidgetMaterialTrack.cpp
@@ -66,7 +66,7 @@ bool UMovieSceneWidgetMaterialTrack::PopulateEvaluationFieldImpl(const TRange<FF
 		{
 			FMovieSceneEvaluationFieldEntityMetaData SectionMetaData = InMetaData;
 			SectionMetaData.Flags = Entry.Flags;
-			SectionMetaData.Condition = MovieSceneHelpers::GetSequenceCondition(this, ParameterSection);
+			SectionMetaData.Condition = MovieSceneHelpers::GetSequenceCondition(this, ParameterSection, true);
 
 			ParameterSection->ExternalPopulateEvaluationField(SectionEffectiveRange, SectionMetaData, OutFieldBuilder);
 		}
