Sequencer: Fix several Binding Lifetime issues.

* When a sequence ends, replaceable bindings with binding lifetime tracks were not invalidating the actor bindings. This meant that if the sequence was re-run, it wasn't calling the dynamic binding code in the replaceables again, but using the cached resolved object from the last time the sequence ran.
* One of the main functions of the binding lifetime track is to call interface functions on bound objects to inform on being bound and unbound from Sequencer. However, this was not working upon invalidation of bindings, or in cases for example when a replaceable failed to bind an object, and then was able to bind it later. This has now been fixed.

The main fix for this was to have the BindingLifetimeSystem run in both the spawn phase and the instantiation phase, and to copy the BindingLifetime component to child entities (that have bound objects). This allows us to add a dependency on the object being bound before calling the OnObjectBoundBySequencer event.

Additionally, as part of this fix, a new child entity initializer was added that only duplicates components to child entities if the parent entity passes a component mask. This allows us to copy GenericObjectBinding to child entities only in the case of BindingLifetime component being present- we need this for the above case.

diff --git a/Engine/Source/Runtime/MovieScene/Private/EntitySystem/BuiltInComponentTypes.cpp b/Engine/Source/Runtime/MovieScene/Private/EntitySystem/BuiltInComponentTypes.cpp
--- a/Engine/Source/Runtime/MovieScene/Private/EntitySystem/BuiltInComponentTypes.cpp
+++ b/Engine/Source/Runtime/MovieScene/Private/EntitySystem/BuiltInComponentTypes.cpp
@@ -466,7 +466,10 @@ FBuiltInComponentTypes::FBuiltInComponentTypes()
 	ComponentRegistry->Factories.DuplicateChildComponent(PropertyBinding);
 	ComponentRegistry->Factories.DuplicateChildComponent(HierarchicalBias);
 	ComponentRegistry->Factories.DuplicateChildComponent(BlendingOrder);
+	ComponentRegistry->Factories.DuplicateChildComponent(BindingLifetime);
+	ComponentRegistry->Factories.ConditionallyDuplicateChildComponent(GenericObjectBinding, FComponentMask({BindingLifetime}));
 
+	
 	// Children always need a Parent - these are initialized by the tasks that create them
 	{
 		ComponentRegistry->Factories.DefineChildComponent(FComponentTypeID::Invalid(), ParentEntity);
diff --git a/Engine/Source/Runtime/MovieScene/Private/EntitySystem/MovieSceneBindingLifetimeSystem.cpp b/Engine/Source/Runtime/MovieScene/Private/EntitySystem/MovieSceneBindingLifetimeSystem.cpp
--- a/Engine/Source/Runtime/MovieScene/Private/EntitySystem/MovieSceneBindingLifetimeSystem.cpp
+++ b/Engine/Source/Runtime/MovieScene/Private/EntitySystem/MovieSceneBindingLifetimeSystem.cpp
@@ -21,6 +21,7 @@
 #include "IMovieScenePlaybackClient.h"
 #include "EntitySystem/MovieSceneSpawnablesSystem.h"
 #include "MovieSceneBindingEventReceiverInterface.h"
+#include "EntitySystem/MovieSceneBoundObjectInstantiator.h"
 
 #include UE_INLINE_GENERATED_CPP_BY_NAME(MovieSceneBindingLifetimeSystem)
 #define LOCTEXT_NAMESPACE "MovieSceneBindingLifetimeSystem"
@@ -31,13 +32,14 @@ UMovieSceneBindingLifetimeSystem::UMovieSceneBindingLifetimeSystem(const FObject
 {
 	using namespace UE::MovieScene;
 
-	Phase = UE::MovieScene::ESystemPhase::Spawn;
+	Phase = UE::MovieScene::ESystemPhase::Spawn | UE::MovieScene::ESystemPhase::Instantiation;
 	FBuiltInComponentTypes* BuiltInComponentTypes = FBuiltInComponentTypes::Get();
 	RelevantComponent = BuiltInComponentTypes->BindingLifetime;
 
 	if (HasAnyFlags(RF_ClassDefaultObject))
 	{
 		DefineImplicitPrerequisite(GetClass(), UMovieSceneSpawnablesSystem::StaticClass());
+		DefineImplicitPrerequisite(UMovieSceneGenericBoundObjectInstantiator::StaticClass(), GetClass());
 	}
 }
 
@@ -45,82 +47,112 @@ void UMovieSceneBindingLifetimeSystem::OnRun(FSystemTaskPrerequisites& InPrerequ
 {
 	using namespace UE::MovieScene;
 	FBuiltInComponentTypes* BuiltInComponents = FBuiltInComponentTypes::Get();
+
 	if (!Linker->EntityManager.Contains(FEntityComponentFilter().Any({ BuiltInComponents->Tags.NeedsLink, BuiltInComponents->Tags.NeedsUnlink })))
 	{
 		return;
 	}
 
 	FInstanceRegistry* InstanceRegistry = Linker->GetInstanceRegistry();
-	FPreAnimatedStateExtension& PreAnimatedState = Linker->PreAnimatedState;
+	TSharedRef<FMovieSceneEntitySystemRunner> Runner = Linker->GetRunner();
 
-	bool bLink = false;
-	auto SetBindingActivation = [InstanceRegistry, &PreAnimatedState, &bLink](FMovieSceneEntityID EntityID, FInstanceHandle InstanceHandle, const FMovieSceneBindingLifetimeComponentData& BindingLifetime)
+	if (Runner->GetCurrentPhase() == ESystemPhase::Spawn)
 	{
-		const FSequenceInstance& SequenceInstance = InstanceRegistry->GetInstance(InstanceHandle);
-
-		FMovieSceneSequenceID SequenceID = SequenceInstance.GetSequenceID();
-		TSharedRef<const FSharedPlaybackState> SharedPlaybackState = SequenceInstance.GetSharedPlaybackState();
-		FMovieSceneEvaluationState* EvaluationState = SharedPlaybackState->FindCapability<FMovieSceneEvaluationState>();
-		IMovieScenePlayer* Player = FPlayerIndexPlaybackCapability::GetPlayer(SharedPlaybackState);
-		if (EvaluationState)
-		{
-			// For now we use the linking/unlinking of the inactive ranges to set the binding activations
-			if (BindingLifetime.BindingLifetimeState == EMovieSceneBindingLifetimeState::Inactive)
-			{
-				EvaluationState->SetBindingActivation(BindingLifetime.BindingGuid, SequenceID, !bLink);
-			}
-			else
+		bool bLink = false;
+		auto SetBindingActivation = [InstanceRegistry, &bLink](FMovieSceneEntityID EntityID, FGuid ObjectBindingID, const UObject* const* OptionalBoundObject, FInstanceHandle InstanceHandle, const FMovieSceneBindingLifetimeComponentData& BindingLifetime)
 			{
-				TArrayView<TWeakObjectPtr<>> BoundObjects = EvaluationState->FindBoundObjects(BindingLifetime.BindingGuid, SequenceID, SharedPlaybackState);
-				for (TWeakObjectPtr<> WeakBoundObject : BoundObjects)
+				const FSequenceInstance& SequenceInstance = InstanceRegistry->GetInstance(InstanceHandle);
+
+				FMovieSceneSequenceID SequenceID = SequenceInstance.GetSequenceID();
+				TSharedRef<const FSharedPlaybackState> SharedPlaybackState = SequenceInstance.GetSharedPlaybackState();
+				FMovieSceneEvaluationState* EvaluationState = SharedPlaybackState->FindCapability<FMovieSceneEvaluationState>();
+				IMovieScenePlayer* Player = FPlayerIndexPlaybackCapability::GetPlayer(SharedPlaybackState);
+				if (EvaluationState)
 				{
-					if (UObject* BoundObject = WeakBoundObject.Get())
+					// For now we use the linking/unlinking of the inactive ranges to set the binding activations
+					if (BindingLifetime.BindingLifetimeState == EMovieSceneBindingLifetimeState::Inactive)
+					{
+						EvaluationState->SetBindingActivation(ObjectBindingID, SequenceID, !bLink);
+					}
+					else if (!bLink)
 					{
-						if (BoundObject->Implements<UMovieSceneBindingEventReceiverInterface>())
+						if (OptionalBoundObject)
 						{
-							TScriptInterface<IMovieSceneBindingEventReceiverInterface> BindingEventReceiver = BoundObject;
-							if (BindingEventReceiver.GetObject() && Player)
+							if (UObject* BoundObject = const_cast<UObject*>(*OptionalBoundObject))
 							{
-								FMovieSceneObjectBindingID BindingID = UE::MovieScene::FRelativeObjectBindingID(MovieSceneSequenceID::Root, SequenceID, BindingLifetime.BindingGuid, SharedPlaybackState);
-								if (bLink)
+								if (BoundObject->Implements<UMovieSceneBindingEventReceiverInterface>())
 								{
-									IMovieSceneBindingEventReceiverInterface::Execute_OnObjectBoundBySequencer(BindingEventReceiver.GetObject(), Cast<UMovieSceneSequencePlayer>(Player->AsUObject()), BindingID);
-								}
-								else
-								{
-									IMovieSceneBindingEventReceiverInterface::Execute_OnObjectUnboundBySequencer(BindingEventReceiver.GetObject(), Cast<UMovieSceneSequencePlayer>(Player->AsUObject()), BindingID);
+									const TScriptInterface<IMovieSceneBindingEventReceiverInterface> BindingEventReceiver = BoundObject;
+									if (BindingEventReceiver.GetObject() && Player)
+									{
+										FMovieSceneObjectBindingID BindingID = UE::MovieScene::FRelativeObjectBindingID(MovieSceneSequenceID::Root, SequenceID, ObjectBindingID, SharedPlaybackState);
+										// The OnBound event is handled by the BindingLifetimeActivationSystem
+										IMovieSceneBindingEventReceiverInterface::Execute_OnObjectUnboundBySequencer(BindingEventReceiver.GetObject(), Cast<UMovieSceneSequencePlayer>(Player->AsUObject()), BindingID);
+									}
 								}
 							}
 						}
+
+						// Invalidate the binding, forcing it to be rebound
+						EvaluationState->Invalidate(ObjectBindingID, SequenceID);
 					}
 				}
-				if (!bLink)
+			};
+		// Unlink stale bindinglifetime entities
+		FEntityTaskBuilder()
+			.ReadEntityIDs()
+			.Read(BuiltInComponents->GenericObjectBinding)
+			.ReadOptional(BuiltInComponents->BoundObject)
+			.Read(BuiltInComponents->InstanceHandle)
+			.Read(BuiltInComponents->BindingLifetime)
+			.FilterAll({ BuiltInComponents->Tags.NeedsUnlink })
+			.Iterate_PerEntity(&Linker->EntityManager, SetBindingActivation);
+
+		// Link new bindinglifetime entities
+		bLink = true;
+		FEntityTaskBuilder()
+			.ReadEntityIDs()
+			.Read(BuiltInComponents->GenericObjectBinding)
+			.ReadOptional(BuiltInComponents->BoundObject)
+			.Read(BuiltInComponents->InstanceHandle)
+			.Read(BuiltInComponents->BindingLifetime)
+			.FilterAll({ BuiltInComponents->Tags.NeedsLink })
+			.Iterate_PerEntity(&Linker->EntityManager, SetBindingActivation);
+	}
+	else // Instantiation- we only care here about linking activate entities
+	{
+		auto SendBoundMessage = [InstanceRegistry](FMovieSceneEntityID EntityID, FGuid ObjectBindingID, FInstanceHandle InstanceHandle, const FMovieSceneBindingLifetimeComponentData& BindingLifetime, UObject* BoundObject)
+			{
+				const FSequenceInstance& SequenceInstance = InstanceRegistry->GetInstance(InstanceHandle);
+
+				FMovieSceneSequenceID SequenceID = SequenceInstance.GetSequenceID();
+				TSharedRef<const FSharedPlaybackState> SharedPlaybackState = SequenceInstance.GetSharedPlaybackState();
+				IMovieScenePlayer* Player = FPlayerIndexPlaybackCapability::GetPlayer(SharedPlaybackState);
+				// For now we use the linking/unlinking of the inactive ranges to set the binding activations
+				if (BindingLifetime.BindingLifetimeState == EMovieSceneBindingLifetimeState::Active)
 				{
-					// Invalidate the binding, forcing it to be rebound
-					EvaluationState->Invalidate(BindingLifetime.BindingGuid, SequenceID);
+					if (BoundObject && BoundObject->Implements<UMovieSceneBindingEventReceiverInterface>())
+					{
+						TScriptInterface<IMovieSceneBindingEventReceiverInterface> BindingEventReceiver = BoundObject;
+						if (BindingEventReceiver.GetObject() && Player)
+						{
+							FMovieSceneObjectBindingID BindingID = UE::MovieScene::FRelativeObjectBindingID(MovieSceneSequenceID::Root, SequenceID, ObjectBindingID, SharedPlaybackState);
+							// The OnBound event is handled by the BindingLifetimeActivationSystem
+							IMovieSceneBindingEventReceiverInterface::Execute_OnObjectBoundBySequencer(BindingEventReceiver.GetObject(), Cast<UMovieSceneSequencePlayer>(Player->AsUObject()), BindingID);
+						}
+					}
 				}
-			}
-		}
-	};
-
-	FBuiltInComponentTypes* BuiltInComponentTypes = FBuiltInComponentTypes::Get();
-
-	// Unlink stale bindinglifetime entities
-	FEntityTaskBuilder()
-		.ReadEntityIDs()
-		.Read(BuiltInComponents->InstanceHandle)
-		.Read(BuiltInComponentTypes->BindingLifetime)
-		.FilterAll({ BuiltInComponents->Tags.NeedsUnlink })
-		.Iterate_PerEntity(&Linker->EntityManager, SetBindingActivation);
-
-	// Link new bindinglifetime entities
-	bLink = true;
-	FEntityTaskBuilder()
-		.ReadEntityIDs()
-		.Read(BuiltInComponents->InstanceHandle)
-		.Read(BuiltInComponentTypes->BindingLifetime)
-		.FilterAll({ BuiltInComponents->Tags.NeedsLink })
-		.Iterate_PerEntity(&Linker->EntityManager, SetBindingActivation);
+			};
+		
+		FEntityTaskBuilder()
+			.ReadEntityIDs()
+			.Read(BuiltInComponents->GenericObjectBinding)
+			.Read(BuiltInComponents->InstanceHandle)
+			.Read(BuiltInComponents->BindingLifetime)
+			.Read(BuiltInComponents->BoundObject)
+			.FilterAll({ BuiltInComponents->Tags.NeedsLink })
+			.Iterate_PerEntity(&Linker->EntityManager, SendBoundMessage);
+	}
 }
 
 #undef LOCTEXT_NAMESPACE
diff --git a/Engine/Source/Runtime/MovieScene/Private/Sections/MovieSceneBindingLifetimeSection.cpp b/Engine/Source/Runtime/MovieScene/Private/Sections/MovieSceneBindingLifetimeSection.cpp
--- a/Engine/Source/Runtime/MovieScene/Private/Sections/MovieSceneBindingLifetimeSection.cpp
+++ b/Engine/Source/Runtime/MovieScene/Private/Sections/MovieSceneBindingLifetimeSection.cpp
@@ -40,8 +40,9 @@ void UMovieSceneBindingLifetimeSection::ImportEntityImpl(UMovieSceneEntitySystem
 
 	OutImportedEntity->AddBuilder(
 		FEntityBuilder()
+		.Add(BuiltInComponentTypes->GenericObjectBinding, Params.GetObjectBindingID())
 		.AddConditional(BuiltInComponentTypes->SpawnableBinding, Params.GetObjectBindingID(), Params.GetObjectBindingID().IsValid() && bSpawnable)
-		.AddConditional(BuiltInComponentTypes->BindingLifetime, FMovieSceneBindingLifetimeComponentData{ Params.GetObjectBindingID(), EMovieSceneBindingLifetimeState::Active }, Params.GetObjectBindingID().IsValid())
+		.AddConditional(BuiltInComponentTypes->BindingLifetime, FMovieSceneBindingLifetimeComponentData{ EMovieSceneBindingLifetimeState::Active }, Params.GetObjectBindingID().IsValid())
 	);
 
 
diff --git a/Engine/Source/Runtime/MovieScene/Private/Tracks/MovieSceneBindingLifetimeTrack.cpp b/Engine/Source/Runtime/MovieScene/Private/Tracks/MovieSceneBindingLifetimeTrack.cpp
--- a/Engine/Source/Runtime/MovieScene/Private/Tracks/MovieSceneBindingLifetimeTrack.cpp
+++ b/Engine/Source/Runtime/MovieScene/Private/Tracks/MovieSceneBindingLifetimeTrack.cpp
@@ -95,7 +95,8 @@ void UMovieSceneBindingLifetimeTrack::ImportEntityImpl(UMovieSceneEntitySystemLi
 
 		OutImportedEntity->AddBuilder(
 			FEntityBuilder()
-			.AddConditional(BuiltInComponentTypes->BindingLifetime, FMovieSceneBindingLifetimeComponentData{ Params.GetObjectBindingID(), EMovieSceneBindingLifetimeState::Inactive }, Params.GetObjectBindingID().IsValid())
+			.Add(BuiltInComponentTypes->GenericObjectBinding, Params.GetObjectBindingID())
+			.AddConditional(BuiltInComponentTypes->BindingLifetime, FMovieSceneBindingLifetimeComponentData{ EMovieSceneBindingLifetimeState::Inactive }, Params.GetObjectBindingID().IsValid())
 		);
 }
 
diff --git a/Engine/Source/Runtime/MovieScene/Public/EntitySystem/BuiltInComponentTypes.h b/Engine/Source/Runtime/MovieScene/Public/EntitySystem/BuiltInComponentTypes.h
--- a/Engine/Source/Runtime/MovieScene/Public/EntitySystem/BuiltInComponentTypes.h
+++ b/Engine/Source/Runtime/MovieScene/Public/EntitySystem/BuiltInComponentTypes.h
@@ -337,8 +337,6 @@ enum class EMovieSceneBindingLifetimeState : uint8
 
 struct FMovieSceneBindingLifetimeComponentData
 {
-	FGuid BindingGuid;
-
 	EMovieSceneBindingLifetimeState BindingLifetimeState = EMovieSceneBindingLifetimeState::Active;
 };
 
diff --git a/Engine/Source/Runtime/MovieScene/Public/EntitySystem/MovieSceneBindingLifetimeSystem.h b/Engine/Source/Runtime/MovieScene/Public/EntitySystem/MovieSceneBindingLifetimeSystem.h
--- a/Engine/Source/Runtime/MovieScene/Public/EntitySystem/MovieSceneBindingLifetimeSystem.h
+++ b/Engine/Source/Runtime/MovieScene/Public/EntitySystem/MovieSceneBindingLifetimeSystem.h
@@ -21,9 +21,7 @@ public:
 
 	MOVIESCENE_API UMovieSceneBindingLifetimeSystem(const FObjectInitializer& ObjInit);
 
-	MOVIESCENE_API static FMovieSceneAnimTypeID GetAnimTypeID();
-
 private:
 
 	MOVIESCENE_API virtual void OnRun(FSystemTaskPrerequisites& InPrerequisites, FSystemSubsequentTasks& Subsequents) override;
-};
+};
\ No newline at end of file
diff --git a/Engine/Source/Runtime/MovieScene/Public/EntitySystem/MovieSceneEntityFactory.h b/Engine/Source/Runtime/MovieScene/Public/EntitySystem/MovieSceneEntityFactory.h
--- a/Engine/Source/Runtime/MovieScene/Public/EntitySystem/MovieSceneEntityFactory.h
+++ b/Engine/Source/Runtime/MovieScene/Public/EntitySystem/MovieSceneEntityFactory.h
@@ -62,6 +62,14 @@ struct FEntityFactories
 	template<typename ComponentType>
 	void DuplicateChildComponent(TComponentTypeID<ComponentType> InComponent);
 
+	/**
+	 * Makes the given component automatically copied from a parent entity to all its children entities,
+	 * but only if the parent entity passes the given InParentComponentMask.
+	 * @note: include "EntitySystem/MovieSceneEntityFactoryTemplates.h" for definition
+	 */
+	template<typename ComponentType>
+	void ConditionallyDuplicateChildComponent(TComponentTypeID<ComponentType> InComponent, FComponentMask InParentComponentMask);
+
 	/**
 	 * Specifies that if a component is present on a parent entity, the given child component should
 	 * be created on any child entity, and initialized with the given initializer.
diff --git a/Engine/Source/Runtime/MovieScene/Public/EntitySystem/MovieSceneEntityFactoryTemplates.h b/Engine/Source/Runtime/MovieScene/Public/EntitySystem/MovieSceneEntityFactoryTemplates.h
--- a/Engine/Source/Runtime/MovieScene/Public/EntitySystem/MovieSceneEntityFactoryTemplates.h
+++ b/Engine/Source/Runtime/MovieScene/Public/EntitySystem/MovieSceneEntityFactoryTemplates.h
@@ -107,6 +107,23 @@ struct TDuplicateChildEntityInitializer : FChildEntityInitializer
 	}
 };
 
+/* Duplicates child components, but only if the parent entity passes the given component mask*/
+template<typename ComponentType>
+struct TConditionalDuplicateChildEntityInitializer : TDuplicateChildEntityInitializer<ComponentType>
+{
+	explicit TConditionalDuplicateChildEntityInitializer(TComponentTypeID<ComponentType> InComponent, FComponentMask InParentComponentMask)
+		: TDuplicateChildEntityInitializer<ComponentType>(InComponent)
+		, ParentComponentMask(InParentComponentMask)
+	{}
+
+	virtual bool IsRelevant(const FComponentMask& InParentType, const FComponentMask& InChildType) const override
+	{
+		return TDuplicateChildEntityInitializer<ComponentType>::IsRelevant(InParentType, InChildType) && InParentType.ContainsAll(ParentComponentMask);
+	}
+
+	FComponentMask ParentComponentMask;
+};
+
 struct FObjectFactoryBatch : FChildEntityFactory
 {
 	void Add(int32 EntityIndex, UObject* BoundObject);
@@ -160,6 +177,12 @@ inline void FEntityFactories::DuplicateChildComponent(TComponentTypeID<Component
 	DefineChildComponent(TDuplicateChildEntityInitializer<ComponentType>(InComponent));
 }
 
+template<typename ComponentType>
+inline void FEntityFactories::ConditionallyDuplicateChildComponent(TComponentTypeID<ComponentType> InComponent, FComponentMask InParentComponentMask)
+{
+	DefineChildComponent(TConditionalDuplicateChildEntityInitializer<ComponentType>(InComponent, InParentComponentMask));
+}
+
 template<typename ParentComponent, typename ChildComponent, typename InitializerCallback>
 inline void FEntityFactories::DefineChildComponent(TComponentTypeID<ParentComponent> InParentType, TComponentTypeID<ChildComponent> InChildType, InitializerCallback&& InInitializer)
 {
diff --git a/Engine/Source/Runtime/MovieScene/Public/EntitySystem/MovieSceneEntityFactoryTypes.h b/Engine/Source/Runtime/MovieScene/Public/EntitySystem/MovieSceneEntityFactoryTypes.h
--- a/Engine/Source/Runtime/MovieScene/Public/EntitySystem/MovieSceneEntityFactoryTypes.h
+++ b/Engine/Source/Runtime/MovieScene/Public/EntitySystem/MovieSceneEntityFactoryTypes.h
@@ -53,7 +53,7 @@ struct FChildEntityInitializer
 {
 	virtual ~FChildEntityInitializer(){}
 
-	bool IsRelevant(const FComponentMask& InParentType, const FComponentMask& InChildType) const
+	virtual bool IsRelevant(const FComponentMask& InParentType, const FComponentMask& InChildType) const
 	{
 		// Entity initializers with no parent component are always valid
 		const bool bHasParentComponent = !ParentComponent || InParentType.Contains(ParentComponent);
