Sequencer: fix more non-deterministic cooking, and missing recursion in cook optimization pass

- The OptimizeForCook() function was not recursing into sub-sequences, so those sub-sequences were only optimized when saved on their own... which means that if the root sequence was saved first, it would then have been compiled first. Compilation is recursive, so when we get to save the sub-sequence, an assert would have fired that we are modifying an already-compiled sequence.

- This fix should work because the optimization pass should be deterministic, and shouldn't do anything if run more than once.

- Also add a signature-change suppressor scope to make sure we don't modify the signatures while optimizing. This was previously done by checking IsRunningCookCommandlet(), but that doesn't catch cooking from the editor (which is soft-deprecated but still exists).

diff --git a/Engine/Source/Runtime/MovieScene/Private/MovieSceneSection.cpp b/Engine/Source/Runtime/MovieScene/Private/MovieSceneSection.cpp
--- a/Engine/Source/Runtime/MovieScene/Private/MovieSceneSection.cpp
+++ b/Engine/Source/Runtime/MovieScene/Private/MovieSceneSection.cpp
@@ -1032,8 +1032,6 @@ ECookOptimizationFlags UMovieSceneSection::GetCookOptimizationFlags() const
 
 void UMovieSceneSection::RemoveForCook()
 {
-	Modify();
-
 	for (const FMovieSceneChannelEntry& Entry : GetChannelProxy().GetAllEntries())
 	{
 		for (FMovieSceneChannel* Channel : Entry.GetChannels())
diff --git a/Engine/Source/Runtime/MovieScene/Private/MovieSceneSequence.cpp b/Engine/Source/Runtime/MovieScene/Private/MovieSceneSequence.cpp
--- a/Engine/Source/Runtime/MovieScene/Private/MovieSceneSequence.cpp
+++ b/Engine/Source/Runtime/MovieScene/Private/MovieSceneSequence.cpp
@@ -213,12 +213,16 @@ void UMovieSceneSequence::Serialize(FArchive& Ar)
 
 void UMovieSceneSequence::OptimizeForCook()
 {
+	// Suppress any change to signature GUIDs, because that could cause cooking indeterminism.
+	UE::MovieScene::FScopedSignedObjectModifySuppress SignatureChangeSupression(true);
+
 	UMovieScene* MovieScene = GetMovieScene();
 	if (!MovieScene)
 	{
 		return;
 	}
 
+	// Go through the root tracks.
 	for (int32 TrackIndex = 0; TrackIndex < MovieScene->GetTracks().Num(); )
 	{
 		UMovieSceneTrack* Track = MovieScene->GetTracks()[TrackIndex];
@@ -232,7 +236,10 @@ void UMovieSceneSequence::OptimizeForCook()
 		++TrackIndex;
 	}
 
-	// Go through the tracks again and look at sections
+	// Go through the root tracks again and look at sections.
+	// If a section points to a sub-sequence, also optimize that sub-sequence. We might end up optimizing
+	// some of these sub-sequences multiple times, if they're used in more than one place, but any 
+	// subsequent times should not do anything.
 	for (int32 TrackIndex = 0; TrackIndex < MovieScene->GetTracks().Num(); ++TrackIndex)
 	{
 		UMovieSceneTrack* Track =  MovieScene->GetTracks()[TrackIndex];
@@ -248,16 +255,24 @@ void UMovieSceneSequence::OptimizeForCook()
 					UE_LOG(LogMovieScene, Display, TEXT("Removing muted section: %s from: %s"), *Section->GetPathName(), *Track->GetDisplayName().ToString());
 					continue;
 				}
+				if (UMovieSceneSubSection* SubSection = Cast<UMovieSceneSubSection>(Section))
+				{
+					if (UMovieSceneSequence* SubSequence = SubSection->GetSequence())
+					{
+						SubSequence->OptimizeForCook();
+					}
+				}
 				++SectionIndex;
 			}
 		}
 	}
 
+	// Go through object bindings.
 	for (int32 ObjectBindingIndex = 0; ObjectBindingIndex < MovieScene->GetBindings().Num(); )
 	{
 		bool bRemoveObject = false;
 
-		// First, look to remove the object
+		// First, see if we need to remove the object.
 		for (int32 TrackIndex = 0; TrackIndex < MovieScene->GetBindings()[ObjectBindingIndex].GetTracks().Num(); ++TrackIndex)
 		{
 			UMovieSceneTrack* Track = MovieScene->GetBindings()[ObjectBindingIndex].GetTracks()[TrackIndex];
@@ -268,7 +283,7 @@ void UMovieSceneSequence::OptimizeForCook()
 			}
 		}
 
-		// Look to remove tracks
+		// Then, remove any appropriate tracks, or all tracks if we decided to remove the object altogether.
 		for (int32 TrackIndex = 0; TrackIndex < MovieScene->GetBindings()[ObjectBindingIndex].GetTracks().Num(); )
 		{
 			UMovieSceneTrack* Track = MovieScene->GetBindings()[ObjectBindingIndex].GetTracks()[TrackIndex];
@@ -282,7 +297,8 @@ void UMovieSceneSequence::OptimizeForCook()
 			++TrackIndex;
 		}
 
-		// Go through the tracks again and look at sections
+		// Go through the tracks again and look at sections.
+		// Once again, we recurse into sub-sequences if needed (see previous comment).
 		for (int32 TrackIndex = 0; TrackIndex < MovieScene->GetBindings()[ObjectBindingIndex].GetTracks().Num(); ++TrackIndex)
 		{
 			UMovieSceneTrack* Track = MovieScene->GetBindings()[ObjectBindingIndex].GetTracks()[TrackIndex];
@@ -298,6 +314,13 @@ void UMovieSceneSequence::OptimizeForCook()
 						UE_LOG(LogMovieScene, Display, TEXT("Removing muted section: %s from: %s"), *Section->GetPathName(), *Track->GetDisplayName().ToString());
 						continue;
 					}
+					if (UMovieSceneSubSection* SubSection = Cast<UMovieSceneSubSection>(Section))
+					{
+						if (UMovieSceneSequence* SubSequence = SubSection->GetSequence())
+						{
+							SubSequence->OptimizeForCook();
+						}
+					}
 					++SectionIndex;
 				}
 			}
@@ -408,4 +431,4 @@ ETrackSupport UMovieSceneSequence::IsTrackSupported(TSubclassOf<class UMovieScen
 	return IsTrackSupportedImpl(InTrackClass); 
 }
 
-#endif
\ No newline at end of file
+#endif
diff --git a/Engine/Source/Runtime/MovieScene/Private/MovieSceneSignedObject.cpp b/Engine/Source/Runtime/MovieScene/Private/MovieSceneSignedObject.cpp
--- a/Engine/Source/Runtime/MovieScene/Private/MovieSceneSignedObject.cpp
+++ b/Engine/Source/Runtime/MovieScene/Private/MovieSceneSignedObject.cpp
@@ -17,6 +17,8 @@ namespace MovieScene
 static TSet<TWeakPtr<IDeferredSignedObjectFlushSignal>> GDeferredSignedObjectFlushSignals;
 static TWeakPtr<IDeferredSignedObjectChangeHandler> GDeferredSignedObjectChangeHandler;
 static uint32 GScopedSignedObjectDeferCount = 0;
+static uint32 GScopedSignedObjectSuppressCount = 0;
+static uint32 GScopedSignedObjectWarnModifyCompiledData = 0;
 
 void SignalScopedSignedObjectModifyFlush()
 {
@@ -53,6 +55,29 @@ FScopedSignedObjectModifyDefer::~FScopedSignedObjectModifyDefer()
 	}
 }
 
+FScopedSignedObjectModifySuppress::FScopedSignedObjectModifySuppress(bool bInWarnModifyCompiledData)
+	: bWarnModifyCompiledData(bInWarnModifyCompiledData)
+{
+	++GScopedSignedObjectSuppressCount;
+
+	if (bWarnModifyCompiledData)
+	{
+		++GScopedSignedObjectWarnModifyCompiledData;
+	}
+}
+
+FScopedSignedObjectModifySuppress::~FScopedSignedObjectModifySuppress()
+{
+	--GScopedSignedObjectSuppressCount;
+	check(GScopedSignedObjectSuppressCount >= 0);
+
+	if (bWarnModifyCompiledData)
+	{
+		--GScopedSignedObjectWarnModifyCompiledData;
+		check(GScopedSignedObjectWarnModifyCompiledData >= 0);
+	}
+}
+
 } // namespace MovieScene
 } // namespace UE
 
@@ -117,7 +142,7 @@ void UMovieSceneSignedObject::MarkAsChanged()
 {
 	using namespace UE::MovieScene;
 
-	if (IsRunningCookCommandlet())
+	if (GScopedSignedObjectWarnModifyCompiledData > 0)
 	{
 		// During cooking, we don't want to regenerate new GUIDs, since they would be different
 		// every time and trigger non-detrministic cooking errors. We therefore bail out early
@@ -139,6 +164,11 @@ void UMovieSceneSignedObject::MarkAsChanged()
 						 "Current object: %s ; Current sequence: %s"),
 					*GetPathName(), *OuterSequence->GetPathName());
 		}
+	}
+
+	// Completely skip signature changes and notifications if we are suppressing things.
+	if (GScopedSignedObjectSuppressCount > 0)
+	{
 		return;
 	}
 
diff --git a/Engine/Source/Runtime/MovieScene/Private/MovieSceneTrack.cpp b/Engine/Source/Runtime/MovieScene/Private/MovieSceneTrack.cpp
--- a/Engine/Source/Runtime/MovieScene/Private/MovieSceneTrack.cpp
+++ b/Engine/Source/Runtime/MovieScene/Private/MovieSceneTrack.cpp
@@ -336,8 +336,6 @@ ECookOptimizationFlags UMovieSceneTrack::GetCookOptimizationFlags() const
 
 void UMovieSceneTrack::RemoveForCook()
 {
-	Modify();
-
 	for (UMovieSceneSection* Section : GetAllSections())
 	{
 		if (Section)
diff --git a/Engine/Source/Runtime/MovieScene/Public/MovieSceneSignedObject.h b/Engine/Source/Runtime/MovieScene/Public/MovieSceneSignedObject.h
--- a/Engine/Source/Runtime/MovieScene/Public/MovieSceneSignedObject.h
+++ b/Engine/Source/Runtime/MovieScene/Public/MovieSceneSignedObject.h
@@ -53,6 +53,14 @@ private:
 	bool bForceFlush;
 };
 
+struct FScopedSignedObjectModifySuppress
+{
+	MOVIESCENE_API FScopedSignedObjectModifySuppress(bool bInWarnModifyCompiledData = false);
+	MOVIESCENE_API ~FScopedSignedObjectModifySuppress();
+
+private:
+	bool bWarnModifyCompiledData;
+};
 
 } // namespace MovieScene
 } // namespace UE
