Sequencer: re-add the global cook signature-change suppressor

diff --git a/Engine/Source/Runtime/MovieScene/Private/MovieSceneSignedObject.cpp b/Engine/Source/Runtime/MovieScene/Private/MovieSceneSignedObject.cpp
--- a/Engine/Source/Runtime/MovieScene/Private/MovieSceneSignedObject.cpp
+++ b/Engine/Source/Runtime/MovieScene/Private/MovieSceneSignedObject.cpp
@@ -9,6 +9,14 @@
 
 #include UE_INLINE_GENERATED_CPP_BY_NAME(MovieSceneSignedObject)
 
+bool GMovieSceneSuppressSignatureChangeInCook = true;
+static FAutoConsoleVariableRef CVarMovieSceneSuppressSignatureChangeInCook(
+	TEXT("Sequencer.SuppressSignatureChangeInCook"),
+	GMovieSceneSuppressSignatureChangeInCook,
+	TEXT("(Default: true) Globally suppress signature GUID changes inside cook commandlets."),
+	ECVF_Default
+);
+
 namespace UE
 {
 namespace MovieScene
@@ -142,7 +150,9 @@ void UMovieSceneSignedObject::MarkAsChanged()
 {
 	using namespace UE::MovieScene;
 
-	if (GScopedSignedObjectWarnModifyCompiledData > 0)
+	const bool bGlobalCookSuppression = (GMovieSceneSuppressSignatureChangeInCook && IsRunningCookCommandlet());
+
+	if (GScopedSignedObjectWarnModifyCompiledData > 0 || bGlobalCookSuppression)
 	{
 		// During cooking, we don't want to regenerate new GUIDs, since they would be different
 		// every time and trigger non-detrministic cooking errors. We therefore bail out early
@@ -167,7 +177,7 @@ void UMovieSceneSignedObject::MarkAsChanged()
 	}
 
 	// Completely skip signature changes and notifications if we are suppressing things.
-	if (GScopedSignedObjectSuppressCount > 0)
+	if (GScopedSignedObjectSuppressCount > 0 || bGlobalCookSuppression)
 	{
 		return;
 	}
