Sequencer: Migrated Movie Scene Debugging to use new visualizer debugging state

diff --git a/Engine/Plugins/MovieScene/TemplateSequence/Source/TemplateSequence/Private/SequenceCameraShake.cpp b/Engine/Plugins/MovieScene/TemplateSequence/Source/TemplateSequence/Private/SequenceCameraShake.cpp
--- a/Engine/Plugins/MovieScene/TemplateSequence/Source/TemplateSequence/Private/SequenceCameraShake.cpp
+++ b/Engine/Plugins/MovieScene/TemplateSequence/Source/TemplateSequence/Private/SequenceCameraShake.cpp
@@ -9,9 +9,6 @@
 
 #include UE_INLINE_GENERATED_CPP_BY_NAME(SequenceCameraShake)
 
-#if !IS_MONOLITHIC
-UE_SELECT_ANY UE::MovieScene::FEntityManager*& GEntityManagerForDebugging = UE::MovieScene::GEntityManagerForDebuggingVisualizers;
-#endif
 
 USequenceCameraShakePattern::USequenceCameraShakePattern(const FObjectInitializer& ObjInit)
 	: Super(ObjInit)
diff --git a/Engine/Source/Editor/MovieSceneTools/Private/MovieSceneToolsModule.cpp b/Engine/Source/Editor/MovieSceneTools/Private/MovieSceneToolsModule.cpp
--- a/Engine/Source/Editor/MovieSceneTools/Private/MovieSceneToolsModule.cpp
+++ b/Engine/Source/Editor/MovieSceneTools/Private/MovieSceneToolsModule.cpp
@@ -138,10 +138,6 @@
 
 TAutoConsoleVariable<bool> CVarDuplicateLinkedAnimSequence(TEXT("Sequencer.DuplicateLinkedAnimSequence"), false, TEXT("When true when we duplicate a level sequence that has a linked anim sequence it will duplicate and link the anim sequencel, if false we leave any link alone."));
 
-#if !IS_MONOLITHIC
-UE_SELECT_ANY UE::MovieScene::FEntityManager*& GEntityManagerForDebugging = UE::MovieScene::GEntityManagerForDebuggingVisualizers;
-#endif
-
 void FMovieSceneToolsModule::StartupModule()
 {
 	using namespace UE::Sequencer;
diff --git a/Engine/Source/Editor/Sequencer/Private/SequencerModule.cpp b/Engine/Source/Editor/Sequencer/Private/SequencerModule.cpp
--- a/Engine/Source/Editor/Sequencer/Private/SequencerModule.cpp
+++ b/Engine/Source/Editor/Sequencer/Private/SequencerModule.cpp
@@ -66,10 +66,6 @@
 #include "Engine/Font.h"
 #include "CanvasTypes.h"
 
-#if !IS_MONOLITHIC
-UE_SELECT_ANY UE::MovieScene::FEntityManager*& GEntityManagerForDebugging = UE::MovieScene::GEntityManagerForDebuggingVisualizers;
-#endif
-
 namespace UE::Sequencer::Private
 {
 	static const TMap<EPropertyKeyedStatus, FName> KeyedStatusStyleNames =
diff --git a/Engine/Source/Runtime/MovieScene/Private/EntitySystem/MovieSceneEntityManager.cpp b/Engine/Source/Runtime/MovieScene/Private/EntitySystem/MovieSceneEntityManager.cpp
--- a/Engine/Source/Runtime/MovieScene/Private/EntitySystem/MovieSceneEntityManager.cpp
+++ b/Engine/Source/Runtime/MovieScene/Private/EntitySystem/MovieSceneEntityManager.cpp
@@ -18,8 +18,6 @@
 
 #include "AutoRTFM/AutoRTFM.h"
 
-UE_SELECT_ANY UE::MovieScene::FEntityManager*& GEntityManagerForDebugging = UE::MovieScene::GEntityManagerForDebuggingVisualizers;
-
 #if DO_GUARD_SLOW
 
 bool GCheckMovieSceneEntityManagerInvariants = false;
@@ -421,6 +419,8 @@ FEntityManager::FEntityManager()
 		GRichComponentDebugging = FPlatformMisc::IsDebuggerPresent();
 		GRichComponentDebuggingInitialized = true;
 	}
+
+	RichComponentDebuggingPtr = &GRichComponentDebugging;
 #endif
 }
 
diff --git a/Engine/Source/Runtime/MovieScene/Private/MovieSceneModule.cpp b/Engine/Source/Runtime/MovieScene/Private/MovieSceneModule.cpp
--- a/Engine/Source/Runtime/MovieScene/Private/MovieSceneModule.cpp
+++ b/Engine/Source/Runtime/MovieScene/Private/MovieSceneModule.cpp
@@ -8,7 +8,9 @@
 #include "HAL/IConsoleManager.h"
 #include "Logging/MessageLog.h"
 #include "Misc/UObjectToken.h"
+#include "Modules/VisualizerDebuggingState.h"
 #include "EntitySystem/BuiltInComponentTypes.h"
+#include "EntitySystem/MovieSceneEntityManager.h"
 
 DEFINE_LOG_CATEGORY(LogMovieScene);
 DEFINE_LOG_CATEGORY(LogMovieSceneECS);
@@ -88,15 +90,23 @@ public:
 
 	virtual void StartupModule() override
 	{
+		using namespace UE::Core;
+		using namespace UE::MovieScene;
+
 		FModuleManager::Get().LoadModuleChecked("UniversalObjectLocator");
 
+		// Register the global debug ptr
+		FGuid GEntityManagerForDebuggingGuid(0xaf15add3, 0xd76b4d2c, 0xa79d574b, 0xb4cd9706);
+		EVisualizerDebuggingStateResult Result =
+			FVisualizerDebuggingState::Assign(GEntityManagerForDebuggingGuid, &GEntityManagerForDebuggingVisualizers);
+
 		struct FNoopDefaultDeleter
 		{
 			void operator()(FMovieSceneModule* Object) const {}
 		};
 		ModuleHandle = MakeShareable(this, FNoopDefaultDeleter());
 
-		UE::MovieScene::FBuiltInComponentTypes::Get();
+		FBuiltInComponentTypes::Get();
 	}
 
 	virtual void ShutdownModule() override
diff --git a/Engine/Source/Runtime/MovieScene/Public/EntitySystem/MovieSceneEntityManager.h b/Engine/Source/Runtime/MovieScene/Public/EntitySystem/MovieSceneEntityManager.h
--- a/Engine/Source/Runtime/MovieScene/Public/EntitySystem/MovieSceneEntityManager.h
+++ b/Engine/Source/Runtime/MovieScene/Public/EntitySystem/MovieSceneEntityManager.h
@@ -1011,6 +1011,11 @@ private:
 	/** The value of this manager's SystemSerialNumber the last time any entities were allocated, freed, or mutated in some way (this does not include component values being written to) */
 	uint64 StructureMutationSystemSerialNumber;
 
+	/** Debugging ptr for natvis */
+#if UE_MOVIESCENE_ENTITY_DEBUG
+	const bool* RichComponentDebuggingPtr;
+#endif
+
 	/** Serially incrementing unique identifier that is assigned to each new entity allocation that is created */
 	uint32 NextAllocationID;
 
diff --git a/Engine/Source/Runtime/MovieSceneTracks/Private/MovieSceneTracksModule.cpp b/Engine/Source/Runtime/MovieSceneTracks/Private/MovieSceneTracksModule.cpp
--- a/Engine/Source/Runtime/MovieSceneTracks/Private/MovieSceneTracksModule.cpp
+++ b/Engine/Source/Runtime/MovieSceneTracks/Private/MovieSceneTracksModule.cpp
@@ -11,9 +11,6 @@
 #include "Systems/MovieSceneMaterialSystem.h"
 #include "MovieSceneTracksComponentTypes.h"
 
-#if !IS_MONOLITHIC
-UE_SELECT_ANY UE::MovieScene::FEntityManager*& GEntityManagerForDebugging = UE::MovieScene::GEntityManagerForDebuggingVisualizers;
-#endif
 
 DEFINE_STAT(MovieSceneEval_ReinitializeBoundMaterials);
 
diff --git a/Engine/Source/Runtime/UMG/Private/UMGModule.cpp b/Engine/Source/Runtime/UMG/Private/UMGModule.cpp
--- a/Engine/Source/Runtime/UMG/Private/UMGModule.cpp
+++ b/Engine/Source/Runtime/UMG/Private/UMGModule.cpp
@@ -10,10 +10,6 @@
 
 DEFINE_LOG_CATEGORY(LogUMG);
 
-#if !IS_MONOLITHIC
-UE_SELECT_ANY UE::MovieScene::FEntityManager*& GEntityManagerForDebugging = UE::MovieScene::GEntityManagerForDebuggingVisualizers;
-#endif
-
 class FUMGModule : public IUMGModule
 {
 public:
