Sequencer: Anim Layers: Keying non-one additive/override weights need to work with the S Key. Found in production

diff --git a/Engine/Plugins/Animation/ControlRig/Source/ControlRigEditor/Private/Sequencer/ControlRigParameterTrackEditor.cpp b/Engine/Plugins/Animation/ControlRig/Source/ControlRigEditor/Private/Sequencer/ControlRigParameterTrackEditor.cpp
--- a/Engine/Plugins/Animation/ControlRig/Source/ControlRigEditor/Private/Sequencer/ControlRigParameterTrackEditor.cpp
+++ b/Engine/Plugins/Animation/ControlRig/Source/ControlRigEditor/Private/Sequencer/ControlRigParameterTrackEditor.cpp
@@ -4217,10 +4217,12 @@ FKeyPropertyResult FControlRigParameterTrackEditor::AddKeysToControlRigHandle(UO
 		{
 			if (!bTrackCreated)
 			{
+				//make sure to use weight on section to key
+				Weight = MovieSceneHelpers::CalculateWeightForBlending(SectionToKey, KeyTime);
 				ModifyOurGeneratedKeysByCurrentAndWeight(InObject, InControlRig, RigControlName, Track, SectionToKey, EvaluateTime, GeneratedKeys, Weight);
 			}
 			const UMovieSceneControlRigParameterSection* ParamSection = Cast<UMovieSceneControlRigParameterSection>(SectionToKey);
-			if (!ParamSection->GetDoNotKey())
+			if (ParamSection && !ParamSection->GetDoNotKey())
 			{
 				KeyPropertyResult |= AddKeysToSection(SectionToKey, KeyTime, GeneratedKeys, KeyMode, EKeyFrameTrackEditorSetDefault::SetDefaultOnAddKeys);
 			}
diff --git a/Engine/Source/Runtime/MovieScene/Private/MovieSceneCommonHelpers.cpp b/Engine/Source/Runtime/MovieScene/Private/MovieSceneCommonHelpers.cpp
--- a/Engine/Source/Runtime/MovieScene/Private/MovieSceneCommonHelpers.cpp
+++ b/Engine/Source/Runtime/MovieScene/Private/MovieSceneCommonHelpers.cpp
@@ -551,10 +551,12 @@ float MovieSceneHelpers::CalculateWeightForBlending(UMovieSceneSection* SectionT
 	float Weight = 1.0f;
 	UMovieSceneTrack* Track = SectionToKey->GetTypedOuter<UMovieSceneTrack>();
 	FOptionalMovieSceneBlendType BlendType = SectionToKey->GetBlendType();
-	if (Track && BlendType.IsValid() && (BlendType.Get() == EMovieSceneBlendType::Additive || BlendType.Get() == EMovieSceneBlendType::Absolute))
+	if (Track && BlendType.IsValid() && (( BlendType.Get() == EMovieSceneBlendType::Additive) || 
+										 ( BlendType.Get() == EMovieSceneBlendType::Absolute) || 
+										 (BlendType.Get() == EMovieSceneBlendType::Override)  ))
 	{
 		//if additive weight is just the inverse of any weight on it
-		if (BlendType.Get() == EMovieSceneBlendType::Additive)
+		if ((BlendType.Get() == EMovieSceneBlendType::Additive) || (BlendType.Get() == EMovieSceneBlendType::Override))
 		{
 			float TotalWeightValue = SectionToKey->GetTotalWeightValue(Time);
 			Weight = !FMath::IsNearlyZero(TotalWeightValue) ? 1.0f / TotalWeightValue : 0.0f;
