Sequencer: Fix divide by 0 with DX when there are 2 overlapping keys.

Change FCachedInterpolation::ComputeExtents check to ensure

#rnx

diff --git a/Engine/Source/Runtime/MovieScene/Private/Channels/MovieSceneInterpolation.cpp b/Engine/Source/Runtime/MovieScene/Private/Channels/MovieSceneInterpolation.cpp
--- a/Engine/Source/Runtime/MovieScene/Private/Channels/MovieSceneInterpolation.cpp
+++ b/Engine/Source/Runtime/MovieScene/Private/Channels/MovieSceneInterpolation.cpp
@@ -224,9 +224,11 @@ FInterpolationExtents FCachedInterpolation::ComputeExtents(FFrameTime From, FFra
 		Extents.AddPoint(WeightedCubic->Evaluate(To),   To);
 	}
 
-	check(
+	ensureMsgf(
 		Extents.MinValueTime >= From && Extents.MinValueTime <= To &&
-		Extents.MaxValueTime >= From && Extents.MaxValueTime <= To);
+		Extents.MaxValueTime >= From && Extents.MaxValueTime <= To,
+		TEXT("ComputeExtents resulted in Min: %f and Max: %f, but expected to be between From: %f To: %f"),
+		Extents.MinValueTime.AsDecimal(), Extents.MaxValueTime.AsDecimal(), From.AsDecimal(), To.AsDecimal());
 
 	return Extents;
 }
@@ -552,6 +554,11 @@ FCubicInterpolation FQuadraticInterpolation::Integral(double ConstantOffset) con
 
 double FCubicInterpolation::Evaluate(FFrameTime InTime) const
 {
+	if (FMath::IsNearlyEqual(DX, 0.0))
+	{
+		return A;
+	}
+
 	const double X = (InTime - Origin).AsDecimal() / DX;
 	return A*X*X*X + B*X*X + C*X + Constant;
 }
@@ -609,6 +616,11 @@ FQuarticInterpolation FCubicInterpolation::Integral(double ConstantOffset) const
 
 double FQuarticInterpolation::Evaluate(FFrameTime InTime) const
 {
+	if (FMath::IsNearlyEqual(DX, 0.0))
+	{
+		return A;
+	}
+
 	const double X = (InTime - Origin).AsDecimal() / DX;
 	return A*X*X*X*X + B*X*X*X + C*X*X + D*X + Constant;
 }
@@ -765,6 +777,11 @@ FCubicInterpolation FCubicBezierInterpolation::AsCubic() const
 
 double FCubicBezierInterpolation::Evaluate(FFrameTime InTime) const
 {
+	if (FMath::IsNearlyEqual(DX, 0.0))
+	{
+		return P3;
+	}
+
 	const float Interp = static_cast<float>((InTime - Origin).AsDecimal() / DX);
 	return UE::Curves::BezierInterp(P0, P1, P2, P3, Interp);
 }
