Fix for interrogation causing entities to have both NeedsLink and NeedsUnlink tags

#rnx

diff --git a/Engine/Source/Runtime/MovieScene/Private/EntitySystem/MovieSceneEntityManager.cpp b/Engine/Source/Runtime/MovieScene/Private/EntitySystem/MovieSceneEntityManager.cpp
--- a/Engine/Source/Runtime/MovieScene/Private/EntitySystem/MovieSceneEntityManager.cpp
+++ b/Engine/Source/Runtime/MovieScene/Private/EntitySystem/MovieSceneEntityManager.cpp
@@ -469,6 +469,7 @@ void FEntityManager::Destroy()
 	AllocationsWithCapacity.Reset();
 	EntityAllocationMasks.Reset();
 	EntityAllocations.Reset();
+	ParentToChild.Reset();
 
 	OnStructureChanged();
 }
