Sequencer: Change GetPossessedObjectClass to do LOAD_EditorOnly so that it doesn't load during cooking (which is what can happen with LoadSynchronous)

Added GetLoadedPossessedObjectClass() which will retrieve the object class if it is loaded. Currently unused.

diff --git a/Engine/Source/Runtime/MovieScene/Public/MovieScenePossessable.h b/Engine/Source/Runtime/MovieScene/Public/MovieScenePossessable.h
--- a/Engine/Source/Runtime/MovieScene/Public/MovieScenePossessable.h
+++ b/Engine/Source/Runtime/MovieScene/Public/MovieScenePossessable.h
@@ -102,14 +102,24 @@ public:
 #if WITH_EDITORONLY_DATA
 
 	/**
-	 * Get the class of the possessed object. Can return null if the class hasn't been loaded by any other means yet.
+	 * Get the class of the possessed object. Will load the class if it hasn't been loaded yet.
 	 *
 	 * @return Object class.
 	 * @see GetGuid, GetName
 	 */
 	const UClass* GetPossessedObjectClass() const
 	{
-		return (const UClass*)(PossessedObjectClass.LoadSynchronous());
+		return LoadObject<UClass>(nullptr, *PossessedObjectClass.ToSoftObjectPath().ToString(), nullptr, LOAD_EditorOnly);
+	}
+
+	/**
+	 * Get the class of the possessed object. Can return null if the class hasn't been loaded by any other means yet.
+	 *
+	 * @return Object class.
+	 */
+	const UClass* GetLoadedPossessedObjectClass() const
+	{
+		return (const UClass*)(PossessedObjectClass.Get());
 	}
 
 	/**
