- Fix mem stomp from links dangling if SSequencerSection is destroyed after UMovieSceneSignedObject

#rnx

diff --git a/Engine/Source/Runtime/MovieScene/Public/EventHandlers/MovieSceneDataEventContainer.h b/Engine/Source/Runtime/MovieScene/Public/EventHandlers/MovieSceneDataEventContainer.h
--- a/Engine/Source/Runtime/MovieScene/Public/EventHandlers/MovieSceneDataEventContainer.h
+++ b/Engine/Source/Runtime/MovieScene/Public/EventHandlers/MovieSceneDataEventContainer.h
@@ -49,6 +49,17 @@ template<typename EventInterface>
 struct TDataEventContainer
 {
 #if UE_MOVIESCENE_EVENTS
+	~TDataEventContainer()
+	{
+		if (IntrusiveHandlers)
+		{
+			IntrusiveHandlers->Unlink();
+		}
+		if (NonIntrusiveHandlers)
+		{
+			NonIntrusiveHandlers->Unlink();
+		}
+	}
 
 	template<typename FuncType, typename... ArgTypes>
 	void Trigger(FuncType&& Func, ArgTypes&&... Args) const
