Fix for static analysis complaint

diff --git a/Engine/Source/Runtime/MovieScene/Private/Channels/MovieSceneCurveChannelImpl.cpp b/Engine/Source/Runtime/MovieScene/Private/Channels/MovieSceneCurveChannelImpl.cpp
--- a/Engine/Source/Runtime/MovieScene/Private/Channels/MovieSceneCurveChannelImpl.cpp
+++ b/Engine/Source/Runtime/MovieScene/Private/Channels/MovieSceneCurveChannelImpl.cpp
@@ -918,7 +918,7 @@ void TMovieSceneCurveChannelImpl<ChannelType>::RemapMusicalFrameTimes(ChannelTyp
 		NewPreviousTime = NewThisTime;
 	}
 
-	if (MovedItems.Num() > 0 && InChannel && InChannel->OnKeyMovedEvent().IsBound())
+	if (MovedItems.Num() > 0 && InChannel->OnKeyMovedEvent().IsBound())
 	{
 		InChannel->OnKeyMovedEvent().Broadcast(InChannel, MovedItems);
 	}
