Sequencer: Change MakeUniqueBindingName to work similar to FActorLabelUtilities::SetActorLabelUnique so that object bindings in Sequencer are labelled sequentially as Cube, Cube2, Cube 3, Cube 4 instead of Cube, Cube (1), Cube (1) (1), Cube (2) (1), etc.

diff --git a/Engine/Source/Runtime/MovieScene/Private/MovieSceneCommonHelpers.cpp b/Engine/Source/Runtime/MovieScene/Private/MovieSceneCommonHelpers.cpp
--- a/Engine/Source/Runtime/MovieScene/Private/MovieSceneCommonHelpers.cpp
+++ b/Engine/Source/Runtime/MovieScene/Private/MovieSceneCommonHelpers.cpp
@@ -591,49 +591,93 @@ float MovieSceneHelpers::CalculateWeightForBlending(UMovieSceneSection* SectionT
 	return Weight;
 }
 
+bool SplitBindingLabel(FString& InOutLabel, int32& OutIdx)
+{
+	// Look at the label and see if it ends in a number and separate them
+	const TArray<TCHAR, FString::AllocatorType>& LabelCharArray = InOutLabel.GetCharArray();
+	for (int32 CharIdx = LabelCharArray.Num() - 1; CharIdx >= 0; CharIdx--)
+	{
+		if (CharIdx == 0 || !FChar::IsDigit(LabelCharArray[CharIdx - 1]))
+		{
+			FString Idx = InOutLabel.RightChop(CharIdx);
+			if (Idx.Len() > 0)
+			{
+				InOutLabel.LeftInline(CharIdx);
+				OutIdx = FCString::Atoi(*Idx);
+				return true;
+			}
+			break;
+		}
+	}
+	return false;
+}
+
 FString MovieSceneHelpers::MakeUniqueBindingName(UMovieScene* MovieScene, const FString& InName)
 {
-	FString NewName = InName;
+	FString Prefix = InName;
+	FString ModifiedActorLabel = InName;
+	int32   LabelIdx = 0;
 
-	auto DuplNameSpawnable = [&](FMovieSceneSpawnable& InSpawnable)
+	TArray<FString> Names;
+	for (int32 Index = 0; Index < MovieScene->GetSpawnableCount(); ++Index)
 	{
-		return InSpawnable.GetName() == NewName;
-	};
-
-	auto DuplNamePossessable = [&](FMovieScenePossessable& InPossessable)
+		Names.Add(MovieScene->GetSpawnable(Index).GetName());
+	}
+	for (int32 Index = 0; Index < MovieScene->GetPossessableCount(); ++Index)
 	{
-		return InPossessable.GetName() == NewName;
-	};
+		Names.Add(MovieScene->GetPossessable(Index).GetName());
+	}
 
-	int32 Index = 2;
-	FString UniqueString;
-	while (MovieScene->FindSpawnable(DuplNameSpawnable) || MovieScene->FindPossessable(DuplNamePossessable))
+	if (Names.Contains(ModifiedActorLabel))
 	{
-		NewName.RemoveFromEnd(UniqueString);
-		UniqueString = FString::Printf(TEXT(" (%d)"), Index++);
-		NewName += UniqueString;
+		// See if the current label ends in a number, and try to create a new label based on that
+		if (!SplitBindingLabel(Prefix, LabelIdx))
+		{
+			// If there wasn't a number on there, append a number, starting from 2 (1 before incrementing below)
+			LabelIdx = 1;
+		}
+
+		// Update the actor label until we find one that doesn't already exist
+		while (Names.Contains(ModifiedActorLabel))
+		{
+			++LabelIdx;
+			ModifiedActorLabel = FString::Printf(TEXT("%s%d"), *Prefix, LabelIdx);
+		}
 	}
-	return NewName;
+
+	return ModifiedActorLabel;
 }
 
 FString MovieSceneHelpers::MakeUniqueSpawnableName(UMovieScene* MovieScene, const FString& InName)
 {
-	FString NewName = InName;
+	FString Prefix = InName;
+	FString ModifiedActorLabel = InName;
+	int32   LabelIdx = 0;
 
-	auto DuplName = [&](FMovieSceneSpawnable& InSpawnable)
+	TArray<FString> Names;
+	for (int32 Index = 0; Index < MovieScene->GetSpawnableCount(); ++Index)
 	{
-		return InSpawnable.GetName() == NewName;
-	};
+		Names.Add(MovieScene->GetSpawnable(Index).GetName());
+	}
 
-	int32 Index = 2;
-	FString UniqueString;
-	while (MovieScene->FindSpawnable(DuplName))
+	if (Names.Contains(ModifiedActorLabel))
 	{
-		NewName.RemoveFromEnd(UniqueString);
-		UniqueString = FString::Printf(TEXT(" (%d)"), Index++);
-		NewName += UniqueString;
+		// See if the current label ends in a number, and try to create a new label based on that
+		if (!SplitBindingLabel(Prefix, LabelIdx))
+		{
+			// If there wasn't a number on there, append a number, starting from 2 (1 before incrementing below)
+			LabelIdx = 1;
+		}
+
+		// Update the actor label until we find one that doesn't already exist
+		while (Names.Contains(ModifiedActorLabel))
+		{
+			++LabelIdx;
+			ModifiedActorLabel = FString::Printf(TEXT("%s%d"), *Prefix, LabelIdx);
+		}
 	}
-	return NewName;
+
+	return ModifiedActorLabel;
 }
 
 UObject* MovieSceneHelpers::MakeSpawnableTemplateFromInstance(UObject& InSourceObject, UMovieScene* InMovieScene, FName InName)
