[AutoRTFM] Make `GetAnimTypeID` transactionally safe.

diff --git a/Engine/Source/Runtime/MovieScene/Public/Evaluation/MovieSceneAnimTypeID.h b/Engine/Source/Runtime/MovieScene/Public/Evaluation/MovieSceneAnimTypeID.h
--- a/Engine/Source/Runtime/MovieScene/Public/Evaluation/MovieSceneAnimTypeID.h
+++ b/Engine/Source/Runtime/MovieScene/Public/Evaluation/MovieSceneAnimTypeID.h
@@ -5,7 +5,8 @@
 #include "Containers/Array.h"
 #include "CoreMinimal.h"
 #include "CoreTypes.h"
-#include "HAL/ThreadSafeCounter.h"
+#include "Misc/ScopeLock.h"
+#include "Misc/TransactionallySafeMutex.h"
 #include "Templates/TypeHash.h"
 
 class FString;
@@ -103,14 +104,12 @@ struct TMovieSceneAnimTypeIDContainer
 	/** Get the type unique animation type identifier for the specified predicate */
 	FMovieSceneAnimTypeID GetAnimTypeID(const DataType& InPredicate)
 	{
-		// Crude spinlock
-		while (Lock.Set(1));
+		UE::TScopeLock ScopeLock(Mutex);
 
 		int32 Index = Data.IndexOfByKey(InPredicate);
 		if (Index != INDEX_NONE)
 		{
 			FMovieSceneAnimTypeID Value = TypeIDs[Index];
-			Lock.Set(0);
 			return Value;
 		}
 
@@ -119,7 +118,6 @@ struct TMovieSceneAnimTypeIDContainer
 		FTypeID NewID(this, TypeIDs.Num());
 		TypeIDs.Add(NewID);
 
-		Lock.Set(0);
 		return NewID;
 	}
 
@@ -133,7 +131,7 @@ private:
 		}
 	};
 
-	FThreadSafeCounter Lock;
+	FTransactionallySafeMutex Mutex;
 
 	/** Array of existing type identifiers */
 	TArray<FMovieSceneAnimTypeID> TypeIDs;
