Sequencer: fix a memory corruption problem with the evaluation hook system

diff --git a/Engine/Source/Runtime/MovieScene/Private/EntitySystem/MovieSceneEvaluationHookSystem.cpp b/Engine/Source/Runtime/MovieScene/Private/EntitySystem/MovieSceneEvaluationHookSystem.cpp
--- a/Engine/Source/Runtime/MovieScene/Private/EntitySystem/MovieSceneEvaluationHookSystem.cpp
+++ b/Engine/Source/Runtime/MovieScene/Private/EntitySystem/MovieSceneEvaluationHookSystem.cpp
@@ -231,6 +231,7 @@ void UMovieSceneEvaluationHookSystem::TriggerAllEvents()
 	SCOPE_CYCLE_COUNTER(MovieSceneECS_GenericHooks);
 
 	FInstanceRegistry* InstanceRegistry = Linker->GetInstanceRegistry();
+	FPreAnimatedEvaluationHookCaptureSources* EvaluationHookMetaData = Linker->PreAnimatedState.GetEvaluationHookMetaData();
 
 	// We need to clean our state before actually triggering the events because one of those events could
 	// call back into an evaluation (for instance, by starting play on another sequence). If we don't clean
@@ -240,14 +241,17 @@ void UMovieSceneEvaluationHookSystem::TriggerAllEvents()
 
 	for (TPair<FMovieSceneEvaluationInstanceKey, FMovieSceneEvaluationHookEventContainer>& Pair : LocalEvents)
 	{
-		const FSequenceInstance& SequenceInstance = InstanceRegistry->GetInstance(Pair.Key.InstanceHandle);
-
-		FMovieSceneContext RootContext = SequenceInstance.GetContext();
-		TSharedRef<FSharedPlaybackState> SharedPlaybackState = SequenceInstance.GetSharedPlaybackState();
-		FPreAnimatedEvaluationHookCaptureSources* EvaluationHookMetaData = Linker->PreAnimatedState.GetEvaluationHookMetaData();
-
 		for (const FMovieSceneEvaluationHookEvent& Event : Pair.Value.Events)
 		{
+			// Re-fetch the sequence instance in each loop. This is because the hook we execute might run
+			// arbitrary code that, for instance, starts other sequences. In that case, the instance registry
+			// would grow, and rellocate all sequence instances. If we held onto one of them, we could
+			// end up holding deallocated memory.
+			const FSequenceInstance& SequenceInstance = InstanceRegistry->GetInstance(Pair.Key.InstanceHandle);
+
+			FMovieSceneContext RootContext = SequenceInstance.GetContext();
+			TSharedRef<FSharedPlaybackState> SharedPlaybackState = SequenceInstance.GetSharedPlaybackState();
+
 			FScopedPreAnimatedCaptureSource CaptureSource(SharedPlaybackState, Event.Hook.Interface.GetObject(), Event.SequenceID, Event.bRestoreState);
 
 			FEvaluationHookParams Params = {
