Sequencer: add debugging visualizer for playback capabilities

diff --git a/Engine/Source/Editor/Sequencer/Private/Capabilities/CameraCutViewTargetCacheCapability.cpp b/Engine/Source/Editor/Sequencer/Private/Capabilities/CameraCutViewTargetCacheCapability.cpp
--- a/Engine/Source/Editor/Sequencer/Private/Capabilities/CameraCutViewTargetCacheCapability.cpp
+++ b/Engine/Source/Editor/Sequencer/Private/Capabilities/CameraCutViewTargetCacheCapability.cpp
@@ -4,6 +4,7 @@
 
 namespace UE::MovieScene
 {
-	const TPlaybackCapabilityID<FCameraCutViewTargetCacheCapability> FCameraCutViewTargetCacheCapability::ID
-		= TPlaybackCapabilityID<FCameraCutViewTargetCacheCapability>::Register();
-}
\ No newline at end of file
+
+UE_DEFINE_MOVIESCENE_PLAYBACK_CAPABILITY(FCameraCutViewTargetCacheCapability)
+
+}
diff --git a/Engine/Source/Editor/Sequencer/Private/Capabilities/CameraCutViewTargetCacheCapability.h b/Engine/Source/Editor/Sequencer/Private/Capabilities/CameraCutViewTargetCacheCapability.h
--- a/Engine/Source/Editor/Sequencer/Private/Capabilities/CameraCutViewTargetCacheCapability.h
+++ b/Engine/Source/Editor/Sequencer/Private/Capabilities/CameraCutViewTargetCacheCapability.h
@@ -12,8 +12,7 @@ namespace UE::MovieScene
 {
 	struct FCameraCutViewTargetCacheCapability
 	{
-		/** Playback capability ID */
-		static const TPlaybackCapabilityID<FCameraCutViewTargetCacheCapability> ID;
+		UE_DECLARE_MOVIESCENE_PLAYBACK_CAPABILITY(FCameraCutViewTargetCacheCapability)
 		
 		/** The last evaluated view target. */
 		TWeakObjectPtr<UCameraComponent> LastViewTargetCamera;
diff --git a/Engine/Source/Runtime/MovieScene/Private/Evaluation/EventContextsPlaybackCapability.cpp b/Engine/Source/Runtime/MovieScene/Private/Evaluation/EventContextsPlaybackCapability.cpp
--- a/Engine/Source/Runtime/MovieScene/Private/Evaluation/EventContextsPlaybackCapability.cpp
+++ b/Engine/Source/Runtime/MovieScene/Private/Evaluation/EventContextsPlaybackCapability.cpp
@@ -5,7 +5,7 @@
 namespace UE::MovieScene
 {
 
-TPlaybackCapabilityID<IEventContextsPlaybackCapability> IEventContextsPlaybackCapability::ID = TPlaybackCapabilityID<IEventContextsPlaybackCapability>::Register();
+UE_DEFINE_MOVIESCENE_PLAYBACK_CAPABILITY(IEventContextsPlaybackCapability)
 
 }  // namespace UE::MovieScene
 
diff --git a/Engine/Source/Runtime/MovieScene/Private/Evaluation/EventTriggerControlPlaybackCapability.cpp b/Engine/Source/Runtime/MovieScene/Private/Evaluation/EventTriggerControlPlaybackCapability.cpp
--- a/Engine/Source/Runtime/MovieScene/Private/Evaluation/EventTriggerControlPlaybackCapability.cpp
+++ b/Engine/Source/Runtime/MovieScene/Private/Evaluation/EventTriggerControlPlaybackCapability.cpp
@@ -5,7 +5,7 @@
 namespace UE::MovieScene
 {
 
-TPlaybackCapabilityID<FEventTriggerControlPlaybackCapability> FEventTriggerControlPlaybackCapability::ID = TPlaybackCapabilityID<FEventTriggerControlPlaybackCapability>::Register();
+UE_DEFINE_MOVIESCENE_PLAYBACK_CAPABILITY(FEventTriggerControlPlaybackCapability)
 
 bool FEventTriggerControlPlaybackCapability::IsDisablingEventTriggers(FFrameTime& OutDisabledUntilTime) const
 {
diff --git a/Engine/Source/Runtime/MovieScene/Private/Evaluation/IMovieScenePlaybackCapability.cpp b/Engine/Source/Runtime/MovieScene/Private/Evaluation/IMovieScenePlaybackCapability.cpp
--- a/Engine/Source/Runtime/MovieScene/Private/Evaluation/IMovieScenePlaybackCapability.cpp
+++ b/Engine/Source/Runtime/MovieScene/Private/Evaluation/IMovieScenePlaybackCapability.cpp
@@ -6,14 +6,31 @@
 namespace UE::MovieScene
 {
 
-FPlaybackCapabilityID FPlaybackCapabilityID::Register()
+FPlaybackCapabilityIDRegistry* GPlaybackCapabilityIDRegistryForDebuggingVisualizers = nullptr;
+
+FPlaybackCapabilityID FPlaybackCapabilityID::Register(const TCHAR* InDebugName)
+{
+	return FPlaybackCapabilityIDRegistry::Get()->RegisterNewID(InDebugName);
+}
+
+FPlaybackCapabilityIDRegistry* FPlaybackCapabilityIDRegistry::Get()
 {
-	static int32 StaticID = 0;
-	FPlaybackCapabilityID NewID{ StaticID++ };
+	static FPlaybackCapabilityIDRegistry Instance;
+	return &Instance;
+}
+
+FPlaybackCapabilityID FPlaybackCapabilityIDRegistry::RegisterNewID(const TCHAR* InDebugName)
+{
+	FPlaybackCapabilityIDInfo NewInfo;
+#if UE_MOVIESCENE_ENTITY_DEBUG
+	NewInfo.DebugName = InDebugName;
+#endif
+	const int32 NewID = Infos.Add(MoveTemp(NewInfo));
 	// The ID is a bit index inside FPlaybackCapabilitiesImpl's AllCapabilities uint32 bitmask,
 	// so we don't want to overflow this.
-	checkf((size_t)StaticID < sizeof(uint32) * CHAR_BIT, TEXT("Exceeded the maximum possible amount of playback capabilities!"));
-	return NewID;
+	checkf((size_t)NewID < sizeof(uint32) * CHAR_BIT, TEXT("Exceeded the maximum possible amount of playback capabilities!"));
+
+	return FPlaybackCapabilityID{ NewID };
 }
 
 } // namespace UE::MovieScene
diff --git a/Engine/Source/Runtime/MovieScene/Private/Evaluation/MovieSceneEvaluationState.cpp b/Engine/Source/Runtime/MovieScene/Private/Evaluation/MovieSceneEvaluationState.cpp
--- a/Engine/Source/Runtime/MovieScene/Private/Evaluation/MovieSceneEvaluationState.cpp
+++ b/Engine/Source/Runtime/MovieScene/Private/Evaluation/MovieSceneEvaluationState.cpp
@@ -21,8 +21,23 @@ DECLARE_CYCLE_STAT(TEXT("Iterate Bound Objects"), MovieSceneEval_IterateBoundObj
 namespace UE::MovieScene
 {
 
-TPlaybackCapabilityID<IObjectBindingNotifyPlaybackCapability> IObjectBindingNotifyPlaybackCapability::ID = TPlaybackCapabilityID<IObjectBindingNotifyPlaybackCapability>::Register();
-TPlaybackCapabilityID<IStaticBindingOverridesPlaybackCapability> IStaticBindingOverridesPlaybackCapability::ID = TPlaybackCapabilityID<IStaticBindingOverridesPlaybackCapability>::Register();
+UE_DEFINE_MOVIESCENE_PLAYBACK_CAPABILITY(IObjectBindingNotifyPlaybackCapability)
+UE_DEFINE_MOVIESCENE_PLAYBACK_CAPABILITY(IStaticBindingOverridesPlaybackCapability)
+
+FMovieSceneEvaluationOperand* FStaticBindingOverrides::GetBindingOverride(const FMovieSceneEvaluationOperand& InOperand)
+{
+	return BindingOverrides.Find(InOperand);
+}
+
+void FStaticBindingOverrides::AddBindingOverride(const FMovieSceneEvaluationOperand& InOperand, const FMovieSceneEvaluationOperand& InOverrideOperand)
+{
+	BindingOverrides.Add(InOperand, InOverrideOperand);
+}
+
+void FStaticBindingOverrides::RemoveBindingOverride(const FMovieSceneEvaluationOperand& InOperand)
+{
+	BindingOverrides.Remove(InOperand);
+}
 
 }  // namespace UE::MovieScene
 
@@ -594,7 +609,7 @@ void FMovieSceneObjectCache::FilterObjectBindings(UObject* PredicateObject, IMov
 	FilterObjectBindings(PredicateObject, Player.GetSharedPlaybackState(), OutBindings);
 }
 
-UE::MovieScene::TPlaybackCapabilityID<FMovieSceneEvaluationState> FMovieSceneEvaluationState::ID = UE::MovieScene::TPlaybackCapabilityID<FMovieSceneEvaluationState>::Register();
+UE_DEFINE_MOVIESCENE_PLAYBACK_CAPABILITY(FMovieSceneEvaluationState)
 
 void FMovieSceneEvaluationState::InvalidateExpiredObjects()
 {
diff --git a/Engine/Source/Runtime/MovieScene/Private/Evaluation/SequenceDirectorPlaybackCapability.cpp b/Engine/Source/Runtime/MovieScene/Private/Evaluation/SequenceDirectorPlaybackCapability.cpp
--- a/Engine/Source/Runtime/MovieScene/Private/Evaluation/SequenceDirectorPlaybackCapability.cpp
+++ b/Engine/Source/Runtime/MovieScene/Private/Evaluation/SequenceDirectorPlaybackCapability.cpp
@@ -10,7 +10,7 @@
 namespace UE::MovieScene
 {
 
-TPlaybackCapabilityID<FSequenceDirectorPlaybackCapability> FSequenceDirectorPlaybackCapability::ID = TPlaybackCapabilityID<FSequenceDirectorPlaybackCapability>::Register();
+UE_DEFINE_MOVIESCENE_PLAYBACK_CAPABILITY(FSequenceDirectorPlaybackCapability)
 
 FString FSequenceDirectorPlaybackCapability::FDirectorInstanceCache::GetReferencerName() const
 {
diff --git a/Engine/Source/Runtime/MovieScene/Private/MovieSceneModule.cpp b/Engine/Source/Runtime/MovieScene/Private/MovieSceneModule.cpp
--- a/Engine/Source/Runtime/MovieScene/Private/MovieSceneModule.cpp
+++ b/Engine/Source/Runtime/MovieScene/Private/MovieSceneModule.cpp
@@ -5,6 +5,7 @@
 #include "IMovieSceneModule.h"
 #include "MovieScene.h"
 #include "Compilation/IMovieSceneTemplateGenerator.h"
+#include "Evaluation/IMovieScenePlaybackCapability.h"
 #include "HAL/IConsoleManager.h"
 #include "Logging/MessageLog.h"
 #include "Misc/UObjectToken.h"
@@ -95,10 +96,13 @@ public:
 
 		FModuleManager::Get().LoadModuleChecked("UniversalObjectLocator");
 
-		// Register the global debug ptr
+		// Register the global debug ptrs
 		FGuid GEntityManagerForDebuggingGuid(0xaf15add3, 0xd76b4d2c, 0xa79d574b, 0xb4cd9706);
-		EVisualizerDebuggingStateResult Result =
-			FVisualizerDebuggingState::Assign(GEntityManagerForDebuggingGuid, &GEntityManagerForDebuggingVisualizers);
+		EVisualizerDebuggingStateResult Result1 = FVisualizerDebuggingState::Assign(GEntityManagerForDebuggingGuid, &GEntityManagerForDebuggingVisualizers);
+
+		FGuid GPlaybackCapabilityIDRegistryForDebuggingGuid(0xf9062bb9, 0xd77c4bc7, 0x861e5d87, 0xd5afb22d);
+		EVisualizerDebuggingStateResult Result2 = FVisualizerDebuggingState::Assign(GPlaybackCapabilityIDRegistryForDebuggingGuid, &GPlaybackCapabilityIDRegistryForDebuggingVisualizers);
+		GPlaybackCapabilityIDRegistryForDebuggingVisualizers = FPlaybackCapabilityIDRegistry::Get();
 
 		struct FNoopDefaultDeleter
 		{
diff --git a/Engine/Source/Runtime/MovieScene/Private/MovieScenePlayer.cpp b/Engine/Source/Runtime/MovieScene/Private/MovieScenePlayer.cpp
--- a/Engine/Source/Runtime/MovieScene/Private/MovieScenePlayer.cpp
+++ b/Engine/Source/Runtime/MovieScene/Private/MovieScenePlayer.cpp
@@ -27,7 +27,7 @@ static FTransactionallySafeRWLock       GGlobalPlayerRegistryLock;
 static TSparseArray<IMovieScenePlayer*> GGlobalPlayerRegistry;
 static TBitArray<> GGlobalPlayerUpdateFlags;
 
-TPlaybackCapabilityID<FPlayerIndexPlaybackCapability> FPlayerIndexPlaybackCapability::ID = TPlaybackCapabilityID<FPlayerIndexPlaybackCapability>::Register();
+UE_DEFINE_MOVIESCENE_PLAYBACK_CAPABILITY(FPlayerIndexPlaybackCapability)
 
 IMovieScenePlayer* FPlayerIndexPlaybackCapability::GetPlayer(TSharedRef<const FSharedPlaybackState> Owner)
 {
@@ -50,7 +50,7 @@ uint16 FPlayerIndexPlaybackCapability::GetPlayerIndex(TSharedRef<const FSharedPl
 } // namespace MovieScene
 } // namespace UE
 
-UE::MovieScene::TPlaybackCapabilityID<IMovieScenePlaybackClient> IMovieScenePlaybackClient::ID = UE::MovieScene::TPlaybackCapabilityID<IMovieScenePlaybackClient>::Register();
+UE_DEFINE_MOVIESCENE_PLAYBACK_CAPABILITY(IMovieScenePlaybackClient)
 
 IMovieScenePlayer::IMovieScenePlayer()
 {
diff --git a/Engine/Source/Runtime/MovieScene/Private/MovieSceneSpawnRegister.cpp b/Engine/Source/Runtime/MovieScene/Private/MovieSceneSpawnRegister.cpp
--- a/Engine/Source/Runtime/MovieScene/Private/MovieSceneSpawnRegister.cpp
+++ b/Engine/Source/Runtime/MovieScene/Private/MovieSceneSpawnRegister.cpp
@@ -8,7 +8,7 @@
 #include "MovieSceneBindingReferences.h"
 #include "Bindings/MovieSceneSpawnableBinding.h"
 
-UE::MovieScene::TPlaybackCapabilityID<FMovieSceneSpawnRegister> FMovieSceneSpawnRegister::ID = UE::MovieScene::TPlaybackCapabilityID<FMovieSceneSpawnRegister>::Register();
+UE_DEFINE_MOVIESCENE_PLAYBACK_CAPABILITY(FMovieSceneSpawnRegister)
 
 FMovieSceneSpawnRegister::FMovieSceneSpawnRegister() = default;
 FMovieSceneSpawnRegister::FMovieSceneSpawnRegister(const FMovieSceneSpawnRegister&) = default;
diff --git a/Engine/Source/Runtime/MovieScene/Private/Tests/MovieScenePlaybackCapabilitiesTests.cpp b/Engine/Source/Runtime/MovieScene/Private/Tests/MovieScenePlaybackCapabilitiesTests.cpp
--- a/Engine/Source/Runtime/MovieScene/Private/Tests/MovieScenePlaybackCapabilitiesTests.cpp
+++ b/Engine/Source/Runtime/MovieScene/Private/Tests/MovieScenePlaybackCapabilitiesTests.cpp
@@ -31,10 +31,10 @@ struct FTestSimpleCapability
 		++TimesDestroyed;
 	}
 
-	static TPlaybackCapabilityID<FTestSimpleCapability> ID;
+	UE_DECLARE_MOVIESCENE_PLAYBACK_CAPABILITY(FTestSimpleCapability)
 };
 
-TPlaybackCapabilityID<FTestSimpleCapability> FTestSimpleCapability::ID = TPlaybackCapabilityID<FTestSimpleCapability>::Register();
+UE_DEFINE_MOVIESCENE_PLAYBACK_CAPABILITY(FTestSimpleCapability)
 
 int32 FTestSimpleCapability::TimesDestroyed = 0;
 
@@ -42,10 +42,10 @@ struct FTestOtherSimpleCapability
 {
 	float FloatValue = 0.f;
 
-	static TPlaybackCapabilityID<FTestOtherSimpleCapability> ID;
+	UE_DECLARE_MOVIESCENE_PLAYBACK_CAPABILITY(FTestOtherSimpleCapability)
 };
 
-TPlaybackCapabilityID<FTestOtherSimpleCapability> FTestOtherSimpleCapability::ID = TPlaybackCapabilityID<FTestOtherSimpleCapability>::Register();
+UE_DEFINE_MOVIESCENE_PLAYBACK_CAPABILITY(FTestOtherSimpleCapability)
 
 // Interface capability and sample implementation.
 struct FTestCapabilityBase
@@ -60,10 +60,10 @@ struct FTestCapabilityBase
 	virtual FString GetStringValue() { return FString(); }
 	virtual void DoSomething() = 0;
 
-	static TPlaybackCapabilityID<FTestCapabilityBase> ID;
+	UE_DECLARE_MOVIESCENE_PLAYBACK_CAPABILITY(FTestCapabilityBase)
 };
 
-TPlaybackCapabilityID<FTestCapabilityBase> FTestCapabilityBase::ID = TPlaybackCapabilityID<FTestCapabilityBase>::Register();
+UE_DEFINE_MOVIESCENE_PLAYBACK_CAPABILITY(FTestCapabilityBase)
 
 int32 FTestCapabilityBase::TimesDestroyedBase = 0;
 
@@ -97,10 +97,10 @@ struct FTestInvalidatableCapability : IPlaybackCapability
 		++TimesInvalidated;
 	}
 
-	static TPlaybackCapabilityID<FTestInvalidatableCapability> ID;
+	UE_DECLARE_MOVIESCENE_PLAYBACK_CAPABILITY(FTestInvalidatableCapability)
 };
 
-TPlaybackCapabilityID<FTestInvalidatableCapability> FTestInvalidatableCapability::ID = TPlaybackCapabilityID<FTestInvalidatableCapability>::Register();
+UE_DEFINE_MOVIESCENE_PLAYBACK_CAPABILITY(FTestInvalidatableCapability)
 
 // Derived version of the simple test capability, but with specific destructor counter
 struct FTestSimpleCapabilityDerived : FTestSimpleCapability
diff --git a/Engine/Source/Runtime/MovieScene/Public/Evaluation/EventContextsPlaybackCapability.h b/Engine/Source/Runtime/MovieScene/Public/Evaluation/EventContextsPlaybackCapability.h
--- a/Engine/Source/Runtime/MovieScene/Public/Evaluation/EventContextsPlaybackCapability.h
+++ b/Engine/Source/Runtime/MovieScene/Public/Evaluation/EventContextsPlaybackCapability.h
@@ -13,8 +13,7 @@ namespace UE::MovieScene
 /** Playback capability for controlling how events are triggered */
 struct MOVIESCENE_API IEventContextsPlaybackCapability
 {
-	/** Playback capability ID */
-	static TPlaybackCapabilityID<IEventContextsPlaybackCapability> ID;
+	UE_DECLARE_MOVIESCENE_PLAYBACK_CAPABILITY(IEventContextsPlaybackCapability)
 
 	/** Virtual destructor */
 	virtual ~IEventContextsPlaybackCapability() {}
diff --git a/Engine/Source/Runtime/MovieScene/Public/Evaluation/EventTriggerControlPlaybackCapability.h b/Engine/Source/Runtime/MovieScene/Public/Evaluation/EventTriggerControlPlaybackCapability.h
--- a/Engine/Source/Runtime/MovieScene/Public/Evaluation/EventTriggerControlPlaybackCapability.h
+++ b/Engine/Source/Runtime/MovieScene/Public/Evaluation/EventTriggerControlPlaybackCapability.h
@@ -12,8 +12,7 @@ namespace UE::MovieScene
 /** Playback capability for controlling how events are triggered */
 struct MOVIESCENE_API FEventTriggerControlPlaybackCapability
 {
-	/** Playback capability ID */
-	static TPlaybackCapabilityID<FEventTriggerControlPlaybackCapability> ID;
+	UE_DECLARE_MOVIESCENE_PLAYBACK_CAPABILITY(FEventTriggerControlPlaybackCapability)
 
 	/**
 	 * Returns whether triggering events should be temporarily disabled.
diff --git a/Engine/Source/Runtime/MovieScene/Public/Evaluation/IMovieScenePlaybackCapability.h b/Engine/Source/Runtime/MovieScene/Public/Evaluation/IMovieScenePlaybackCapability.h
--- a/Engine/Source/Runtime/MovieScene/Public/Evaluation/IMovieScenePlaybackCapability.h
+++ b/Engine/Source/Runtime/MovieScene/Public/Evaluation/IMovieScenePlaybackCapability.h
@@ -2,11 +2,13 @@
 
 #pragma once
 
+#include "Containers/Array.h"
+#include "Containers/UnrealString.h"
 #include "CoreTypes.h"
+#include "EntitySystem/MovieSceneSequenceInstanceHandle.h"
 #include "Misc/AssertionMacros.h"
 #include "Misc/EnumClassFlags.h"
 #include "Templates/SharedPointer.h"
-#include "EntitySystem/MovieSceneSequenceInstanceHandle.h"
 
 class IMovieScenePlayer;
 class UMovieSceneEntitySystemLinker;
@@ -30,7 +32,7 @@ struct FPlaybackCapabilityID
 	}
 
 protected:
-	static MOVIESCENE_API FPlaybackCapabilityID Register();
+	static MOVIESCENE_API FPlaybackCapabilityID Register(const TCHAR* InDebugName);
 };
 
 /**
@@ -56,13 +58,22 @@ private:
 		: FPlaybackCapabilityID{ InIndex }
 	{}
 
+	UE_DEPRECATED(5.6, "Use the version that takes a debug name. If defining a static ID field, please upgrade to UE_DECLARE_MOVIESCENE_PLAYBACK_CAPABILITY_API and UE_DEFINE_MOVIESCENE_PLAYBACK_CAPABILITY instead.")
 	static TPlaybackCapabilityID<T> Register()
 	{
-		FPlaybackCapabilityID StaticID = FPlaybackCapabilityID::Register();
+		return TPlaybackCapabilityID<T>::Register(TEXT("Unknown"));
+	}
+
+	static TPlaybackCapabilityID<T> Register(const TCHAR* InDebugName)
+	{
+		FPlaybackCapabilityID StaticID = FPlaybackCapabilityID::Register(InDebugName);
 		return TPlaybackCapabilityID<T>(StaticID.Index);
 	}
 };
 
+/**
+ * Interface for playback capabilities that want to be notified of various operations.
+ */
 struct IPlaybackCapability
 {
 	virtual ~IPlaybackCapability() {}
@@ -75,4 +86,77 @@ struct IPlaybackCapability
 	virtual void InvalidateCachedData(UMovieSceneEntitySystemLinker* Linker) {}
 };
 
+/**
+ * Structure providing basic information on a playback capability type.
+ */
+struct MOVIESCENE_API FPlaybackCapabilityIDInfo
+{
+#if UE_MOVIESCENE_ENTITY_DEBUG
+	/** Display name for debugging. */
+	FString DebugName;
+#endif
+};
+
+/**
+ * A registry for all known playback capability types.
+ */
+class MOVIESCENE_API FPlaybackCapabilityIDRegistry
+{
+public:
+
+	/** Gets the registry. */
+	static FPlaybackCapabilityIDRegistry* Get();
+
+	/** Registers a new playback capability type with the given display name. */
+	FPlaybackCapabilityID RegisterNewID(const TCHAR* InDebugName);
+
+private:
+
+	TArray<FPlaybackCapabilityIDInfo> Infos;
+};
+
+extern MOVIESCENE_API FPlaybackCapabilityIDRegistry* GPlaybackCapabilityIDRegistryForDebuggingVisualizers;
+
+#if UE_MOVIESCENE_ENTITY_DEBUG
+
+/**
+ * Interface to a playback capability pointer that has a valid v-table for the debugger
+ * to be able to show what it is.
+ */
+struct MOVIESCENE_API IPlaybackCapabilityDebuggingTypedPtr
+{
+	virtual ~IPlaybackCapabilityDebuggingTypedPtr() {}
+	void* Ptr = nullptr;
+};
+
+/**
+ * Actual typed version of the pointer wrapper above, for debugging.
+ */
+template<typename T>
+struct TPlaybackCapabilityDebuggingTypedPtr : IPlaybackCapabilityDebuggingTypedPtr
+{
+	TPlaybackCapabilityDebuggingTypedPtr(T* InPtr)
+	{
+		static_assert(sizeof(TPlaybackCapabilityDebuggingTypedPtr) == sizeof(IPlaybackCapabilityDebuggingTypedPtr), "Size must match");
+		Ptr = InPtr;
+	}
+};
+
+#endif
+
+#define UE_DECLARE_MOVIESCENE_PLAYBACK_CAPABILITY_API(ApiDeclSpec, ClassName)\
+	static ApiDeclSpec UE::MovieScene::TPlaybackCapabilityID<ClassName> GetPlaybackCapabilityID();\
+
+#define UE_DECLARE_MOVIESCENE_PLAYBACK_CAPABILITY(ClassName)\
+	static UE::MovieScene::TPlaybackCapabilityID<ClassName> GetPlaybackCapabilityID();\
+
+#define UE_DEFINE_MOVIESCENE_PLAYBACK_CAPABILITY(ClassName)\
+	UE::MovieScene::TPlaybackCapabilityID<ClassName> ClassName::GetPlaybackCapabilityID()\
+	{\
+		static UE::MovieScene::TPlaybackCapabilityID<ClassName> StaticID = \
+			UE::MovieScene::TPlaybackCapabilityID<ClassName>::Register(TEXT(#ClassName));\
+		return StaticID;\
+	}
+
 } // namespace UE::MovieScene
+
diff --git a/Engine/Source/Runtime/MovieScene/Public/Evaluation/MovieSceneEvaluationState.h b/Engine/Source/Runtime/MovieScene/Public/Evaluation/MovieSceneEvaluationState.h
--- a/Engine/Source/Runtime/MovieScene/Public/Evaluation/MovieSceneEvaluationState.h
+++ b/Engine/Source/Runtime/MovieScene/Public/Evaluation/MovieSceneEvaluationState.h
@@ -42,7 +42,7 @@ namespace UE::MovieScene
 	 */
 	struct MOVIESCENE_API IObjectBindingNotifyPlaybackCapability
 	{
-		static TPlaybackCapabilityID<IObjectBindingNotifyPlaybackCapability> ID;
+		UE_DECLARE_MOVIESCENE_PLAYBACK_CAPABILITY(IObjectBindingNotifyPlaybackCapability)
 
 		/** Called when multiple object bindings have changed. */
 		virtual void NotifyBindingsChanged() {}
@@ -56,7 +56,7 @@ namespace UE::MovieScene
 	 */
 	struct MOVIESCENE_API IStaticBindingOverridesPlaybackCapability
 	{
-		static TPlaybackCapabilityID<IStaticBindingOverridesPlaybackCapability> ID;
+		UE_DECLARE_MOVIESCENE_PLAYBACK_CAPABILITY(IStaticBindingOverridesPlaybackCapability)
 
 		virtual ~IStaticBindingOverridesPlaybackCapability() {}
 
@@ -68,6 +68,21 @@ namespace UE::MovieScene
 		virtual void RemoveBindingOverride(const FMovieSceneEvaluationOperand& InOperand) = 0;
 	};
 
+	/**
+	 * A re-usable implementation of IStaticBindingOverridesPlaybackCapability.
+	 */
+	struct MOVIESCENE_API FStaticBindingOverrides : public IStaticBindingOverridesPlaybackCapability
+	{
+		virtual FMovieSceneEvaluationOperand* GetBindingOverride(const FMovieSceneEvaluationOperand& InOperand) override;
+		virtual void AddBindingOverride(const FMovieSceneEvaluationOperand& InOperand, const FMovieSceneEvaluationOperand& InOverrideOperand) override;
+		virtual void RemoveBindingOverride(const FMovieSceneEvaluationOperand& InOperand) override;
+
+	private:
+		TMap<FMovieSceneEvaluationOperand, FMovieSceneEvaluationOperand> BindingOverrides;
+
+		friend class IMovieScenePlayer;
+	};
+
 }  // namespace UE::MovieScene
 
 /**
@@ -290,7 +305,11 @@ struct FMovieSceneEvaluationState : public UE::MovieScene::IPlaybackCapability
 {
 	using FSharedPlaybackState = UE::MovieScene::FSharedPlaybackState;
 
-	MOVIESCENE_API static UE::MovieScene::TPlaybackCapabilityID<FMovieSceneEvaluationState> ID;
+	UE_DECLARE_MOVIESCENE_PLAYBACK_CAPABILITY_API(MOVIESCENE_API, FMovieSceneEvaluationState)
+
+	FMovieSceneEvaluationState() = default;
+	FMovieSceneEvaluationState(const FMovieSceneEvaluationState&) = delete;
+	FMovieSceneEvaluationState& operator=(const FMovieSceneEvaluationState&) = delete;
 
 	/**
 	 * Assign a sequence to a specific ID
diff --git a/Engine/Source/Runtime/MovieScene/Public/Evaluation/MovieScenePlaybackCapabilities.h b/Engine/Source/Runtime/MovieScene/Public/Evaluation/MovieScenePlaybackCapabilities.h
--- a/Engine/Source/Runtime/MovieScene/Public/Evaluation/MovieScenePlaybackCapabilities.h
+++ b/Engine/Source/Runtime/MovieScene/Public/Evaluation/MovieScenePlaybackCapabilities.h
@@ -5,8 +5,10 @@
 #include "CoreTypes.h"
 #include "Evaluation/IMovieScenePlaybackCapability.h"
 #include "EntitySystem/RelativePtr.h"
+#include "MovieSceneFwd.h"
 #include "Templates/AlignmentTemplates.h"
 #include "Templates/PointerIsConvertibleFromTo.h"
+#include <type_traits>
 
 class IMovieScenePlayer;
 class IMovieScenePlaybackClient;
@@ -195,11 +197,16 @@ struct TPlaybackCapabilityStorageTraits<TSharedPtr<CapabilityType>, CapabilityTy
  */
 struct FPlaybackCapabilityHeader
 {
+#if UE_MOVIESCENE_ENTITY_DEBUG
+	// Extra debugging field.
+	IPlaybackCapabilityDebuggingTypedPtr* DebugPtr = nullptr;
+#endif
+
 	// 2 Bytes
 	TRelativePtr<void, uint16> Capability;
 
-	// 1 Byte
-	uint8 Sizeof = 0;
+	// 2 Bytes
+	uint16 Sizeof = 0;
 
 	// 1 Byte
 	uint8 Alignment = 0;
@@ -405,6 +412,9 @@ protected:
 			// further down), and the base pointer may have changed too (if we re-allocated the buffer)
 			void* NewCapabilityPtr = this->Memory + NewCapabilityOffsets[NewIndex];
 			NewHeaderPtr->Capability = TRelativePtr<void, uint16>(this->Memory, NewCapabilityPtr);
+
+			// The new header will share its debugging pointer with the old header at this point. But the old
+			// header will be gone when we free the entire old memory buffer after step 6.
 		};
 
 		// Step 4
@@ -423,16 +433,24 @@ protected:
 			void* NewCapabilityPtr = Memory + NewCapabilityOffsets[Index];
 			StorageType* NewCapabilityTypedPtr = reinterpret_cast<StorageType*>(NewCapabilityPtr);
 
+			static_assert(sizeof(StorageType) <= MAX_uint16, "Can't store a playback capability with a size above 65535 bits. Store it as a raw or shared pointer.");
+			static_assert(alignof(StorageType) <= MAX_uint8, "Can't store a playback capability with an alignment above 255 bits. Store it as a raw or shared pointer.");
+
 			new (NewHeaderPtr) FPlaybackCapabilityHeader(); // Reset to default constructor
-															//
+
 			NewHeaderPtr->Capability.Reset(Memory, NewCapabilityPtr);
-			NewHeaderPtr->Sizeof = sizeof(StorageType);
-			NewHeaderPtr->Alignment = alignof(StorageType);
+			NewHeaderPtr->Sizeof = static_cast<uint16>(sizeof(StorageType));
+			NewHeaderPtr->Alignment = static_cast<uint8>(alignof(StorageType));
 			NewHeaderPtr->StorageMode = FStorageTraits::GetStorageMode();
-			// The poPtrinter offset is whatever we need to add to the capability pointer in order to get
+			// The pointer offset is whatever we need to add to the capability pointer in order to get
 			// a pointer to capability type itself (in case the storage type is a sub-class of it)
 			// We only need it if the capability object is stored inline inside our memory buffer.
 			NewHeaderPtr->PointerOffset = FStorageTraits::ComputePointerOffset(NewCapabilityTypedPtr);
+
+#if UE_MOVIESCENE_ENTITY_DEBUG
+			// Set the debug pointer.
+			NewHeaderPtr->DebugPtr = new TPlaybackCapabilityDebuggingTypedPtr<StorageType>(NewCapabilityTypedPtr);
+#endif
 		}
 		--Index;
 
@@ -469,7 +487,7 @@ protected:
 
 		// Get the header for this capability
 		const int32 Index = GetCapabilityIndex(CapabilityBit);
-		const FPlaybackCapabilityHeader& Header = GetHeader(static_cast<uint8>(Index));
+		FPlaybackCapabilityHeader& Header = GetHeader(static_cast<uint8>(Index));
 
 		// Check that we are overwriting the same storage mode
 		using FStorageTraits = TPlaybackCapabilityStorageTraits<StorageType, CapabilityType>;
@@ -507,6 +525,12 @@ protected:
 		// new helper functions on the header.
 		Helpers[Index] = TPlaybackCapabilityHelpers<StorageType>::GetHelpers();
 
+#if UE_MOVIESCENE_ENTITY_DEBUG
+		// Update the debug pointer.
+		delete Header.DebugPtr;
+		Header.DebugPtr = new TPlaybackCapabilityDebuggingTypedPtr<StorageType>(TypedCapabilityPtr);
+#endif
+
 		return true;
 	}
 
@@ -516,6 +540,12 @@ protected:
 		return reinterpret_cast<FPlaybackCapabilityHeader*>(Memory)[Index];
 	}
 
+	FPlaybackCapabilityHeader& GetHeader(uint8 Index)
+	{
+		check(Index < Num);
+		return reinterpret_cast<FPlaybackCapabilityHeader*>(Memory)[Index];
+	}
+
 	TArrayView<const FPlaybackCapabilityHeader> GetHeaders() const
 	{
 		return MakeArrayView(reinterpret_cast<FPlaybackCapabilityHeader*>(Memory), Num);
@@ -542,6 +572,41 @@ protected:
 	TArray<FPlaybackCapabilityHelpers> Helpers;
 };
 
+namespace Internal
+{
+
+/** Gets the type of the playback capability ID in a backwards compatible way. */
+template<typename T, typename = void>
+struct TPlaybackCapabilityIDType {};
+
+/** Getter for the new GetPlaybackCapabilityID method. */
+template<typename T>
+struct TPlaybackCapabilityIDType<
+	T,
+	std::enable_if_t<
+		std::is_assignable_v<FPlaybackCapabilityID, decltype(T::GetPlaybackCapabilityID())>
+		>>
+{
+	using Type = std::decay_t<decltype(T::GetPlaybackCapabilityID())>;
+
+	static Type Get() { return T::GetPlaybackCapabilityID(); }
+};
+
+/** Getter for the old ID static field. */
+template<typename T>
+struct TPlaybackCapabilityIDType<
+	T,
+	std::enable_if_t<
+		std::is_assignable_v<FPlaybackCapabilityID, decltype(T::ID)>
+		>>
+{
+	using Type = std::decay_t<decltype(T::ID)>;
+
+	static Type Get() { return T::ID; }
+};
+
+}  // namespace Internal
+
 /**
  * Actual playback capabilities container.
  */
@@ -561,7 +626,7 @@ struct FPlaybackCapabilities : FPlaybackCapabilitiesImpl
 	template<typename T>
 	bool HasCapability() const
 	{
-		uint32 CapabilityBit = 1 << T::ID.Index;
+		uint32 CapabilityBit = 1 << Internal::TPlaybackCapabilityIDType<T>::Get().Index;
 		return FPlaybackCapabilitiesImpl::HasCapability(CapabilityBit);
 	}
 
@@ -571,10 +636,10 @@ struct FPlaybackCapabilities : FPlaybackCapabilitiesImpl
 	{
 		// T must be the base capability type, and not a sub-class, because we are returning a pointer to it,
 		// and we can't check what sort of sub-class we have or not.
-		using CapabilityIDType = decltype(T::ID);
+		using CapabilityIDType = Internal::TPlaybackCapabilityIDType<T>::Type;
 		static_assert(std::is_same<T, typename CapabilityIDType::CapabilityType>::value, "You must pass the actual playback capability type as a template parameter, not a sub-class.");
 
-		const TPlaybackCapabilityID<T> CapabilityID(T::ID);
+		const TPlaybackCapabilityID<T> CapabilityID(Internal::TPlaybackCapabilityIDType<T>::Get());
 		uint32 CapabilityBit = 1 << CapabilityID.Index;
 		FPlaybackCapabilityPtr Ptr = FPlaybackCapabilitiesImpl::FindCapability(CapabilityBit);
 		return Ptr.ResolveOptional<T>();
@@ -586,10 +651,10 @@ struct FPlaybackCapabilities : FPlaybackCapabilitiesImpl
 	{
 		// T must be the base capability type, and not a sub-class, because we are returning a reference to it,
 		// and we can't check what sort of sub-class we have or not.
-		using CapabilityIDType = decltype(T::ID);
+		using CapabilityIDType = Internal::TPlaybackCapabilityIDType<T>::Type;
 		static_assert(std::is_same<T, typename CapabilityIDType::CapabilityType>::value, "You must pass the actual playback capability type as a template parameter, not a sub-class.");
 
-		const TPlaybackCapabilityID<T> CapabilityID(T::ID);
+		const TPlaybackCapabilityID<T> CapabilityID(Internal::TPlaybackCapabilityIDType<T>::Get());
 		uint32 CapabilityBit = 1 << CapabilityID.Index;
 		FPlaybackCapabilityPtr Ptr = FPlaybackCapabilitiesImpl::GetCapabilityChecked(CapabilityBit);
 		return Ptr.ResolveChecked<T>();
@@ -605,10 +670,21 @@ struct FPlaybackCapabilities : FPlaybackCapabilitiesImpl
 	template<typename T, typename ...ArgTypes>
 	T& AddCapability(ArgTypes&&... InArgs)
 	{
-		using CapabilityIDType = decltype(T::ID);
+		using CapabilityIDType = Internal::TPlaybackCapabilityIDType<T>::Type;
 		using CapabilityType = typename CapabilityIDType::CapabilityType;
 
-		CapabilityType& NewCapability = DoAddCapability<T>(T::ID, Forward<ArgTypes>(InArgs)...);
+		CapabilityType& NewCapability = DoAddCapability<T>(Internal::TPlaybackCapabilityIDType<T>::Get(), Forward<ArgTypes>(InArgs)...);
+		return static_cast<T&>(NewCapability);
+	}
+
+	/**
+	 * As per AsCapability, but with an explicit capability ID. This is useful when the first template 
+	 * parameter doesn't have a single static member called ID providing the capability ID.
+	 */
+	template<typename T, typename U, typename ...ArgTypes>
+	T& AddCapability(const TPlaybackCapabilityID<U> CapabilityID, ArgTypes&&... InArgs)
+	{
+		U& NewCapability = DoAddCapability<T>(CapabilityID, Forward<ArgTypes>(InArgs)...);
 		return static_cast<T&>(NewCapability);
 	}
 
@@ -620,10 +696,10 @@ struct FPlaybackCapabilities : FPlaybackCapabilitiesImpl
 	template<typename T>
 	T& AddCapabilityRaw(T* InPointer)
 	{
-		using CapabilityIDType = decltype(T::ID);
+		using CapabilityIDType = Internal::TPlaybackCapabilityIDType<T>::Type;
 		using CapabilityType = typename CapabilityIDType::CapabilityType;
 
-		CapabilityType& NewCapability = DoAddCapability<CapabilityType*>(T::ID, static_cast<CapabilityType*>(InPointer));
+		CapabilityType& NewCapability = DoAddCapability<CapabilityType*>(Internal::TPlaybackCapabilityIDType<T>::Get(), static_cast<CapabilityType*>(InPointer));
 		return static_cast<T&>(NewCapability);
 	}
 	
@@ -635,10 +711,10 @@ struct FPlaybackCapabilities : FPlaybackCapabilitiesImpl
 	template<typename T>
 	T& AddCapabilityShared(TSharedRef<T> InSharedRef)
 	{
-		using CapabilityIDType = decltype(T::ID);
+		using CapabilityIDType = Internal::TPlaybackCapabilityIDType<T>::Type;
 		using CapabilityType = typename CapabilityIDType::CapabilityType;
 
-		CapabilityType& NewCapability = DoAddCapability<TSharedPtr<CapabilityType>>(T::ID, StaticCastSharedRef<CapabilityType>(InSharedRef));
+		CapabilityType& NewCapability = DoAddCapability<TSharedPtr<CapabilityType>>(Internal::TPlaybackCapabilityIDType<T>::Get(), StaticCastSharedRef<CapabilityType>(InSharedRef));
 		return static_cast<T&>(NewCapability);
 	}
 
@@ -652,10 +728,10 @@ struct FPlaybackCapabilities : FPlaybackCapabilitiesImpl
 	template<typename T, typename ...ArgTypes>
 	T& OverwriteCapability(ArgTypes&&... InArgs)
 	{
-		using CapabilityIDType = decltype(T::ID);
+		using CapabilityIDType = Internal::TPlaybackCapabilityIDType<T>::Type;
 		using CapabilityType = typename CapabilityIDType::CapabilityType;
 
-		CapabilityType& NewCapability = DoOverwriteCapability<T>(T::ID, Forward<ArgTypes>(InArgs)...);
+		CapabilityType& NewCapability = DoOverwriteCapability<T>(Internal::TPlaybackCapabilityIDType<T>::Get(), Forward<ArgTypes>(InArgs)...);
 		return static_cast<T&>(NewCapability);
 	}
 
@@ -666,10 +742,10 @@ struct FPlaybackCapabilities : FPlaybackCapabilitiesImpl
 	template<typename T>
 	T& OverwriteCapabilityRaw(T* InPointer)
 	{
-		using CapabilityIDType = decltype(T::ID);
+		using CapabilityIDType = Internal::TPlaybackCapabilityIDType<T>::Type;
 		using CapabilityType = typename CapabilityIDType::CapabilityType;
 
-		CapabilityType& NewCapability = DoOverwriteCapability<CapabilityType*>(T::ID, static_cast<CapabilityType*>(InPointer));
+		CapabilityType& NewCapability = DoOverwriteCapability<CapabilityType*>(Internal::TPlaybackCapabilityIDType<T>::Get(), static_cast<CapabilityType*>(InPointer));
 		return static_cast<T&>(NewCapability);
 	}
 	
@@ -680,10 +756,10 @@ struct FPlaybackCapabilities : FPlaybackCapabilitiesImpl
 	template<typename T>
 	T& OverwriteCapabilityShared(TSharedRef<T> InSharedRef)
 	{
-		using CapabilityIDType = decltype(T::ID);
+		using CapabilityIDType = Internal::TPlaybackCapabilityIDType<T>::Type;
 		using CapabilityType = typename CapabilityIDType::CapabilityType;
 
-		CapabilityType& NewCapability = DoOverwriteCapability<TSharedPtr<CapabilityType>>(T::ID, StaticCastSharedRef<CapabilityType>(InSharedRef));
+		CapabilityType& NewCapability = DoOverwriteCapability<TSharedPtr<CapabilityType>>(Internal::TPlaybackCapabilityIDType<T>::Get(), StaticCastSharedRef<CapabilityType>(InSharedRef));
 		return static_cast<T&>(NewCapability);
 	}
 
diff --git a/Engine/Source/Runtime/MovieScene/Public/Evaluation/SequenceDirectorPlaybackCapability.h b/Engine/Source/Runtime/MovieScene/Public/Evaluation/SequenceDirectorPlaybackCapability.h
--- a/Engine/Source/Runtime/MovieScene/Public/Evaluation/SequenceDirectorPlaybackCapability.h
+++ b/Engine/Source/Runtime/MovieScene/Public/Evaluation/SequenceDirectorPlaybackCapability.h
@@ -18,8 +18,7 @@ struct FSharedPlaybackState;
  */
 struct MOVIESCENE_API FSequenceDirectorPlaybackCapability
 {
-	/** Playback capability ID */
-	static TPlaybackCapabilityID<FSequenceDirectorPlaybackCapability> ID;
+	UE_DECLARE_MOVIESCENE_PLAYBACK_CAPABILITY(FSequenceDirectorPlaybackCapability)
 
 	/** Remove all director blueprint instances */
 	void ResetDirectorInstances();
diff --git a/Engine/Source/Runtime/MovieScene/Public/IMovieScenePlaybackClient.h b/Engine/Source/Runtime/MovieScene/Public/IMovieScenePlaybackClient.h
--- a/Engine/Source/Runtime/MovieScene/Public/IMovieScenePlaybackClient.h
+++ b/Engine/Source/Runtime/MovieScene/Public/IMovieScenePlaybackClient.h
@@ -26,7 +26,7 @@ class IMovieScenePlaybackClient
 public:
 	GENERATED_BODY()
 
-	MOVIESCENE_API static UE::MovieScene::TPlaybackCapabilityID<IMovieScenePlaybackClient> ID;
+	UE_DECLARE_MOVIESCENE_PLAYBACK_CAPABILITY_API(MOVIESCENE_API, IMovieScenePlaybackClient)
 
 	/**
 	 * Locate bound objects that relate to the specified binding ID
diff --git a/Engine/Source/Runtime/MovieScene/Public/IMovieScenePlayer.h b/Engine/Source/Runtime/MovieScene/Public/IMovieScenePlayer.h
--- a/Engine/Source/Runtime/MovieScene/Public/IMovieScenePlayer.h
+++ b/Engine/Source/Runtime/MovieScene/Public/IMovieScenePlayer.h
@@ -39,7 +39,7 @@ namespace UE::MovieScene
 	 */
 	struct FPlayerIndexPlaybackCapability
 	{
-		static TPlaybackCapabilityID<FPlayerIndexPlaybackCapability> ID;
+		UE_DECLARE_MOVIESCENE_PLAYBACK_CAPABILITY(FPlayerIndexPlaybackCapability)
 
 		static MOVIESCENE_API IMovieScenePlayer* GetPlayer(TSharedRef<const FSharedPlaybackState> Owner);
 		static MOVIESCENE_API uint16 GetPlayerIndex(TSharedRef<const FSharedPlaybackState> Owner);
diff --git a/Engine/Source/Runtime/MovieScene/Public/MovieSceneSpawnRegister.h b/Engine/Source/Runtime/MovieScene/Public/MovieSceneSpawnRegister.h
--- a/Engine/Source/Runtime/MovieScene/Public/MovieSceneSpawnRegister.h
+++ b/Engine/Source/Runtime/MovieScene/Public/MovieSceneSpawnRegister.h
@@ -51,7 +51,7 @@ public:
 
 	using FSharedPlaybackState = UE::MovieScene::FSharedPlaybackState;
 
-	static MOVIESCENE_API UE::MovieScene::TPlaybackCapabilityID<FMovieSceneSpawnRegister> ID;
+	UE_DECLARE_MOVIESCENE_PLAYBACK_CAPABILITY_API(MOVIESCENE_API, FMovieSceneSpawnRegister)
 
 	MOVIESCENE_API FMovieSceneSpawnRegister();
 	MOVIESCENE_API FMovieSceneSpawnRegister(const FMovieSceneSpawnRegister&);
diff --git a/Engine/Source/Runtime/MovieSceneTracks/Private/Evaluation/CameraCutPlaybackCapability.cpp b/Engine/Source/Runtime/MovieSceneTracks/Private/Evaluation/CameraCutPlaybackCapability.cpp
--- a/Engine/Source/Runtime/MovieSceneTracks/Private/Evaluation/CameraCutPlaybackCapability.cpp
+++ b/Engine/Source/Runtime/MovieSceneTracks/Private/Evaluation/CameraCutPlaybackCapability.cpp
@@ -5,7 +5,7 @@
 namespace UE::MovieScene
 {
 
-TPlaybackCapabilityID<FCameraCutPlaybackCapability> FCameraCutPlaybackCapability::ID = TPlaybackCapabilityID<FCameraCutPlaybackCapability>::Register();
+UE_DEFINE_MOVIESCENE_PLAYBACK_CAPABILITY(FCameraCutPlaybackCapability)
 
 } // namespace UE::MovieScene
 
diff --git a/Engine/Source/Runtime/MovieSceneTracks/Private/Evaluation/ViewportSettingsPlaybackCapability.cpp b/Engine/Source/Runtime/MovieSceneTracks/Private/Evaluation/ViewportSettingsPlaybackCapability.cpp
--- a/Engine/Source/Runtime/MovieSceneTracks/Private/Evaluation/ViewportSettingsPlaybackCapability.cpp
+++ b/Engine/Source/Runtime/MovieSceneTracks/Private/Evaluation/ViewportSettingsPlaybackCapability.cpp
@@ -5,7 +5,7 @@
 namespace UE::MovieScene
 {
 
-TPlaybackCapabilityID<FViewportSettingsPlaybackCapability> FViewportSettingsPlaybackCapability::ID = TPlaybackCapabilityID<FViewportSettingsPlaybackCapability>::Register();
+UE_DEFINE_MOVIESCENE_PLAYBACK_CAPABILITY(FViewportSettingsPlaybackCapability)
 
 } // namespace UE::MovieScene
 
diff --git a/Engine/Source/Runtime/MovieSceneTracks/Public/Evaluation/CameraCutPlaybackCapability.h b/Engine/Source/Runtime/MovieSceneTracks/Public/Evaluation/CameraCutPlaybackCapability.h
--- a/Engine/Source/Runtime/MovieSceneTracks/Public/Evaluation/CameraCutPlaybackCapability.h
+++ b/Engine/Source/Runtime/MovieSceneTracks/Public/Evaluation/CameraCutPlaybackCapability.h
@@ -25,8 +25,7 @@ namespace UE::MovieScene
 	 */
 	struct MOVIESCENETRACKS_API FCameraCutPlaybackCapability
 	{
-		/** Playback capability ID */
-		static TPlaybackCapabilityID<FCameraCutPlaybackCapability> ID;
+		UE_DECLARE_MOVIESCENE_PLAYBACK_CAPABILITY(FCameraCutPlaybackCapability)
 
 		/** Whether the associated sequence should execute camera cuts */
 		virtual bool ShouldUpdateCameraCut() { return true; }
diff --git a/Engine/Source/Runtime/MovieSceneTracks/Public/Evaluation/ViewportSettingsPlaybackCapability.h b/Engine/Source/Runtime/MovieSceneTracks/Public/Evaluation/ViewportSettingsPlaybackCapability.h
--- a/Engine/Source/Runtime/MovieSceneTracks/Public/Evaluation/ViewportSettingsPlaybackCapability.h
+++ b/Engine/Source/Runtime/MovieSceneTracks/Public/Evaluation/ViewportSettingsPlaybackCapability.h
@@ -42,8 +42,7 @@ namespace UE::MovieScene
  */
 struct MOVIESCENETRACKS_API FViewportSettingsPlaybackCapability
 {
-	/** Playback capability ID */
-	static TPlaybackCapabilityID<FViewportSettingsPlaybackCapability> ID;
+	UE_DECLARE_MOVIESCENE_PLAYBACK_CAPABILITY(FViewportSettingsPlaybackCapability)
 
 	/*
 	 * Set the perspective viewport settings
