Sequencer: fix android and clang builds

#rnx

diff --git a/Engine/Source/Runtime/MovieScene/Public/Evaluation/MovieScenePlaybackCapabilities.h b/Engine/Source/Runtime/MovieScene/Public/Evaluation/MovieScenePlaybackCapabilities.h
--- a/Engine/Source/Runtime/MovieScene/Public/Evaluation/MovieScenePlaybackCapabilities.h
+++ b/Engine/Source/Runtime/MovieScene/Public/Evaluation/MovieScenePlaybackCapabilities.h
@@ -636,7 +636,7 @@ struct FPlaybackCapabilities : FPlaybackCapabilitiesImpl
 	{
 		// T must be the base capability type, and not a sub-class, because we are returning a pointer to it,
 		// and we can't check what sort of sub-class we have or not.
-		using CapabilityIDType = Internal::TPlaybackCapabilityIDType<T>::Type;
+		using CapabilityIDType = typename Internal::TPlaybackCapabilityIDType<T>::Type;
 		static_assert(std::is_same<T, typename CapabilityIDType::CapabilityType>::value, "You must pass the actual playback capability type as a template parameter, not a sub-class.");
 
 		const TPlaybackCapabilityID<T> CapabilityID(Internal::TPlaybackCapabilityIDType<T>::Get());
@@ -651,7 +651,7 @@ struct FPlaybackCapabilities : FPlaybackCapabilitiesImpl
 	{
 		// T must be the base capability type, and not a sub-class, because we are returning a reference to it,
 		// and we can't check what sort of sub-class we have or not.
-		using CapabilityIDType = Internal::TPlaybackCapabilityIDType<T>::Type;
+		using CapabilityIDType = typename Internal::TPlaybackCapabilityIDType<T>::Type;
 		static_assert(std::is_same<T, typename CapabilityIDType::CapabilityType>::value, "You must pass the actual playback capability type as a template parameter, not a sub-class.");
 
 		const TPlaybackCapabilityID<T> CapabilityID(Internal::TPlaybackCapabilityIDType<T>::Get());
@@ -670,7 +670,7 @@ struct FPlaybackCapabilities : FPlaybackCapabilitiesImpl
 	template<typename T, typename ...ArgTypes>
 	T& AddCapability(ArgTypes&&... InArgs)
 	{
-		using CapabilityIDType = Internal::TPlaybackCapabilityIDType<T>::Type;
+		using CapabilityIDType = typename Internal::TPlaybackCapabilityIDType<T>::Type;
 		using CapabilityType = typename CapabilityIDType::CapabilityType;
 
 		CapabilityType& NewCapability = DoAddCapability<T>(Internal::TPlaybackCapabilityIDType<T>::Get(), Forward<ArgTypes>(InArgs)...);
@@ -696,7 +696,7 @@ struct FPlaybackCapabilities : FPlaybackCapabilitiesImpl
 	template<typename T>
 	T& AddCapabilityRaw(T* InPointer)
 	{
-		using CapabilityIDType = Internal::TPlaybackCapabilityIDType<T>::Type;
+		using CapabilityIDType = typename Internal::TPlaybackCapabilityIDType<T>::Type;
 		using CapabilityType = typename CapabilityIDType::CapabilityType;
 
 		CapabilityType& NewCapability = DoAddCapability<CapabilityType*>(Internal::TPlaybackCapabilityIDType<T>::Get(), static_cast<CapabilityType*>(InPointer));
@@ -711,7 +711,7 @@ struct FPlaybackCapabilities : FPlaybackCapabilitiesImpl
 	template<typename T>
 	T& AddCapabilityShared(TSharedRef<T> InSharedRef)
 	{
-		using CapabilityIDType = Internal::TPlaybackCapabilityIDType<T>::Type;
+		using CapabilityIDType = typename Internal::TPlaybackCapabilityIDType<T>::Type;
 		using CapabilityType = typename CapabilityIDType::CapabilityType;
 
 		CapabilityType& NewCapability = DoAddCapability<TSharedPtr<CapabilityType>>(Internal::TPlaybackCapabilityIDType<T>::Get(), StaticCastSharedRef<CapabilityType>(InSharedRef));
@@ -728,7 +728,7 @@ struct FPlaybackCapabilities : FPlaybackCapabilitiesImpl
 	template<typename T, typename ...ArgTypes>
 	T& OverwriteCapability(ArgTypes&&... InArgs)
 	{
-		using CapabilityIDType = Internal::TPlaybackCapabilityIDType<T>::Type;
+		using CapabilityIDType = typename Internal::TPlaybackCapabilityIDType<T>::Type;
 		using CapabilityType = typename CapabilityIDType::CapabilityType;
 
 		CapabilityType& NewCapability = DoOverwriteCapability<T>(Internal::TPlaybackCapabilityIDType<T>::Get(), Forward<ArgTypes>(InArgs)...);
@@ -742,7 +742,7 @@ struct FPlaybackCapabilities : FPlaybackCapabilitiesImpl
 	template<typename T>
 	T& OverwriteCapabilityRaw(T* InPointer)
 	{
-		using CapabilityIDType = Internal::TPlaybackCapabilityIDType<T>::Type;
+		using CapabilityIDType = typename Internal::TPlaybackCapabilityIDType<T>::Type;
 		using CapabilityType = typename CapabilityIDType::CapabilityType;
 
 		CapabilityType& NewCapability = DoOverwriteCapability<CapabilityType*>(Internal::TPlaybackCapabilityIDType<T>::Get(), static_cast<CapabilityType*>(InPointer));
@@ -756,7 +756,7 @@ struct FPlaybackCapabilities : FPlaybackCapabilitiesImpl
 	template<typename T>
 	T& OverwriteCapabilityShared(TSharedRef<T> InSharedRef)
 	{
-		using CapabilityIDType = Internal::TPlaybackCapabilityIDType<T>::Type;
+		using CapabilityIDType = typename Internal::TPlaybackCapabilityIDType<T>::Type;
 		using CapabilityType = typename CapabilityIDType::CapabilityType;
 
 		CapabilityType& NewCapability = DoOverwriteCapability<TSharedPtr<CapabilityType>>(Internal::TPlaybackCapabilityIDType<T>::Get(), StaticCastSharedRef<CapabilityType>(InSharedRef));
