Sequencer: deprecate direct access to State and BindingOverrides on the IMovieScenePlayer interface

This change marks the `State` and `BindingOverrides` public fields of the player interface as deprecated. This is because the underlying player class could be a "legacy player", i.e. a player only created for running legacy evaluator template tracks (not a player that's actually playing the sequence). In that case, the `State` and `BindingOverrides` live as playback capabilities on the shared playback state, and are totally not the same objects as those on the interface. So we want to deprecate direct access to them.

diff --git a/Engine/Plugins/Animation/ControlRig/Source/ControlRig/Private/Sequencer/MovieSceneControlRigParameterTemplate.cpp b/Engine/Plugins/Animation/ControlRig/Source/ControlRig/Private/Sequencer/MovieSceneControlRigParameterTemplate.cpp
--- a/Engine/Plugins/Animation/ControlRig/Source/ControlRig/Private/Sequencer/MovieSceneControlRigParameterTemplate.cpp
+++ b/Engine/Plugins/Animation/ControlRig/Source/ControlRig/Private/Sequencer/MovieSceneControlRigParameterTemplate.cpp
@@ -1241,7 +1241,7 @@ struct FControlRigParameterExecutionToken : IMovieSceneExecutionToken
 		
 		FMovieSceneSequenceID SequenceID = Operand.SequenceID;
 		TArrayView<TWeakObjectPtr<>> BoundObjects = Player.FindBoundObjects(Operand);
-		const UMovieSceneSequence* Sequence = Player.State.FindSequence(Operand.SequenceID);
+		const UMovieSceneSequence* Sequence = Player.GetEvaluationState()->FindSequence(Operand.SequenceID);
 		UControlRig* ControlRig = nullptr;
 		UObject* BoundObject = BoundObjects.Num() > 0 ? BoundObjects[0].Get() : nullptr;
 		if (Sequence && BoundObject)
diff --git a/Engine/Plugins/Animation/ControlRig/Source/ControlRigEditor/Private/ControlRigSequencerEditorLibrary.cpp b/Engine/Plugins/Animation/ControlRig/Source/ControlRigEditor/Private/ControlRigSequencerEditorLibrary.cpp
--- a/Engine/Plugins/Animation/ControlRig/Source/ControlRigEditor/Private/ControlRigSequencerEditorLibrary.cpp
+++ b/Engine/Plugins/Animation/ControlRig/Source/ControlRigEditor/Private/ControlRigSequencerEditorLibrary.cpp
@@ -110,7 +110,7 @@ static TArray<UObject*> GetBoundObjects(UWorld* World, ULevelSequence* LevelSequ
 	// Evaluation needs to occur in order to obtain spawnables
 	Player->SetPlaybackPosition(FMovieSceneSequencePlaybackParams(LevelSequence->GetMovieScene()->GetPlaybackRange().GetLowerBoundValue().Value, EUpdatePositionMethod::Play));
 
-	FMovieSceneSequenceID SequenceId = Player->State.FindSequenceId(LevelSequence);
+	FMovieSceneSequenceID SequenceId = Player->GetEvaluationState()->FindSequenceId(LevelSequence);
 
 	FMovieSceneObjectBindingID ObjectBinding = UE::MovieScene::FFixedObjectBindingID(Binding.BindingID, SequenceId);
 	TArray<UObject*> BoundObjects = Player->GetBoundObjects(ObjectBinding);
diff --git a/Engine/Plugins/Animation/ControlRig/Source/ControlRigEditor/Private/Sequencer/AnimLayers/AnimLayers.cpp b/Engine/Plugins/Animation/ControlRig/Source/ControlRigEditor/Private/Sequencer/AnimLayers/AnimLayers.cpp
--- a/Engine/Plugins/Animation/ControlRig/Source/ControlRigEditor/Private/Sequencer/AnimLayers/AnimLayers.cpp
+++ b/Engine/Plugins/Animation/ControlRig/Source/ControlRigEditor/Private/Sequencer/AnimLayers/AnimLayers.cpp
@@ -118,7 +118,7 @@ static void GetSelectedControlRigsAndBoundObjects(ISequencer* SequencerPtr, TArr
 				{
 					continue;
 				}
-				const FMovieSceneSequenceID& SequenceID = SequencerPtr->State.FindSequenceId(LevelSequence);
+				const FMovieSceneSequenceID& SequenceID = SequencerPtr->GetEvaluationState()->FindSequenceId(LevelSequence);
 
 				for (UMovieSceneTrack* AnyOleTrack : Tracks)
 				{
diff --git a/Engine/Plugins/Experimental/Avalanche/Source/AvalancheSequence/Private/AvaSequence.cpp b/Engine/Plugins/Experimental/Avalanche/Source/AvalancheSequence/Private/AvaSequence.cpp
--- a/Engine/Plugins/Experimental/Avalanche/Source/AvalancheSequence/Private/AvaSequence.cpp
+++ b/Engine/Plugins/Experimental/Avalanche/Source/AvalancheSequence/Private/AvaSequence.cpp
@@ -31,30 +31,6 @@
 
 #define LOCTEXT_NAMESPACE "AvaSequence"
 
-namespace UE::AvaSequence::Private
-{
-	class FTransientPlayer : public IMovieScenePlayer
-	{
-	public:
-		FTransientPlayer(UAvaSequence& InSequence)
-		{
-			Template.Initialize(InSequence, *this, nullptr);
-			State.AssignSequence(MovieSceneSequenceID::Root, InSequence, *this);
-		}
-
-		//~ Begin IMovieScenePlayer
-		virtual FMovieSceneRootEvaluationTemplateInstance& GetEvaluationTemplate() override { return Template; }
-		virtual void SetViewportSettings(const TMap<FViewportClient*, EMovieSceneViewportParams>& ViewportParamsMap) override {}
-		virtual void GetViewportSettings(TMap<FViewportClient*, EMovieSceneViewportParams>& ViewportParamsMap) const override {}
-		virtual EMovieScenePlayerStatus::Type GetPlaybackStatus() const { return EMovieScenePlayerStatus::Stopped; }
-		virtual void SetPlaybackStatus(EMovieScenePlayerStatus::Type InPlaybackStatus) override {}
-		//~ End IMovieScenePlayer
-
-	private:
-		FMovieSceneRootEvaluationTemplateInstance Template;
-	};
-}
-
 UAvaSequence::UAvaSequence(const FObjectInitializer& InObjectInitializer)
 	: Super(InObjectInitializer)
 	, Label(TEXT("Sequence"))
@@ -601,24 +577,40 @@ int32 UAvaSequence::UpdateBindings(const FTopLevelAssetPath* InOldContext, const
 
 TArrayView<TWeakObjectPtr<>> UAvaSequence::FindObjectsFromGuid(const FGuid& InGuid)
 {
+	using namespace UE::MovieScene;
+
 	if (!InGuid.IsValid())
 	{
 		return TArrayView<TWeakObjectPtr<>>();
 	}
 
-	UE::AvaSequence::Private::FTransientPlayer Player(*this);
-	return Player.FindBoundObjects(InGuid, MovieSceneSequenceID::Root);
+	FSharedPlaybackStateCreateParams CreateParams;
+	TSharedRef<FSharedPlaybackState> TransientPlaybackState = MakeShared<FSharedPlaybackState>(*this, CreateParams);
+
+	FMovieSceneEvaluationState State;
+	TransientPlaybackState->AddCapabilityRaw(&State);
+	State.AssignSequence(MovieSceneSequenceID::Root, *this, TransientPlaybackState);
+
+	return State.FindBoundObjects(InGuid, MovieSceneSequenceID::Root, TransientPlaybackState);
 }
 
 FGuid UAvaSequence::FindGuidFromObject(UObject* InObject)
 {
+	using namespace UE::MovieScene;
+
 	if (!InObject)
 	{
 		return FGuid();
 	}
 
-	UE::AvaSequence::Private::FTransientPlayer Player(*this);
-	return Player.FindObjectId(*InObject, MovieSceneSequenceID::Root);
+	FSharedPlaybackStateCreateParams CreateParams;
+	TSharedRef<FSharedPlaybackState> TransientPlaybackState = MakeShared<FSharedPlaybackState>(*this, CreateParams);
+
+	FMovieSceneEvaluationState State;
+	TransientPlaybackState->AddCapabilityRaw(&State);
+	State.AssignSequence(MovieSceneSequenceID::Root, *this, TransientPlaybackState);
+
+	return State.FindObjectId(*InObject, MovieSceneSequenceID::Root, TransientPlaybackState);
 }
 
 #if WITH_EDITOR
diff --git a/Engine/Plugins/Experimental/DaySequence/Source/DaySequence/Private/DaySequence.cpp b/Engine/Plugins/Experimental/DaySequence/Source/DaySequence/Private/DaySequence.cpp
--- a/Engine/Plugins/Experimental/DaySequence/Source/DaySequence/Private/DaySequence.cpp
+++ b/Engine/Plugins/Experimental/DaySequence/Source/DaySequence/Private/DaySequence.cpp
@@ -421,6 +421,8 @@ void UDaySequence::OnDirectorRecompiled(UBlueprint* InCompiledBlueprint)
 
 FGuid UDaySequence::FindOrAddBinding(UObject* InObject)
 {
+	using namespace UE::MovieScene;
+
 	UObject* PlaybackContext = InObject ? InObject->GetWorld() : nullptr;
 	if (!InObject || !PlaybackContext)
 	{
@@ -454,21 +456,15 @@ FGuid UDaySequence::FindOrAddBinding(UObject* InObject)
 
 	// Perform a potentially slow lookup of every possessable binding in the sequence to see if we already have this
 	{
-		class FTransientPlayer : public IMovieScenePlayer
-		{
-		public:
-			FMovieSceneRootEvaluationTemplateInstance Template;
-			virtual FMovieSceneRootEvaluationTemplateInstance& GetEvaluationTemplate() override { return Template; }
-			virtual void UpdateCameraCut(UObject* CameraObject, const EMovieSceneCameraCutParams& CameraCutParams) override {}
-			virtual void SetViewportSettings(const TMap<FViewportClient*, EMovieSceneViewportParams>& ViewportParamsMap) override {}
-			virtual void GetViewportSettings(TMap<FViewportClient*, EMovieSceneViewportParams>& ViewportParamsMap) const override {}
-			virtual EMovieScenePlayerStatus::Type GetPlaybackStatus() const { return EMovieScenePlayerStatus::Stopped; }
-			virtual void SetPlaybackStatus(EMovieScenePlayerStatus::Type InPlaybackStatus) override {}
-		} Player;
-
-		Player.State.AssignSequence(MovieSceneSequenceID::Root, *this, Player);
-
-		FGuid ExistingID = Player.FindObjectId(*InObject, MovieSceneSequenceID::Root);
+		FSharedPlaybackStateCreateParams CreateParams;
+		CreateParams.PlaybackContext = PlaybackContext;
+		TSharedRef<FSharedPlaybackState> TransientPlaybackState = MakeShared<FSharedPlaybackState>(*this, CreateParams);
+
+		FMovieSceneEvaluationState State;
+		TransientPlaybackState->AddCapabilityRaw(&State);
+		State.AssignSequence(MovieSceneSequenceID::Root, *this, TransientPlaybackState);
+
+		FGuid ExistingID = State.FindObjectId(*InObject, MovieSceneSequenceID::Root, TransientPlaybackState);
 		if (ExistingID.IsValid())
 		{
 			return ExistingID;
diff --git a/Engine/Plugins/Experimental/DaySequence/Source/DaySequence/Private/DaySequencePlayer.cpp b/Engine/Plugins/Experimental/DaySequence/Source/DaySequence/Private/DaySequencePlayer.cpp
--- a/Engine/Plugins/Experimental/DaySequence/Source/DaySequence/Private/DaySequencePlayer.cpp
+++ b/Engine/Plugins/Experimental/DaySequence/Source/DaySequence/Private/DaySequencePlayer.cpp
@@ -1281,13 +1281,17 @@ TArray<UObject*> UDaySequencePlayer::GetBoundObjects(FMovieSceneObjectBindingID
 TArray<FMovieSceneObjectBindingID> UDaySequencePlayer::GetObjectBindings(UObject* InObject)
 {
 	TArray<FMovieSceneObjectBindingID> Bindings;
+	PRAGMA_DISABLE_DEPRECATION_WARNINGS
 	State.FilterObjectBindings(InObject, GetSharedPlaybackState(), &Bindings);
+	PRAGMA_ENABLE_DEPRECATION_WARNINGS
 	return Bindings;
 }
 
 void UDaySequencePlayer::RequestInvalidateBinding(FMovieSceneObjectBindingID ObjectBinding)
 {
+	PRAGMA_DISABLE_DEPRECATION_WARNINGS
 	State.Invalidate(ObjectBinding.GetGuid(), ObjectBinding.GetRelativeSequenceID());
+	PRAGMA_ENABLE_DEPRECATION_WARNINGS
 }
 
 UWorld* UDaySequencePlayer::GetPlaybackWorld() const
diff --git a/Engine/Plugins/Experimental/DaySequence/Source/DaySequenceEditor/Private/DaySequenceEditorSpawnRegister.cpp b/Engine/Plugins/Experimental/DaySequence/Source/DaySequenceEditor/Private/DaySequenceEditorSpawnRegister.cpp
--- a/Engine/Plugins/Experimental/DaySequence/Source/DaySequenceEditor/Private/DaySequenceEditorSpawnRegister.cpp
+++ b/Engine/Plugins/Experimental/DaySequence/Source/DaySequenceEditor/Private/DaySequenceEditorSpawnRegister.cpp
@@ -226,7 +226,7 @@ void FDaySequenceEditorSpawnRegister::OnObjectsReplaced(const TMap<UObject*, UOb
 			NewObject->SetFlags(RF_Transient);
 
 			// Invalidate the binding - it will be resolved if it's ever asked for again
-			Sequencer->State.Invalidate(Pair.Key.BindingId, Pair.Key.TemplateID);
+			Sequencer->GetEvaluationState()->Invalidate(Pair.Key.BindingId, Pair.Key.TemplateID);
 		}
 	}
 }
diff --git a/Engine/Plugins/Importers/USDImporter/Source/USDExporter/Private/LevelSequenceExporterUSD.cpp b/Engine/Plugins/Importers/USDImporter/Source/USDExporter/Private/LevelSequenceExporterUSD.cpp
--- a/Engine/Plugins/Importers/USDImporter/Source/USDExporter/Private/LevelSequenceExporterUSD.cpp
+++ b/Engine/Plugins/Importers/USDImporter/Source/USDExporter/Private/LevelSequenceExporterUSD.cpp
@@ -644,7 +644,7 @@ namespace UE::LevelSequenceExporterUSD::Private
 					AActor* Actor = OutSelectedNonSpawnableActors[Index];
 
 					TArray<FMovieSceneObjectBindingID> Bindings;
-					PinnedSequencer->State.FilterObjectBindings(Actor, *PinnedSequencer, &Bindings);
+					PinnedSequencer->GetEvaluationState()->FilterObjectBindings(Actor, *PinnedSequencer, &Bindings);
 
 					OutSelectedSpawnableBindings.Reserve(Bindings.Num() + OutSelectedSpawnableBindings.Num());
 					for (const FMovieSceneObjectBindingID& Binding : Bindings)
@@ -876,7 +876,7 @@ namespace UE::LevelSequenceExporterUSD::Private
 
 		for (FMovieSceneSequenceID SequenceInstance : *InstancesOfThisSequence)
 		{
-			FMovieSceneObjectCache& ObjectCache = Context.Sequencer->State.GetObjectCache(SequenceInstance);
+			FMovieSceneObjectCache& ObjectCache = Context.Sequencer->GetEvaluationState()->GetObjectCache(SequenceInstance);
 
 			// Possessables
 			int32 NumPossessables = MovieScene->GetPossessableCount();
diff --git a/Engine/Plugins/MovieScene/ActorSequence/Source/ActorSequenceEditor/Private/ActorSequenceEditorTabSummoner.cpp b/Engine/Plugins/MovieScene/ActorSequence/Source/ActorSequenceEditor/Private/ActorSequenceEditorTabSummoner.cpp
--- a/Engine/Plugins/MovieScene/ActorSequence/Source/ActorSequenceEditor/Private/ActorSequenceEditorTabSummoner.cpp
+++ b/Engine/Plugins/MovieScene/ActorSequence/Source/ActorSequenceEditor/Private/ActorSequenceEditorTabSummoner.cpp
@@ -592,7 +592,7 @@ public:
 			return;
 		}
 
-		Sequencer->State.ClearObjectCaches(*Sequencer);
+		Sequencer->GetEvaluationState()->ClearObjectCaches(*Sequencer);
 		TSet<UObject*> AllBoundObjects;
 
 		AllBoundObjects.Add(GetOwnerComponent());
diff --git a/Engine/Plugins/MovieScene/LevelSequenceEditor/Source/LevelSequenceEditor/Private/LevelSequenceEditorSubsystem.cpp b/Engine/Plugins/MovieScene/LevelSequenceEditor/Source/LevelSequenceEditor/Private/LevelSequenceEditorSubsystem.cpp
--- a/Engine/Plugins/MovieScene/LevelSequenceEditor/Source/LevelSequenceEditor/Private/LevelSequenceEditorSubsystem.cpp
+++ b/Engine/Plugins/MovieScene/LevelSequenceEditor/Source/LevelSequenceEditor/Private/LevelSequenceEditorSubsystem.cpp
@@ -2726,7 +2726,7 @@ void ULevelSequenceEditorSubsystem::OnFinishedChangingLocators(const FPropertyCh
 			}
 		}
 
-		Sequencer->State.Invalidate(ObjectBindingID, Sequencer->GetFocusedTemplateID());
+		Sequencer->GetEvaluationState()->Invalidate(ObjectBindingID, Sequencer->GetFocusedTemplateID());
 
 		// Update the object class and DisplayName
 		TArrayView<TWeakObjectPtr<>> ObjectsInCurrentSequence = Sequencer->FindObjectsInCurrentSequence(ObjectBindingID);
diff --git a/Engine/Plugins/MovieScene/LevelSequenceEditor/Source/LevelSequenceEditor/Private/Misc/LevelSequenceEditorSpawnRegister.cpp b/Engine/Plugins/MovieScene/LevelSequenceEditor/Source/LevelSequenceEditor/Private/Misc/LevelSequenceEditorSpawnRegister.cpp
--- a/Engine/Plugins/MovieScene/LevelSequenceEditor/Source/LevelSequenceEditor/Private/Misc/LevelSequenceEditorSpawnRegister.cpp
+++ b/Engine/Plugins/MovieScene/LevelSequenceEditor/Source/LevelSequenceEditor/Private/Misc/LevelSequenceEditorSpawnRegister.cpp
@@ -233,7 +233,7 @@ void FLevelSequenceEditorSpawnRegister::OnObjectsReplaced(const TMap<UObject*, U
 			NewObject->SetFlags(RF_Transient);
 
 			// Invalidate the binding - it will be resolved if it's ever asked for again
-			Sequencer->State.Invalidate(Pair.Key.BindingId, Pair.Key.TemplateID);
+			Sequencer->GetEvaluationState()->Invalidate(Pair.Key.BindingId, Pair.Key.TemplateID);
 		}
 	}
 }
diff --git a/Engine/Plugins/MovieScene/SequencerScripting/Source/SequencerScriptingEditor/Private/SequencerTools.cpp b/Engine/Plugins/MovieScene/SequencerScripting/Source/SequencerScriptingEditor/Private/SequencerTools.cpp
--- a/Engine/Plugins/MovieScene/SequencerScripting/Source/SequencerScriptingEditor/Private/SequencerTools.cpp
+++ b/Engine/Plugins/MovieScene/SequencerScripting/Source/SequencerScriptingEditor/Private/SequencerTools.cpp
@@ -136,7 +136,7 @@ TArray<FSequencerBoundObjects> USequencerToolsFunctionLibrary::GetBoundObjects(U
 	TRange<FFrameNumber> SpecifiedRange = InRange.ToNative(Resolution);
 	Player->SetPlaybackPosition(FMovieSceneSequencePlaybackParams(SpecifiedRange.GetLowerBoundValue().Value, EUpdatePositionMethod::Play));
 
-	FMovieSceneSequenceID SequenceId = Player->State.FindSequenceId(InSequence);
+	FMovieSceneSequenceID SequenceId = Player->GetEvaluationState()->FindSequenceId(InSequence);
 
 	TArray<FSequencerBoundObjects> BoundObjects;
 	for (FMovieSceneBindingProxy Binding : InBindings)
@@ -171,7 +171,7 @@ TArray<FSequencerBoundObjects> USequencerToolsFunctionLibrary::GetObjectBindings
 		for (FMovieSceneObjectBindingID ObjectBinding : ObjectBindings)
 		{
 			FMovieSceneSequenceID SequenceID = ObjectBinding.ResolveSequenceID(MovieSceneSequenceID::Root, *Player);
-			FMovieSceneBindingProxy Binding(ObjectBinding.GetGuid(), Player->State.FindSequence(SequenceID));
+			FMovieSceneBindingProxy Binding(ObjectBinding.GetGuid(), Player->GetEvaluationState()->FindSequence(SequenceID));
 			BoundObjects.Add(FSequencerBoundObjects(Binding, TArray<UObject*>({ Object })));
 		}
 	}
@@ -240,8 +240,9 @@ bool ExportFBXInternal(const FSequencerExportFBXParams& InParams, UMovieSceneSeq
 		}
 	}
 
-	Player->State.AssignSequence(MovieSceneSequenceID::Root, *RootSequence, *Player);
-	FMovieSceneSequenceIDRef Template = Player->State.FindSequenceId(Sequence);
+	FMovieSceneEvaluationState* State = Player->GetEvaluationState();
+	State->AssignSequence(MovieSceneSequenceID::Root, *RootSequence, *Player);
+	FMovieSceneSequenceIDRef Template = State->FindSequenceId(Sequence);
 	UnFbx::FFbxExporter::FLevelSequenceNodeNameAdapter NodeNameAdapter(MovieScene, Player, Template);
 
 	FMovieSceneSequenceHierarchy Hierarchy = FMovieSceneSequenceHierarchy();
diff --git a/Engine/Plugins/MovieScene/TemplateSequence/Source/TemplateSequence/Private/TemplateSequence.cpp b/Engine/Plugins/MovieScene/TemplateSequence/Source/TemplateSequence/Private/TemplateSequence.cpp
--- a/Engine/Plugins/MovieScene/TemplateSequence/Source/TemplateSequence/Private/TemplateSequence.cpp
+++ b/Engine/Plugins/MovieScene/TemplateSequence/Source/TemplateSequence/Private/TemplateSequence.cpp
@@ -166,7 +166,7 @@ FGuid UTemplateSequence::FindOrAddBinding(UObject* InObject)
 		OutActor->BindingOverride.bOverridesDefault = false;
 
 		Player->Initialize(this, World, FMovieSceneSequencePlaybackSettings());
-		Player->State.AssignSequence(MovieSceneSequenceID::Root, *this, *Player);
+		Player->GetEvaluationState()->AssignSequence(MovieSceneSequenceID::Root, *this, *Player);
 		Player->PlayTo(FMovieSceneSequencePlaybackParams(MovieScene->GetPlaybackRange().GetLowerBoundValue(), EUpdatePositionMethod::Play), FMovieSceneSequencePlayToParams());
 
 		FGuid FoundBinding;
diff --git a/Engine/Plugins/MovieScene/TemplateSequence/Source/TemplateSequence/Private/TemplateSequenceActor.cpp b/Engine/Plugins/MovieScene/TemplateSequence/Source/TemplateSequence/Private/TemplateSequenceActor.cpp
--- a/Engine/Plugins/MovieScene/TemplateSequence/Source/TemplateSequence/Private/TemplateSequenceActor.cpp
+++ b/Engine/Plugins/MovieScene/TemplateSequence/Source/TemplateSequence/Private/TemplateSequenceActor.cpp
@@ -195,7 +195,7 @@ void ATemplateSequenceActor::SetBinding(AActor* Actor, bool bOverridesDefault)
 		const FGuid RootBindingID = TemplateSequenceObject->GetRootObjectBindingID();
 		if (RootBindingID.IsValid())
 		{
-			SequencePlayer->State.Invalidate(RootBindingID, MovieSceneSequenceID::Root);
+			SequencePlayer->GetEvaluationState()->Invalidate(RootBindingID, MovieSceneSequenceID::Root);
 		}
 	}
 }
diff --git a/Engine/Source/Editor/MovieSceneTools/Private/MovieSceneObjectBindingIDPicker.cpp b/Engine/Source/Editor/MovieSceneTools/Private/MovieSceneObjectBindingIDPicker.cpp
--- a/Engine/Source/Editor/MovieSceneTools/Private/MovieSceneObjectBindingIDPicker.cpp
+++ b/Engine/Source/Editor/MovieSceneTools/Private/MovieSceneObjectBindingIDPicker.cpp
@@ -219,7 +219,7 @@ bool FMovieSceneObjectBindingIDPicker::IsClassAllowed(const TSharedPtr<FSequence
 	auto BindingID = Node->BindingID;
 
 	TSharedPtr<ISequencer> Sequencer = WeakSequencer.Pin();
-	UMovieSceneSequence* BindSequence = Sequencer.IsValid() ? Sequencer->State.FindSequence(BindingID.SequenceID) : nullptr;
+	UMovieSceneSequence* BindSequence = Sequencer.IsValid() ? Sequencer->GetEvaluationState()->FindSequence(BindingID.SequenceID) : nullptr;
 	UMovieScene* BindMovieScene = BindSequence ? BindSequence->GetMovieScene() : nullptr;
 
 	InClass = MovieSceneHelpers::GetBoundObjectClass(BindSequence, BindingID.Guid);
diff --git a/Engine/Source/Editor/Sequencer/Private/LevelEditorSequencerIntegration.cpp b/Engine/Source/Editor/Sequencer/Private/LevelEditorSequencerIntegration.cpp
--- a/Engine/Source/Editor/Sequencer/Private/LevelEditorSequencerIntegration.cpp
+++ b/Engine/Source/Editor/Sequencer/Private/LevelEditorSequencerIntegration.cpp
@@ -481,7 +481,7 @@ void FLevelEditorSequencerIntegration::OnLevelAdded(ULevel* InLevel, UWorld* InW
 		{
 			if (Options.bRequiresLevelEvents)
 			{
-				In.State.ClearObjectCaches(In);
+				In.GetEvaluationState()->ClearObjectCaches(In);
 			}
 		}
 	);
@@ -494,7 +494,7 @@ void FLevelEditorSequencerIntegration::OnLevelRemoved(ULevel* InLevel, UWorld* I
 		{
 			if (Options.bRequiresLevelEvents)
 			{
-				In.State.ClearObjectCaches(In);
+				In.GetEvaluationState()->ClearObjectCaches(In);
 			}
 		}
 	);
@@ -728,7 +728,7 @@ void FLevelEditorSequencerIntegration::OnPreBeginPIE(bool bIsSimulating)
 			{
 				In.OnPlaybackContextChanged();
 				In.RestorePreAnimatedState();
-				In.State.ClearObjectCaches(In);
+				In.GetEvaluationState()->ClearObjectCaches(In);
 				In.RequestEvaluate();
 			}
 		}
@@ -765,7 +765,7 @@ void FLevelEditorSequencerIntegration::OnEndPlayMap()
 			{
 				// Update and clear any stale bindings 
 				In.OnPlaybackContextChanged();
-				In.State.ClearObjectCaches(In);
+				In.GetEvaluationState()->ClearObjectCaches(In);
 				In.ForceEvaluate();
 			}
 		}
@@ -1293,7 +1293,7 @@ void FLevelEditorSequencerIntegration::OnMapChanged(UWorld* World, EMapChangeTyp
 			{
 				In.OnPlaybackContextChanged();
 				In.RestorePreAnimatedState();
-				In.State.ClearObjectCaches(In);
+				In.GetEvaluationState()->ClearObjectCaches(In);
 
 				// Notify data changed to enqueue an evaluate
 				In.NotifyMovieSceneDataChanged(EMovieSceneDataChangeType::Unknown);
diff --git a/Engine/Source/Editor/Sequencer/Private/MVVM/ViewModels/ObjectBindingModel.cpp b/Engine/Source/Editor/Sequencer/Private/MVVM/ViewModels/ObjectBindingModel.cpp
--- a/Engine/Source/Editor/Sequencer/Private/MVVM/ViewModels/ObjectBindingModel.cpp
+++ b/Engine/Source/Editor/Sequencer/Private/MVVM/ViewModels/ObjectBindingModel.cpp
@@ -447,8 +447,9 @@ FSlateColor FObjectBindingModel::GetLabelColor() const
 	// Find the last objecting binding ancestor and ask it for the invalid color to use.
 	// e.g. Spawnables don't have valid object bindings when their track hasn't spawned them yet,
 	// so we override the default behavior of red with a gray so that users don't think there is something wrong.
+	FMovieSceneEvaluationState* EvaluationState = Sequencer->GetEvaluationState();
 	TFunction<FSlateColor(const FObjectBindingModel&)> GetObjectBindingAncestorInvalidLabelColor = [&](const FObjectBindingModel& InObjectBindingModel) -> FSlateColor {
-		if (!Sequencer->State.GetBindingActivation(InObjectBindingModel.GetObjectGuid(), OwnerModel->GetSequenceID()))
+		if (!EvaluationState->GetBindingActivation(InObjectBindingModel.GetObjectGuid(), OwnerModel->GetSequenceID()))
 		{
 			return FSlateColor::UseSubduedForeground();
 		}
diff --git a/Engine/Source/Editor/Sequencer/Private/Sequencer.cpp b/Engine/Source/Editor/Sequencer/Private/Sequencer.cpp
--- a/Engine/Source/Editor/Sequencer/Private/Sequencer.cpp
+++ b/Engine/Source/Editor/Sequencer/Private/Sequencer.cpp
@@ -375,7 +375,9 @@ void FSequencer::InitSequencer(const FSequencerInitParams& InitParams
 
 		FDelegateHandle OnBlueprintCompiledHandle = GEditor->OnBlueprintCompiled().AddLambda([&]
 		{
+			PRAGMA_DISABLE_DEPRECATION_WARNINGS
 			State.InvalidateExpiredObjects();
+			PRAGMA_ENABLE_DEPRECATION_WARNINGS
 
 			// Force re-evaluation since animated state was restored in PreCompile
 			bNeedsEvaluate = true;
@@ -468,7 +470,9 @@ void FSequencer::InitSequencer(const FSequencerInitParams& InitParams
 
 			if (bAnythingReplaced)
 			{
+				PRAGMA_DISABLE_DEPRECATION_WARNINGS
 				State.InvalidateExpiredObjects();
+				PRAGMA_ENABLE_DEPRECATION_WARNINGS
 			}
 		});
 		AcquiredResources.Add([=] { FCoreUObjectDelegates::OnObjectsReplaced.Remove(OnObjectsReplacedHandle); });
@@ -1369,10 +1373,12 @@ void FSequencer::FocusSequenceInstance(UMovieSceneSubSection& InSubSection)
 
 	UMovieSceneSequence* FocusedSequence = GetFocusedMovieSceneSequence();
 
+	PRAGMA_DISABLE_DEPRECATION_WARNINGS
 	if (!State.FindSequence(SequenceID))
 	{
 		State.AssignSequence(SequenceID, *FocusedSequence, *this);
 	}
+	PRAGMA_ENABLE_DEPRECATION_WARNINGS
 
 	// Ensure that the director BP is registered with the action database
 	if (FMovieSceneSequenceEditor* SequenceEditor = FMovieSceneSequenceEditor::Find(FocusedSequence))
@@ -1733,10 +1739,12 @@ void FSequencer::RerunConstructionScripts()
 		}
 	}
 
+	PRAGMA_DISABLE_DEPRECATION_WARNINGS
 	for (TPair<FMovieSceneSequenceID, FGuid> BoundGuid : BoundGuids)
 	{
 		State.Invalidate(BoundGuid.Value, BoundGuid.Key);
 	}
+	PRAGMA_ENABLE_DEPRECATION_WARNINGS
 }
 
 void FSequencer::GetConstructionScriptActors(UMovieScene* MovieScene, FMovieSceneSequenceIDRef SequenceID, TSet<TWeakObjectPtr<AActor> >& BoundActors, TArray < TPair<FMovieSceneSequenceID, FGuid> >& BoundGuids)
@@ -2683,7 +2691,9 @@ void FSequencer::NotifyMovieSceneDataChanged( EMovieSceneDataChangeType DataChan
 
 	bGlobalMarkedFramesCached = false;
 	bNeedsEvaluate = true;
+	PRAGMA_DISABLE_DEPRECATION_WARNINGS
 	State.ClearObjectCaches(*this);
+	PRAGMA_ENABLE_DEPRECATION_WARNINGS
 
 	UpdatePlaybackRange();
 	OnMovieSceneDataChangedDelegate.Broadcast(DataChangeType);
@@ -3601,7 +3611,9 @@ void FSequencer::UpdateCachedPlaybackContextAndClient()
 	if (CachedPlaybackContext != NewPlaybackContext || CachedPlaybackClient != NewPlaybackClient)
 	{
 		PrePossessionViewTargets.Reset();
+		PRAGMA_DISABLE_DEPRECATION_WARNINGS
 		State.ClearObjectCaches(*this);
+		PRAGMA_ENABLE_DEPRECATION_WARNINGS
 		RestorePreAnimatedState();
 
 		CachedPlaybackContext = NewPlaybackContext;
@@ -3613,7 +3625,9 @@ void FSequencer::UpdateCachedPlaybackContextAndClient()
 
 void FSequencer::UpdateCachedCameraActors()
 {
+	PRAGMA_DISABLE_DEPRECATION_WARNINGS
 	const uint32 CurrentStateSerial = State.GetSerialNumber();
+	PRAGMA_ENABLE_DEPRECATION_WARNINGS
 	if (CurrentStateSerial == LastKnownStateSerial)
 	{
 		return;
@@ -3653,7 +3667,9 @@ void FSequencer::UpdateCachedCameraActors()
 					BindingGuids.Add(Possessable.GetGuid());
 				}
 
+				PRAGMA_DISABLE_DEPRECATION_WARNINGS
 				const FMovieSceneObjectCache& ObjectCache = State.GetObjectCache(SequenceID);
+				PRAGMA_ENABLE_DEPRECATION_WARNINGS
 				for (const FGuid& BindingGuid : BindingGuids)
 				{
 					for (TWeakObjectPtr<> BoundObject : ObjectCache.IterateBoundObjects(BindingGuid))
@@ -4021,8 +4037,10 @@ FGuid FSequencer::FindUnspawnedObjectGuid(UObject& InObject)
 				// The only way to find the object now is to resolve all the child bindings, and see if they are the same
 				for (const FGuid& ChildGuid : ParentSpawnable->GetChildPossessables())
 				{
-						TArrayView<TWeakObjectPtr<>> BoundObjects = State.FindBoundObjects(ChildGuid, GetFocusedTemplateID(), GetSharedPlaybackState());
-						const bool bHasObject = BoundObjects.Contains(&InObject);
+					PRAGMA_DISABLE_DEPRECATION_WARNINGS
+					TArrayView<TWeakObjectPtr<>> BoundObjects = State.FindBoundObjects(ChildGuid, GetFocusedTemplateID(), GetSharedPlaybackState());
+					PRAGMA_ENABLE_DEPRECATION_WARNINGS
+					const bool bHasObject = BoundObjects.Contains(&InObject);
 					if (bHasObject)
 					{
 						return ChildGuid;
@@ -4802,7 +4820,9 @@ UObject* FSequencer::FindSpawnedObjectOrTemplate(const FGuid& BindingId)
 
 		if (ParentObject)
 		{
+			PRAGMA_DISABLE_DEPRECATION_WARNINGS
 			TArrayView<TWeakObjectPtr<>> BoundObjects = State.FindBoundObjects(BindingId, GetFocusedTemplateID(), GetSharedPlaybackState());
+			PRAGMA_ENABLE_DEPRECATION_WARNINGS
 			for (TWeakObjectPtr<> WeakObj : BoundObjects)
 			{
 				if (UObject* Obj = WeakObj.Get())
@@ -10792,7 +10812,9 @@ void FSequencer::RebindPossessableReferences()
 		}
 
 		TArray<UObject*, TInlineAllocator<1>>& References = AllObjects.FindOrAdd(Possessable.GetGuid());
+		PRAGMA_DISABLE_DEPRECATION_WARNINGS
 		TArrayView<TWeakObjectPtr<>> BoundObjects = State.FindBoundObjects(Possessable.GetGuid(), GetFocusedTemplateID(), GetSharedPlaybackState());
+		PRAGMA_ENABLE_DEPRECATION_WARNINGS
 		for (TWeakObjectPtr<> WeakObj : BoundObjects)
 		{
 			if (UObject* Obj = WeakObj.Get())
diff --git a/Engine/Source/Editor/Sequencer/Private/SequencerUtilities.cpp b/Engine/Source/Editor/Sequencer/Private/SequencerUtilities.cpp
--- a/Engine/Source/Editor/Sequencer/Private/SequencerUtilities.cpp
+++ b/Engine/Source/Editor/Sequencer/Private/SequencerUtilities.cpp
@@ -2735,7 +2735,7 @@ FGuid TryCreateCustomBinding(TSharedPtr<ISequencer> Sequencer, UObject* CustomBi
 				FTransformData TransformData;
 				Sequencer->GetSpawnRegister().SetupDefaultsForSpawnable(SpawnedObject, NewID, TransformData, Sequencer.ToSharedRef(), Sequencer->GetSequencerSettings());
 			}
-			Sequencer->State.Invalidate(NewID, Sequencer->GetFocusedTemplateID());
+			Sequencer->GetEvaluationState()->Invalidate(NewID, Sequencer->GetFocusedTemplateID());
 			Sequencer->ForceEvaluate();
 
 			// We don't call these events in the case bSetupDefaults is false because they may add tracks.
@@ -2950,7 +2950,7 @@ FGuid CreateGenericBinding(TSharedPtr<ISequencer> Sequencer, UMovieSceneSequence
 				BindingReferences->AddOrReplaceBinding(NewID, MoveTemp(Locator), InParams.BindingIndex);
 				if (Sequencer)
 				{
-					Sequencer->State.Invalidate(NewID, Sequencer->GetFocusedTemplateID());
+					Sequencer->GetEvaluationState()->Invalidate(NewID, Sequencer->GetFocusedTemplateID());
 				}
 			}
 			else
@@ -3861,8 +3861,10 @@ FGuid FSequencerUtilities::AssignActor(TSharedRef<ISequencer> Sequencer, AActor*
 		// Replace
 		OwnerMovieScene->ReplacePossessable(OldComponentGuid, NewPossessable);
 		OwnerSequence->UnbindPossessableObjects(OldComponentGuid);
-		Sequencer->State.Invalidate(OldComponentGuid, FocusedGuid);
-		Sequencer->State.Invalidate(NewPossessable.GetGuid(), FocusedGuid);
+		
+		FMovieSceneEvaluationState* State = Sequencer->GetEvaluationState();
+		State->Invalidate(OldComponentGuid, FocusedGuid);
+		State->Invalidate(NewPossessable.GetGuid(), FocusedGuid);
 
 		NewComponentGuids.Add(NewPossessable.GetGuid());
 	};
@@ -3932,8 +3934,9 @@ FGuid FSequencerUtilities::AssignActor(TSharedRef<ISequencer> Sequencer, AActor*
 	OwnerMovieScene->ReplacePossessable(InObjectBinding, NewPossessableActor);
 	OwnerSequence->UnbindPossessableObjects(InObjectBinding);
 
-	Sequencer->State.Invalidate(InObjectBinding, Sequencer->GetFocusedTemplateID());
-	Sequencer->State.Invalidate(NewPossessableActor.GetGuid(), Sequencer->GetFocusedTemplateID());
+	FMovieSceneEvaluationState* State = Sequencer->GetEvaluationState();
+	State->Invalidate(InObjectBinding, Sequencer->GetFocusedTemplateID());
+	State->Invalidate(NewPossessableActor.GetGuid(), Sequencer->GetFocusedTemplateID());
 
 	for (const FGuid& NewComponentGuid : NewComponentGuids)
 	{
diff --git a/Engine/Source/Editor/UMGEditor/Private/WidgetBlueprintEditor.cpp b/Engine/Source/Editor/UMGEditor/Private/WidgetBlueprintEditor.cpp
--- a/Engine/Source/Editor/UMGEditor/Private/WidgetBlueprintEditor.cpp
+++ b/Engine/Source/Editor/UMGEditor/Private/WidgetBlueprintEditor.cpp
@@ -1746,7 +1746,7 @@ void FWidgetBlueprintEditor::UpdatePreview(UBlueprint* InBlueprint, bool bInForc
 	{
 		if (TSharedPtr<ISequencer> Sequencer = SequencerPtr.Pin())
 		{
-			Sequencer->State.ClearObjectCaches(*Sequencer);
+			Sequencer->GetEvaluationState()->ClearObjectCaches(*Sequencer);
 			Sequencer->ForceEvaluate();
 		}
 	}
@@ -2675,4 +2675,4 @@ void FWidgetBlueprintEditor::UpdateTrackName(FGuid ObjectId)
 	}
 }
 
-#undef LOCTEXT_NAMESPACE
\ No newline at end of file
+#undef LOCTEXT_NAMESPACE
diff --git a/Engine/Source/Editor/UnrealEd/Private/Fbx/FbxMainExport.cpp b/Engine/Source/Editor/UnrealEd/Private/Fbx/FbxMainExport.cpp
--- a/Engine/Source/Editor/UnrealEd/Private/Fbx/FbxMainExport.cpp
+++ b/Engine/Source/Editor/UnrealEd/Private/Fbx/FbxMainExport.cpp
@@ -1837,8 +1837,10 @@ TRange<FFrameNumber> FLevelSequenceAnimTrackAdapter::GetSequenceRange() const
 		TRange<FFrameNumber> SetRange(OptionalRange.GetValue().Key, (OptionalRange.GetValue().Value));
 		return SetRange;
 	}
-	MovieScenePlayer->State.AssignSequence(MovieSceneSequenceID::Root, *RootMovieSceneSequence, *MovieScenePlayer);
-	FMovieSceneSequenceIDRef Template = MovieScenePlayer->State.FindSequenceId(MovieSceneSequence);
+
+	FMovieSceneEvaluationState* State = MovieScenePlayer->GetEvaluationState();
+	State->AssignSequence(MovieSceneSequenceID::Root, *RootMovieSceneSequence, *MovieScenePlayer);
+	FMovieSceneSequenceIDRef Template = State->FindSequenceId(MovieSceneSequence);
 
 	FMovieSceneSequenceHierarchy Hierarchy = FMovieSceneSequenceHierarchy();
 	UMovieSceneCompiledDataManager::CompileHierarchy(RootMovieSceneSequence, &Hierarchy,EMovieSceneServerClientMask::All);
diff --git a/Engine/Source/Runtime/LevelSequence/Private/LevelSequenceActor.cpp b/Engine/Source/Runtime/LevelSequence/Private/LevelSequenceActor.cpp
--- a/Engine/Source/Runtime/LevelSequence/Private/LevelSequenceActor.cpp
+++ b/Engine/Source/Runtime/LevelSequence/Private/LevelSequenceActor.cpp
@@ -428,7 +428,7 @@ void ALevelSequenceActor::SetBinding(FMovieSceneObjectBindingID Binding, const T
 		if (GetSequencePlayer())
 		{
 			FMovieSceneSequenceID SequenceID = Binding.ResolveSequenceID(MovieSceneSequenceID::Root, *GetSequencePlayer());
-			GetSequencePlayer()->State.Invalidate(Binding.GetGuid(), SequenceID);
+			GetSequencePlayer()->GetEvaluationState()->Invalidate(Binding.GetGuid(), SequenceID);
 		}
 	}
 }
@@ -466,7 +466,7 @@ void ALevelSequenceActor::AddBinding(FMovieSceneObjectBindingID Binding, AActor*
 		if (GetSequencePlayer())
 		{
 			FMovieSceneSequenceID SequenceID = Binding.ResolveSequenceID(MovieSceneSequenceID::Root, *GetSequencePlayer());
-			GetSequencePlayer()->State.Invalidate(Binding.GetGuid(), SequenceID);
+			GetSequencePlayer()->GetEvaluationState()->Invalidate(Binding.GetGuid(), SequenceID);
 		}
 	}
 }
@@ -504,7 +504,7 @@ void ALevelSequenceActor::RemoveBinding(FMovieSceneObjectBindingID Binding, AAct
 		if (GetSequencePlayer())
 		{
 			FMovieSceneSequenceID SequenceID = Binding.ResolveSequenceID(MovieSceneSequenceID::Root, *GetSequencePlayer());
-			GetSequencePlayer()->State.Invalidate(Binding.GetGuid(), SequenceID);
+			GetSequencePlayer()->GetEvaluationState()->Invalidate(Binding.GetGuid(), SequenceID);
 		}
 	}
 }
@@ -542,7 +542,7 @@ void ALevelSequenceActor::ResetBinding(FMovieSceneObjectBindingID Binding)
 		if (GetSequencePlayer())
 		{
 			FMovieSceneSequenceID SequenceID = Binding.ResolveSequenceID(MovieSceneSequenceID::Root, *GetSequencePlayer());
-			GetSequencePlayer()->State.Invalidate(Binding.GetGuid(), SequenceID);
+			GetSequencePlayer()->GetEvaluationState()->Invalidate(Binding.GetGuid(), SequenceID);
 		}
 	}
 }
@@ -552,7 +552,7 @@ void ALevelSequenceActor::ResetBindings()
 	BindingOverrides->ResetBindings();
 	if (GetSequencePlayer())
 	{
-		GetSequencePlayer()->State.ClearObjectCaches(*GetSequencePlayer());
+		GetSequencePlayer()->GetEvaluationState()->ClearObjectCaches(*GetSequencePlayer());
 	}
 }
 
diff --git a/Engine/Source/Runtime/MovieScene/Private/Compilation/MovieSceneCompiledVolatilityManager.cpp b/Engine/Source/Runtime/MovieScene/Private/Compilation/MovieSceneCompiledVolatilityManager.cpp
--- a/Engine/Source/Runtime/MovieScene/Private/Compilation/MovieSceneCompiledVolatilityManager.cpp
+++ b/Engine/Source/Runtime/MovieScene/Private/Compilation/MovieSceneCompiledVolatilityManager.cpp
@@ -152,16 +152,16 @@ void FCompiledDataVolatilityManager::UpdateCachedSignatures()
 	TSharedPtr<const FSharedPlaybackState> SharedPlaybackState = WeakSharedPlaybackState.Pin();
 	FMovieSceneCompiledDataID RootDataID = SharedPlaybackState->GetRootCompiledDataID();
 	UMovieSceneCompiledDataManager* CompiledDataManager = SharedPlaybackState->GetCompiledDataManager();
-	IMovieScenePlayer* Player = FPlayerIndexPlaybackCapability::GetPlayer(SharedPlaybackState.ToSharedRef());
+	FMovieSceneEvaluationState* State = SharedPlaybackState->FindCapability<FMovieSceneEvaluationState>();
 
 	{
 		const FMovieSceneCompiledDataEntry& RootEntry = CompiledDataManager->GetEntryRef(RootDataID);
 		CachedCompilationSignatures.Add(MovieSceneSequenceID::Root, RootEntry.CompiledSignature);
 
 		UMovieSceneSequence* RootSequence = RootEntry.GetSequence();
-		if (RootSequence && Player)
+		if (RootSequence && State)
 		{
-			Player->State.AssignSequence(MovieSceneSequenceID::Root, *RootSequence, Player->GetSharedPlaybackState());
+			State->AssignSequence(MovieSceneSequenceID::Root, *RootSequence, SharedPlaybackState.ToSharedRef());
 		}
 	}
 
@@ -181,9 +181,9 @@ void FCompiledDataVolatilityManager::UpdateCachedSignatures()
 			CachedCompilationSignatures.Add(SubData.Key, SubEntry.CompiledSignature);
 
 			UMovieSceneSequence* Sequence = SubData.Value.GetSequence();
-			if (Sequence && Player)
+			if (Sequence && State)
 			{
-				Player->State.AssignSequence(SubData.Key, *Sequence, Player->GetSharedPlaybackState());
+				State->AssignSequence(SubData.Key, *Sequence, SharedPlaybackState.ToSharedRef());
 			}
 		}
 	}
diff --git a/Engine/Source/Runtime/MovieScene/Private/Evaluation/Instances/MovieSceneTrackEvaluator.cpp b/Engine/Source/Runtime/MovieScene/Private/Evaluation/Instances/MovieSceneTrackEvaluator.cpp
--- a/Engine/Source/Runtime/MovieScene/Private/Evaluation/Instances/MovieSceneTrackEvaluator.cpp
+++ b/Engine/Source/Runtime/MovieScene/Private/Evaluation/Instances/MovieSceneTrackEvaluator.cpp
@@ -134,7 +134,10 @@ void FMovieSceneTrackEvaluator::Evaluate(FMovieSceneContext Context, IMovieScene
 	CallSetupTearDown(Player, &DelayedRestore);
 
 	// Ensure any null objects are not cached
-	Player.State.InvalidateExpiredObjects();
+	if (FMovieSceneEvaluationState* State = Player.GetEvaluationState())
+	{
+		State->InvalidateExpiredObjects();
+	}
 
 	// Accumulate execution tokens into this structure
 	EvaluateGroup(*GroupToEvaluate, Context, Player);
diff --git a/Engine/Source/Runtime/MovieScene/Private/Evaluation/MovieSceneEvaluationTemplateInstance.cpp b/Engine/Source/Runtime/MovieScene/Private/Evaluation/MovieSceneEvaluationTemplateInstance.cpp
--- a/Engine/Source/Runtime/MovieScene/Private/Evaluation/MovieSceneEvaluationTemplateInstance.cpp
+++ b/Engine/Source/Runtime/MovieScene/Private/Evaluation/MovieSceneEvaluationTemplateInstance.cpp
@@ -148,8 +148,10 @@ void FMovieSceneRootEvaluationTemplateInstance::Initialize(UMovieSceneSequence&
 			}
 		}
 
+		PRAGMA_DISABLE_DEPRECATION_WARNINGS
 		Player.State.PersistentEntityData.Reset();
 		Player.State.PersistentSharedData.Reset();
+		PRAGMA_ENABLE_DEPRECATION_WARNINGS
 		
 		// Build a new linker and immediately attach the runner to it, so that we can find both
 		// when initializing the new root instance and shared playback state.
diff --git a/Engine/Source/Runtime/MovieScene/Private/Evaluation/PersistentEvaluationData.cpp b/Engine/Source/Runtime/MovieScene/Private/Evaluation/PersistentEvaluationData.cpp
--- a/Engine/Source/Runtime/MovieScene/Private/Evaluation/PersistentEvaluationData.cpp
+++ b/Engine/Source/Runtime/MovieScene/Private/Evaluation/PersistentEvaluationData.cpp
@@ -7,8 +7,8 @@
 
 FPersistentEvaluationData::FPersistentEvaluationData(IMovieScenePlayer& InPlayer)
 	: Player(InPlayer)
-	, EntityData(Player.State.PersistentEntityData)
-	, SharedData(Player.State.PersistentSharedData)
+	, EntityData(Player.GetEvaluationState()->PersistentEntityData)
+	, SharedData(Player.GetEvaluationState()->PersistentSharedData)
 {}
 
 const FMovieSceneSequenceInstanceData* FPersistentEvaluationData::GetInstanceData() const
diff --git a/Engine/Source/Runtime/MovieScene/Private/MovieScenePlayer.cpp b/Engine/Source/Runtime/MovieScene/Private/MovieScenePlayer.cpp
--- a/Engine/Source/Runtime/MovieScene/Private/MovieScenePlayer.cpp
+++ b/Engine/Source/Runtime/MovieScene/Private/MovieScenePlayer.cpp
@@ -52,7 +52,9 @@ uint16 FPlayerIndexPlaybackCapability::GetPlayerIndex(TSharedRef<const FSharedPl
 
 UE_DEFINE_MOVIESCENE_PLAYBACK_CAPABILITY(IMovieScenePlaybackClient)
 
+PRAGMA_DISABLE_DEPRECATION_WARNINGS
 IMovieScenePlayer::IMovieScenePlayer()
+	: BindingOverrides(StaticBindingOverrides.BindingOverrides)
 {
 	FTransactionallySafeWriteScopeLock ScopeLock(UE::MovieScene::GGlobalPlayerRegistryLock);
 
@@ -62,7 +64,9 @@ IMovieScenePlayer::IMovieScenePlayer()
 	UE::MovieScene::GGlobalPlayerUpdateFlags.PadToNum(UniqueIndex + 1, false);
 	UE::MovieScene::GGlobalPlayerUpdateFlags[UniqueIndex] = 0;
 }
+PRAGMA_ENABLE_DEPRECATION_WARNINGS
 
+PRAGMA_DISABLE_DEPRECATION_WARNINGS
 IMovieScenePlayer::~IMovieScenePlayer()
 {	
 	FTransactionallySafeWriteScopeLock ScopeLock(UE::MovieScene::GGlobalPlayerRegistryLock);
@@ -70,6 +74,7 @@ IMovieScenePlayer::~IMovieScenePlayer()
 	UE::MovieScene::GGlobalPlayerUpdateFlags[UniqueIndex] = 0;
 	UE::MovieScene::GGlobalPlayerRegistry.RemoveAt(UniqueIndex, 1);
 }
+PRAGMA_ENABLE_DEPRECATION_WARNINGS
 
 IMovieScenePlayer* IMovieScenePlayer::Get(uint16 InUniqueIndex)
 {
@@ -156,7 +161,8 @@ TArrayView<TWeakObjectPtr<>> IMovieScenePlayer::FindBoundObjects(const FGuid& Ob
 
 	if (TSharedPtr<const FSharedPlaybackState> SharedPlaybackState = FindSharedPlaybackState())
 	{
-		return State.FindBoundObjects(ObjectBindingID, SequenceID, SharedPlaybackState.ToSharedRef());
+		FMovieSceneEvaluationState* ActualState = GetEvaluationState();
+		return ActualState->FindBoundObjects(ObjectBindingID, SequenceID, SharedPlaybackState.ToSharedRef());
 	}
 
 	return TArrayView<TWeakObjectPtr<>>();
@@ -195,6 +201,31 @@ TSharedRef<const UE::MovieScene::FSharedPlaybackState> IMovieScenePlayer::GetSha
 	return ConstCastSharedRef<const UE::MovieScene::FSharedPlaybackState>(const_cast<IMovieScenePlayer*>(this)->GetEvaluationTemplate().GetSharedPlaybackState().ToSharedRef());
 }
 
+FMovieSceneEvaluationOperand* IMovieScenePlayer::GetBindingOverride(const FMovieSceneEvaluationOperand& InOperand)
+{
+	if (UE::MovieScene::IStaticBindingOverridesPlaybackCapability* ActualOverrides = GetStaticBindingOverrides())
+	{
+		return ActualOverrides->GetBindingOverride(InOperand);
+	}
+	return nullptr;
+}
+
+void IMovieScenePlayer::AddBindingOverride(const FMovieSceneEvaluationOperand& InOperand, const FMovieSceneEvaluationOperand& InOverrideOperand)
+{
+	if (UE::MovieScene::IStaticBindingOverridesPlaybackCapability* ActualOverrides = GetStaticBindingOverrides())
+	{
+		ActualOverrides->AddBindingOverride(InOperand, InOverrideOperand);
+	}
+}
+
+void IMovieScenePlayer::RemoveBindingOverride(const FMovieSceneEvaluationOperand& InOperand)
+{
+	if (UE::MovieScene::IStaticBindingOverridesPlaybackCapability* ActualOverrides = GetStaticBindingOverrides())
+	{
+		ActualOverrides->RemoveBindingOverride(InOperand);
+	}
+}
+
 void IMovieScenePlayer::ResetDirectorInstances()
 {
 	using namespace UE::MovieScene;
@@ -266,12 +297,47 @@ FGuid IMovieScenePlayer::CreateBinding(UMovieSceneSequence* InSequence, UObject*
 	return FGuid();
 }
 
+FMovieSceneEvaluationState* IMovieScenePlayer::GetEvaluationState()
+{
+	using namespace UE::MovieScene;
+
+	// Return the playback capability, which is generally the same as the State member variable if
+	// InitializeRootInstance has been called, but that member variable will be removed in the future.
+	// In addition to this, if our underlying type is FMovieSceneLegacyPlayer, the playback capability
+	// is NOT the same as our State member variable.
+	if (TSharedPtr<FSharedPlaybackState> SharedPlaybackState = FindSharedPlaybackState())
+	{
+		return SharedPlaybackState->FindCapability<FMovieSceneEvaluationState>();
+	}
+
+	return nullptr;
+}
+
+UE::MovieScene::IStaticBindingOverridesPlaybackCapability* IMovieScenePlayer::GetStaticBindingOverrides()
+{
+	using namespace UE::MovieScene;
+
+	// Return the playback capability, which is generally the same as "this" if InitializeRootInstance 
+	// has been called, but we want to deprecate the BindingOverrides member field in the future.
+	// In addition to this, if our underlying type is FMovieSceneLegacyPlayer, the playback capability
+	// is NOT the same as "this".
+	if (TSharedPtr<FSharedPlaybackState> SharedPlaybackState = FindSharedPlaybackState())
+	{
+		return SharedPlaybackState->FindCapability<IStaticBindingOverridesPlaybackCapability>();
+	}
+
+	return nullptr;
+}
+
 void IMovieScenePlayer::InitializeRootInstance(TSharedRef<UE::MovieScene::FSharedPlaybackState> NewSharedPlaybackState)
 {
 	using namespace UE::MovieScene;
 
 	NewSharedPlaybackState->AddCapability<FPlayerIndexPlaybackCapability>(UniqueIndex);
+
+	PRAGMA_DISABLE_DEPRECATION_WARNINGS
 	NewSharedPlaybackState->AddCapabilityRaw(&State);
+	PRAGMA_ENABLE_DEPRECATION_WARNINGS
 
 	// Only add the spawnregister if it is different from the default 'null' register (which does nothing)
 	FMovieSceneSpawnRegister* SpawnRegister = &GetSpawnRegister();
@@ -280,7 +346,7 @@ void IMovieScenePlayer::InitializeRootInstance(TSharedRef<UE::MovieScene::FShare
 		NewSharedPlaybackState->AddCapabilityRaw(SpawnRegister);
 	}
 	NewSharedPlaybackState->AddCapabilityRaw((IObjectBindingNotifyPlaybackCapability*)this);
-	NewSharedPlaybackState->AddCapabilityRaw((IStaticBindingOverridesPlaybackCapability*)this);
+	NewSharedPlaybackState->AddCapabilityRaw(&StaticBindingOverrides);
 
 	if (IMovieScenePlaybackClient* PlaybackClient = GetPlaybackClient())
 	{
diff --git a/Engine/Source/Runtime/MovieScene/Private/MovieSceneSequencePlayer.cpp b/Engine/Source/Runtime/MovieScene/Private/MovieSceneSequencePlayer.cpp
--- a/Engine/Source/Runtime/MovieScene/Private/MovieSceneSequencePlayer.cpp
+++ b/Engine/Source/Runtime/MovieScene/Private/MovieSceneSequencePlayer.cpp
@@ -1567,13 +1567,17 @@ TArray<UObject*> UMovieSceneSequencePlayer::GetBoundObjects(FMovieSceneObjectBin
 TArray<FMovieSceneObjectBindingID> UMovieSceneSequencePlayer::GetObjectBindings(UObject* InObject)
 {
 	TArray<FMovieSceneObjectBindingID> Bindings;
+	PRAGMA_DISABLE_DEPRECATION_WARNINGS
 	State.FilterObjectBindings(InObject, GetSharedPlaybackState(), &Bindings);
+	PRAGMA_ENABLE_DEPRECATION_WARNINGS
 	return Bindings;
 }
 
 void UMovieSceneSequencePlayer::RequestInvalidateBinding(FMovieSceneObjectBindingID ObjectBinding)
 {
+	PRAGMA_DISABLE_DEPRECATION_WARNINGS
 	State.Invalidate(ObjectBinding.GetGuid(), ObjectBinding.GetRelativeSequenceID());
+	PRAGMA_ENABLE_DEPRECATION_WARNINGS
 }
 
 UWorld* UMovieSceneSequencePlayer::GetPlaybackWorld() const
diff --git a/Engine/Source/Runtime/MovieScene/Public/Evaluation/MovieSceneEvaluationState.h b/Engine/Source/Runtime/MovieScene/Public/Evaluation/MovieSceneEvaluationState.h
--- a/Engine/Source/Runtime/MovieScene/Public/Evaluation/MovieSceneEvaluationState.h
+++ b/Engine/Source/Runtime/MovieScene/Public/Evaluation/MovieSceneEvaluationState.h
@@ -80,7 +80,7 @@ namespace UE::MovieScene
 	private:
 		TMap<FMovieSceneEvaluationOperand, FMovieSceneEvaluationOperand> BindingOverrides;
 
-		friend class IMovieScenePlayer;
+		friend class ::IMovieScenePlayer;
 	};
 
 }  // namespace UE::MovieScene
diff --git a/Engine/Source/Runtime/MovieScene/Public/IMovieScenePlayer.h b/Engine/Source/Runtime/MovieScene/Public/IMovieScenePlayer.h
--- a/Engine/Source/Runtime/MovieScene/Public/IMovieScenePlayer.h
+++ b/Engine/Source/Runtime/MovieScene/Public/IMovieScenePlayer.h
@@ -85,7 +85,6 @@ using EMovieSceneCameraCutParams = FMovieSceneCameraCutParams;
  */
 class IMovieScenePlayer 
 	: public UE::MovieScene::IObjectBindingNotifyPlaybackCapability
-	, public UE::MovieScene::IStaticBindingOverridesPlaybackCapability
 {
 public:
 	MOVIESCENE_API IMovieScenePlayer();
@@ -177,26 +176,17 @@ public:
 	/** 
 	 * Retrieves any override for the given operand
 	 */
-	virtual FMovieSceneEvaluationOperand* GetBindingOverride(const FMovieSceneEvaluationOperand& InOperand) override
-	{
-		return BindingOverrides.Find(InOperand);
-	}
+	FMovieSceneEvaluationOperand* GetBindingOverride(const FMovieSceneEvaluationOperand& InOperand);
 
 	/** 
 	 * Adds an override for the given operand 
 	 */
-	virtual void AddBindingOverride(const FMovieSceneEvaluationOperand& InOperand, const FMovieSceneEvaluationOperand& InOverrideOperand) override
-	{
-		BindingOverrides.Add(InOperand, InOverrideOperand);
-	}
+	void AddBindingOverride(const FMovieSceneEvaluationOperand& InOperand, const FMovieSceneEvaluationOperand& InOverrideOperand);
 
 	/** 
 	 * Removes any override set for the given operand
 	 */
-	virtual void RemoveBindingOverride(const FMovieSceneEvaluationOperand& InOperand) override
-	{
-		BindingOverrides.Remove(InOperand);
-	}
+	void RemoveBindingOverride(const FMovieSceneEvaluationOperand& InOperand);
 
 	/**
 	 * Remove all director blueprint instances
@@ -297,7 +287,9 @@ public:
 	 */
 	FGuid FindObjectId(UObject& InObject, FMovieSceneSequenceIDRef SequenceID)
 	{
+		PRAGMA_DISABLE_DEPRECATION_WARNINGS
 		return State.FindObjectId(InObject, SequenceID, GetSharedPlaybackState());
+		PRAGMA_ENABLE_DEPRECATION_WARNINGS
 	}
 
 	/**
@@ -311,7 +303,9 @@ public:
 	*/
 	FGuid FindCachedObjectId(UObject& InObject, FMovieSceneSequenceIDRef SequenceID)
 	{
+		PRAGMA_DISABLE_DEPRECATION_WARNINGS
 		return State.FindCachedObjectId(InObject, SequenceID, GetSharedPlaybackState());
+		PRAGMA_ENABLE_DEPRECATION_WARNINGS
 	}
 
 	/**
@@ -345,7 +339,9 @@ public:
 	void RestorePreAnimatedState()
 	{
 		PreAnimatedState.RestorePreAnimatedState();
+		PRAGMA_DISABLE_DEPRECATION_WARNINGS
 		State.ClearObjectCaches(GetSharedPlaybackState());
+		PRAGMA_ENABLE_DEPRECATION_WARNINGS
 	}
 
 	/**
@@ -354,10 +350,11 @@ public:
 	void DiscardPreAnimatedState()
 	{
 		PreAnimatedState.DiscardPreAnimatedState();
+		PRAGMA_DISABLE_DEPRECATION_WARNINGS
 		State.ClearObjectCaches(GetSharedPlaybackState());
+		PRAGMA_ENABLE_DEPRECATION_WARNINGS
 	}
 
-
 	/**
 	 * Invalidate any cached state contained within this player causing all entities to be forcibly re-linked and evaluated
 	 */
@@ -425,23 +422,35 @@ protected:
 	 */
 	MOVIESCENE_API virtual void ResolveBoundObjects(UE::UniversalObjectLocator::FResolveParams& ResolveParams, const FGuid& InBindingId, FMovieSceneSequenceID SequenceID, UMovieSceneSequence& Sequence, TArray<UObject*, TInlineAllocator<1>>& OutObjects) const;
 
+public:
+
+	/** Gets the evaluation state that stores global state to do with the playback operation */
+	MOVIESCENE_API FMovieSceneEvaluationState* GetEvaluationState();
+
+	/** Gets the list of binding overrides to use for the sequence */
+	MOVIESCENE_API UE::MovieScene::IStaticBindingOverridesPlaybackCapability* GetStaticBindingOverrides();
 
 public:
 
 	/** Evaluation state that stores global state to do with the playback operation */
+	UE_DEPRECATED(5.6, "Please use GetEvaluationState().")
 	FMovieSceneEvaluationState State;
 
 	/** Container that stores any per-animated state tokens  */
 	FMovieScenePreAnimatedState PreAnimatedState;
 
 	/** List of binding overrides to use for the sequence */
-	TMap<FMovieSceneEvaluationOperand, FMovieSceneEvaluationOperand> BindingOverrides;
+	UE_DEPRECATED(5.6, "Please use GetStaticBindingOverrides().")
+	TMap<FMovieSceneEvaluationOperand, FMovieSceneEvaluationOperand>& BindingOverrides;
 
 private:
 
 	/** Null register that asserts on use */
 	FNullMovieSceneSpawnRegister NullRegister;
 
+	/** Static binding overrides */
+	UE::MovieScene::FStaticBindingOverrides StaticBindingOverrides;
+
 	/** This player's unique Index */
 	uint16 UniqueIndex;
 };
diff --git a/Engine/Source/Runtime/MovieSceneTracks/Private/EntitySystem/Interrogation/MovieSceneInterrogationLinker.cpp b/Engine/Source/Runtime/MovieSceneTracks/Private/EntitySystem/Interrogation/MovieSceneInterrogationLinker.cpp
--- a/Engine/Source/Runtime/MovieSceneTracks/Private/EntitySystem/Interrogation/MovieSceneInterrogationLinker.cpp
+++ b/Engine/Source/Runtime/MovieSceneTracks/Private/EntitySystem/Interrogation/MovieSceneInterrogationLinker.cpp
@@ -878,7 +878,8 @@ FInterrogationChannel FSystemInterrogator::ImportLocalTransforms(USceneComponent
 {
 	check(SceneComponent);
 
-	UMovieSceneSequence* Sequence = InPlayer->State.FindSequence(SequenceID);
+	FMovieSceneEvaluationState* State = InPlayer->GetEvaluationState();
+	UMovieSceneSequence* Sequence = State->FindSequence(SequenceID);
 	if (!Sequence)
 	{
 		return FInterrogationChannel::Invalid();
@@ -902,7 +903,7 @@ FInterrogationChannel FSystemInterrogator::ImportLocalTransforms(USceneComponent
 	}
 
 	// Find the binding that corresponds to the component directly
-	FGuid ObjectBindingID = InPlayer->State.FindCachedObjectId(*SceneComponent, SequenceID, *InPlayer);
+	FGuid ObjectBindingID = State->FindCachedObjectId(*SceneComponent, SequenceID, *InPlayer);
 	if (ObjectBindingID.IsValid())
 	{
 		ImportTransformTracks(*Sequence->GetMovieScene()->FindBinding(ObjectBindingID), Channel);
@@ -912,7 +913,7 @@ FInterrogationChannel FSystemInterrogator::ImportLocalTransforms(USceneComponent
 	AActor* Owner = SceneComponent->GetOwner();
 	if (SceneComponent == Owner->GetRootComponent())
 	{
-		FGuid OwnerObjectBindingID = InPlayer->State.FindCachedObjectId(*Owner, SequenceID, *InPlayer);
+		FGuid OwnerObjectBindingID = State->FindCachedObjectId(*Owner, SequenceID, *InPlayer);
 		if (OwnerObjectBindingID.IsValid())
 		{
 			ImportTransformTracks(*Sequence->GetMovieScene()->FindBinding(OwnerObjectBindingID), Channel);
