Sequencer: auto-create a "legacy player" to allow playing legacy tracks.

diff --git a/Engine/Source/Runtime/MovieScene/Private/EntitySystem/MovieSceneSequenceInstance.cpp b/Engine/Source/Runtime/MovieScene/Private/EntitySystem/MovieSceneSequenceInstance.cpp
--- a/Engine/Source/Runtime/MovieScene/Private/EntitySystem/MovieSceneSequenceInstance.cpp
+++ b/Engine/Source/Runtime/MovieScene/Private/EntitySystem/MovieSceneSequenceInstance.cpp
@@ -16,6 +16,7 @@
 
 #include "IMovieScenePlayer.h"
 #include "Misc/AssertionMacros.h"
+#include "MovieSceneLegacyPlayer.h"
 #include "MovieSceneSequence.h"
 #include "MovieSceneSequencePlayer.h"
 #include "MovieSceneTimeHelpers.h"
@@ -148,6 +149,28 @@ void FSequenceInstance::InitializeLegacyEvaluator()
 	if (EnumHasAnyFlags(CompiledEntry.AccumulatedMask, EMovieSceneSequenceCompilerMask::EvaluationTemplate))
 	{
 		IMovieScenePlayer* Player = FPlayerIndexPlaybackCapability::GetPlayer(SharedPlaybackState);
+		if (!Player)
+		{
+			if (ILegacyPlayerProviderPlaybackCapability* LegacyPlayerProvider = SharedPlaybackState->FindCapability<ILegacyPlayerProviderPlaybackCapability>())
+			{
+				UE_LOG(LogMovieScene, Log, TEXT("Initializing a legacy player for sequence '%s'."), *GetNameSafe(SharedPlaybackState->GetRootSequence()));
+				Player = LegacyPlayerProvider->CreateLegacyPlayer(SharedPlaybackState);
+				if (Player)
+				{
+					SharedPlaybackState->SetOrAddCapability<FPlayerIndexPlaybackCapability>(Player->GetUniqueIndex());
+				}
+			}
+		}
+		if (!Player)
+		{
+			UE_LOG(LogMovieScene, Log, TEXT("Initializing a default legacy player for sequence '%s'."), *GetNameSafe(SharedPlaybackState->GetRootSequence()));
+			InitializeLegacyPlayer();
+			Player = LegacyPlayer.Get();
+			if (ensure(Player))
+			{
+				SharedPlaybackState->SetOrAddCapability<FPlayerIndexPlaybackCapability>(Player->GetUniqueIndex());
+			}
+		}
 		if (ensureMsgf(Player, TEXT("Sequence has legacy templates but no player! Legacy tracks will be disabled!")))
 		{
 			UpdateFlags |= ESequenceInstanceUpdateFlags::HasLegacyTemplates;
@@ -170,6 +193,16 @@ void FSequenceInstance::InitializeLegacyEvaluator()
 	}
 }
 
+void FSequenceInstance::InitializeLegacyPlayer()
+{
+	IMovieScenePlayer* Player = FPlayerIndexPlaybackCapability::GetPlayer(SharedPlaybackState);
+	if (!Player)
+	{
+		LegacyPlayer = MakeUnique<FMovieSceneLegacyPlayer>(SharedPlaybackState);
+		SharedPlaybackState->AddCapability<FPlayerIndexPlaybackCapability>(LegacyPlayer->GetUniqueIndex());
+	}
+}
+
 void FSequenceInstance::InvalidateCachedData(ESequenceInstanceInvalidationType InvalidationType)
 {
 	ensureMsgf(bInitialized, TEXT("Sequence instance hasn't been initialized yet!"));
diff --git a/Engine/Source/Runtime/MovieScene/Private/Evaluation/MovieSceneEvaluationTemplateInstance.cpp b/Engine/Source/Runtime/MovieScene/Private/Evaluation/MovieSceneEvaluationTemplateInstance.cpp
--- a/Engine/Source/Runtime/MovieScene/Private/Evaluation/MovieSceneEvaluationTemplateInstance.cpp
+++ b/Engine/Source/Runtime/MovieScene/Private/Evaluation/MovieSceneEvaluationTemplateInstance.cpp
@@ -24,7 +24,7 @@ DECLARE_CYCLE_STAT(TEXT("Entire Evaluation Cost"), MovieSceneEval_EntireEvaluati
 
 FMovieSceneRootEvaluationTemplateInstance::FMovieSceneRootEvaluationTemplateInstance()
 	: EntitySystemLinker(nullptr)
-	, RootID(MovieSceneSequenceID::Root)
+	, bOwnsSharedPlaybackState(false)
 {
 #if WITH_EDITOR
 	EmulatedNetworkMask = EMovieSceneServerClientMask::All;
@@ -43,7 +43,8 @@ void FMovieSceneRootEvaluationTemplateInstance::TearDown()
 	SharedPlaybackState.Reset();
 
 	// Avoid redundant work if the linker is being destroyed anyway
-	if (RootInstanceHandle.IsValid() &&
+	if (bOwnsSharedPlaybackState && 
+			RootInstanceHandle.IsValid() &&
 			EntitySystemLinker && 
 			IsValidChecked(EntitySystemLinker) && 
 			!EntitySystemLinker->IsUnreachable() && 
@@ -123,8 +124,6 @@ void FMovieSceneRootEvaluationTemplateInstance::Initialize(UMovieSceneSequence&
 		SharedPlaybackState->GetRootSequence() : nullptr;
 	bReinitialize |= (PreviousRootSequence != &InRootSequence);
 
-	RootID = MovieSceneSequenceID::Root;
-
 	if (bReinitialize)
 	{
 		// Tear down previous sequence first.
@@ -132,7 +131,7 @@ void FMovieSceneRootEvaluationTemplateInstance::Initialize(UMovieSceneSequence&
 			SharedPlaybackState->GetRootInstanceHandle() : FRootInstanceHandle();
 		SharedPlaybackState.Reset();
 
-		if (OldRootInstanceHandle.IsValid())
+		if (OldRootInstanceHandle.IsValid() && bOwnsSharedPlaybackState)
 		{
 			if (PreviousRunner)
 			{
@@ -172,9 +171,18 @@ void FMovieSceneRootEvaluationTemplateInstance::Initialize(UMovieSceneSequence&
 		{
 			Player.PreAnimatedState.Initialize(EntitySystemLinker, NewRootInstanceHandle);
 		}
+
+		bOwnsSharedPlaybackState = true;
 	}
 }
 
+void FMovieSceneRootEvaluationTemplateInstance::Initialize(TSharedRef<UE::MovieScene::FSharedPlaybackState> InSharedPlaybackState)
+{
+	EntitySystemLinker = InSharedPlaybackState->GetLinker();
+	SharedPlaybackState = InSharedPlaybackState;
+	bOwnsSharedPlaybackState = false;
+}
+
 void FMovieSceneRootEvaluationTemplateInstance::EvaluateSynchronousBlocking(FMovieSceneContext Context, IMovieScenePlayer& Player)
 {
 	EvaluateSynchronousBlocking(Context);
@@ -426,6 +434,11 @@ void FMovieSceneRootEvaluationTemplateInstance::PlaybackContextChanged(IMovieSce
 {
 	using namespace UE::MovieScene;
 
+	if (!ensure(bOwnsSharedPlaybackState))
+	{
+		return;
+	}
+
 	// Only the playback context changed, so we keep the same sequence, runner, and compiled data manager.
 	UMovieSceneSequence* RootSequence = SharedPlaybackState->GetRootSequence();
 	UMovieSceneCompiledDataManager* CompiledDataManager = SharedPlaybackState->GetCompiledDataManager();
diff --git a/Engine/Source/Runtime/MovieScene/Private/MovieSceneLegacyPlayer.cpp b/Engine/Source/Runtime/MovieScene/Private/MovieSceneLegacyPlayer.cpp
new file mode 100644
--- /dev/null
+++ b/Engine/Source/Runtime/MovieScene/Private/MovieSceneLegacyPlayer.cpp
@@ -0,0 +1,101 @@
+// Copyright Epic Games, Inc. All Rights Reserved.
+
+#include "MovieSceneLegacyPlayer.h"
+
+#include "EntitySystem/MovieSceneEntitySystemLinker.h"
+#include "EntitySystem/MovieSceneInstanceRegistry.h"
+#include "EntitySystem/MovieSceneSharedPlaybackState.h"
+#include "Evaluation/SequenceDirectorPlaybackCapability.h"
+#include "IMovieScenePlaybackClient.h"
+#include "MovieSceneSpawnRegister.h"
+
+namespace UE::MovieScene
+{
+
+UE_DEFINE_MOVIESCENE_PLAYBACK_CAPABILITY(ILegacyPlayerProviderPlaybackCapability)
+
+}  // namespace UE::MovieScene
+
+FMovieSceneLegacyPlayer::FMovieSceneLegacyPlayer()
+{
+}
+
+FMovieSceneLegacyPlayer::FMovieSceneLegacyPlayer(TSharedRef<UE::MovieScene::FSharedPlaybackState> InSharedPlaybackState)
+{
+	EvaluationTemplateInstance.Initialize(InSharedPlaybackState);
+}
+
+FMovieSceneRootEvaluationTemplateInstance& FMovieSceneLegacyPlayer::GetEvaluationTemplate()
+{
+	return EvaluationTemplateInstance;
+}
+
+UMovieSceneEntitySystemLinker* FMovieSceneLegacyPlayer::ConstructEntitySystemLinker()
+{
+	ensureMsgf(false, TEXT("This legacy player should never have to construct a linker."));
+	return nullptr;
+}
+
+UObject* FMovieSceneLegacyPlayer::AsUObject()
+{
+	return nullptr;
+}
+
+EMovieScenePlayerStatus::Type FMovieSceneLegacyPlayer::GetPlaybackStatus() const
+{
+	using namespace UE::MovieScene;
+
+	if (SharedPlaybackState)
+	{
+		FRootInstanceHandle RootInstanceHandle = SharedPlaybackState->GetRootInstanceHandle();
+		UMovieSceneEntitySystemLinker* Linker = SharedPlaybackState->GetLinker();
+		FInstanceRegistry* InstanceRegisty = Linker->GetInstanceRegistry();
+		const FSequenceInstance& RootInstance = InstanceRegisty->GetInstance(RootInstanceHandle);
+		const FMovieSceneContext& Context = RootInstance.GetContext();
+		return Context.GetStatus();
+	}
+	return EMovieScenePlayerStatus::Stopped;
+}
+
+void FMovieSceneLegacyPlayer::SetPlaybackStatus(EMovieScenePlayerStatus::Type InPlaybackStatus)
+{
+	// Ignored.
+	ensureMsgf(false, TEXT("Please don't set playback status from evaluation templates."));
+}
+
+IMovieScenePlaybackClient* FMovieSceneLegacyPlayer::GetPlaybackClient()
+{
+	if (SharedPlaybackState)
+	{
+		return SharedPlaybackState->FindCapability<IMovieScenePlaybackClient>();
+	}
+	return nullptr;
+}
+
+FMovieSceneSpawnRegister& FMovieSceneLegacyPlayer::GetSpawnRegister()
+{
+	if (SharedPlaybackState)
+	{
+		FMovieSceneSpawnRegister* SpawnRegister = SharedPlaybackState->FindCapability<FMovieSceneSpawnRegister>();
+		if (SpawnRegister)
+		{
+			return *SpawnRegister;
+		}
+	}
+	return IMovieScenePlayer::GetSpawnRegister();
+}
+
+UObject* FMovieSceneLegacyPlayer::GetPlaybackContext() const
+{
+	if (SharedPlaybackState)
+	{
+		return SharedPlaybackState->GetPlaybackContext();
+	}
+	return nullptr;
+}
+
+void FMovieSceneLegacyPlayer::InitializeRootInstance(TSharedRef<UE::MovieScene::FSharedPlaybackState> NewSharedPlaybackState)
+{
+	ensureMsgf(false, TEXT("The legacy player should never initialize sequences: it only wraps already initialized ones."));
+}
+
diff --git a/Engine/Source/Runtime/MovieScene/Private/MovieSceneSpawnable.cpp b/Engine/Source/Runtime/MovieScene/Private/MovieSceneSpawnable.cpp
--- a/Engine/Source/Runtime/MovieScene/Private/MovieSceneSpawnable.cpp
+++ b/Engine/Source/Runtime/MovieScene/Private/MovieSceneSpawnable.cpp
@@ -66,9 +66,12 @@ FName FMovieSceneSpawnable::GetNetAddressableName(TSharedRef<const UE::MovieScen
 	{
 		AddressingContext = Player->AsUObject();
 	}
-	else if (UObject* PlaybackContext = SharedPlaybackState->GetPlaybackContext())
+	if (!AddressingContext)
 	{
-		AddressingContext = PlaybackContext;
+		if (UObject* PlaybackContext = SharedPlaybackState->GetPlaybackContext())
+		{
+			AddressingContext = PlaybackContext;
+		}
 	}
 
 	if (!AddressingContext)
diff --git a/Engine/Source/Runtime/MovieScene/Public/EntitySystem/MovieSceneSequenceInstance.h b/Engine/Source/Runtime/MovieScene/Public/EntitySystem/MovieSceneSequenceInstance.h
--- a/Engine/Source/Runtime/MovieScene/Public/EntitySystem/MovieSceneSequenceInstance.h
+++ b/Engine/Source/Runtime/MovieScene/Public/EntitySystem/MovieSceneSequenceInstance.h
@@ -16,6 +16,7 @@
 #include "MovieSceneSequenceID.h"
 #include "Templates/UniquePtr.h"
 
+class FMovieSceneLegacyPlayer;
 class IMovieScenePlayer;
 class UMovieSceneEntitySystemLinker;
 class UObject;
@@ -382,6 +383,7 @@ public:
 private:
 
 	void InitializeLegacyEvaluator();
+	void InitializeLegacyPlayer();
 
 private:
 
@@ -394,6 +396,8 @@ private:
 
 	/** For top-level sequences only - legacy track template evaluator for the entire sequence */
 	TUniquePtr<FMovieSceneTrackEvaluator> LegacyEvaluator;
+	/** For top-level sequences only - legacy player */
+	TUniquePtr<FMovieSceneLegacyPlayer> LegacyPlayer;
 	/** For top-level sequences only - interface for either a flat or hierarchical entity updater */
 	TUniquePtr<ISequenceUpdater> SequenceUpdater;
 	/** For top-level sequences only - a utility class that is used to ensure that volatile sequences are up to date. Only valid in editor, or for sequences that have the volatile flag. */
diff --git a/Engine/Source/Runtime/MovieScene/Public/Evaluation/MovieSceneEvaluationTemplateInstance.h b/Engine/Source/Runtime/MovieScene/Public/Evaluation/MovieSceneEvaluationTemplateInstance.h
--- a/Engine/Source/Runtime/MovieScene/Public/Evaluation/MovieSceneEvaluationTemplateInstance.h
+++ b/Engine/Source/Runtime/MovieScene/Public/Evaluation/MovieSceneEvaluationTemplateInstance.h
@@ -52,6 +52,12 @@ public:
 	 */
 	MOVIESCENE_API void Initialize(UMovieSceneSequence& RootSequence, IMovieScenePlayer& Player, UMovieSceneCompiledDataManager* InCompiledDataManager);
 
+	/**
+	 * Initialize this template instance as a "dormant" wrapper around an already
+	 * existing shared playback state.
+	 */
+	MOVIESCENE_API void Initialize(TSharedRef<UE::MovieScene::FSharedPlaybackState> InSharedPlaybackState);
+
 	/**
 	 * Evaluate this sequence in a synchronous way.
 	 *
@@ -127,13 +133,15 @@ private:
 
 private:
 
+	/** The linker that we are using, cached from SharedPlaybackState. */
 	UPROPERTY()
 	TObjectPtr<UMovieSceneEntitySystemLinker> EntitySystemLinker;
 
-	/** The playback state for the hierarchy of sequence instances */
+	/** The playback state for the hierarchy of sequence instances. */
 	TSharedPtr<UE::MovieScene::FSharedPlaybackState> SharedPlaybackState;
 
-	FMovieSceneSequenceID RootID;
+	/** Whether we should destroy the sequence hierarchy when we are destroyed. */
+	bool bOwnsSharedPlaybackState;
 
 #if WITH_EDITOR
 	EMovieSceneServerClientMask EmulatedNetworkMask;
diff --git a/Engine/Source/Runtime/MovieScene/Public/MovieSceneLegacyPlayer.h b/Engine/Source/Runtime/MovieScene/Public/MovieSceneLegacyPlayer.h
new file mode 100644
--- /dev/null
+++ b/Engine/Source/Runtime/MovieScene/Public/MovieSceneLegacyPlayer.h
@@ -0,0 +1,56 @@
+// Copyright Epic Games, Inc. All Rights Reserved.
+
+#pragma once
+
+#include "IMovieScenePlayer.h"
+#include "Evaluation/MovieSceneEvaluationTemplateInstance.h"
+#include "Templates/SharedPointerFwd.h"
+
+namespace UE::MovieScene
+{
+
+	struct FSharedPlaybackState;
+
+	struct MOVIESCENE_API ILegacyPlayerProviderPlaybackCapability
+	{
+		UE_DECLARE_MOVIESCENE_PLAYBACK_CAPABILITY(ILegacyPlayerProviderPlaybackCapability)
+
+		virtual ~ILegacyPlayerProviderPlaybackCapability() {}
+
+		virtual IMovieScenePlayer* CreateLegacyPlayer(TSharedRef<FSharedPlaybackState> InSharedPlaybackState) = 0;
+	};
+
+}  // namespace UE::MovieScene
+
+/**
+ * A legacy player implementation only meant for running legacy tracks using the older
+ * evaluation template system.
+ */
+class FMovieSceneLegacyPlayer : public IMovieScenePlayer
+{
+public:
+
+	FMovieSceneLegacyPlayer();
+	FMovieSceneLegacyPlayer(TSharedRef<UE::MovieScene::FSharedPlaybackState> InSharedPlaybackState);
+
+public:
+
+	// IMovieScenePlayer interface.
+	virtual FMovieSceneRootEvaluationTemplateInstance& GetEvaluationTemplate() override;
+	virtual UMovieSceneEntitySystemLinker* ConstructEntitySystemLinker() override;
+	virtual UObject* AsUObject() override;
+	virtual EMovieScenePlayerStatus::Type GetPlaybackStatus() const override;
+	virtual void SetPlaybackStatus(EMovieScenePlayerStatus::Type InPlaybackStatus) override;
+	virtual IMovieScenePlaybackClient* GetPlaybackClient() override;
+	virtual FMovieSceneSpawnRegister& GetSpawnRegister() override;
+	virtual UObject* GetPlaybackContext() const override;
+
+	virtual void InitializeRootInstance(TSharedRef<UE::MovieScene::FSharedPlaybackState> NewSharedPlaybackState) override;
+
+private:
+
+	TSharedPtr<UE::MovieScene::FSharedPlaybackState> SharedPlaybackState;
+
+	FMovieSceneRootEvaluationTemplateInstance EvaluationTemplateInstance;
+};
+
diff --git a/Engine/Source/Runtime/MovieSceneTracks/Private/Bindings/MovieSceneSpawnableActorBinding.cpp b/Engine/Source/Runtime/MovieSceneTracks/Private/Bindings/MovieSceneSpawnableActorBinding.cpp
--- a/Engine/Source/Runtime/MovieSceneTracks/Private/Bindings/MovieSceneSpawnableActorBinding.cpp
+++ b/Engine/Source/Runtime/MovieSceneTracks/Private/Bindings/MovieSceneSpawnableActorBinding.cpp
@@ -242,9 +242,12 @@ FName UMovieSceneSpawnableActorBindingBase::GetNetAddressableName(TSharedRef<con
 	{
 		AddressingContext = Player->AsUObject();
 	}
-	else if (UObject* PlaybackContext = SharedPlaybackState->GetPlaybackContext())
+	if (!AddressingContext)
 	{
-		AddressingContext = PlaybackContext;
+		if (UObject* PlaybackContext = SharedPlaybackState->GetPlaybackContext())
+		{
+			AddressingContext = PlaybackContext;
+		}
 	}
 
 	if (!AddressingContext)
@@ -488,4 +491,4 @@ FText UMovieSceneSpawnableActorBinding::GetBindingTypePrettyName() const
 }
 #endif
 
-#undef LOCTEXT_NAMESPACE
\ No newline at end of file
+#undef LOCTEXT_NAMESPACE
diff --git a/Engine/Source/Runtime/MovieSceneTracks/Private/Systems/MovieSceneEventSystems.cpp b/Engine/Source/Runtime/MovieSceneTracks/Private/Systems/MovieSceneEventSystems.cpp
--- a/Engine/Source/Runtime/MovieSceneTracks/Private/Systems/MovieSceneEventSystems.cpp
+++ b/Engine/Source/Runtime/MovieSceneTracks/Private/Systems/MovieSceneEventSystems.cpp
@@ -91,13 +91,13 @@ void UMovieSceneEventSystem::TriggerAllEvents()
 
 		FFrameTime SkipUntil;
 		bool bSkipTrigger = false;
-		if (IMovieScenePlayer* Player = FPlayerIndexPlaybackCapability::GetPlayer(SharedPlaybackState))
+		if (FEventTriggerControlPlaybackCapability* TriggerControlCapability = SharedPlaybackState->FindCapability<FEventTriggerControlPlaybackCapability>())
 		{
-			bSkipTrigger = Player->IsDisablingEventTriggers(SkipUntil);
+			bSkipTrigger = TriggerControlCapability->IsDisablingEventTriggers(SkipUntil);
 		}
-		else if (FEventTriggerControlPlaybackCapability* TriggerControlCapability = SharedPlaybackState->FindCapability<FEventTriggerControlPlaybackCapability>())
+		else if (IMovieScenePlayer* Player = FPlayerIndexPlaybackCapability::GetPlayer(SharedPlaybackState))
 		{
-			bSkipTrigger = TriggerControlCapability->IsDisablingEventTriggers(SkipUntil);
+			bSkipTrigger = Player->IsDisablingEventTriggers(SkipUntil);
 		}
 
 		if (RootInstance.GetContext().GetDirection() == EPlayDirection::Forwards)
@@ -152,18 +152,18 @@ void UMovieSceneEventSystem::TriggerEvents(TArrayView<const FMovieSceneEventTrig
 	}
 
 	TArray<UObject*> GlobalContexts;
-	if (IMovieScenePlayer* Player = FPlayerIndexPlaybackCapability::GetPlayer(SharedPlaybackState))
+	if (IEventContextsPlaybackCapability* EventContextsCapability = SharedPlaybackState->FindCapability<IEventContextsPlaybackCapability>())
+	{
+		// Query the playback capability first.
+		GlobalContexts = EventContextsCapability->GetEventContexts();
+	}
+	else if (IMovieScenePlayer* Player = FPlayerIndexPlaybackCapability::GetPlayer(SharedPlaybackState))
 	{
 		// Check the player interface first. If implemented, we get what we want. If not, it will fall back to the default
 		// implementation that queries the playback capabality (see below), which covers the cases of a player that uses
 		// the new API.
 		GlobalContexts = Player->GetEventContexts();
 	}
-	else if (IEventContextsPlaybackCapability* EventContextsCapability = SharedPlaybackState->FindCapability<IEventContextsPlaybackCapability>())
-	{
-		// No player... let's query the playback capability.
-		GlobalContexts = EventContextsCapability->GetEventContexts();
-	}
 
 	for (const FMovieSceneEventTriggerData& Event : Events)
 	{
