Sequencer: new playback manager structure for player implementations

diff --git a/Engine/Source/Runtime/MovieScene/Private/Evaluation/MovieScenePlaybackManager.cpp b/Engine/Source/Runtime/MovieScene/Private/Evaluation/MovieScenePlaybackManager.cpp
new file mode 100644
--- /dev/null
+++ b/Engine/Source/Runtime/MovieScene/Private/Evaluation/MovieScenePlaybackManager.cpp
@@ -0,0 +1,409 @@
+// Copyright Epic Games, Inc. All Rights Reserved.
+
+#include "Evaluation/MovieScenePlaybackManager.h"
+
+#include "Channels/MovieSceneTimeWarpChannel.h"
+#include "Misc/CoreMiscDefines.h"
+#include "MovieScene.h"
+#include "MovieSceneFwd.h"
+#include "MovieSceneSequence.h"
+
+FMovieScenePlaybackManager::FMovieScenePlaybackManager()
+{
+}
+
+FMovieScenePlaybackManager::FMovieScenePlaybackManager(UMovieSceneSequence* InSequence)
+{
+	Initialize(InSequence);
+}
+
+void FMovieScenePlaybackManager::Initialize(UMovieSceneSequence* InSequence)
+{
+	if (!ensure(InSequence))
+	{
+		return;
+	}
+
+	UMovieScene* MovieScene = InSequence->GetMovieScene();
+	const EMovieSceneEvaluationType EvaluationType = MovieScene->GetEvaluationType();
+
+	DisplayRate = MovieScene->GetDisplayRate();
+
+	// Make our playback position work *only* in ticks. We will handle conversion to/from frames ourselves.
+	const FFrameRate TickResolution = MovieScene->GetTickResolution();
+	PlaybackPosition.SetTimeBase(TickResolution, TickResolution, EvaluationType);
+
+	const TRange<FFrameNumber> PlaybackRange = MovieScene->GetPlaybackRange();
+	SequenceStartTick = UE::MovieScene::DiscreteInclusiveLower(PlaybackRange);
+	SequenceEndTick = UE::MovieScene::DiscreteExclusiveUpper(PlaybackRange);
+
+	ResetPlaybackSettings();
+
+	PlaybackPosition.Reset(SequenceStartTick);
+}
+
+void FMovieScenePlaybackManager::ResetPlaybackSettings()
+{
+	StartOffsetTicks = EndOffsetTicks = FFrameNumber(0);
+	NumLoopsToPlay = 1;
+	NumLoopsCompleted = 0;
+	PlayRate = 1.0;
+	PlayDirection = EPlayDirection::Forwards;
+}
+
+void FMovieScenePlaybackManager::Update(float InDeltaSeconds, FContexts& OutContexts)
+{
+	using namespace UE::MovieScene;
+
+	if (PlaybackStatus != EMovieScenePlayerStatus::Playing)
+	{
+		return;
+	}
+	
+	// Get the new time, advanced by InDeltaSeconds.
+	const FFrameTime PreviousTick = PlaybackPosition.GetCurrentPosition();
+	const FFrameTime DeltaTicks = PlaybackPosition.GetInputRate().AsFrameTime(
+			(IsPlayingForward() ? InDeltaSeconds : (-InDeltaSeconds)) * PlayRate);
+	const FFrameTime NextTick = PreviousTick + DeltaTicks;
+
+	FFrameTime WarpedNextTick = NextTick;
+	if (bTransformPlaybackTime)
+	{
+		if (TimeTransform.FindFirstWarpDomain() == ETimeWarpChannelDomain::PlayRate)
+		{
+			WarpedNextTick = TimeTransform.TransformTime(WarpedNextTick);
+		}
+	}
+
+	InternalUpdateToTick(WarpedNextTick.RoundToFrame(), OutContexts);
+}
+
+void FMovieScenePlaybackManager::UpdateTo(const FFrameTime NextTime, FContexts& OutContexts)
+{
+	const FFrameTime NextTick = ConvertFrameTime(NextTime, DisplayRate, PlaybackPosition.GetInputRate());
+	InternalUpdateToTick(NextTick.RoundToFrame(), OutContexts);
+}
+
+void FMovieScenePlaybackManager::InternalUpdateToTick(const FFrameNumber NextTick, FContexts& OutContexts)
+{
+	using namespace UE::MovieScene;
+
+	// If we are stopped, just move the playhead to the given time, without generating evaluation contexts.
+	if (PlaybackStatus == EMovieScenePlayerStatus::Stopped)
+	{
+		PlaybackPosition.Reset(NextTick);
+		return;
+	}
+
+	// Check we have some loop counters that make sense.
+	if (NumLoopsToPlay > 0 && !ensure(NumLoopsCompleted < NumLoopsToPlay))
+	{
+		PlaybackStatus = EMovieScenePlayerStatus::Stopped;
+		PlaybackPosition.Reset(NextTick);
+	}
+
+	// Gather some information about this update.
+	const bool bShouldJump = 
+		(PlaybackStatus != EMovieScenePlayerStatus::Playing && PlaybackStatus != EMovieScenePlayerStatus::Scrubbing);
+
+	const FFrameNumber EffectiveStartTick = SequenceStartTick + StartOffsetTicks;
+	const FFrameNumber EffectiveEndTick = SequenceEndTick - EndOffsetTicks;
+	const FFrameNumber EffectiveDurationTicks = FMath::Max(FFrameNumber(0), EffectiveEndTick - EffectiveStartTick);
+	const FFrameNumber LastValidTick = GetLastValidTick();
+	// IMPORTANT: we assume that LastValidTick is less than the duration (current implementation is 
+	//			  duration minus one tick).
+
+	const bool bIsPlayingForward = IsPlayingForward();
+	const FFrameNumber LoopStartTick = bIsPlayingForward ? EffectiveStartTick : LastValidTick;
+	const FFrameNumber LoopLastTick = bIsPlayingForward ? LastValidTick : EffectiveStartTick;
+	const FFrameNumber LoopDurationTicks = EffectiveDurationTicks;
+
+	// Figure out if we crossed the loop-end boundary.
+	if ((bIsPlayingForward && NextTick > LoopLastTick) ||
+			((!bIsPlayingForward && NextTick < LoopLastTick)))
+	{
+		// Compute how many times we crossed the loop-end boundary.
+		const FFrameNumber LoopRelativeTick = NextTick - LoopStartTick;
+		const int32 NumLoopingsOver = FMath::Abs(LoopRelativeTick.Value) / LoopDurationTicks.Value;
+		ensure(NumLoopingsOver > 0);
+
+		// Massage this a bit with the following rules:
+		//
+		// 1) Don't go over the number of loops to play unless we're playing indefinitely.
+		//
+		// 2) Add an extra completed loop if we reached the end (as opposed to crossing it) because
+		//    that should count as "completing a loop".
+		//
+		const int32 NumLoopsNewlyCompleted = ((NumLoopsToPlay > 0) ?
+				FMath::Min(NumLoopingsOver, NumLoopsToPlay - NumLoopsCompleted) :
+				NumLoopingsOver);
+
+		// Play the last bit of the loop if we are looping and doing any sort of dissections.
+		if (DissectLooping != EMovieSceneLoopDissection::None)
+		{
+			OutContexts.Add(
+					FMovieSceneContext(PlaybackPosition.PlayTo(LoopLastTick), PlaybackStatus)
+					.SetHasJumped(bShouldJump)
+					);
+		}
+
+		// See if we need to generate more update ranges for the loops. This can happen if we had a
+		// large delta-time, and the duration of a loop is pretty short (i.e. we could have looped 
+		// several times in one update).
+		if (NumLoopingsOver > 1)
+		{
+			const int32 ExtraLoops = NumLoopingsOver - 1;
+
+			if (DissectLooping == EMovieSceneLoopDissection::DissectAll)
+			{
+				// If we dissect the looping, we add an explicit update for each loop, from start to end.
+				if (!bPingPongPlayback)
+				{
+					for (int32 Index = 0; Index < ExtraLoops; ++Index)
+					{
+						PlaybackPosition.Reset(LoopStartTick);
+						OutContexts.Add(
+								FMovieSceneContext(PlaybackPosition.PlayTo(LoopLastTick), PlaybackStatus)
+								.SetHasJumped(true)
+								);
+					}
+				}
+				else
+				{
+					ReversePlayDirection();
+
+					for (int32 Index = 0; Index < ExtraLoops; ++Index)
+					{
+						if (IsPlayingForward())
+						{
+							PlaybackPosition.Reset(EffectiveStartTick);
+							OutContexts.Add(
+									FMovieSceneContext(PlaybackPosition.PlayTo(LastValidTick), PlaybackStatus)
+									.SetHasJumped(true)
+									);
+						}
+						else
+						{
+							PlaybackPosition.Reset(LastValidTick);
+							OutContexts.Add(
+									FMovieSceneContext(PlaybackPosition.PlayTo(EffectiveStartTick), PlaybackStatus)
+									.SetHasJumped(true)
+									);
+						}
+
+						ReversePlayDirection();
+					}
+				}
+			}
+			else
+			{
+				// We are not dissecting loops so don't emit extra evaluation contexts.
+				// If we are ping-pong'ing, we at least need to keep track of which way we are now going.
+				if (bPingPongPlayback && (ExtraLoops + 1) % 2 != 0)
+				{
+					ReversePlayDirection();
+				}
+			}
+		}
+
+		// Complete the loops we said we completed.
+		NumLoopsCompleted += NumLoopsNewlyCompleted;
+
+		if (NumLoopsToPlay > 0 && NumLoopsCompleted >= NumLoopsToPlay)
+		{
+			// If we have played all the loops we needed to play, we can stop playback.
+			// However, if we don't dissect looping, we didn't finish the loop, so do it now.
+			if (DissectLooping == EMovieSceneLoopDissection::None)
+			{
+				OutContexts.Add(
+						FMovieSceneContext(PlaybackPosition.PlayTo(LoopLastTick), PlaybackStatus)
+						.SetHasJumped(bShouldJump)
+						);
+			}
+
+			PlaybackStatus = EMovieScenePlayerStatus::Stopped;
+		}
+		else
+		{
+			// Start the next loop with any overplay from the update.
+			//
+			// When playing forward, we have, e.g.:
+			// loop = [-30, -10], next time = -5, relative time = 25, mod(25, 20) = 5
+			//
+			// When playing backwards, we have, e.g.:
+			// loop = [-30, -10], next time = -35, relative time = -25, mod(-25, 20) = -5
+			const FFrameNumber OverplayTicks = LoopRelativeTick.Value % LoopDurationTicks.Value;
+
+			// Don't use LoopStartTick/LoopLastTick here because if we are ping-pong'ing, we would
+			// be going the other way. Use the most recent value of PlayDirection (via IsPlayingForward())
+			// for the same reason.
+			// Also, reverse OverplayTicks when ping-pong'in since we're playing this overplay in the
+			// reverse direction.
+			const FFrameTime EffectiveOverplayTicks = 
+				(IsPlayingForward() ? EffectiveStartTick : LastValidTick) + 
+				((!bPingPongPlayback) ? OverplayTicks : (-OverplayTicks));
+			PlaybackPosition.Reset(IsPlayingForward() ? EffectiveStartTick : LastValidTick);
+			OutContexts.Add(
+					FMovieSceneContext(PlaybackPosition.PlayTo(EffectiveOverplayTicks), PlaybackStatus)
+					.SetHasJumped(true)
+					);
+
+			// If the overplay leads us exactly to the last tick of the loop, let's count that
+			// as a completed loop... but only if that finishes playback. Otherwise, we'll wait
+			// for the next update to loop over in order to avoid counting that loop twice.
+			if (EffectiveOverplayTicks == LoopLastTick 
+					&& NumLoopsToPlay > 0
+					&& NumLoopsCompleted == NumLoopsToPlay - 1)
+			{
+				++NumLoopsCompleted;
+
+				PlaybackStatus = EMovieScenePlayerStatus::Stopped;
+			}
+		}
+	}
+	else
+	{
+		// We haven't crossed a loop-end boundary... just chug along.
+		OutContexts.Add(
+				FMovieSceneContext(PlaybackPosition.PlayTo(NextTick), PlaybackStatus)
+				.SetHasJumped(bShouldJump)
+				);
+
+		// If we were updated right up to the last tick of the loop, let's count that as a completed
+		// loop... but only if that finishes playback. Otherwise, we'll wait for the next update to
+		// loop over in order to avoid counting that loop twice.
+		if (NextTick == LoopLastTick
+				&& NumLoopsToPlay > 0
+				&& NumLoopsCompleted == NumLoopsToPlay - 1)
+		{
+			++NumLoopsCompleted;
+
+			PlaybackStatus = EMovieScenePlayerStatus::Stopped;
+		}
+	}
+
+	ensure(OutContexts.Num() > 0);
+	ensure(OutContexts.Num() == 1 || (DissectLooping != EMovieSceneLoopDissection::None));
+}
+
+FFrameNumber FMovieScenePlaybackManager::GetLastValidTick() const
+{
+	// TODO: handle precision problems with float SubFrame, or change SubFrame to double.
+	return SequenceEndTick - EndOffsetTicks - 1;  // minus one tick for exclusive end frame.
+}
+
+FMovieSceneContext FMovieScenePlaybackManager::UpdateAtCurrentTime() const
+{
+	return FMovieSceneContext(PlaybackPosition.GetCurrentPositionAsRange(), PlaybackStatus);
+}
+
+FFrameTime FMovieScenePlaybackManager::GetCurrentTime() const
+{
+	const FFrameRate TickResolution = PlaybackPosition.GetOutputRate();
+	return ConvertFrameTime(PlaybackPosition.GetCurrentPosition(), TickResolution, DisplayRate);
+}
+
+void FMovieScenePlaybackManager::SetCurrentTime(const FFrameTime& InFrameTime)
+{
+	const FFrameRate TickResolution = PlaybackPosition.GetOutputRate();
+	const FFrameNumber CurrentTick = ConvertFrameTime(InFrameTime, DisplayRate, TickResolution).RoundToFrame();
+	PlaybackPosition.Reset(CurrentTick);
+}
+
+TRange<FFrameTime> FMovieScenePlaybackManager::GetEffectivePlaybackRange() const
+{
+	const FFrameNumber StartTick = SequenceStartTick + StartOffsetTicks;
+	const FFrameNumber EndTick = SequenceEndTick - EndOffsetTicks;
+
+	const FFrameRate TickResolution = PlaybackPosition.GetOutputRate();
+
+	const FFrameTime StartFrame = ConvertFrameTime(StartTick, TickResolution, DisplayRate);
+	const FFrameTime EndFrame = ConvertFrameTime(EndTick, TickResolution, DisplayRate);
+
+	return TRange<FFrameTime>(TRangeBound<FFrameTime>::Inclusive(StartFrame), TRangeBound<FFrameTime>::Exclusive(EndFrame));
+}
+
+FFrameTime FMovieScenePlaybackManager::GetEffectiveStartTime() const
+{
+	const FFrameRate TickResolution = PlaybackPosition.GetOutputRate();
+	return ConvertFrameTime(SequenceStartTick + StartOffsetTicks, TickResolution, DisplayRate);
+}
+
+FFrameTime FMovieScenePlaybackManager::GetEffectiveEndTime() const
+{
+	const FFrameRate TickResolution = PlaybackPosition.GetOutputRate();
+	return ConvertFrameTime(SequenceEndTick + EndOffsetTicks, TickResolution, DisplayRate);
+}
+
+void FMovieScenePlaybackManager::SetStartOffset(const FFrameTime& InStartOffset)
+{
+	const FFrameRate TickResolution = PlaybackPosition.GetOutputRate();
+	const FFrameNumber InStartOffsetTicks = ConvertFrameTime(InStartOffset, DisplayRate, TickResolution).RoundToFrame();
+
+	SetStartAndEndOffsetTicks(InStartOffsetTicks, EndOffsetTicks);
+}
+
+void FMovieScenePlaybackManager::SetEndOffset(const FFrameTime& InEndOffset)
+{
+	const FFrameRate TickResolution = PlaybackPosition.GetOutputRate();
+	const FFrameNumber InEndOffsetTicks = ConvertFrameTime(InEndOffset, DisplayRate, TickResolution).RoundToFrame();
+
+	SetStartAndEndOffsetTicks(StartOffsetTicks, InEndOffsetTicks);
+}
+
+void FMovieScenePlaybackManager::SetEndOffsetAsTime(const FFrameTime& InEndTime)
+{
+	const FFrameRate TickResolution = PlaybackPosition.GetOutputRate();
+	const FFrameNumber InEndTick = ConvertFrameTime(InEndTime, DisplayRate, TickResolution).RoundToFrame();
+
+	const FFrameNumber InEndOffsetTicks = SequenceEndTick - InEndTick;
+	SetStartAndEndOffsetTicks(StartOffsetTicks, InEndOffsetTicks);
+}
+
+void FMovieScenePlaybackManager::SetStartAndEndOffsetTicks(FFrameNumber InStartOffsetTicks, FFrameNumber InEndOffsetTicks)
+{
+	const FFrameNumber SequenceDurationTicks = SequenceEndTick - SequenceStartTick;
+
+	StartOffsetTicks = FMath::Min(
+			FMath::Max(InStartOffsetTicks, FFrameNumber(0)),
+			SequenceDurationTicks);
+
+	EndOffsetTicks = FMath::Min(
+			FMath::Max(InEndOffsetTicks, FFrameNumber(0)),
+			SequenceDurationTicks - StartOffsetTicks);
+}
+
+FFrameTime FMovieScenePlaybackManager::GetStartOffset() const
+{
+	const FFrameRate TickResolution = PlaybackPosition.GetOutputRate();
+	return ConvertFrameTime(StartOffsetTicks, TickResolution, DisplayRate);
+}
+
+FFrameTime FMovieScenePlaybackManager::GetEndOffset() const
+{
+	const FFrameRate TickResolution = PlaybackPosition.GetOutputRate();
+	return ConvertFrameTime(EndOffsetTicks, TickResolution, DisplayRate);
+}
+
+void FMovieScenePlaybackManager::SetNumLoopsToPlay(int32 InNumLoopsToPlay)
+{
+	NumLoopsToPlay = InNumLoopsToPlay;
+}
+
+void FMovieScenePlaybackManager::ReversePlayDirection()
+{
+	if (PlayDirection == EPlayDirection::Forwards)
+	{
+		PlayDirection = EPlayDirection::Backwards;
+	}
+	else
+	{
+		PlayDirection = EPlayDirection::Forwards;
+	}
+}
+
+void FMovieScenePlaybackManager::SetPingPongPlayback(bool bInPingPongPlayback)
+{
+	bPingPongPlayback = bInPingPongPlayback;
+}
+
diff --git a/Engine/Source/Runtime/MovieScene/Private/Tests/MovieScenePlaybackManagerTests.cpp b/Engine/Source/Runtime/MovieScene/Private/Tests/MovieScenePlaybackManagerTests.cpp
new file mode 100644
--- /dev/null
+++ b/Engine/Source/Runtime/MovieScene/Private/Tests/MovieScenePlaybackManagerTests.cpp
@@ -0,0 +1,528 @@
+// Copyright Epic Games, Inc. All Rights Reserved.
+
+#include "CoreTypes.h"
+#include "Evaluation/MovieScenePlaybackManager.h"
+#include "Misc/AutomationTest.h"
+#include "MovieSceneFwd.h"
+#include "MovieSceneTimeHelpers.h"
+#include "Tests/MovieSceneTestObjects.h"
+#include "UObject/Package.h"
+#include "UObject/UObjectGlobals.h"
+
+#define LOCTEXT_NAMESPACE "MovieScenePlaybackManagerTests"
+
+namespace UE::MovieScene::Test
+{
+
+struct FMakeTestSequenceParams
+{
+	int32 StartTick = 0;
+	int32 DurationTicks = 60000 * 10;
+
+	int32 TickResolution = 60000;
+	int32 DisplayRate = 30;
+
+	FMakeTestSequenceParams()
+	{}
+	FMakeTestSequenceParams(int32 InDurationTicks) 
+		: DurationTicks(InDurationTicks)
+	{}
+	FMakeTestSequenceParams(int32 InStartTick, int32 InDurationTicks) 
+		: StartTick(InStartTick)
+		, DurationTicks(InDurationTicks)
+	{}
+};
+
+UMovieSceneSequence* MakeTestSequence(const FMakeTestSequenceParams& Params)
+{
+	UTestMovieSceneSequence* Sequence = NewObject<UTestMovieSceneSequence>(GetTransientPackage());
+	
+	Sequence->MovieScene->SetDisplayRate(FFrameRate(Params.DisplayRate, 1));
+	Sequence->MovieScene->SetTickResolutionDirectly(FFrameRate(Params.TickResolution, 1));
+	Sequence->MovieScene->SetPlaybackRange(FFrameNumber(Params.StartTick), Params.DurationTicks);
+	
+	return Sequence;
+}
+
+TRange<FFrameTime> MakeDiscreteTimeRange(FFrameTime MinInclusive, FFrameTime MaxExclusive)
+{
+	return TRange<FFrameTime>(
+			TRangeBound<FFrameTime>::Inclusive(MinInclusive),
+			TRangeBound<FFrameTime>::Exclusive(MaxExclusive));
+}
+
+TRange<FFrameTime> MakeContinuedTimeRange(FFrameTime MinExclusive, FFrameTime MaxInclusive)
+{
+	return TRange<FFrameTime>(
+			TRangeBound<FFrameTime>::Exclusive(MinExclusive),
+			TRangeBound<FFrameTime>::Inclusive(MaxInclusive));
+}
+
+TRange<FFrameTime> MakeHullTimeRange(FFrameTime MinInclusive, FFrameTime MaxInclusive)
+{
+	return TRange<FFrameTime>(
+			TRangeBound<FFrameTime>::Inclusive(MinInclusive),
+			TRangeBound<FFrameTime>::Inclusive(MaxInclusive));
+}
+
+}  // namespace UE::MovieScene::Test
+
+IMPLEMENT_SIMPLE_AUTOMATION_TEST(
+		FMovieScenePlaybackManagerTestStartEndTimes, 
+		"System.Engine.Sequencer.PlaybackManager.StartEndTimes", 
+		EAutomationTestFlags::EditorContext | EAutomationTestFlags::EngineFilter)
+bool FMovieScenePlaybackManagerTestStartEndTimes::RunTest(const FString& Parameters)
+{
+	using namespace UE::MovieScene;
+	using namespace UE::MovieScene::Test;
+
+	FMakeTestSequenceParams MakeParams;
+	UMovieSceneSequence* Sequence = MakeTestSequence(MakeParams);
+
+	FMovieScenePlaybackManager Manager(Sequence);
+
+	UTEST_EQUAL("PlaybackRange", 
+			Manager.GetEffectivePlaybackRange(),
+			MakeDiscreteTimeRange(0, 30 * 10));
+
+	Manager.SetStartOffset(30 * 1);
+
+	UTEST_EQUAL("PlaybackRange", 
+			Manager.GetEffectivePlaybackRange(),
+			MakeDiscreteTimeRange(30 * 1, 30 * 10));
+
+	Manager.SetEndOffset(FFrameTime(30 * 2, 0.5));
+
+	UTEST_EQUAL("PlaybackRange", 
+			Manager.GetEffectivePlaybackRange(),
+			MakeDiscreteTimeRange(FFrameTime(30 * 1), FFrameTime(30 * 7 + 29, 0.5)));
+
+	Manager.SetEndOffsetAsTime(30 * 8);
+
+	UTEST_EQUAL("PlaybackRange", 
+			Manager.GetEffectivePlaybackRange(),
+			MakeDiscreteTimeRange(30 * 1, 30 * 8));
+
+	return true;
+}
+
+IMPLEMENT_SIMPLE_AUTOMATION_TEST(
+		FMovieScenePlaybackManagerTestLoopingNoDissection, 
+		"System.Engine.Sequencer.PlaybackManager.Looping.NoDissection", 
+		EAutomationTestFlags::EditorContext | EAutomationTestFlags::EngineFilter)
+bool FMovieScenePlaybackManagerTestLoopingNoDissection::RunTest(const FString& Parameters)
+{
+	using namespace UE::MovieScene;
+	using namespace UE::MovieScene::Test;
+
+	FMakeTestSequenceParams MakeParams;
+	UMovieSceneSequence* Sequence = MakeTestSequence(MakeParams);
+
+	FMovieScenePlaybackManager Manager(Sequence);
+	Manager.SetDissectLooping(EMovieSceneLoopDissection::None);
+	Manager.SetNumLoopsToPlay(4);
+	Manager.SetPlaybackStatus(EMovieScenePlayerStatus::Playing);
+
+	const int32 TR = 60000;
+	const FFrameTime LastValidTick(10 * TR - 1);
+
+	FMovieScenePlaybackManager::FContexts Contexts;
+
+	{
+		UTEST_EQUAL("NumLoops", Manager.GetNumLoopsCompleted(), 0);
+
+		Contexts.Reset();
+		Manager.UpdateTo(30 * 5, Contexts);
+		UTEST_EQUAL("NumContexts", Contexts.Num(), 1);
+		UTEST_EQUAL("Context0", Contexts[0].GetRange(), MakeHullTimeRange(0, TR * 5));
+		UTEST_FALSE("Context0_Jumped", Contexts[0].HasJumped());
+		UTEST_EQUAL("NumLoops", Manager.GetNumLoopsCompleted(), 0);
+		UTEST_EQUAL("Status", Manager.GetPlaybackStatus(), EMovieScenePlayerStatus::Playing);
+
+		Contexts.Reset();
+		Manager.UpdateTo(30 * 12, Contexts);
+		UTEST_EQUAL("NumContexts", Contexts.Num(), 1);
+		UTEST_EQUAL("Context0", Contexts[0].GetRange(), MakeHullTimeRange(0, TR * 2));
+		UTEST_TRUE("Context0_Jumped", Contexts[0].HasJumped());
+		UTEST_EQUAL("NumLoops", Manager.GetNumLoopsCompleted(), 1);
+		UTEST_EQUAL("Status", Manager.GetPlaybackStatus(), EMovieScenePlayerStatus::Playing);
+
+		Contexts.Reset();
+		Manager.UpdateTo(30 * 23, Contexts);
+		UTEST_EQUAL("NumContexts", Contexts.Num(), 1);
+		UTEST_EQUAL("Context0", Contexts[0].GetRange(), MakeHullTimeRange(0, TR * 3));
+		UTEST_TRUE("Context0_Jumped", Contexts[0].HasJumped());
+		UTEST_EQUAL("NumLoops", Manager.GetNumLoopsCompleted(), 3);
+		UTEST_EQUAL("Status", Manager.GetPlaybackStatus(), EMovieScenePlayerStatus::Playing);
+
+		Contexts.Reset();
+		Manager.UpdateTo(30 * 11, Contexts);
+		UTEST_EQUAL("NumContexts", Contexts.Num(), 1);
+		UTEST_EQUAL("Context0", Contexts[0].GetRange(), MakeContinuedTimeRange(TR * 3, LastValidTick));
+		UTEST_FALSE("Context0_Jumped", Contexts[0].HasJumped());
+		UTEST_EQUAL("NumLoops", Manager.GetNumLoopsCompleted(), 4);
+		UTEST_EQUAL("Status", Manager.GetPlaybackStatus(), EMovieScenePlayerStatus::Stopped);
+	}
+
+	Manager.ResetNumLoopsCompleted();
+	Manager.SetPlayDirection(EPlayDirection::Backwards);
+	Manager.SetPlaybackStatus(EMovieScenePlayerStatus::Playing);
+
+	{
+		UTEST_EQUAL("NumLoops", Manager.GetNumLoopsCompleted(), 0);
+
+		Manager.SetCurrentTime(30 * 10);
+
+		Contexts.Reset();
+		Manager.UpdateTo(30 * 5, Contexts);
+		UTEST_EQUAL("NumContexts", Contexts.Num(), 1);
+		UTEST_EQUAL("Context0", Contexts[0].GetRange(), MakeHullTimeRange(TR * 5, TR * 10));
+		UTEST_FALSE("Context0_Jumped", Contexts[0].HasJumped());
+		UTEST_EQUAL("Context0_Backwards", Contexts[0].GetDirection(), EPlayDirection::Backwards);
+		UTEST_EQUAL("NumLoops", Manager.GetNumLoopsCompleted(), 0);
+		UTEST_EQUAL("Status", Manager.GetPlaybackStatus(), EMovieScenePlayerStatus::Playing);
+
+		Contexts.Reset();
+		Manager.UpdateTo(30 * -2, Contexts);
+		UTEST_EQUAL("NumContexts", Contexts.Num(), 1);
+		UTEST_EQUAL("Context0", Contexts[0].GetRange(), MakeHullTimeRange(TR * 8, LastValidTick));
+		UTEST_TRUE ("Context0_Jumped", Contexts[0].HasJumped());
+		UTEST_EQUAL("Context0_Backwards", Contexts[0].GetDirection(), EPlayDirection::Backwards);
+		UTEST_EQUAL("NumLoops", Manager.GetNumLoopsCompleted(), 1);
+		UTEST_EQUAL("Status", Manager.GetPlaybackStatus(), EMovieScenePlayerStatus::Playing);
+
+		Contexts.Reset();
+		Manager.UpdateTo(30 * -13, Contexts);
+		UTEST_EQUAL("NumContexts", Contexts.Num(), 1);
+		UTEST_EQUAL("Context0", Contexts[0].GetRange(), MakeHullTimeRange(TR * 7, LastValidTick));
+		UTEST_TRUE ("Context0_Jumped", Contexts[0].HasJumped());
+		UTEST_EQUAL("Context0_Backwards", Contexts[0].GetDirection(), EPlayDirection::Backwards);
+		UTEST_EQUAL("NumLoops", Manager.GetNumLoopsCompleted(), 3);
+		UTEST_EQUAL("Status", Manager.GetPlaybackStatus(), EMovieScenePlayerStatus::Playing);
+
+		Contexts.Reset();
+		Manager.UpdateTo(30 * -1, Contexts);
+		UTEST_EQUAL("NumContexts", Contexts.Num(), 1);
+		UTEST_EQUAL("Context0", Contexts[0].GetRange(), MakeDiscreteTimeRange(0, TR * 7));
+		UTEST_FALSE("Context0_Jumped", Contexts[0].HasJumped());
+		UTEST_EQUAL("Context0_Backwards", Contexts[0].GetDirection(), EPlayDirection::Backwards);
+		UTEST_EQUAL("NumLoops", Manager.GetNumLoopsCompleted(), 4);
+		UTEST_EQUAL("Status", Manager.GetPlaybackStatus(), EMovieScenePlayerStatus::Stopped);
+	}
+
+	return true;
+}
+
+IMPLEMENT_SIMPLE_AUTOMATION_TEST(
+		FMovieScenePlaybackManagerTestLoopingDissectOne, 
+		"System.Engine.Sequencer.PlaybackManager.Looping.DissectOne", 
+		EAutomationTestFlags::EditorContext | EAutomationTestFlags::EngineFilter)
+bool FMovieScenePlaybackManagerTestLoopingDissectOne::RunTest(const FString& Parameters)
+{
+	using namespace UE::MovieScene;
+	using namespace UE::MovieScene::Test;
+
+	FMakeTestSequenceParams MakeParams;
+	UMovieSceneSequence* Sequence = MakeTestSequence(MakeParams);
+
+	FMovieScenePlaybackManager Manager(Sequence);
+	Manager.SetDissectLooping(EMovieSceneLoopDissection::DissectOne);
+	Manager.SetNumLoopsToPlay(4);
+	Manager.SetPlaybackStatus(EMovieScenePlayerStatus::Playing);
+
+	const int32 TR = 60000;
+	const FFrameTime LastValidTick(10 * TR - 1);
+
+	FMovieScenePlaybackManager::FContexts Contexts;
+
+	{
+
+		UTEST_EQUAL("NumLoops", Manager.GetNumLoopsCompleted(), 0);
+
+		Contexts.Reset();
+		Manager.UpdateTo(30 * 5, Contexts);
+		UTEST_EQUAL("NumContexts", Contexts.Num(), 1);
+		UTEST_EQUAL("Context0", Contexts[0].GetRange(), MakeHullTimeRange(0, TR * 5));
+		UTEST_FALSE("Context0_Jumped", Contexts[0].HasJumped());
+		UTEST_EQUAL("NumLoops", Manager.GetNumLoopsCompleted(), 0);
+		UTEST_EQUAL("Status", Manager.GetPlaybackStatus(), EMovieScenePlayerStatus::Playing);
+
+		Contexts.Reset();
+		Manager.UpdateTo(30 * 12, Contexts);
+		UTEST_EQUAL("NumContexts", Contexts.Num(), 2);
+		UTEST_EQUAL("Context0", Contexts[0].GetRange(), MakeContinuedTimeRange(TR * 5, LastValidTick));
+		UTEST_FALSE("Context0_Jumped", Contexts[0].HasJumped());
+		UTEST_EQUAL("Context1", Contexts[1].GetRange(), MakeHullTimeRange(0, TR * 2));
+		UTEST_TRUE ("Context1_Jumped", Contexts[1].HasJumped());
+		UTEST_EQUAL("NumLoops", Manager.GetNumLoopsCompleted(), 1);
+		UTEST_EQUAL("Status", Manager.GetPlaybackStatus(), EMovieScenePlayerStatus::Playing);
+
+		Contexts.Reset();
+		Manager.UpdateTo(30 * 23, Contexts);
+		UTEST_EQUAL("NumContexts", Contexts.Num(), 2);
+		UTEST_EQUAL("Context0", Contexts[0].GetRange(), MakeContinuedTimeRange(TR * 2, LastValidTick));
+		UTEST_FALSE("Context0_Jumped", Contexts[0].HasJumped());
+		UTEST_EQUAL("Context1", Contexts[1].GetRange(), MakeHullTimeRange(0, TR * 3));
+		UTEST_TRUE ("Context1_Jumped", Contexts[1].HasJumped());
+		UTEST_EQUAL("NumLoops", Manager.GetNumLoopsCompleted(), 3);
+		UTEST_EQUAL("Status", Manager.GetPlaybackStatus(), EMovieScenePlayerStatus::Playing);
+
+		Contexts.Reset();
+		Manager.UpdateTo(30 * 11, Contexts);
+		UTEST_EQUAL("NumContexts", Contexts.Num(), 1);
+		UTEST_EQUAL("Context0", Contexts[0].GetRange(), MakeContinuedTimeRange(TR * 3, LastValidTick));
+		UTEST_FALSE("Context0_Jumped", Contexts[0].HasJumped());
+		UTEST_EQUAL("NumLoops", Manager.GetNumLoopsCompleted(), 4);
+		UTEST_EQUAL("Status", Manager.GetPlaybackStatus(), EMovieScenePlayerStatus::Stopped);
+
+	}
+
+	Manager.ResetNumLoopsCompleted();
+	Manager.SetPlayDirection(EPlayDirection::Backwards);
+	Manager.SetPlaybackStatus(EMovieScenePlayerStatus::Playing);
+
+	{
+		UTEST_EQUAL("NumLoops", Manager.GetNumLoopsCompleted(), 0);
+
+		Manager.SetCurrentTime(30 * 10);
+
+		Contexts.Reset();
+		Manager.UpdateTo(30 * 5, Contexts);
+		UTEST_EQUAL("NumContexts", Contexts.Num(), 1);
+		UTEST_EQUAL("Context0", Contexts[0].GetRange(), MakeHullTimeRange(TR * 5, TR * 10));
+		UTEST_FALSE("Context0_Jumped", Contexts[0].HasJumped());
+		UTEST_EQUAL("Context0_Backwards", Contexts[0].GetDirection(), EPlayDirection::Backwards);
+		UTEST_EQUAL("NumLoops", Manager.GetNumLoopsCompleted(), 0);
+		UTEST_EQUAL("Status", Manager.GetPlaybackStatus(), EMovieScenePlayerStatus::Playing);
+
+		Contexts.Reset();
+		Manager.UpdateTo(30 * -2, Contexts);
+		UTEST_EQUAL("NumContexts", Contexts.Num(), 2);
+		UTEST_EQUAL("Context0", Contexts[0].GetRange(), MakeDiscreteTimeRange(0, TR * 5));
+		UTEST_FALSE("Context0_Jumped", Contexts[0].HasJumped());
+		UTEST_EQUAL("Context0_Backwards", Contexts[0].GetDirection(), EPlayDirection::Backwards);
+		UTEST_EQUAL("Context1", Contexts[1].GetRange(), MakeHullTimeRange(TR * 8, LastValidTick));
+		UTEST_TRUE ("Context1_Jumped", Contexts[1].HasJumped());
+		UTEST_EQUAL("Context1_Backwards", Contexts[1].GetDirection(), EPlayDirection::Backwards);
+		UTEST_EQUAL("NumLoops", Manager.GetNumLoopsCompleted(), 1);
+		UTEST_EQUAL("Status", Manager.GetPlaybackStatus(), EMovieScenePlayerStatus::Playing);
+
+		Contexts.Reset();
+		Manager.UpdateTo(30 * -13, Contexts);
+		UTEST_EQUAL("NumContexts", Contexts.Num(), 2);
+		UTEST_EQUAL("Context0", Contexts[0].GetRange(), MakeDiscreteTimeRange(0, TR * 8));
+		UTEST_FALSE("Context0_Jumped", Contexts[0].HasJumped());
+		UTEST_EQUAL("Context0_Backwards", Contexts[0].GetDirection(), EPlayDirection::Backwards);
+		UTEST_EQUAL("Context1", Contexts[1].GetRange(), MakeHullTimeRange(TR * 7, LastValidTick));
+		UTEST_TRUE ("Context1_Jumped", Contexts[1].HasJumped());
+		UTEST_EQUAL("Context1_Backwards", Contexts[1].GetDirection(), EPlayDirection::Backwards);
+		UTEST_EQUAL("NumLoops", Manager.GetNumLoopsCompleted(), 3);
+		UTEST_EQUAL("Status", Manager.GetPlaybackStatus(), EMovieScenePlayerStatus::Playing);
+
+		Contexts.Reset();
+		Manager.UpdateTo(30 * -1, Contexts);
+		UTEST_EQUAL("NumContexts", Contexts.Num(), 1);
+		UTEST_EQUAL("Context0", Contexts[0].GetRange(), MakeDiscreteTimeRange(0, TR * 7));
+		UTEST_FALSE("Context0_Jumped", Contexts[0].HasJumped());
+		UTEST_EQUAL("Context0_Backwards", Contexts[0].GetDirection(), EPlayDirection::Backwards);
+		UTEST_EQUAL("NumLoops", Manager.GetNumLoopsCompleted(), 4);
+		UTEST_EQUAL("Status", Manager.GetPlaybackStatus(), EMovieScenePlayerStatus::Stopped);
+	}
+
+	return true;
+}
+
+IMPLEMENT_SIMPLE_AUTOMATION_TEST(
+		FMovieScenePlaybackManagerTestLoopingDissectAll, 
+		"System.Engine.Sequencer.PlaybackManager.Looping.DissectAll", 
+		EAutomationTestFlags::EditorContext | EAutomationTestFlags::EngineFilter)
+bool FMovieScenePlaybackManagerTestLoopingDissectAll::RunTest(const FString& Parameters)
+{
+	using namespace UE::MovieScene;
+	using namespace UE::MovieScene::Test;
+
+	FMakeTestSequenceParams MakeParams;
+	UMovieSceneSequence* Sequence = MakeTestSequence(MakeParams);
+
+	FMovieScenePlaybackManager Manager(Sequence);
+	Manager.SetDissectLooping(EMovieSceneLoopDissection::DissectAll);
+	Manager.SetNumLoopsToPlay(4);
+	Manager.SetPlaybackStatus(EMovieScenePlayerStatus::Playing);
+
+	const int32 TR = 60000;
+	const FFrameTime LastValidTick(10 * TR - 1);
+
+	FMovieScenePlaybackManager::FContexts Contexts;
+
+	{
+
+		UTEST_EQUAL("NumLoops", Manager.GetNumLoopsCompleted(), 0);
+
+		Contexts.Reset();
+		Manager.UpdateTo(30 * 5, Contexts);
+		UTEST_EQUAL("NumContexts", Contexts.Num(), 1);
+		UTEST_EQUAL("Context0", Contexts[0].GetRange(), MakeHullTimeRange(0, TR * 5));
+		UTEST_FALSE("Context0_Jumped", Contexts[0].HasJumped());
+		UTEST_EQUAL("NumLoops", Manager.GetNumLoopsCompleted(), 0);
+		UTEST_EQUAL("Status", Manager.GetPlaybackStatus(), EMovieScenePlayerStatus::Playing);
+
+		Contexts.Reset();
+		Manager.UpdateTo(30 * 12, Contexts);
+		UTEST_EQUAL("NumContexts", Contexts.Num(), 2);
+		UTEST_EQUAL("Context0", Contexts[0].GetRange(), MakeContinuedTimeRange(TR * 5, LastValidTick));
+		UTEST_FALSE("Context0_Jumped", Contexts[0].HasJumped());
+		UTEST_EQUAL("Context1", Contexts[1].GetRange(), MakeHullTimeRange(0, TR * 2));
+		UTEST_TRUE ("Context1_Jumped", Contexts[1].HasJumped());
+		UTEST_EQUAL("NumLoops", Manager.GetNumLoopsCompleted(), 1);
+		UTEST_EQUAL("Status", Manager.GetPlaybackStatus(), EMovieScenePlayerStatus::Playing);
+
+		Contexts.Reset();
+		Manager.UpdateTo(30 * 23, Contexts);
+		UTEST_EQUAL("NumContexts", Contexts.Num(), 3);
+		UTEST_EQUAL("Context0", Contexts[0].GetRange(), MakeContinuedTimeRange(TR * 2, LastValidTick));
+		UTEST_FALSE("Context0_Jumped", Contexts[0].HasJumped());
+		UTEST_EQUAL("Context1", Contexts[1].GetRange(), MakeHullTimeRange(0, LastValidTick));
+		UTEST_TRUE ("Context1_Jumped", Contexts[1].HasJumped());
+		UTEST_EQUAL("Context2", Contexts[2].GetRange(), MakeHullTimeRange(0, TR * 3));
+		UTEST_TRUE ("Context2_Jumped", Contexts[2].HasJumped());
+		UTEST_EQUAL("NumLoops", Manager.GetNumLoopsCompleted(), 3);
+		UTEST_EQUAL("Status", Manager.GetPlaybackStatus(), EMovieScenePlayerStatus::Playing);
+
+		Contexts.Reset();
+		Manager.UpdateTo(30 * 11, Contexts);
+		UTEST_EQUAL("NumContexts", Contexts.Num(), 1);
+		UTEST_EQUAL("Context0", Contexts[0].GetRange(), MakeContinuedTimeRange(TR * 3, LastValidTick));
+		UTEST_FALSE("Context0_Jumped", Contexts[0].HasJumped());
+		UTEST_EQUAL("NumLoops", Manager.GetNumLoopsCompleted(), 4);
+		UTEST_EQUAL("Status", Manager.GetPlaybackStatus(), EMovieScenePlayerStatus::Stopped);
+
+	}
+
+	Manager.ResetNumLoopsCompleted();
+	Manager.SetPlayDirection(EPlayDirection::Backwards);
+	Manager.SetPlaybackStatus(EMovieScenePlayerStatus::Playing);
+
+	{
+		UTEST_EQUAL("NumLoops", Manager.GetNumLoopsCompleted(), 0);
+
+		Manager.SetCurrentTime(30 * 10);
+
+		Contexts.Reset();
+		Manager.UpdateTo(30 * 5, Contexts);
+		UTEST_EQUAL("NumContexts", Contexts.Num(), 1);
+		UTEST_EQUAL("Context0", Contexts[0].GetRange(), MakeHullTimeRange(TR * 5, TR * 10));
+		UTEST_FALSE("Context0_Jumped", Contexts[0].HasJumped());
+		UTEST_EQUAL("NumLoops", Manager.GetNumLoopsCompleted(), 0);
+		UTEST_EQUAL("Status", Manager.GetPlaybackStatus(), EMovieScenePlayerStatus::Playing);
+
+		Contexts.Reset();
+		Manager.UpdateTo(30 * -2, Contexts);
+		UTEST_EQUAL("NumContexts", Contexts.Num(), 2);
+		UTEST_EQUAL("Context0", Contexts[0].GetRange(), MakeDiscreteTimeRange(0, TR * 5));
+		UTEST_FALSE("Context0_Jumped", Contexts[0].HasJumped());
+		UTEST_EQUAL("Context0_Backwards", Contexts[0].GetDirection(), EPlayDirection::Backwards);
+		UTEST_EQUAL("Context1", Contexts[1].GetRange(), MakeHullTimeRange(TR * 8, LastValidTick));
+		UTEST_TRUE ("Context1_Jumped", Contexts[1].HasJumped());
+		UTEST_EQUAL("Context1_Backwards", Contexts[1].GetDirection(), EPlayDirection::Backwards);
+		UTEST_EQUAL("NumLoops", Manager.GetNumLoopsCompleted(), 1);
+		UTEST_EQUAL("Status", Manager.GetPlaybackStatus(), EMovieScenePlayerStatus::Playing);
+
+		Contexts.Reset();
+		Manager.UpdateTo(30 * -13, Contexts);
+		UTEST_EQUAL("NumContexts", Contexts.Num(), 3);
+		UTEST_EQUAL("Context0", Contexts[0].GetRange(), MakeDiscreteTimeRange(0, TR * 8));
+		UTEST_FALSE("Context0_Jumped", Contexts[0].HasJumped());
+		UTEST_EQUAL("Context0_Backwards", Contexts[0].GetDirection(), EPlayDirection::Backwards);
+		UTEST_EQUAL("Context1", Contexts[1].GetRange(), MakeHullTimeRange(0, LastValidTick));
+		UTEST_TRUE ("Context1_Jumped", Contexts[1].HasJumped());
+		UTEST_EQUAL("Context1_Backwards", Contexts[1].GetDirection(), EPlayDirection::Backwards);
+		UTEST_EQUAL("Context2", Contexts[2].GetRange(), MakeHullTimeRange(TR * 7, LastValidTick));
+		UTEST_TRUE ("Context2_Jumped", Contexts[2].HasJumped());
+		UTEST_EQUAL("Context2_Backwards", Contexts[2].GetDirection(), EPlayDirection::Backwards);
+		UTEST_EQUAL("NumLoops", Manager.GetNumLoopsCompleted(), 3);
+		UTEST_EQUAL("Status", Manager.GetPlaybackStatus(), EMovieScenePlayerStatus::Playing);
+
+		Contexts.Reset();
+		Manager.UpdateTo(30 * -1, Contexts);
+		UTEST_EQUAL("NumContexts", Contexts.Num(), 1);
+		UTEST_EQUAL("Context0", Contexts[0].GetRange(), MakeDiscreteTimeRange(0, TR * 7));
+		UTEST_FALSE("Context0_Jumped", Contexts[0].HasJumped());
+		UTEST_EQUAL("Context0_Backwards", Contexts[0].GetDirection(), EPlayDirection::Backwards);
+		UTEST_EQUAL("NumLoops", Manager.GetNumLoopsCompleted(), 4);
+		UTEST_EQUAL("Status", Manager.GetPlaybackStatus(), EMovieScenePlayerStatus::Stopped);
+	}
+
+	return true;
+}
+
+IMPLEMENT_SIMPLE_AUTOMATION_TEST(
+		FMovieScenePlaybackManagerTestUpdateToEnd, 
+		"System.Engine.Sequencer.PlaybackManager.UpdateToEnd", 
+		EAutomationTestFlags::EditorContext | EAutomationTestFlags::EngineFilter)
+bool FMovieScenePlaybackManagerTestUpdateToEnd::RunTest(const FString& Parameters)
+{
+	using namespace UE::MovieScene;
+	using namespace UE::MovieScene::Test;
+
+	FMakeTestSequenceParams MakeParams;
+	UMovieSceneSequence* Sequence = MakeTestSequence(MakeParams);
+
+	FMovieScenePlaybackManager Manager(Sequence);
+	Manager.SetNumLoopsToPlay(1);
+
+	const int32 TR = 60000;
+	const FFrameTime LastValidTick(10 * TR - 1);
+	const FFrameTime LastValidFrame = ConvertFrameTime(LastValidTick, FFrameRate(60000, 1), FFrameRate(30, 1));
+
+	FMovieScenePlaybackManager::FContexts Contexts;
+
+	TArray<EMovieSceneLoopDissection> Dissections{
+		EMovieSceneLoopDissection::None,
+		EMovieSceneLoopDissection::DissectOne,
+		EMovieSceneLoopDissection::DissectAll
+	};
+
+	for (EMovieSceneLoopDissection Dissection : Dissections)
+	{
+		Manager.SetDissectLooping(Dissection);
+
+		Manager.ResetNumLoopsCompleted();
+		Manager.SetPlayDirection(EPlayDirection::Forwards);
+		Manager.SetPlaybackStatus(EMovieScenePlayerStatus::Playing);
+		Manager.SetCurrentTime(0);
+
+		{
+			UTEST_EQUAL("NumLoops", Manager.GetNumLoopsCompleted(), 0);
+
+			Contexts.Reset();
+			Manager.UpdateTo(LastValidFrame, Contexts);
+			UTEST_EQUAL("NumContexts", Contexts.Num(), 1);
+			UTEST_EQUAL("Context0", Contexts[0].GetRange(), MakeHullTimeRange(0, LastValidTick));
+			UTEST_FALSE("Context0_Jumped", Contexts[0].HasJumped());
+			UTEST_EQUAL("Context0_Forwards", Contexts[0].GetDirection(), EPlayDirection::Forwards);
+			UTEST_EQUAL("NumLoops", Manager.GetNumLoopsCompleted(), 1);
+			UTEST_EQUAL("Status", Manager.GetPlaybackStatus(), EMovieScenePlayerStatus::Stopped);
+		}
+
+		Manager.ResetNumLoopsCompleted();
+		Manager.SetPlayDirection(EPlayDirection::Backwards);
+		Manager.SetPlaybackStatus(EMovieScenePlayerStatus::Playing);
+		Manager.SetCurrentTime(LastValidFrame);
+
+		{
+			UTEST_EQUAL("NumLoops", Manager.GetNumLoopsCompleted(), 0);
+
+			Contexts.Reset();
+			Manager.UpdateTo(0, Contexts);
+			UTEST_EQUAL("NumContexts", Contexts.Num(), 1);
+			UTEST_EQUAL("Context0", Contexts[0].GetRange(), MakeHullTimeRange(0, LastValidTick));
+			UTEST_FALSE("Context0_Jumped", Contexts[0].HasJumped());
+			UTEST_EQUAL("Context0_Backwards", Contexts[0].GetDirection(), EPlayDirection::Backwards);
+			UTEST_EQUAL("NumLoops", Manager.GetNumLoopsCompleted(), 1);
+			UTEST_EQUAL("Status", Manager.GetPlaybackStatus(), EMovieScenePlayerStatus::Stopped);
+		}
+	}
+
+	return true;
+}
+
diff --git a/Engine/Source/Runtime/MovieScene/Public/Evaluation/MovieScenePlaybackManager.h b/Engine/Source/Runtime/MovieScene/Public/Evaluation/MovieScenePlaybackManager.h
new file mode 100644
--- /dev/null
+++ b/Engine/Source/Runtime/MovieScene/Public/Evaluation/MovieScenePlaybackManager.h
@@ -0,0 +1,225 @@
+// Copyright Epic Games, Inc. All Rights Reserved.
+
+#pragma once
+
+#include "EntitySystem/MovieSceneSharedPlaybackState.h"
+#include "Evaluation/MovieScenePlayback.h"
+#include "Evaluation/MovieSceneSequenceTransform.h"
+#include "MovieSceneFwd.h"
+
+class UMovieSceneSequence;
+
+namespace UE::MovieScene { struct FSharedPlaybackState; }
+
+/**
+ * Describes whether to dissect looping when playing back a sequence.
+ *
+ * Looping dissection means that, when a sequence loops, more than one evaluation context
+ * is returned by the FMovieScenePlaybackManager::Update() methods: one context for updating 
+ * to the "last valid time", and one context for updating through the first few frames of 
+ * the next loop.
+ */
+enum class EMovieSceneLoopDissection
+{
+	/** 
+	 * Do no dissect loop. Jump directly to the current time of the next loop, using an
+	 * evaluation range from the first frame to the current time.
+	 * WARNING: this means that the last few frames of the previous loop won't be evaluated!
+	 */
+	None,
+
+	/**
+	 * Only dissect one loop. Emit an evaluation range for the last few frames of the 
+	 * current loop, and another evaluation range for the first few frames of the next 
+	 * loop.
+	 */
+	DissectOne,
+
+	/**
+	 * Dissect all loops. As per DissectOne, but also emit evaluation ranges for entire loops
+	 * if the delta-time is large enough, and/or the sequence is short enough, that we might
+	 * have gone through entire extra loops.
+	 */
+	DissectAll
+};
+
+/**
+ * A utility class that can manage a playing sequence's status and current time, while also
+ * handling looping, ping-ponging, and other playback modes.
+ *
+ * All public APIs take and return times in display rate (e.g. 30fps frames). Internally, 
+ * everything is treated as ticks (e.g. 60000 ticks/sec).
+ */
+class FMovieScenePlaybackManager
+{
+public:
+
+	using FContexts = TArray<FMovieSceneContext, TInlineAllocator<2>>;
+
+	/** Creates a default playback manager. */
+	MOVIESCENE_API FMovieScenePlaybackManager();
+
+	/** Creates a playback manager immediately initialized with the given sequence. */
+	MOVIESCENE_API FMovieScenePlaybackManager(UMovieSceneSequence* InSequence);
+
+	/** 
+	 * Initializes the playback manager for the given sequence.
+	 * This resets all playback settings (play rate, start/end offsets, etc.) to their default
+	 * values.
+	 */
+	MOVIESCENE_API void Initialize(UMovieSceneSequence* InSequence);
+
+	/**
+	 * Updates the playback state and returns the evaluation contexts to use for evaluating
+	 * the sequence. More than one context may be returned if the sequence loops once or more,
+	 * and "looping dissection" is enabled.
+	 *
+	 * @param InDeltaSeconds  The delta-time to update with
+	 * @param OutContexts     One or more update contexts to evaluate the sequence
+	 */
+	MOVIESCENE_API void Update(float InDeltaSeconds, FContexts& OutContexts);
+
+	/**
+	 * As per the other Update method, but takes a time (in display rate) instead of seconds.
+	 */
+	MOVIESCENE_API void UpdateTo(const FFrameTime NextTime, FContexts& OutContexts);
+
+	/**
+	 * Returns an evaluation context for the current time, i.e. using a zero-width evaluation
+	 * range set around the current time.
+	 */
+	MOVIESCENE_API FMovieSceneContext UpdateAtCurrentTime() const;
+
+	/** Gets the current playback position, in display rate. */
+	MOVIESCENE_API FFrameTime GetCurrentTime() const;
+
+	/** Sets the current playback position, in display rate. */
+	MOVIESCENE_API void SetCurrentTime(const FFrameTime& InFrameTime);
+
+	/** Get the effective playback range, in display rate, taking into account start/end offsets. */
+	MOVIESCENE_API TRange<FFrameTime> GetEffectivePlaybackRange() const;
+
+	/** Get the effective playback start time, in display rate, taking into account any start offset. */
+	MOVIESCENE_API FFrameTime GetEffectiveStartTime() const;
+	/** Get the effective playback end time, in display rate, taking into account any end offset. */
+	MOVIESCENE_API FFrameTime GetEffectiveEndTime() const;
+
+public:
+
+	/** Gets the display rate of the sequence. */
+	FFrameRate GetDisplayRate() const { return DisplayRate; }
+	/** Gets the tick resolution of the sequence. */
+	FFrameRate GetTickResolution() const { return PlaybackPosition.GetOutputRate(); }
+
+	/** Gets whether looping dissection is enabled (see EMovieSceneLoopDissection). */
+	EMovieSceneLoopDissection GetDissectLooping() const { return DissectLooping; }
+	/** Sets whether looping dissection is enabled (see EMovieSceneLoopDissection). */
+	void SetDissectLooping(EMovieSceneLoopDissection InDissectLooping) { DissectLooping = InDissectLooping; }
+
+	/** Get the playback status. */
+	EMovieScenePlayerStatus::Type GetPlaybackStatus() const { return PlaybackStatus; }
+	/** Set the playback status. */
+	void SetPlaybackStatus(EMovieScenePlayerStatus::Type InPlaybackStatus) { PlaybackStatus = InPlaybackStatus; }
+
+	/** Gets the number of loops to play before ending playback. */
+	int32 GetNumLoopsToPlay() const { return NumLoopsToPlay; }
+	/** Sets the number of loops to play before ending playback. */
+	MOVIESCENE_API void SetNumLoopsToPlay(int32 InNumLoopsToPlay);
+
+	/** Returns the current number of loops completed so far. */
+	int32 GetNumLoopsCompleted() const { return NumLoopsCompleted; }
+	/** Reset the number of completed loops to zero. */
+	void ResetNumLoopsCompleted() { NumLoopsCompleted = 0; }
+
+	/** Gets the play-rate. */
+	double GetPlayRate() const { return PlayRate; }
+	/** Sets the play-rate. */
+	void SetPlayRate(double InPlayRate) { PlayRate = InPlayRate; }
+
+	/** Gets whether the sequence should be playing forwards. */
+	bool IsPlayingForward() const { return PlayDirection == EPlayDirection::Forwards; } 
+	/** Gets whether the sequence should be playing in reverse. */
+	bool IsPlayingBackward() const { return PlayDirection == EPlayDirection::Backwards; } 
+	/** Gets the playback direction. */
+	EPlayDirection GetPlayDirection() const { return PlayDirection; }
+	/** Sets the playback direction. */
+	void SetPlayDirection(EPlayDirection InPlayDirection) { PlayDirection = InPlayDirection; }
+	/** Reverses the playback direction. */
+	MOVIESCENE_API void ReversePlayDirection();
+
+	/** Gets whether the sequence should ping-pong between forwards and backwards playback. */
+	bool IsPingPongPlayback() const { return bPingPongPlayback; }
+	/** 
+	 * Sets whether the sequence should ping-pong between forwards and backwards playback.
+	 * Note that a "loop" is still one play through of the sequence, whether forwards or 
+	 * backwards, so odd numbers of loops corresponds to a "ping" without a "pong".
+	 */
+	MOVIESCENE_API void SetPingPongPlayback(bool bInPingPongPlayback);
+
+	/** Gets the start offset of the playback in display rate. */
+	MOVIESCENE_API FFrameTime GetStartOffset() const;
+	/** Gets the end offset of the playback in display rate. */
+	MOVIESCENE_API FFrameTime GetEndOffset() const;
+	/** Sets the start offset of the playback in display rate. */
+	MOVIESCENE_API void SetStartOffset(const FFrameTime& InStartOffset);
+	/** Sets the end offset of the playback in display rate. */
+	MOVIESCENE_API void SetEndOffset(const FFrameTime& InEndOffset);
+	/** Sets the end time of the playback in display rate. */
+	MOVIESCENE_API void SetEndOffsetAsTime(const FFrameTime& InEndTime);
+
+	/** Gets whether the playback time transform should be applied to time updates. */
+	bool ShouldTransformPlaybackTime() const { return bTransformPlaybackTime; }
+	/** Sets whether the playback time transform should be applied to time updates. */
+	void SetTransformPlaybackTime(bool bInTransformPlaybackTime) { bTransformPlaybackTime = bInTransformPlaybackTime; }
+
+	/**
+	 * Gets the time transform to apply to the current time when updating.
+	 * Only used if ShouldTransformPlaybackTime is true.
+	 */
+	const FMovieSceneSequenceTransform& GetPlaybackTimeTransform() const { return TimeTransform; }
+
+	/*
+	 * Sets the time transform to apply to the current time when updating.
+	 * Only used if ShouldTransformPlaybackTime is true.
+	 */
+	void SetPlaybackTimeTransform(const FMovieSceneSequenceTransform& InTimeTransform) { TimeTransform = InTimeTransform; }
+
+private:
+
+	void InternalUpdateToTick(const FFrameNumber NextTick, FContexts& OutContexts);
+
+	void ResetPlaybackSettings();
+
+	void SetStartAndEndOffsetTicks(FFrameNumber InStartOffsetTicks, FFrameNumber InEndOffsetTicks);
+
+	FFrameNumber GetLastValidTick() const;
+
+private:
+
+	FMovieScenePlaybackPosition PlaybackPosition;
+
+	FFrameRate DisplayRate;
+
+	EMovieScenePlayerStatus::Type PlaybackStatus;
+	EPlayDirection PlayDirection;
+
+	FFrameNumber SequenceStartTick;
+	FFrameNumber SequenceEndTick;
+
+	FFrameNumber StartOffsetTicks;
+	FFrameNumber EndOffsetTicks;
+
+	int32 NumLoopsToPlay = 1;
+	int32 NumLoopsCompleted = 0;
+
+	double PlayRate = 1.0;
+
+	FMovieSceneSequenceTransform TimeTransform;
+
+	EMovieSceneLoopDissection DissectLooping = EMovieSceneLoopDissection::DissectOne;
+
+	bool bSupportsTimeWarp = true;
+	bool bPingPongPlayback = false;
+	bool bTransformPlaybackTime = false;
+};
+
