Fix undef'ed LOCTEXT_NAMESPACE

#rnx

diff --git a/Engine/Source/Runtime/MovieScene/Private/Tests/MovieScenePlaybackManagerTests.cpp b/Engine/Source/Runtime/MovieScene/Private/Tests/MovieScenePlaybackManagerTests.cpp
--- a/Engine/Source/Runtime/MovieScene/Private/Tests/MovieScenePlaybackManagerTests.cpp
+++ b/Engine/Source/Runtime/MovieScene/Private/Tests/MovieScenePlaybackManagerTests.cpp
@@ -526,3 +526,5 @@ bool FMovieScenePlaybackManagerTestUpdateToEnd::RunTest(const FString& Parameter
 	return true;
 }
 
+#undef LOCTEXT_NAMESPACE
+
