Sequencer: Fix evaluating the wrong time for 'evaluate nearest section' leading edges, and rendering of loop points

diff --git a/Engine/Source/Runtime/MovieScene/Private/Compilation/MovieSceneEvaluationTreePopulationRules.cpp b/Engine/Source/Runtime/MovieScene/Private/Compilation/MovieSceneEvaluationTreePopulationRules.cpp
--- a/Engine/Source/Runtime/MovieScene/Private/Compilation/MovieSceneEvaluationTreePopulationRules.cpp
+++ b/Engine/Source/Runtime/MovieScene/Private/Compilation/MovieSceneEvaluationTreePopulationRules.cpp
@@ -155,7 +155,7 @@ void FEvaluationTreePopulationRules::PopulateNearestSection(TArrayView<UMovieSce
 			}
 			else
 			{
-				ForcedTime = UE::MovieScene::DiscreteInclusiveLower(NodeToCopy.Range())-1;
+				ForcedTime = UE::MovieScene::DiscreteInclusiveLower(NodeToCopy.Range());
 			}
 
 			if (NodeToCopy)
diff --git a/Engine/Source/Runtime/MovieScene/Private/Evaluation/MovieSceneSequenceTransform.cpp b/Engine/Source/Runtime/MovieScene/Private/Evaluation/MovieSceneSequenceTransform.cpp
--- a/Engine/Source/Runtime/MovieScene/Private/Evaluation/MovieSceneSequenceTransform.cpp
+++ b/Engine/Source/Runtime/MovieScene/Private/Evaluation/MovieSceneSequenceTransform.cpp
@@ -373,12 +373,24 @@ TRange<FFrameTime> FMovieSceneNestedSequenceTransform::ComputeTraversedHull(cons
 bool FMovieSceneNestedSequenceTransform::ExtractBoundariesWithinRange(const TRange<FFrameTime>& Range, const TFunctionRef<bool(FFrameTime)>& InVisitor) const
 {
 	EMovieSceneTimeWarpType WarpType = TimeScale.GetType();
+
+	auto VisitorOffsetWrapper = [this, Range, InVisitor](FFrameTime Time)
+	{
+		// Factor the loop offset
+		Time -= this->Offset;
+		if (Range.Contains(Time))
+		{
+			return InVisitor(Time);
+		}
+		return true;
+	};
+
 	switch (WarpType)
 	{
 	case EMovieSceneTimeWarpType::Loop:
-		return TimeScale.AsLoop().ExtractBoundariesWithinRange(Range, InVisitor);
+		return TimeScale.AsLoop().ExtractBoundariesWithinRange(Range, VisitorOffsetWrapper);
 	case EMovieSceneTimeWarpType::LoopFloat:
-		return TimeScale.AsLoopFloat().ExtractBoundariesWithinRange(Range, InVisitor);
+		return TimeScale.AsLoopFloat().ExtractBoundariesWithinRange(Range, VisitorOffsetWrapper);
 	default:
 		return false;
 	}
