Speculative mac compile error fix

diff --git a/Engine/Source/Runtime/MovieScene/Public/EntitySystem/MovieSceneEntitySystemTask.h b/Engine/Source/Runtime/MovieScene/Public/EntitySystem/MovieSceneEntitySystemTask.h
--- a/Engine/Source/Runtime/MovieScene/Public/EntitySystem/MovieSceneEntitySystemTask.h
+++ b/Engine/Source/Runtime/MovieScene/Public/EntitySystem/MovieSceneEntitySystemTask.h
@@ -774,7 +774,7 @@ struct TEntityTaskComponentsImpl<TIntegerSequence<int, Indices...>, T...>
 
 	TEntityTaskComponents<T..., FFilterMatchPassthrough> PassthroughFilter(const FEntityComponentFilter& InFilter) const
 	{
-		return TEntityTaskComponents<T..., FFilterMatchPassthrough >(CommonParams, Accessors.template Get<Indices>()..., FFilterMatchPassthrough(InFilter));
+		return TEntityTaskComponents<T..., FFilterMatchPassthrough >(CommonParams, Accessors.template Get<Indices>()..., FFilterMatchPassthrough{ InFilter });
 	}
 
 	bool HasBeenWrittenToSince(uint32 InSystemVersion)
