Fixed incorrect indexing into sparse array in FUpdateGroupsTask::PreTask

diff --git a/Engine/Source/Runtime/MovieScene/Private/EntitySystem/MovieSceneEntityGroupingSystem.cpp b/Engine/Source/Runtime/MovieScene/Private/EntitySystem/MovieSceneEntityGroupingSystem.cpp
--- a/Engine/Source/Runtime/MovieScene/Private/EntitySystem/MovieSceneEntityGroupingSystem.cpp
+++ b/Engine/Source/Runtime/MovieScene/Private/EntitySystem/MovieSceneEntityGroupingSystem.cpp
@@ -134,12 +134,12 @@ struct FUpdateGroupsTask
 
 		// Run the handlers' pre-task callback.
 		
-		for (int32 HandlerIndex = 0; HandlerIndex < System->GroupHandlers.Num(); ++HandlerIndex)
+		for (auto It = System->GroupHandlers.CreateIterator(); It; ++It)
 		{
-			const UMovieSceneEntityGroupingSystem::FEntityGroupingHandlerInfo& HandlerInfo = System->GroupHandlers[HandlerIndex];
+			const UMovieSceneEntityGroupingSystem::FEntityGroupingHandlerInfo& HandlerInfo = *It;
 			if (HandlerInfo.Handler.IsValid())
 			{
-				FEntityGroupBuilder Builder(System, FEntityGroupingPolicyKey{ HandlerIndex });
+				FEntityGroupBuilder Builder(System, FEntityGroupingPolicyKey{ It.GetIndex() });
 				HandlerInfo.Handler->PreTask(&Builder);
 			}
 		}
@@ -183,12 +183,12 @@ struct FUpdateGroupsTask
 	void PostTask()
 	{
 		// Run the handlers' post-task callback.
-		for (int32 HandlerIndex = 0; HandlerIndex < System->GroupHandlers.Num(); ++HandlerIndex)
+		for (auto It = System->GroupHandlers.CreateIterator(); It; ++It)
 		{
-			const UMovieSceneEntityGroupingSystem::FEntityGroupingHandlerInfo& HandlerInfo = System->GroupHandlers[HandlerIndex];
+			const UMovieSceneEntityGroupingSystem::FEntityGroupingHandlerInfo& HandlerInfo = *It;
 			if (HandlerInfo.Handler.IsValid())
 			{
-				FEntityGroupBuilder Builder(System, FEntityGroupingPolicyKey{ HandlerIndex });
+				FEntityGroupBuilder Builder(System, FEntityGroupingPolicyKey{ It.GetIndex() });
 				HandlerInfo.Handler->PostTask(bFreeGroupIDs, &Builder);
 			}
 		}
