Add support for EMovieSceneTimeWarpType::FixedPlayRate in FMovieSceneTimeWarpVariant::ScaleBy so that we can support reverse play clips in Sequencer Playlists

diff --git a/Engine/Source/Runtime/MovieScene/Private/Variants/MovieSceneTimeWarpVariant.cpp b/Engine/Source/Runtime/MovieScene/Private/Variants/MovieSceneTimeWarpVariant.cpp
--- a/Engine/Source/Runtime/MovieScene/Private/Variants/MovieSceneTimeWarpVariant.cpp
+++ b/Engine/Source/Runtime/MovieScene/Private/Variants/MovieSceneTimeWarpVariant.cpp
@@ -114,6 +114,13 @@ void FMovieSceneTimeWarpVariant::ScaleBy(double ScaleFactor)
 {
 	switch(GetType())
 	{
+	case EMovieSceneTimeWarpType::FixedPlayRate:
+		{
+			const double NewPlayRate = AsFixedPlayRate() * ScaleFactor;
+			// Use Set() here in order to properly handle (near-)zero play rates
+			Set(NewPlayRate);
+		}
+		break;
 	case EMovieSceneTimeWarpType::Custom:
 		if (UMovieSceneTimeWarpGetter* Custom = AsCustom())
 		{
@@ -515,4 +522,4 @@ bool FMovieSceneTimeWarpVariant::ImportTextItem(const TCHAR*& Buffer, int32 Port
 		ensureMsgf(false, TEXT("Unimplemented type found when importing text!"));
 		return false;
 	}
-}
\ No newline at end of file
+}
