Sequencer: Fixed hbias not working correctly for Ignored entities

diff --git a/Engine/Source/Runtime/MovieScene/Public/EntitySystem/MovieScenePropertyComponentHandler.h b/Engine/Source/Runtime/MovieScene/Public/EntitySystem/MovieScenePropertyComponentHandler.h
--- a/Engine/Source/Runtime/MovieScene/Public/EntitySystem/MovieScenePropertyComponentHandler.h
+++ b/Engine/Source/Runtime/MovieScene/Public/EntitySystem/MovieScenePropertyComponentHandler.h
@@ -376,6 +376,7 @@ struct TPropertyComponentHandlerImpl<PropertyTraits, TPropertyMetaData<MetaDataT
 		.ReadAllOf(Definition.GetMetaDataComponent<MetaDataTypes>(MetaDataIndices)...)
 		.ReadAllOf(Composites[CompositeIndices].ComponentTypeID.ReinterpretCast<CompositeTypes>()...)
 		.FilterAll({ Definition.PropertyType })
+		.FilterNone({ BuiltInComponents->Tags.Ignored })
 		.SetStat(Definition.StatID)
 		.SetDesiredThread(Linker->EntityManager.GetGatherThread())
 		.template Fork_PerAllocation<CompleteSetterTask>(&Linker->EntityManager, TaskScheduler, Definition.CustomPropertyRegistration);
@@ -399,6 +400,7 @@ struct TPropertyComponentHandlerImpl<PropertyTraits, TPropertyMetaData<MetaDataT
 				.ReadAnyOf(Composites[CompositeIndices].ComponentTypeID.ReinterpretCast<CompositeTypes>()...)
 				.FilterAny({ CompletePropertyMask })
 				.FilterAll({ Definition.PropertyType })
+				.FilterNone({ BuiltInComponents->Tags.Ignored })
 				.FilterOut(CompletePropertyMask)
 				.SetStat(Definition.StatID)
 				.SetDesiredThread(Linker->EntityManager.GetGatherThread())
@@ -417,6 +419,7 @@ struct TPropertyComponentHandlerImpl<PropertyTraits, TPropertyMetaData<MetaDataT
 		.ReadAllOf(Definition.GetMetaDataComponent<MetaDataTypes>(MetaDataIndices)...)
 		.ReadAllOf(Composites[CompositeIndices].ComponentTypeID.ReinterpretCast<CompositeTypes>()...)
 		.FilterAll({ Definition.PropertyType })
+		.FilterNone({ BuiltInComponents->Tags.Ignored })
 		.SetStat(Definition.StatID)
 		.SetDesiredThread(Linker->EntityManager.GetGatherThread())
 		.template Dispatch_PerAllocation<CompleteSetterTask>(&Linker->EntityManager, InPrerequisites, &Subsequents, Definition.CustomPropertyRegistration);
@@ -441,6 +444,7 @@ struct TPropertyComponentHandlerImpl<PropertyTraits, TPropertyMetaData<MetaDataT
 				.FilterAny({ CompletePropertyMask })
 				.FilterAll({ Definition.PropertyType })
 				.FilterOut(CompletePropertyMask)
+				.FilterNone({ BuiltInComponents->Tags.Ignored })
 				.SetStat(Definition.StatID)
 				.SetDesiredThread(Linker->EntityManager.GetGatherThread())
 				.template Dispatch_PerAllocation<PartialSetterTask>(&Linker->EntityManager, InPrerequisites, &Subsequents, Definition.CustomPropertyRegistration, Composites);
