- Movied previously committed squencer changes for music mode into a new Musical Mode plugin

diff --git a/Engine/Source/Runtime/MovieScene/Private/MovieSceneTimeHelpers.cpp b/Engine/Source/Runtime/MovieScene/Private/MovieSceneTimeHelpers.cpp
--- a/Engine/Source/Runtime/MovieScene/Private/MovieSceneTimeHelpers.cpp
+++ b/Engine/Source/Runtime/MovieScene/Private/MovieSceneTimeHelpers.cpp
@@ -15,24 +15,20 @@
 #include "MovieSceneSequence.h"
 #include "Channels/MovieSceneChannelProxy.h"
 #include "Channels/MovieSceneChannel.h"
+#include "IMovieSceneRetimingInterface.h"
 
 namespace UE
 {
 namespace MovieScene
 {
 
-TRange<FFrameNumber> MigrateFrameRange(const TRange<FFrameNumber>& SourceRange, FFrameRate SourceRate, FFrameRate DestinationRate)
+TRange<FFrameNumber> MigrateFrameRange(const TRange<FFrameNumber>& SourceRange, const IRetimingInterface& RetimingInterface)
 {
-	if (SourceRate == DestinationRate)
-	{
-		return SourceRange;
-	}
-
 	TRange<FFrameNumber> NewRange = TRange<FFrameNumber>::All();
 
 	if (!SourceRange.GetLowerBound().IsOpen())
 	{
-		const FFrameNumber FrameNumber = ConvertFrameTime(SourceRange.GetLowerBoundValue(), SourceRate, DestinationRate).RoundToFrame();
+		const FFrameNumber FrameNumber = RetimingInterface.RemapTime(SourceRange.GetLowerBoundValue());
 
 		NewRange.SetLowerBound(
 			SourceRange.GetLowerBound().IsExclusive()
@@ -43,7 +39,7 @@ TRange<FFrameNumber> MigrateFrameRange(const TRange<FFrameNumber>& SourceRange,
 
 	if (!SourceRange.GetUpperBound().IsOpen())
 	{
-		const FFrameNumber FrameNumber = ConvertFrameTime(SourceRange.GetUpperBoundValue(), SourceRate, DestinationRate).RoundToFrame();
+		const FFrameNumber FrameNumber = RetimingInterface.RemapTime(SourceRange.GetUpperBoundValue());
 
 		NewRange.SetUpperBound(
 			SourceRange.GetUpperBound().IsExclusive()
@@ -55,7 +51,7 @@ TRange<FFrameNumber> MigrateFrameRange(const TRange<FFrameNumber>& SourceRange,
 	return NewRange;
 }
 
-void MigrateFrameTimes(FFrameRate SourceRate, FFrameRate DestinationRate, UMovieSceneSection* Section)
+void MigrateFrameTimes(const IRetimingInterface& Retimer, UMovieSceneSection* Section)
 {
 	Section->Modify();
 	const bool bSectionWasLocked = Section->IsLocked();
@@ -65,13 +61,13 @@ void MigrateFrameTimes(FFrameRate SourceRate, FFrameRate DestinationRate, UMovie
 
 	if (Section->HasStartFrame())
 	{
-		FFrameNumber NewLowerBoundFrame = ConvertFrameTime(Section->GetInclusiveStartFrame(), SourceRate, DestinationRate).FloorToFrame();
+		FFrameNumber NewLowerBoundFrame = Retimer.RemapTime(Section->GetInclusiveStartFrame());
 		NewLowerBound = TRangeBound<FFrameNumber>::Inclusive(NewLowerBoundFrame);
 	}
 
 	if (Section->HasEndFrame())
 	{
-		FFrameNumber NewUpperBoundValue = ConvertFrameTime(Section->GetExclusiveEndFrame(), SourceRate, DestinationRate).FloorToFrame();
+		FFrameNumber NewUpperBoundValue = Retimer.RemapTime(Section->GetExclusiveEndFrame());
 		NewUpperBound = TRangeBound<FFrameNumber>::Exclusive(NewUpperBoundValue);
 	}
 
@@ -79,90 +75,61 @@ void MigrateFrameTimes(FFrameRate SourceRate, FFrameRate DestinationRate, UMovie
 
 	if (Section->GetPreRollFrames() > 0)
 	{
-		FFrameNumber NewPreRollFrameCount = ConvertFrameTime(FFrameTime(Section->GetPreRollFrames()), SourceRate, DestinationRate).FloorToFrame();
+		FFrameNumber NewPreRollFrameCount = Retimer.RemapTime(FFrameNumber(Section->GetPreRollFrames()));
 		Section->SetPreRollFrames(NewPreRollFrameCount.Value);
 	}
 
 	if (Section->GetPostRollFrames() > 0)
 	{
-		FFrameNumber NewPostRollFrameCount = ConvertFrameTime(FFrameTime(Section->GetPostRollFrames()), SourceRate, DestinationRate).FloorToFrame();
+		FFrameNumber NewPostRollFrameCount = Retimer.RemapTime(FFrameNumber(Section->GetPostRollFrames()));
 		Section->SetPostRollFrames(NewPostRollFrameCount.Value);
 	}
 
-	Section->MigrateFrameTimes(SourceRate, DestinationRate);
+	Section->MigrateFrameTimes(Retimer);
 
-	Section->Easing.AutoEaseInDuration    = ConvertFrameTime(Section->Easing.AutoEaseInDuration,    SourceRate, DestinationRate).FloorToFrame().Value;
-	Section->Easing.AutoEaseOutDuration   = ConvertFrameTime(Section->Easing.AutoEaseOutDuration,   SourceRate, DestinationRate).FloorToFrame().Value;
-	Section->Easing.ManualEaseInDuration  = ConvertFrameTime(Section->Easing.ManualEaseInDuration,  SourceRate, DestinationRate).FloorToFrame().Value;
-	Section->Easing.ManualEaseOutDuration = ConvertFrameTime(Section->Easing.ManualEaseOutDuration, SourceRate, DestinationRate).FloorToFrame().Value;
+	Section->Easing.AutoEaseInDuration    = Retimer.RemapTime(FFrameNumber(Section->Easing.AutoEaseInDuration)).Value;
+	Section->Easing.AutoEaseOutDuration   = Retimer.RemapTime(FFrameNumber(Section->Easing.AutoEaseOutDuration)).Value;
+	Section->Easing.ManualEaseInDuration  = Retimer.RemapTime(FFrameNumber(Section->Easing.ManualEaseInDuration)).Value;
+	Section->Easing.ManualEaseOutDuration = Retimer.RemapTime(FFrameNumber(Section->Easing.ManualEaseOutDuration)).Value;
 
 	for (const FMovieSceneChannelEntry& Entry : Section->GetChannelProxy().GetAllEntries())
 	{
 		for (FMovieSceneChannel* Channel : Entry.GetChannels())
 		{
-			Channel->ChangeFrameResolution(SourceRate, DestinationRate);
+			Channel->RemapTimes(Retimer);
 		}
 	}
 
 	Section->SetIsLocked(bSectionWasLocked);
 }
 
-void RemapMusicalFrameTimes(UMovieSceneSection* Section, const UFrameBasedMusicMap* OldMap, const UFrameBasedMusicMap* NewMap)
-{
-	Section->Modify();
-	const bool bSectionWasLocked = Section->IsLocked();
-	Section->SetIsLocked(false);
-
-	// We don't need to adjust the section 'Range' here because that happens in the next function call...
-
-	Section->RemapMusicalFrameTimes(OldMap, NewMap);
-
-	if (Section->GetPreRollFrames() > 0)
-	{
-		FFrameNumber NewPreRollFrameCount = NewMap->RemapFrameTime(FFrameTime(Section->GetPreRollFrames()), OldMap).FloorToFrame();
-		Section->SetPreRollFrames(NewPreRollFrameCount.Value);
-	}
-
-	if (Section->GetPostRollFrames() > 0)
-	{
-		FFrameNumber NewPostRollFrameCount = NewMap->RemapFrameTime(FFrameTime(Section->GetPostRollFrames()), OldMap).FloorToFrame();
-		Section->SetPostRollFrames(NewPostRollFrameCount.Value);
-	}
-
-	// NOT SURE THESE SHOULD CHANGE WHEN REMAPPING TO NEW MUSIC TIMINGS!
-	// QUESTION FOR ENGINE TEAM
-	Section->Easing.AutoEaseInDuration = NewMap->RemapFrameTime(Section->Easing.AutoEaseInDuration, OldMap).FloorToFrame().Value;
-	Section->Easing.AutoEaseOutDuration = NewMap->RemapFrameTime(Section->Easing.AutoEaseOutDuration, OldMap).FloorToFrame().Value;
-	Section->Easing.ManualEaseInDuration = NewMap->RemapFrameTime(Section->Easing.ManualEaseInDuration, OldMap).FloorToFrame().Value;
-	Section->Easing.ManualEaseOutDuration = NewMap->RemapFrameTime(Section->Easing.ManualEaseOutDuration, OldMap).FloorToFrame().Value;
-
-	Section->SetIsLocked(bSectionWasLocked);
-}
-
-void MigrateFrameTimes(FFrameRate SourceRate, FFrameRate DestinationRate, UMovieSceneTrack* Track)
+void MigrateFrameTimes(const IRetimingInterface& Retimer, UMovieSceneTrack* Track)
 {
 	FScopedSlowTask SlowTask(Track->GetAllSections().Num());
 
 	for (UMovieSceneSection* Section : Track->GetAllSections())
 	{
 		SlowTask.EnterProgressFrame();
-		MigrateFrameTimes(SourceRate, DestinationRate, Section);
+		MigrateFrameTimes(Retimer, Section);
 	}
 }
 
-void RemapMusicalFrameTimes(UMovieSceneTrack* Track, const UFrameBasedMusicMap* OldMap, const UFrameBasedMusicMap* NewMap)
+void TimeHelpers::MigrateFrameTimes(FFrameRate SourceRate, FFrameRate DestinationRate, UMovieScene* MovieScene, bool bApplyRecursively)
 {
-	FScopedSlowTask SlowTask(Track->GetAllSections().Num());
-
-	for (UMovieSceneSection* Section : Track->GetAllSections())
-	{
-		SlowTask.EnterProgressFrame();
-		RemapMusicalFrameTimes(Section, OldMap, NewMap);
-	}
+	FFrameRateRetiming Retimer(SourceRate, DestinationRate);
+	MigrateFrameTimes(Retimer, MovieScene, bApplyRecursively);
 }
 
-void TimeHelpers::MigrateFrameTimes(FFrameRate SourceRate, FFrameRate DestinationRate, UMovieScene* MovieScene, bool bApplyRecursively)
+void TimeHelpers::MigrateFrameTimes(const IRetimingInterface& Retimer, UMovieScene* MovieScene, bool bApplyRecursively)
 {
+	MovieScene->Modify();
+#if WITH_EDITOR
+	const bool bMovieSceneReadOnly = MovieScene->IsReadOnly();
+	MovieScene->SetReadOnly(false);
+#endif
+
+	Retimer.Begin(MovieScene);
+
 	int32 TotalNumTracks = MovieScene->GetTracks().Num() + (MovieScene->GetCameraCutTrack() ? 1 : 0);
 	for (const FMovieSceneBinding& Binding : MovieScene->GetBindings())
 	{
@@ -172,21 +139,15 @@ void TimeHelpers::MigrateFrameTimes(FFrameRate SourceRate, FFrameRate Destinatio
 	FScopedSlowTask SlowTask(TotalNumTracks, NSLOCTEXT("MovieScene", "ChangingTickResolution", "Migrating sequence frame timing"));
 	SlowTask.MakeDialogDelayed(0.25f, true);
 
-	MovieScene->Modify();
-#if WITH_EDITOR
-	const bool bMovieSceneReadOnly = MovieScene->IsReadOnly();
-	MovieScene->SetReadOnly(false);
-#endif
-
-	MovieScene->SetPlaybackRange(MigrateFrameRange(MovieScene->GetPlaybackRange(), SourceRate, DestinationRate));
+	MovieScene->SetPlaybackRange(MigrateFrameRange(MovieScene->GetPlaybackRange(), Retimer));
 #if WITH_EDITORONLY_DATA
-	MovieScene->SetSelectionRange(MigrateFrameRange(MovieScene->GetSelectionRange(), SourceRate, DestinationRate));
+	MovieScene->SetSelectionRange(MigrateFrameRange(MovieScene->GetSelectionRange(), Retimer));
 #endif
 
 	for (UMovieSceneTrack* Track : MovieScene->GetTracks())
 	{
 		SlowTask.EnterProgressFrame();
-		UE::MovieScene::MigrateFrameTimes(SourceRate, DestinationRate, Track);
+		UE::MovieScene::MigrateFrameTimes(Retimer, Track);
 
 		// We iterate through recursively here (and not in MigrateFrameTimes) so that the movie scene is taken
 		// into account for locking/modifying/etc.
@@ -202,7 +163,11 @@ void TimeHelpers::MigrateFrameTimes(FFrameRate SourceRate, FFrameRate Destinatio
 						UMovieScene* ChildMovieScene = SubSequence->GetMovieScene();
 						if (ChildMovieScene)
 						{
-							TimeHelpers::MigrateFrameTimes(SourceRate, DestinationRate, ChildMovieScene, bApplyRecursively);
+							TUniquePtr<IRetimingInterface> ChildRetimer = Retimer.RecurseInto(ChildMovieScene);
+							if (ChildRetimer)
+							{
+								TimeHelpers::MigrateFrameTimes(*ChildRetimer, ChildMovieScene, bApplyRecursively);
+							}
 						}
 					}
 				}
@@ -213,7 +178,7 @@ void TimeHelpers::MigrateFrameTimes(FFrameRate SourceRate, FFrameRate Destinatio
 	if (UMovieSceneTrack* Track = MovieScene->GetCameraCutTrack())
 	{
 		SlowTask.EnterProgressFrame();
-		UE::MovieScene::MigrateFrameTimes(SourceRate, DestinationRate, Track);
+		UE::MovieScene::MigrateFrameTimes(Retimer, Track);
 	}
 
 	for (const FMovieSceneBinding& Binding : MovieScene->GetBindings())
@@ -221,7 +186,7 @@ void TimeHelpers::MigrateFrameTimes(FFrameRate SourceRate, FFrameRate Destinatio
 		for (UMovieSceneTrack* Track : Binding.GetTracks())
 		{
 			SlowTask.EnterProgressFrame();
-			UE::MovieScene::MigrateFrameTimes(SourceRate, DestinationRate, Track);
+			UE::MovieScene::MigrateFrameTimes(Retimer, Track);
 		}
 	}
 
@@ -233,7 +198,7 @@ void TimeHelpers::MigrateFrameTimes(FFrameRate SourceRate, FFrameRate Destinatio
 
 		for (FMovieSceneMarkedFrame& MarkedFrame : MarkedFrames)
 		{
-			MarkedFrame.FrameNumber = ConvertFrameTime(MarkedFrame.FrameNumber, SourceRate, DestinationRate).RoundToFrame();
+			MarkedFrame.FrameNumber = Retimer.RemapTime(MarkedFrame.FrameNumber);
 
 			// Add it back in
 			MovieScene->AddMarkedFrame(MarkedFrame);
@@ -241,77 +206,9 @@ void TimeHelpers::MigrateFrameTimes(FFrameRate SourceRate, FFrameRate Destinatio
 
 		// Ensure they're in order as they may not have been before.
 		MovieScene->SortMarkedFrames();
-
-	}
-
-	MovieScene->SetTickResolutionDirectly(DestinationRate);
-#if WITH_EDITOR
-	MovieScene->SetReadOnly(bMovieSceneReadOnly);
-#endif
-}
-
-void TimeHelpers::RemapMusicalFrameTimes(UMovieScene* MovieScene, UFrameBasedMusicMap* NewMap)
-{
-	int32 TotalNumTracks = MovieScene->GetTracks().Num() + (MovieScene->GetCameraCutTrack() ? 1 : 0);
-	for (const FMovieSceneBinding& Binding : MovieScene->GetBindings())
-	{
-		TotalNumTracks += Binding.GetTracks().Num();
-	}
-
-	FScopedSlowTask SlowTask(TotalNumTracks, NSLOCTEXT("MovieScene", "ChangingTickResolution", "Migrating sequence frame timing"));
-	SlowTask.MakeDialogDelayed(0.25f, true);
-
-	MovieScene->Modify();
-#if WITH_EDITOR
-	const bool bMovieSceneReadOnly = MovieScene->IsReadOnly();
-	MovieScene->SetReadOnly(false);
-#endif
-
-	const UFrameBasedMusicMap* SourceMap = MovieScene->GetMusicMap();
-
-	if (SourceMap)
-	{
-		for (UMovieSceneTrack* Track : MovieScene->GetTracks())
-		{
-			SlowTask.EnterProgressFrame();
-			UE::MovieScene::RemapMusicalFrameTimes(Track, SourceMap, NewMap);
-		}
-
-		if (UMovieSceneTrack* Track = MovieScene->GetCameraCutTrack())
-		{
-			SlowTask.EnterProgressFrame();
-			UE::MovieScene::RemapMusicalFrameTimes(Track, SourceMap, NewMap);
-		}
-
-		for (const FMovieSceneBinding& Binding : MovieScene->GetBindings())
-		{
-			for (UMovieSceneTrack* Track : Binding.GetTracks())
-			{
-				SlowTask.EnterProgressFrame();
-				UE::MovieScene::RemapMusicalFrameTimes(Track, SourceMap, NewMap);
-			}
-		}
-
-		{
-			TArray<FMovieSceneMarkedFrame> MarkedFrames = MovieScene->GetMarkedFrames();
-
-			// Clear the marked frames as the returned array is immutable
-			MovieScene->DeleteMarkedFrames();
-
-			for (FMovieSceneMarkedFrame& MarkedFrame : MarkedFrames)
-			{
-				MarkedFrame.FrameNumber = NewMap->RemapFrameTime(MarkedFrame.FrameNumber, SourceMap).RoundToFrame();
-
-				// Add it back in
-				MovieScene->AddMarkedFrame(MarkedFrame);
-			}
-
-			// Ensure they're in order as they may not have been before.
-			MovieScene->SortMarkedFrames();
-		}
 	}
 
-	MovieScene->SetMusicMapDirectly(NewMap);
+	Retimer.End(MovieScene);
 
 #if WITH_EDITOR
 	MovieScene->SetReadOnly(bMovieSceneReadOnly);
diff --git a/Engine/Source/Runtime/MovieScene/Private/Sections/MovieSceneSectionTimingParameters.cpp b/Engine/Source/Runtime/MovieScene/Private/Sections/MovieSceneSectionTimingParameters.cpp
--- a/Engine/Source/Runtime/MovieScene/Private/Sections/MovieSceneSectionTimingParameters.cpp
+++ b/Engine/Source/Runtime/MovieScene/Private/Sections/MovieSceneSectionTimingParameters.cpp
@@ -3,9 +3,9 @@
 #include "Sections/MovieSceneSectionTimingParameters.h"
 #include "Evaluation/MovieSceneSequenceTransform.h"
 #include "Misc/FrameRate.h"
-#include "Variants/MovieSceneMusicTimeTransform.h"
 #include "Variants/MovieSceneTimeWarpVariantPayloads.h"
 #include "MovieScene.h"
+#include "MovieSceneClock.h"
 
 FMovieSceneSequenceTransform FMovieSceneSectionTimingParametersSeconds::MakeTransform(const FFrameRate& OuterFrameRate, const TRange<FFrameNumber>& OuterRange, double SourceDuration, double InnerPlayRate) const
 {
@@ -149,77 +149,6 @@ FMovieSceneSequenceTransform FMovieSceneSectionTimingParametersFrames::MakeTrans
 	return Result;
 }
 
-MOVIESCENE_API FMovieSceneSequenceTransform FMovieSceneSectionTimingParametersFrames::MakeMusicalTransform(const FFrameRate& OuterFrameRate, const TRange<FFrameNumber>& OuterRange, UMovieScene* OuterMovieScene, const FFrameRate& InnerFrameRate, const TRange<FFrameNumber>& InnerRange, UMovieScene* InnerMovieScene) const
-{
-	FMovieSceneSequenceTransform Result;
-
-	check(OuterRange.HasLowerBound());
-	check(InnerRange.HasLowerBound() && InnerRange.HasUpperBound());
-	check(OuterMovieScene);
-	check(InnerMovieScene);
-
-
-	const UFrameBasedMusicMap* OuterMap = OuterMovieScene->GetMusicMap();
-	const UFrameBasedMusicMap* InnerMap = InnerMovieScene->GetMusicMap();
-	check(OuterMap);
-	check(InnerMap);
-
-	// ----------------------------------------------------------------------------
-	// First things first, subtract the section start bound
-	AddPositionInOuterAsOffset(Result, -OuterRange.GetLowerBoundValue());
-
-	// ----------------------------------------------------------------------------
-	// FrameRate conversion
-	if (InnerFrameRate != OuterFrameRate)
-	{
-		AddFrameRateConversion(Result, OuterFrameRate, InnerFrameRate);
-	}
-
-	AddMusicalTransform(Result, OuterRange.GetLowerBoundValue(), OuterMovieScene, InnerMovieScene);
-
-	// ----------------------------------------------------------------------------
-	// Time warp
-	AddPlayRate(Result);
-
-
-	// InnerStartOffset and InnerEndOffset are in frames that are both...
-	// a. In units of the outer music map, and...
-	// b. a count from the BEGINNING of the outer music map as a duration.
-	// So we have some conversions to use...
-	FMusicalTime StartOffsetAsMusicalTime = OuterMap->GetMusicalTime(InnerStartOffset);
-	FFrameNumber RealInnerStartOffset = InnerMap->GetFrameTime(StartOffsetAsMusicalTime).RoundToFrame();
-	FMusicalTime EndOffsetAsMusicalTime = OuterMap->GetMusicalTime(InnerEndOffset);
-	FFrameNumber RealInnerEndOffset = InnerMap->GetFrameTime(EndOffsetAsMusicalTime).RoundToFrame();
-
-	FMusicalTime LoopStartOffsetAsMusicalTime = OuterMap->GetMusicalTime(FirstLoopStartOffset);
-	FFrameNumber RealFirstLoopStartOffset = InnerMap->GetFrameTime(LoopStartOffsetAsMusicalTime).RoundToFrame();
-
-
-	FFrameNumber StartTime = InnerRange.GetLowerBoundValue() + RealInnerStartOffset;
-	FFrameNumber EndTime = InnerRange.GetUpperBoundValue() - RealInnerEndOffset;
-	FFrameNumber Duration = EndTime - StartTime;
-
-	FFrameNumber LoopOffset(bLoop ? RealFirstLoopStartOffset : 0);
-
-	FFrameNumber NegativeRateOffset = 0;
-	if (PlayRate.GetType() == EMovieSceneTimeWarpType::FixedPlayRate && PlayRate.AsFixedPlayRate() < 0.0)
-	{
-		NegativeRateOffset = Duration;
-	}
-
-	// Start offset
-	AddInnerStartOffset(Result, StartTime + LoopOffset + NegativeRateOffset);
-
-	// ----------------------------------------------------------------------------
-	// Looping or clamping
-	if (Duration > 0)
-	{
-		AddLoopingOrClampingAndReverse(Result, Duration, -StartTime);
-	}
-
-	return Result;
-}
-
 void FMovieSceneSectionTimingParametersFrames::AddPositionInOuterAsOffset(FMovieSceneSequenceTransform& Transform, const FFrameTime& Offset) const
 {
 	Transform.Add(FMovieSceneTimeTransform(Offset));
@@ -264,12 +193,4 @@ void FMovieSceneSectionTimingParametersFrames::AddLoopingOrClampingAndReverse(FM
 		// Reverse
 		Result.Add(FMovieSceneTimeTransform(Duration, -1.f));
 	}
-}
-
-void FMovieSceneSectionTimingParametersFrames::AddMusicalTransform(FMovieSceneSequenceTransform& Result, const FFrameNumber& Offset, UMovieScene* OuterMovieScene, UMovieScene* InnerMovieScene) const
-{
-	check(OuterMovieScene->HasMusicMap() && InnerMovieScene->HasMusicMap());
-	UMovieSceneMusicTimeTransform* Mapper = NewObject<UMovieSceneMusicTimeTransform>(InnerMovieScene);
-	Mapper->Initialize(OuterMovieScene, InnerMovieScene, Offset);
-	Result.Add(0, FMovieSceneTimeWarpVariant(Mapper));
 }
\ No newline at end of file
diff --git a/Engine/Source/Runtime/MovieScene/Private/Sections/MovieSceneSubSection.cpp b/Engine/Source/Runtime/MovieScene/Private/Sections/MovieSceneSubSection.cpp
--- a/Engine/Source/Runtime/MovieScene/Private/Sections/MovieSceneSubSection.cpp
+++ b/Engine/Source/Runtime/MovieScene/Private/Sections/MovieSceneSubSection.cpp
@@ -3,6 +3,7 @@
 #include "Sections/MovieSceneSubSection.h"
 #include "Tracks/MovieSceneSubTrack.h"
 #include "MovieScene.h"
+#include "MovieSceneClock.h"
 #include "MovieSceneTrack.h"
 #include "MovieSceneSequence.h"
 #include "MovieSceneTimeHelpers.h"
@@ -21,7 +22,6 @@
 #include "Tracks/MovieSceneTimeWarpTrack.h"
 #include "Sections/MovieSceneSectionTimingParameters.h"
 #include "Variants/MovieSceneTimeWarpGetter.h"
-#include "Variants/MovieSceneMusicTimeTransform.h"
 
 
 #include UE_INLINE_GENERATED_CPP_BY_NAME(MovieSceneSubSection)
@@ -234,15 +234,13 @@ FMovieSceneSequenceTransform UMovieSceneSubSection::OuterToInnerTransform_NoInne
 	};
 
 	// determine if we need to generate a musical transform or a standard transform...
-	if (!MovieScenePtr->HasMusicMap() || !GetTypedOuter<UMovieScene>()->HasMusicMap())
+	FMovieSceneSequenceTransform ClockResult;
+	if (MovieScenePtr->GetCustomClock() && MovieScenePtr->GetCustomClock()->MakeSubSequenceTransform(TimingParams, this, ClockResult))
 	{
-		FMovieSceneSequenceTransform Result = TimingParams.MakeTransform(OuterFrameRate, SubRange, InnerFrameRate, MovieScenePtr->GetPlaybackRange());
-		return Result;
+		return ClockResult;
 	}
 
-	FMovieSceneSequenceTransform Result = TimingParams.MakeMusicalTransform(OuterFrameRate, SubRange, GetTypedOuter<UMovieScene>(), InnerFrameRate, MovieScenePtr->GetPlaybackRange(), MovieScenePtr);
-
-	return Result;
+	return TimingParams.MakeTransform(OuterFrameRate, SubRange, InnerFrameRate, MovieScenePtr->GetPlaybackRange());
 }
 
 FMovieSceneSequenceTransform UMovieSceneSubSection::OuterToInnerTransform() const
@@ -678,49 +676,22 @@ void UMovieSceneSubSection::GetSnapTimes(TArray<FFrameNumber>& OutSnapTimes, boo
 
 void UMovieSceneSubSection::MigrateFrameTimes(FFrameRate SourceRate, FFrameRate DestinationRate)
 {
-	if (Parameters.StartFrameOffset.Value > 0)
-	{
-		FFrameNumber NewStartFrameOffset = ConvertFrameTime(FFrameTime(Parameters.StartFrameOffset), SourceRate, DestinationRate).FloorToFrame();
-		Parameters.StartFrameOffset = NewStartFrameOffset;
-	}
-	if (Parameters.EndFrameOffset.Value > 0)
-	{
-		FFrameNumber NewEndFrameOffset = ConvertFrameTime(FFrameTime(Parameters.EndFrameOffset), SourceRate, DestinationRate).FloorToFrame();
-		Parameters.EndFrameOffset = NewEndFrameOffset;
-	}
-	if (Parameters.FirstLoopStartFrameOffset.Value > 0)
-	{
-		FFrameNumber NewFirstLoopStartFrameOffset = ConvertFrameTime(FFrameTime(Parameters.FirstLoopStartFrameOffset), SourceRate, DestinationRate).FloorToFrame();
-		Parameters.FirstLoopStartFrameOffset = NewFirstLoopStartFrameOffset;
-	}
+	MigrateFrameTimes(UE::MovieScene::FFrameRateRetiming(SourceRate, DestinationRate));
 }
 
-MOVIESCENE_API void UMovieSceneSubSection::RemapMusicalFrameTimes(const UFrameBasedMusicMap* OldMap, const UFrameBasedMusicMap* NewMap)
+void UMovieSceneSubSection::MigrateFrameTimes(const UE::MovieScene::IRetimingInterface& Retimer)
 {
-	Super::RemapMusicalFrameTimes(OldMap, NewMap);
-	
-	UMovieSceneSequence* SequencePtr = GetSequence();
-	UMovieScene* MovieScenePtr = SequencePtr ? SequencePtr->GetMovieScene() : nullptr;
-	if (!MovieScenePtr || !MovieScenePtr->HasMusicMap())
-	{
-		return;
-	}
-
-	// The start, end, and loop-start offsets are in musical units
 	if (Parameters.StartFrameOffset.Value > 0)
 	{
-		FFrameNumber NewStartFrameOffset = NewMap->RemapFrameTime(FFrameTime(Parameters.StartFrameOffset), OldMap).RoundToFrame();
-		Parameters.StartFrameOffset = NewStartFrameOffset;
+		Parameters.StartFrameOffset = Retimer.RemapTime(Parameters.StartFrameOffset);
 	}
 	if (Parameters.EndFrameOffset.Value > 0)
 	{
-		FFrameNumber NewEndFrameOffset = NewMap->RemapFrameTime(FFrameTime(Parameters.EndFrameOffset), OldMap).RoundToFrame();
-		Parameters.EndFrameOffset = NewEndFrameOffset;
+		Parameters.EndFrameOffset = Retimer.RemapTime(Parameters.EndFrameOffset);
 	}
 	if (Parameters.FirstLoopStartFrameOffset.Value > 0)
 	{
-		FFrameNumber NewFirstLoopStartFrameOffset = NewMap->RemapFrameTime(FFrameTime(Parameters.FirstLoopStartFrameOffset), OldMap).RoundToFrame();
-		Parameters.FirstLoopStartFrameOffset = NewFirstLoopStartFrameOffset;
+		Parameters.FirstLoopStartFrameOffset = Retimer.RemapTime(Parameters.FirstLoopStartFrameOffset);
 	}
 }
 
diff --git a/Engine/Source/Runtime/MovieScene/Public/Channels/MovieSceneAudioTriggerChannel.h b/Engine/Source/Runtime/MovieScene/Public/Channels/MovieSceneAudioTriggerChannel.h
--- a/Engine/Source/Runtime/MovieScene/Public/Channels/MovieSceneAudioTriggerChannel.h
+++ b/Engine/Source/Runtime/MovieScene/Public/Channels/MovieSceneAudioTriggerChannel.h
@@ -90,8 +90,7 @@ public:
 	MOVIESCENE_API virtual void DuplicateKeys(TArrayView<const FKeyHandle> InHandles, TArrayView<FKeyHandle> OutNewHandles) override;
 	MOVIESCENE_API virtual void DeleteKeys(TArrayView<const FKeyHandle> InHandles) override;
 	MOVIESCENE_API virtual void DeleteKeysFrom(FFrameNumber InTime, bool bDeleteKeysBefore) override;
-	MOVIESCENE_API virtual void ChangeFrameResolution(FFrameRate SourceRate, FFrameRate DestinationRate) override;
-	MOVIESCENE_API virtual void RemapMusicalFrameTimes(const UFrameBasedMusicMap* OldMap, const UFrameBasedMusicMap* NewMap) override;
+	MOVIESCENE_API virtual void RemapTimes(const UE::MovieScene::IRetimingInterface& Retimer) override;
 	MOVIESCENE_API virtual TRange<FFrameNumber> ComputeEffectiveRange() const override;
 	MOVIESCENE_API virtual int32 GetNumKeys() const override;
 	MOVIESCENE_API virtual void Reset() override;
diff --git a/Engine/Source/Runtime/MovieScene/Public/Channels/MovieSceneBoolChannel.h b/Engine/Source/Runtime/MovieScene/Public/Channels/MovieSceneBoolChannel.h
--- a/Engine/Source/Runtime/MovieScene/Public/Channels/MovieSceneBoolChannel.h
+++ b/Engine/Source/Runtime/MovieScene/Public/Channels/MovieSceneBoolChannel.h
@@ -115,8 +115,7 @@ public:
 	MOVIESCENE_API virtual void DuplicateKeys(TArrayView<const FKeyHandle> InHandles, TArrayView<FKeyHandle> OutNewHandles) override;
 	MOVIESCENE_API virtual void DeleteKeys(TArrayView<const FKeyHandle> InHandles) override;
 	MOVIESCENE_API virtual void DeleteKeysFrom(FFrameNumber InTime, bool bDeleteKeysBefore) override;
-	MOVIESCENE_API virtual void ChangeFrameResolution(FFrameRate SourceRate, FFrameRate DestinationRate) override;
-	MOVIESCENE_API virtual void RemapMusicalFrameTimes(const UFrameBasedMusicMap* OldMap, const UFrameBasedMusicMap* NewMap) override;
+	MOVIESCENE_API virtual void RemapTimes(const UE::MovieScene::IRetimingInterface& Retimer) override;
 	MOVIESCENE_API virtual TRange<FFrameNumber> ComputeEffectiveRange() const override;
 	MOVIESCENE_API virtual int32 GetNumKeys() const override;
 	MOVIESCENE_API virtual void Reset() override;
diff --git a/Engine/Source/Runtime/MovieScene/Public/Channels/MovieSceneChannel.h b/Engine/Source/Runtime/MovieScene/Public/Channels/MovieSceneChannel.h
--- a/Engine/Source/Runtime/MovieScene/Public/Channels/MovieSceneChannel.h
+++ b/Engine/Source/Runtime/MovieScene/Public/Channels/MovieSceneChannel.h
@@ -19,7 +19,11 @@ struct FFrameNumber;
 struct FFrameRate;
 struct FKeyDataOptimizationParams;
 struct FMovieSceneChannel;
-class UFrameBasedMusicMap;
+
+namespace UE::MovieScene
+{
+	struct IRetimingInterface;
+}
 
 /*
 *  Events that are fired when adding, deleting or moving keys
@@ -187,16 +191,13 @@ struct FMovieSceneChannel
 	 * @param SourceRate      The previous frame resolution that the channel is currently in
 	 * @param DestinationRate The desired new frame resolution. All keys should be transformed into this time-base.
 	 */
-	virtual void ChangeFrameResolution(FFrameRate SourceRate, FFrameRate DestinationRate)
-	{}
+	UE_DEPRECATED(5.6, "Please use RemapTimes")
+	MOVIESCENE_API virtual void ChangeFrameResolution(FFrameRate SourceRate, FFrameRate DestinationRate);
 
 	/**
-	 * Called when the music map is changed in a sequence that supports music mode editing
-	 *
-	 * @param OldMap The music map the keyframes were originally in reference to.
-	 * @param NewMap The music map the keyframes should be "remapped" to.
+	 * Remap the times contained within this channel using a custom remapper
 	 */
-	virtual void RemapMusicalFrameTimes(const UFrameBasedMusicMap* OldMap, const UFrameBasedMusicMap* NewMap)
+	virtual void RemapTimes(const UE::MovieScene::IRetimingInterface& Retimer)
 	{}
 
 	/**
diff --git a/Engine/Source/Runtime/MovieScene/Public/Channels/MovieSceneChannelData.h b/Engine/Source/Runtime/MovieScene/Public/Channels/MovieSceneChannelData.h
--- a/Engine/Source/Runtime/MovieScene/Public/Channels/MovieSceneChannelData.h
+++ b/Engine/Source/Runtime/MovieScene/Public/Channels/MovieSceneChannelData.h
@@ -16,12 +16,20 @@
 #include "Templates/UnrealTypeTraits.h"
 #include "UObject/Class.h"
 #include "UObject/ObjectMacros.h"
+#include "IMovieSceneRetimingInterface.h"
 
 #include "MovieSceneChannelData.generated.h"
 
 struct FFrameRate;
 struct FKeyDataOptimizationParams;
 
+namespace UE::MovieScene
+{
+	template<typename ValueType>
+	void OnRemapChannelKeyTime(const FMovieSceneChannel* Channel, const IRetimingInterface& Retimer, FFrameNumber PreviousTime, FFrameNumber CurrentTime, ValueType& InOutValue)
+	{}
+}
+
 /** A map of key handles that is copyable, but does not copy data on copy */
 USTRUCT()
 struct FMovieSceneKeyHandleMap : public FKeyHandleLookupTable
@@ -164,16 +172,9 @@ struct FMovieSceneChannelData
 	 * @param SourceRate      The frame rate the channel is currently in
 	 * @param DestinationRate The new frame rate to convert the channel to
 	 */
+	UE_DEPRECATED(5.6, "Please use RemapTimes")
 	MOVIESCENE_API void ChangeFrameResolution(FFrameRate SourceRate, FFrameRate DestinationRate);
 
-	/**
-	 * Moves the location of keys to account for a changing music map.
-	 *
-	 * @param OldMap   The music map that the key's original time was in reference to.
-	 * @param NewMap   The music map that the key's new time should be in refernce to.
-	 */
-	MOVIESCENE_API void RemapMusicalFrameTimes(const UFrameBasedMusicMap* OldMap, const UFrameBasedMusicMap* NewMap);
-
 	/**
 	 * Get all the keys in the given range. Resulting arrays must be the same size where indices correspond to both arrays.
 	 *
@@ -360,6 +361,83 @@ struct TMovieSceneChannelData : FMovieSceneChannelData
 		return MoveKey(KeyIndex, InNewTime);
 	}
 
+
+
+	/**
+	 * Remap the times of all the keys in this channel using an abstract retimer
+	 *
+	 * @param Retimer         Custom retimer object that can convert from an old time to a new time
+	 */
+	void RemapTimes(const UE::MovieScene::IRetimingInterface& Retimer)
+	{
+		using namespace UE::MovieScene;
+
+		const int32 Num = Times->Num();
+		if (Num == 0)
+		{
+			return;
+		}
+
+		bool bIsSorted = true;
+		{
+			int32 Index = 0;
+
+			// Remap the first time
+			{
+				FFrameNumber& NewTime  = (*Times)[Index];
+				FFrameNumber  OldTime  = NewTime;
+
+				NewTime = Retimer.RemapTime(OldTime);
+				OnRemapChannelKeyTime(OwningChannel, Retimer, OldTime, NewTime, (*Values)[Index]);
+			}
+			++Index;
+
+			// Remap others keeping track of whether we need resorting
+			for ( ; Index < Num; ++Index)
+			{
+				FFrameNumber& NewTime  = (*Times)[Index];
+				FFrameNumber  OldTime  = NewTime;
+				FFrameNumber  Previous = (*Times)[Index-1];
+
+				NewTime = Retimer.RemapTime(OldTime);
+				bIsSorted &= NewTime >= Previous;
+
+				OnRemapChannelKeyTime(OwningChannel, Retimer, OldTime, NewTime, (*Values)[Index]);
+			}
+		}
+
+		if (!bIsSorted)
+		{
+			// Have to re-sort all the keys...
+			TArray<int32, TInlineAllocator<128>> Indices;
+			for (int32 Index = 0; Index < Num; ++Index)
+			{
+				Indices.Add(Index);
+			}
+
+			Algo::Sort(Indices, [this](int32 A, int32 B){
+				return (*Times)[A] < (*Times)[B];
+			});
+
+			for (int32 Index = 0; Index < Num; ++Index)
+			{
+				KeyHandles->MoveHandle(Index, Indices[Index]);
+			}
+
+			// Move all the keys by repeatedly swapping until we get the correct index
+			for (int32 Index = 0; Index < Num; ++Index)
+			{
+				while (Indices[Index] != Index)
+				{
+					const int32 DestIndex = Indices[Index];
+					Swap((*Times)[Index], (*Times)[DestIndex]);
+					Swap((*Values)[Index], (*Values)[DestIndex]);
+					Swap(Indices[Index], Indices[DestIndex]);
+				}
+			}
+		}
+	}
+
 	/**
 	 * Remove the key at a given index
 	 *
diff --git a/Engine/Source/Runtime/MovieScene/Public/Channels/MovieSceneDoubleChannel.h b/Engine/Source/Runtime/MovieScene/Public/Channels/MovieSceneDoubleChannel.h
--- a/Engine/Source/Runtime/MovieScene/Public/Channels/MovieSceneDoubleChannel.h
+++ b/Engine/Source/Runtime/MovieScene/Public/Channels/MovieSceneDoubleChannel.h
@@ -31,6 +31,7 @@
 
 class FArchive;
 struct FPropertyTag;
+struct FMovieSceneDoubleValue;
 template<typename> struct TMovieSceneCurveChannelImpl;
 template <typename T> struct TIsPODType;
 
@@ -38,6 +39,7 @@ namespace UE::MovieScene
 {
 	enum class EInverseEvaluateFlags : uint8;
 	struct FPiecewiseCurve;
+	MOVIESCENE_API void OnRemapChannelKeyTime(const FMovieSceneChannel* Channel, const IRetimingInterface& Retimer, FFrameNumber PreviousTime, FFrameNumber CurrentTime, FMovieSceneDoubleValue& InOutValue);
 }
 namespace UE::MovieScene::Interpolation
 {
@@ -259,8 +261,7 @@ public:
 	MOVIESCENE_API virtual void DuplicateKeys(TArrayView<const FKeyHandle> InHandles, TArrayView<FKeyHandle> OutNewHandles) override;
 	MOVIESCENE_API virtual void DeleteKeys(TArrayView<const FKeyHandle> InHandles) override;
 	MOVIESCENE_API virtual void DeleteKeysFrom(FFrameNumber InTime, bool bDeleteKeysBefore) override;
-	MOVIESCENE_API virtual void ChangeFrameResolution(FFrameRate SourceRate, FFrameRate DestinationRate) override;
-	MOVIESCENE_API virtual void RemapMusicalFrameTimes(const UFrameBasedMusicMap* OldMap, const UFrameBasedMusicMap* NewMap) override;
+	MOVIESCENE_API virtual void RemapTimes(const UE::MovieScene::IRetimingInterface& Retimer) override;
 	MOVIESCENE_API virtual TRange<FFrameNumber> ComputeEffectiveRange() const override;
 	MOVIESCENE_API virtual int32 GetNumKeys() const override;
 	MOVIESCENE_API virtual void Reset() override;
diff --git a/Engine/Source/Runtime/MovieScene/Public/Channels/MovieSceneFloatChannel.h b/Engine/Source/Runtime/MovieScene/Public/Channels/MovieSceneFloatChannel.h
--- a/Engine/Source/Runtime/MovieScene/Public/Channels/MovieSceneFloatChannel.h
+++ b/Engine/Source/Runtime/MovieScene/Public/Channels/MovieSceneFloatChannel.h
@@ -34,9 +34,14 @@ struct FPropertyTag;
 template<typename> struct TMovieSceneCurveChannelImpl;
 template <typename T> struct TIsPODType;
 
-namespace UE::MovieScene::Interpolation
+namespace UE::MovieScene
 {
-	struct FCachedInterpolation;
+	MOVIESCENE_API void OnRemapChannelKeyTime(const FMovieSceneChannel* Channel, const IRetimingInterface& Retimer, FFrameNumber PreviousTime, FFrameNumber CurrentTime, FMovieSceneFloatValue& InOutValue);
+
+	namespace Interpolation
+	{
+		struct FCachedInterpolation;
+	}
 }
 
 USTRUCT()
@@ -198,8 +203,7 @@ public:
 	MOVIESCENE_API virtual void DuplicateKeys(TArrayView<const FKeyHandle> InHandles, TArrayView<FKeyHandle> OutNewHandles) override;
 	MOVIESCENE_API virtual void DeleteKeys(TArrayView<const FKeyHandle> InHandles) override;
 	MOVIESCENE_API virtual void DeleteKeysFrom(FFrameNumber InTime, bool bDeleteKeysBefore) override;
-	MOVIESCENE_API virtual void ChangeFrameResolution(FFrameRate SourceRate, FFrameRate DestinationRate) override;
-	MOVIESCENE_API virtual void RemapMusicalFrameTimes(const UFrameBasedMusicMap* OldMap, const UFrameBasedMusicMap* NewMap) override;
+	MOVIESCENE_API virtual void RemapTimes(const UE::MovieScene::IRetimingInterface& Retimer) override;
 	MOVIESCENE_API virtual TRange<FFrameNumber> ComputeEffectiveRange() const override;
 	MOVIESCENE_API virtual int32 GetNumKeys() const override;
 	MOVIESCENE_API virtual void Reset() override;
diff --git a/Engine/Source/Runtime/MovieScene/Public/Evaluation/IMovieSceneCustomClockSource.h b/Engine/Source/Runtime/MovieScene/Public/Evaluation/IMovieSceneCustomClockSource.h
--- a/Engine/Source/Runtime/MovieScene/Public/Evaluation/IMovieSceneCustomClockSource.h
+++ b/Engine/Source/Runtime/MovieScene/Public/Evaluation/IMovieSceneCustomClockSource.h
@@ -42,12 +42,6 @@ public:
 
 	UFUNCTION()
 	virtual FFrameTime OnRequestCurrentTime(const FQualifiedFrameTime& InCurrentTime, float InPlayRate) { return 0; }
-
-	UFUNCTION()
-	virtual bool OnSupportsMusicalTime() { return false; }
-
-	UFUNCTION()
-	virtual FMusicalTime OnRequestCurrentMusicalTime(const FQualifiedFrameTime& InCurrentTime, float InPlayRate) { return FMusicalTime(); }
 };
 
 struct FMovieSceneTimeController_Custom : FMovieSceneTimeController
@@ -60,10 +54,6 @@ private:
 	MOVIESCENE_API virtual void OnStartPlaying(const FQualifiedFrameTime& InStartTime) override final;
 	MOVIESCENE_API virtual void OnStopPlaying(const FQualifiedFrameTime& InStopTime) override final;
 	MOVIESCENE_API virtual FFrameTime OnRequestCurrentTime(const FQualifiedFrameTime& InCurrentTime, float InPlayRate) override final;
-	MOVIESCENE_API virtual bool OnSupportsMusicalTime() override final;
-	MOVIESCENE_API virtual FMusicalTime OnRequestCurrentMusicalTime(const FQualifiedFrameTime& InCurrentTime, float InPlayRate) override final;
-
-private:
 
 	void ResolveInterfacePtr();
 
diff --git a/Engine/Source/Runtime/MovieScene/Public/EventHandlers/ISequenceDataEventHandler.h b/Engine/Source/Runtime/MovieScene/Public/EventHandlers/ISequenceDataEventHandler.h
--- a/Engine/Source/Runtime/MovieScene/Public/EventHandlers/ISequenceDataEventHandler.h
+++ b/Engine/Source/Runtime/MovieScene/Public/EventHandlers/ISequenceDataEventHandler.h
@@ -9,9 +9,12 @@
 class UMovieScene;
 class UMovieSceneFolder;
 class UMovieSceneTrack;
+class UMovieSceneClock;
 struct FGuid;
 struct FMovieSceneBinding;
 
+enum class EUpdateClockSource : uint8;
+
 namespace UE
 {
 namespace MovieScene
@@ -39,6 +42,8 @@ public:
 
 	virtual void OnBindingParentChanged(const FGuid& ObjectBindingID, const FGuid& NewParent) {}
 
+	virtual void OnClockChanged(EUpdateClockSource NewClock, UMovieSceneClock* NewCustomClock) {}
+
 };
 
 } // namespace MovieScene
diff --git a/Engine/Source/Runtime/MovieScene/Public/EventHandlers/MovieSceneDataEventContainer.h b/Engine/Source/Runtime/MovieScene/Public/EventHandlers/MovieSceneDataEventContainer.h
--- a/Engine/Source/Runtime/MovieScene/Public/EventHandlers/MovieSceneDataEventContainer.h
+++ b/Engine/Source/Runtime/MovieScene/Public/EventHandlers/MovieSceneDataEventContainer.h
@@ -64,13 +64,21 @@ struct TDataEventContainer
 	template<typename FuncType, typename... ArgTypes>
 	void Trigger(FuncType&& Func, ArgTypes&&... Args) const
 	{
-		for (typename TIntrusiveLinkedList<TIntrusiveEventHandler<EventInterface>>::TIterator It(IntrusiveHandlers); It; ++It)
+		for (typename TIntrusiveLinkedList<TIntrusiveEventHandler<EventInterface>>::TIterator It(IntrusiveHandlers); It; )
 		{
-			Invoke(Func, *It, Forward<ArgTypes>(Args)...);
+			TIntrusiveEventHandler<EventInterface>& Current = *It;
+			// Increment before invoking to ensure that removals inside the invocation do not cause the iteration to return prematurely
+			++It;
+
+			Invoke(Func, Current, Forward<ArgTypes>(Args)...);
 		}
-		for (typename TNonIntrusiveEventHandler<EventInterface>::TIterator It(NonIntrusiveHandlers); It; ++It)
+		for (typename TNonIntrusiveEventHandler<EventInterface>::TIterator It(NonIntrusiveHandlers); It; )
 		{
-			Invoke(Func, *It, Forward<ArgTypes>(Args)...);
+			EventInterface* Current = *It;
+			// Increment before invoking to ensure that removals inside the invocation do not cause the iteration to return prematurely
+			++It;
+
+			Invoke(Func, Current, Forward<ArgTypes>(Args)...);
 		}
 	}
 
diff --git a/Engine/Source/Runtime/MovieScene/Public/IMovieSceneRetimingInterface.h b/Engine/Source/Runtime/MovieScene/Public/IMovieSceneRetimingInterface.h
new file mode 100644
--- /dev/null
+++ b/Engine/Source/Runtime/MovieScene/Public/IMovieSceneRetimingInterface.h
@@ -0,0 +1,48 @@
+// Copyright Epic Games, Inc. All Rights Reserved.
+
+#pragma once
+
+#include "CoreTypes.h"
+#include "Misc/FrameRate.h"
+
+struct FFrameNumber;
+struct FFrameTime;
+class UMovieScene;
+
+namespace UE::MovieScene
+{
+
+struct IRetimingInterface
+{
+	virtual ~IRetimingInterface() {}
+
+	MOVIESCENE_API FFrameNumber RemapTime(FFrameNumber InTime) const;
+
+	virtual double GetScale() const = 0;
+	virtual FFrameTime RemapTime(FFrameTime InTime) const = 0;
+
+	virtual TUniquePtr<IRetimingInterface> RecurseInto(UMovieScene* InMovieScene) const = 0;
+
+	virtual void Begin(UMovieScene* InMovieScene) const = 0;
+	virtual void End(UMovieScene* InMovieScene) const = 0;
+};
+
+
+struct MOVIESCENE_API FFrameRateRetiming : IRetimingInterface
+{
+	FFrameRate SourceRate;
+	FFrameRate DestinationRate;
+
+	FFrameRateRetiming(FFrameRate InSourceRate, FFrameRate InDestinationRate);
+
+	virtual double GetScale() const override;
+	virtual FFrameTime RemapTime(FFrameTime InTime) const override;
+
+	virtual TUniquePtr<IRetimingInterface> RecurseInto(UMovieScene* InMovieScene) const override;
+
+	virtual void Begin(UMovieScene* InMovieScene) const override;
+	virtual void End(UMovieScene* InMovieScene) const override;
+};
+
+
+} // namespace UE::MovieScene
\ No newline at end of file
diff --git a/Engine/Source/Runtime/MovieScene/Public/MovieScene.h b/Engine/Source/Runtime/MovieScene/Public/MovieScene.h
--- a/Engine/Source/Runtime/MovieScene/Public/MovieScene.h
+++ b/Engine/Source/Runtime/MovieScene/Public/MovieScene.h
@@ -11,7 +11,6 @@
 #include "CoreTypes.h"
 #include "Delegates/Delegate.h"
 #include "EventHandlers/MovieSceneDataEventContainer.h"
-#include "FrameBasedMusicMap.h"
 #include "HAL/PlatformCrt.h"
 #include "Internationalization/Text.h"
 #include "Math/Color.h"
@@ -48,6 +47,7 @@
 class FArchive;
 class FObjectPreSaveContext;
 class UClass;
+class UMovieSceneClock;
 class UK2Node;
 class UMovieSceneFolder;
 class UMovieSceneSection;
@@ -847,6 +847,17 @@ public:
 		return ClockSource;
 	}
 
+	UMovieSceneClock* GetCustomClock() const
+	{
+		return ClockSource == EUpdateClockSource::Custom ? CustomClock : nullptr;
+	}
+
+	template<typename T>
+	T* GetCustomClockAs() const
+	{
+		return Cast<T>(GetCustomClock());
+	}
+
 	/**
 	 * Retrieve a time controller from this sequence instance, if the clock source is set to custom
 	 */
@@ -855,23 +866,15 @@ public:
 	/**
 	 * Assign the clock source to be used for this moviescene
 	 */
-	void SetClockSource(EUpdateClockSource InNewClockSource)
-	{
-		ClockSource = InNewClockSource;
-		if (ClockSource != EUpdateClockSource::Custom)
-		{
-			CustomClockSourcePath.Reset();
-		}
-	}
+	MOVIESCENE_API void SetClockSource(EUpdateClockSource InNewClockSource);
 
 	/**
 	 * Assign the clock source to be used for this moviescene
 	 */
-	void SetClockSource(UObject* InNewClockSource)
-	{
-		ClockSource = EUpdateClockSource::Custom;
-		CustomClockSourcePath = InNewClockSource;
-	}
+	UE_DEPRECATED(5.6, "Please call SetCustomClock with an instance of UMovieSceneExternalClock")
+	MOVIESCENE_API void SetClockSource(UObject* InNewClockSource);
+
+	MOVIESCENE_API void SetCustomClock(UMovieSceneClock* InNewClockSource);
 
 	/**
 	 * Get the earliest timecode source out of all of the movie scene sections contained within this movie scene.
@@ -1199,34 +1202,6 @@ public:
 	/* Called by the compiler to empty the list of generated conditions*/
 	void ResetGeneratedConditions() { GeneratedConditions.Reset(); }
 
-
-	MOVIESCENE_API bool HasMusicMap() const
-	{
-		return (bool)MusicMap;
-	}
-
-	MOVIESCENE_API void RemoveMusicMap()
-	{
-		MusicMap = nullptr;
-	}
-
-	MOVIESCENE_API const UFrameBasedMusicMap* GetMusicMap() const
-	{
-		return MusicMap.Get();
-	}
-
-	MOVIESCENE_API UFrameBasedMusicMap* GetMusicMap()
-	{
-		return MusicMap.Get();
-	}
-
-	void SetMusicMapDirectly(UFrameBasedMusicMap* InMusicMap)
-	{
-		check(InMusicMap);
-		check(InMusicMap->GetOuter() == this);
-		MusicMap = InMusicMap;
-	}
-
 protected:
 
 	/**
@@ -1307,6 +1282,10 @@ private:
 	UPROPERTY(Instanced)
 	TObjectPtr<UMovieSceneTrack> CameraCutTrack;
 
+	/** A custom clock implementation */
+	UPROPERTY(Instanced)
+	TObjectPtr<UMovieSceneClock> CustomClock;
+
 	/** User-defined selection range. */
 	UPROPERTY()
 	FMovieSceneFrameRange SelectionRange;
@@ -1330,8 +1309,10 @@ private:
 	UPROPERTY()
 	EUpdateClockSource ClockSource;
 
+#if WITH_EDITORONLY_DATA
 	UPROPERTY()
-	FSoftObjectPath CustomClockSourcePath;
+	FSoftObjectPath CustomClockSourcePath_DEPRECATED;
+#endif
 
 	/** The set of user-marked frames */
 	UPROPERTY()
@@ -1341,9 +1322,6 @@ private:
 	UPROPERTY()
 	TArray<TObjectPtr<UMovieSceneGroupCondition>> GeneratedConditions;
 
-	UPROPERTY()
-	TObjectPtr<UFrameBasedMusicMap> MusicMap;
-
 #if WITH_EDITORONLY_DATA
 
 	/** Indicates whether this movie scene is read only */
diff --git a/Engine/Source/Runtime/MovieScene/Public/MovieSceneSection.h b/Engine/Source/Runtime/MovieScene/Public/MovieSceneSection.h
--- a/Engine/Source/Runtime/MovieScene/Public/MovieSceneSection.h
+++ b/Engine/Source/Runtime/MovieScene/Public/MovieSceneSection.h
@@ -55,7 +55,6 @@ struct FMovieSceneSequenceID;
 struct FPropertyChangedEvent;
 struct FQualifiedFrameTime;
 struct FMovieSceneTimeWarpVariant;
-class UFrameBasedMusicMap;
 
 namespace UE::MovieScene
 {
@@ -64,6 +63,7 @@ namespace UE::MovieScene
 	struct FFixedObjectBindingID;
 	struct FImportedEntity;
 	struct FSharedPlaybackState;
+	struct IRetimingInterface;
 }
 
 /** Enumeration defining how a section's channel proxy behaves. */
@@ -608,9 +608,7 @@ public:
 
 	/* Migrate the frame times of the movie scene section from the source frame rate to the destination frame rate */
 	virtual void MigrateFrameTimes(FFrameRate SourceRate, FFrameRate DestinationRate) {}
-
-	/* Remap the frame times of the movie scene section from the source music map to the destination music map */
-	MOVIESCENE_API virtual void RemapMusicalFrameTimes(const UFrameBasedMusicMap* OldMap, const UFrameBasedMusicMap* NewMap);
+	virtual void MigrateFrameTimes(const UE::MovieScene::IRetimingInterface& Retimer) {}
 
 	/** When guid bindings are updated to allow this section to fix-up any internal bindings */
 	virtual void OnBindingIDsUpdated(const TMap<UE::MovieScene::FFixedObjectBindingID, UE::MovieScene::FFixedObjectBindingID>& OldFixedToNewFixedMap, FMovieSceneSequenceID LocalSequenceID, TSharedRef<UE::MovieScene::FSharedPlaybackState> SharedPlaybackState) {}
diff --git a/Engine/Source/Runtime/MovieScene/Public/MovieSceneTimeController.h b/Engine/Source/Runtime/MovieScene/Public/MovieSceneTimeController.h
--- a/Engine/Source/Runtime/MovieScene/Public/MovieSceneTimeController.h
+++ b/Engine/Source/Runtime/MovieScene/Public/MovieSceneTimeController.h
@@ -6,7 +6,6 @@
 #include "Misc/QualifiedFrameTime.h"
 #include "MovieSceneFwd.h"
 #include "Misc/Attribute.h"
-#include "Misc/MusicalTime.h"
 #include "UObject/ObjectMacros.h"
 
 class IMovieScenePlayer;
@@ -44,16 +43,6 @@ public:
 	 */
 	MOVIESCENE_API FFrameTime RequestCurrentTime(const FQualifiedFrameTime& InCurrentTime, float InPlayRate, FFrameRate InDisplayRate);
 
-	MOVIESCENE_API bool SupportsMusicalTime();
-
-	/**
-	 * Request the current musical time based on the specified existing time and play rate.
-	 *
-	 * @param InCurrentTime    The current time of the sequence
-	 * @param InPlayRate       The current play rate of the sequence, multiplied by any world actor settings global dilation
-	 */
-	MOVIESCENE_API FMusicalTime RequestCurrentMusicalTime(const FQualifiedFrameTime& InCurrentTime, float InPlayRate, FFrameRate InDisplayRate);
-
 	/**
 	 * Called when the status of the owning IMovieScenePlayer has changed
 	 */
@@ -70,8 +59,6 @@ protected:
 	virtual void OnStartPlaying(const FQualifiedFrameTime& InStartTime){}
 	virtual void OnStopPlaying(const FQualifiedFrameTime& InStopTime){}
 	virtual FFrameTime OnRequestCurrentTime(const FQualifiedFrameTime& InCurrentTime, float InPlayRate) = 0;
-	virtual FMusicalTime OnRequestCurrentMusicalTime(const FQualifiedFrameTime& InCurrentTime, float InPlayRate) { return FMusicalTime(); }
-	virtual bool OnSupportsMusicalTime() { return false; }
 
 protected:
 
@@ -150,7 +137,7 @@ protected:
 	TOptional<FFrameTime> LastCurrentFrameTime;
 	TOptional<double> TimeSinceCurrentFrameTime;
 
-	double AccumulatedFrameTime;
+	double AccumulatedFrameTime = 0.0;
 };
 
 
diff --git a/Engine/Source/Runtime/MovieScene/Public/MovieSceneTimeHelpers.h b/Engine/Source/Runtime/MovieScene/Public/MovieSceneTimeHelpers.h
--- a/Engine/Source/Runtime/MovieScene/Public/MovieSceneTimeHelpers.h
+++ b/Engine/Source/Runtime/MovieScene/Public/MovieSceneTimeHelpers.h
@@ -14,29 +14,25 @@
 
 class UMovieScene;
 struct FFrameRate;
-class UFrameBasedMusicMap;
 
 namespace UE
 {
 namespace MovieScene
 {
 
+struct IRetimingInterface;
+
 class TimeHelpers
 {
 public:
 
-/**
- * Migrate the frame times of the movie scene from the source frame rate to the destination frame rate
- */
+	/**
+	 * Migrate the frame times of the movie scene from the source frame rate to the destination frame rate
+	 */
 
-static MOVIESCENE_API void MigrateFrameTimes(FFrameRate SourceRate, FFrameRate DestinationRate, UMovieScene* MovieScene, bool bApplyRecursively = false);
+	static MOVIESCENE_API void MigrateFrameTimes(FFrameRate SourceRate, FFrameRate DestinationRate, UMovieScene* MovieScene, bool bApplyRecursively = false);
 
-/**
- * Remap musical frame times from the current music map to a new music map. This remapping will cause the 
- * frames to maintain their musical position (bar.beat), but will likely cause them to change there position
- * in seconds as the new map may have different tempo(s).
- */
-static MOVIESCENE_API void RemapMusicalFrameTimes(UMovieScene* MovieScene, UFrameBasedMusicMap* NewMap);
+	static MOVIESCENE_API void MigrateFrameTimes(const IRetimingInterface& Retimer, UMovieScene* MovieScene, bool bApplyRecursively = false);
 
 };
 
diff --git a/Engine/Source/Runtime/MovieScene/Public/Sections/MovieSceneSectionTimingParameters.h b/Engine/Source/Runtime/MovieScene/Public/Sections/MovieSceneSectionTimingParameters.h
--- a/Engine/Source/Runtime/MovieScene/Public/Sections/MovieSceneSectionTimingParameters.h
+++ b/Engine/Source/Runtime/MovieScene/Public/Sections/MovieSceneSectionTimingParameters.h
@@ -151,25 +151,18 @@ struct FMovieSceneSectionTimingParametersFrames
 	 */
 	MOVIESCENE_API FMovieSceneSequenceTransform MakeTransform(const FFrameRate& OuterFrameRate, const TRange<FFrameNumber>& OuterRange, const FFrameRate& InnerFrameRate, const TRange<FFrameNumber>& InnerRange) const;
 
-	/**
-	 * Make a musical transform out of these parameters
-	 */
-	MOVIESCENE_API FMovieSceneSequenceTransform MakeMusicalTransform(const FFrameRate& OuterFrameRate, const TRange<FFrameNumber>& OuterRange, UMovieScene* OuterMovieScene, const FFrameRate& InnerFrameRate, const TRange<FFrameNumber>& InnerRange, UMovieScene* InnerMovieScene) const;
-
 	/**
 	 * Helpers
 	 */
-	void AddPositionInOuterAsOffset(FMovieSceneSequenceTransform& Transform, const FFrameTime& Offset) const;
-
-	void AddPlayRate(FMovieSceneSequenceTransform& Transform) const;
+	MOVIESCENE_API void AddPositionInOuterAsOffset(FMovieSceneSequenceTransform& Transform, const FFrameTime& Offset) const;
 
-	void AddFrameRateConversion(FMovieSceneSequenceTransform& Transform, const FFrameRate& OuterFrameRate, const FFrameRate& InnerFrameRate) const;
+	MOVIESCENE_API void AddPlayRate(FMovieSceneSequenceTransform& Transform) const;
 
-	void AddInnerStartOffset(FMovieSceneSequenceTransform& Transform, const FFrameTime& Offset) const;
+	MOVIESCENE_API void AddFrameRateConversion(FMovieSceneSequenceTransform& Transform, const FFrameRate& OuterFrameRate, const FFrameRate& InnerFrameRate) const;
 
-	void AddLoopingOrClampingAndReverse(FMovieSceneSequenceTransform& Result, const FFrameNumber& Duration, const FFrameNumber& Offset) const;
+	MOVIESCENE_API void AddInnerStartOffset(FMovieSceneSequenceTransform& Transform, const FFrameTime& Offset) const;
 
-	void AddMusicalTransform(FMovieSceneSequenceTransform& Result, const FFrameNumber& Offset, UMovieScene* OuterMovieScene, UMovieScene* InnerMovieScene) const;
+	MOVIESCENE_API void AddLoopingOrClampingAndReverse(FMovieSceneSequenceTransform& Result, const FFrameNumber& Duration, const FFrameNumber& Offset) const;
 
 };
 
diff --git a/Engine/Source/Runtime/MovieScene/Public/Sections/MovieSceneSubSection.h b/Engine/Source/Runtime/MovieScene/Public/Sections/MovieSceneSubSection.h
--- a/Engine/Source/Runtime/MovieScene/Public/Sections/MovieSceneSubSection.h
+++ b/Engine/Source/Runtime/MovieScene/Public/Sections/MovieSceneSubSection.h
@@ -36,8 +36,6 @@ class UMovieScene;
 class UMovieSceneEntitySystemLinker;
 class UMovieSceneSequence;
 class UObject;
-namespace UE { namespace MovieScene { struct FEntityImportParams; } }
-namespace UE { namespace MovieScene { struct FImportedEntity; } }
 struct FFrame;
 struct FFrameRate;
 struct FMovieSceneEvaluationTemplate;
@@ -47,6 +45,13 @@ struct FPropertyChangedEvent;
 struct FQualifiedFrameTime;
 struct FMovieSceneTransformMask;
 
+namespace UE::MovieScene
+{
+	struct IRetimingInterface;
+	struct FEntityImportParams;
+	struct FImportedEntity;
+}
+
 DECLARE_DELEGATE_OneParam(FOnSequenceChanged, UMovieSceneSequence* /*Sequence*/);
 
 struct FSubSequenceInstanceDataParams
@@ -228,7 +233,7 @@ public:
 	virtual TOptional<FFrameTime> GetOffsetTime() const override { return TOptional<FFrameTime>(FFrameTime(Parameters.StartFrameOffset)); }
 	MOVIESCENE_API virtual void GetSnapTimes(TArray<FFrameNumber>& OutSnapTimes, bool bGetSectionBorders) const override;
 	MOVIESCENE_API virtual void MigrateFrameTimes(FFrameRate SourceRate, FFrameRate DestinationRate) override;
-	MOVIESCENE_API virtual void RemapMusicalFrameTimes(const UFrameBasedMusicMap* OldMap, const UFrameBasedMusicMap* NewMap) override;
+	MOVIESCENE_API virtual void MigrateFrameTimes(const UE::MovieScene::IRetimingInterface& Retimer) override;
 	MOVIESCENE_API virtual FMovieSceneTimeWarpVariant* GetTimeWarp() override;
 	MOVIESCENE_API virtual UObject* GetSourceObject() const override;
 
diff --git a/Engine/Source/Runtime/MovieSceneTracks/Private/Channels/MovieSceneCameraShakeSourceTriggerChannel.cpp b/Engine/Source/Runtime/MovieSceneTracks/Private/Channels/MovieSceneCameraShakeSourceTriggerChannel.cpp
--- a/Engine/Source/Runtime/MovieSceneTracks/Private/Channels/MovieSceneCameraShakeSourceTriggerChannel.cpp
+++ b/Engine/Source/Runtime/MovieSceneTracks/Private/Channels/MovieSceneCameraShakeSourceTriggerChannel.cpp
@@ -34,14 +34,9 @@ void FMovieSceneCameraShakeSourceTriggerChannel::DeleteKeysFrom(FFrameNumber InT
 	GetData().DeleteKeysFrom(InTime, bDeleteKeysBefore);
 }
 
-void FMovieSceneCameraShakeSourceTriggerChannel::ChangeFrameResolution(FFrameRate SourceRate, FFrameRate DestinationRate)
+void FMovieSceneCameraShakeSourceTriggerChannel::RemapTimes(const UE::MovieScene::IRetimingInterface& Retimer)
 {
-	GetData().ChangeFrameResolution(SourceRate, DestinationRate);
-}
-
-void FMovieSceneCameraShakeSourceTriggerChannel::RemapMusicalFrameTimes(const UFrameBasedMusicMap* OldMap, const UFrameBasedMusicMap* NewMap)
-{
-	GetData().RemapMusicalFrameTimes(OldMap, NewMap);
+	GetData().RemapTimes(Retimer);
 }
 
 TRange<FFrameNumber> FMovieSceneCameraShakeSourceTriggerChannel::ComputeEffectiveRange() const
diff --git a/Engine/Source/Runtime/MovieSceneTracks/Private/Channels/MovieSceneEventChannel.cpp b/Engine/Source/Runtime/MovieSceneTracks/Private/Channels/MovieSceneEventChannel.cpp
--- a/Engine/Source/Runtime/MovieSceneTracks/Private/Channels/MovieSceneEventChannel.cpp
+++ b/Engine/Source/Runtime/MovieSceneTracks/Private/Channels/MovieSceneEventChannel.cpp
@@ -34,15 +34,11 @@ void FMovieSceneEventChannel::DeleteKeysFrom(FFrameNumber InTime, bool bDeleteKe
 	GetData().DeleteKeysFrom(InTime, bDeleteKeysBefore);
 }
 
-void FMovieSceneEventChannel::ChangeFrameResolution(FFrameRate SourceRate, FFrameRate DestinationRate)
+void FMovieSceneEventChannel::RemapTimes(const UE::MovieScene::IRetimingInterface& Retimer)
 {
-	GetData().ChangeFrameResolution(SourceRate, DestinationRate);
+	GetData().RemapTimes(Retimer);
 }
 
-void FMovieSceneEventChannel::RemapMusicalFrameTimes(const UFrameBasedMusicMap* OldMap, const UFrameBasedMusicMap* NewMap)
-{
-	GetData().RemapMusicalFrameTimes(OldMap, NewMap);
-}
 
 TRange<FFrameNumber> FMovieSceneEventChannel::ComputeEffectiveRange() const
 {
diff --git a/Engine/Source/Runtime/MovieSceneTracks/Private/Channels/MovieSceneStringChannel.cpp b/Engine/Source/Runtime/MovieSceneTracks/Private/Channels/MovieSceneStringChannel.cpp
--- a/Engine/Source/Runtime/MovieSceneTracks/Private/Channels/MovieSceneStringChannel.cpp
+++ b/Engine/Source/Runtime/MovieSceneTracks/Private/Channels/MovieSceneStringChannel.cpp
@@ -102,14 +102,9 @@ int32 FMovieSceneStringChannel::GetIndex(FKeyHandle Handle)
 	return GetData().GetIndex(Handle);
 }
 
-void FMovieSceneStringChannel::ChangeFrameResolution(FFrameRate SourceRate, FFrameRate DestinationRate)
+void FMovieSceneStringChannel::RemapTimes(const UE::MovieScene::IRetimingInterface& Retimer)
 {
-	GetData().ChangeFrameResolution(SourceRate, DestinationRate);
-}
-
-void FMovieSceneStringChannel::RemapMusicalFrameTimes(const UFrameBasedMusicMap* OldMap, const UFrameBasedMusicMap* NewMap)
-{
-	GetData().RemapMusicalFrameTimes(OldMap, NewMap);
+	GetData().RemapTimes(Retimer);
 }
 
 TRange<FFrameNumber> FMovieSceneStringChannel::ComputeEffectiveRange() const
diff --git a/Engine/Source/Runtime/MovieSceneTracks/Private/Sections/MovieSceneActorReferenceSection.cpp b/Engine/Source/Runtime/MovieSceneTracks/Private/Sections/MovieSceneActorReferenceSection.cpp
--- a/Engine/Source/Runtime/MovieSceneTracks/Private/Sections/MovieSceneActorReferenceSection.cpp
+++ b/Engine/Source/Runtime/MovieSceneTracks/Private/Sections/MovieSceneActorReferenceSection.cpp
@@ -70,14 +70,9 @@ int32 FMovieSceneActorReferenceData::GetIndex(FKeyHandle Handle)
 	return GetData().GetIndex(Handle);
 }
 
-void FMovieSceneActorReferenceData::ChangeFrameResolution(FFrameRate SourceRate, FFrameRate DestinationRate)
+void FMovieSceneActorReferenceData::RemapTimes(const UE::MovieScene::IRetimingInterface& Retimer)
 {
-	GetData().ChangeFrameResolution(SourceRate, DestinationRate);
-}
-
-void FMovieSceneActorReferenceData::RemapMusicalFrameTimes(const UFrameBasedMusicMap* OldMap, const UFrameBasedMusicMap* NewMap)
-{
-	GetData().RemapMusicalFrameTimes(OldMap, NewMap);
+	GetData().RemapTimes(Retimer);
 }
 
 TRange<FFrameNumber> FMovieSceneActorReferenceData::ComputeEffectiveRange() const
diff --git a/Engine/Source/Runtime/MovieSceneTracks/Private/Sections/MovieSceneEventSection.cpp b/Engine/Source/Runtime/MovieSceneTracks/Private/Sections/MovieSceneEventSection.cpp
--- a/Engine/Source/Runtime/MovieSceneTracks/Private/Sections/MovieSceneEventSection.cpp
+++ b/Engine/Source/Runtime/MovieSceneTracks/Private/Sections/MovieSceneEventSection.cpp
@@ -369,14 +369,9 @@ void FMovieSceneEventSectionData::DeleteKeysFrom(FFrameNumber InTime, bool bDele
 	GetData().DeleteKeysFrom(InTime, bDeleteKeysBefore);
 }
 
-void FMovieSceneEventSectionData::ChangeFrameResolution(FFrameRate SourceRate, FFrameRate DestinationRate)
+void FMovieSceneEventSectionData::RemapTimes(const UE::MovieScene::IRetimingInterface& Retimer)
 {
-	GetData().ChangeFrameResolution(SourceRate, DestinationRate);
-}
-
-void FMovieSceneEventSectionData::RemapMusicalFrameTimes(const UFrameBasedMusicMap* OldMap, const UFrameBasedMusicMap* NewMap)
-{
-	GetData().RemapMusicalFrameTimes(OldMap, NewMap);
+	GetData().RemapTimes(Retimer);
 }
 
 TRange<FFrameNumber> FMovieSceneEventSectionData::ComputeEffectiveRange() const
diff --git a/Engine/Source/Runtime/MovieSceneTracks/Public/Channels/MovieSceneCameraShakeSourceTriggerChannel.h b/Engine/Source/Runtime/MovieSceneTracks/Public/Channels/MovieSceneCameraShakeSourceTriggerChannel.h
--- a/Engine/Source/Runtime/MovieSceneTracks/Public/Channels/MovieSceneCameraShakeSourceTriggerChannel.h
+++ b/Engine/Source/Runtime/MovieSceneTracks/Public/Channels/MovieSceneCameraShakeSourceTriggerChannel.h
@@ -81,8 +81,7 @@ public:
 	MOVIESCENETRACKS_API virtual void DuplicateKeys(TArrayView<const FKeyHandle> InHandles, TArrayView<FKeyHandle> OutNewHandles) override;
 	MOVIESCENETRACKS_API virtual void DeleteKeys(TArrayView<const FKeyHandle> InHandles) override;
 	MOVIESCENETRACKS_API virtual void DeleteKeysFrom(FFrameNumber InTime, bool bDeleteKeysBefore) override;
-	MOVIESCENETRACKS_API virtual void ChangeFrameResolution(FFrameRate SourceRate, FFrameRate DestinationRate) override;
-	MOVIESCENETRACKS_API virtual void RemapMusicalFrameTimes(const UFrameBasedMusicMap* OldMap, const UFrameBasedMusicMap* NewMap) override;
+	MOVIESCENETRACKS_API virtual void RemapTimes(const UE::MovieScene::IRetimingInterface& Retimer) override;
 	MOVIESCENETRACKS_API virtual TRange<FFrameNumber> ComputeEffectiveRange() const override;
 	MOVIESCENETRACKS_API virtual int32 GetNumKeys() const override;
 	MOVIESCENETRACKS_API virtual void Reset() override;
diff --git a/Engine/Source/Runtime/MovieSceneTracks/Public/Channels/MovieSceneEventChannel.h b/Engine/Source/Runtime/MovieSceneTracks/Public/Channels/MovieSceneEventChannel.h
--- a/Engine/Source/Runtime/MovieSceneTracks/Public/Channels/MovieSceneEventChannel.h
+++ b/Engine/Source/Runtime/MovieSceneTracks/Public/Channels/MovieSceneEventChannel.h
@@ -55,8 +55,7 @@ public:
 	MOVIESCENETRACKS_API virtual void DuplicateKeys(TArrayView<const FKeyHandle> InHandles, TArrayView<FKeyHandle> OutNewHandles) override;
 	MOVIESCENETRACKS_API virtual void DeleteKeys(TArrayView<const FKeyHandle> InHandles) override;
 	MOVIESCENETRACKS_API virtual void DeleteKeysFrom(FFrameNumber InTime, bool bDeleteKeysBefore) override;
-	MOVIESCENETRACKS_API virtual void ChangeFrameResolution(FFrameRate SourceRate, FFrameRate DestinationRate) override;
-	MOVIESCENETRACKS_API virtual void RemapMusicalFrameTimes(const UFrameBasedMusicMap* OldMap, const UFrameBasedMusicMap* NewMap) override;
+	MOVIESCENETRACKS_API virtual void RemapTimes(const UE::MovieScene::IRetimingInterface& Retimer) override;
 	MOVIESCENETRACKS_API virtual TRange<FFrameNumber> ComputeEffectiveRange() const override;
 	MOVIESCENETRACKS_API virtual int32 GetNumKeys() const override;
 	MOVIESCENETRACKS_API virtual void Reset() override;
diff --git a/Engine/Source/Runtime/Slate/Private/Widgets/Input/SSpinBox.cpp b/Engine/Source/Runtime/Slate/Private/Widgets/Input/SSpinBox.cpp
--- a/Engine/Source/Runtime/Slate/Private/Widgets/Input/SSpinBox.cpp
+++ b/Engine/Source/Runtime/Slate/Private/Widgets/Input/SSpinBox.cpp
@@ -53,9 +53,15 @@ void SSpinBox<NumericType>::Construct(const FArguments& InArgs)
 	Style = InArgs._Style;
 
 	SetForegroundColor(InArgs._Style->ForegroundColor);
-	Interface = InArgs._TypeInterface.IsValid() ? InArgs._TypeInterface : MakeShareable(new TDefaultNumericTypeInterface<NumericType>);
+	InterfaceAttr = InArgs._TypeInterface;
 
-	if (InArgs._TypeInterface.IsValid() && Interface->GetOnSettingChanged())
+	if (!InterfaceAttr.IsBound() && InterfaceAttr.Get() == nullptr)
+	{
+		InterfaceAttr = MakeShared<TDefaultNumericTypeInterface<NumericType>>();
+	}
+
+	TSharedPtr<INumericTypeInterface<NumericType>> Interface = InterfaceAttr.Get();
+	if (Interface->GetOnSettingChanged())
 	{
 		Interface->GetOnSettingChanged()->AddSP(this, &SSpinBox::ResetCachedValueString);
 	}
@@ -596,6 +602,8 @@ FReply SSpinBox<NumericType>::OnMouseWheel(const FGeometry& MyGeometry, const FP
 		const double NewValue = InternalValue + (Sign * Step);
 		const NumericType RoundedNewValue = RoundIfIntegerValue(NewValue);
 
+		TSharedPtr<INumericTypeInterface<NumericType>> Interface = InterfaceAttr.Get();
+
 		// First SetText is to update the value before calling CommitValue. Otherwise, when the text lose
 		// focus from the CommitValue, it will override the value we just committed.
 		// The second SetText is to update the text to the InternalValue since it could have been clamped.
@@ -761,23 +769,23 @@ void SSpinBox<NumericType>::SetMaxSliderValue(const TAttribute<TOptional<Numeric
 
 template <typename NumericType>
 int32 SSpinBox<NumericType>::GetMinFractionalDigits() const
-{ return Interface->GetMinFractionalDigits(); }
+{ return InterfaceAttr.Get()->GetMinFractionalDigits(); }
 
 template <typename NumericType>
 void SSpinBox<NumericType>::SetMinFractionalDigits(const TAttribute<TOptional<int32>>& InMinFractionalDigits)
 {
-	Interface->SetMinFractionalDigits((InMinFractionalDigits.Get().IsSet()) ? InMinFractionalDigits.Get() : MinFractionalDigits);
+	InterfaceAttr.Get()->SetMinFractionalDigits((InMinFractionalDigits.Get().IsSet()) ? InMinFractionalDigits.Get() : MinFractionalDigits);
 	bCachedValueStringDirty = true;
 }
 
 template <typename NumericType>
 int32 SSpinBox<NumericType>::GetMaxFractionalDigits() const
-{ return Interface->GetMaxFractionalDigits(); }
+{ return InterfaceAttr.Get()->GetMaxFractionalDigits(); }
 
 template <typename NumericType>
 void SSpinBox<NumericType>::SetMaxFractionalDigits(const TAttribute<TOptional<int32>>& InMaxFractionalDigits)
 {
-	Interface->SetMaxFractionalDigits((InMaxFractionalDigits.Get().IsSet()) ? InMaxFractionalDigits.Get() : MaxFractionalDigits);
+	InterfaceAttr.Get()->SetMaxFractionalDigits((InMaxFractionalDigits.Get().IsSet()) ? InMaxFractionalDigits.Get() : MaxFractionalDigits);
 	bCachedValueStringDirty = true;
 }
 
@@ -873,7 +881,7 @@ FString SSpinBox<NumericType>::GetValueAsString() const
 	}
 
 	bCachedValueStringDirty = false;
-	return Interface->ToString(CurrentValue);
+	return InterfaceAttr.Get()->ToString(CurrentValue);
 }
 
 template <typename NumericType>
@@ -893,6 +901,8 @@ void SSpinBox<NumericType>::TextField_OnTextChanged(const FText& NewText)
 		FString Data = NewText.ToString();
 		int32 NumValidChars = Data.Len();
 
+		TSharedPtr<INumericTypeInterface<NumericType>> Interface = InterfaceAttr.Get();
+
 		for (int32 Index = 0; Index < Data.Len(); ++Index)
 		{
 			if (!Interface->IsCharacterValid(Data[Index]))
@@ -928,6 +938,7 @@ void SSpinBox<NumericType>::TextField_OnTextCommitted(const FText& NewText, ETex
 		ExitTextMode();
 	}
 
+	TSharedPtr<INumericTypeInterface<NumericType>> Interface = InterfaceAttr.Get();
 	TOptional<NumericType> NewValue = Interface->FromString(NewText.ToString(), ValueAttribute.Get());
 	if (NewValue.IsSet())
 	{
@@ -1011,6 +1022,8 @@ void SSpinBox<NumericType>::CommitValue(NumericType NewValue, double NewSpinValu
 	const NumericType CurrentValue = ValueAttribute.Get();
 	if (CachedExternalValue != CurrentValue || bCachedValueStringDirty)
 	{
+		TSharedPtr<INumericTypeInterface<NumericType>> Interface = InterfaceAttr.Get();
+
 		CachedExternalValue = ValueAttribute.Get();
 		CachedValueString = Interface->ToString(CachedExternalValue);
 		bCachedValueStringDirty = false;
@@ -1064,7 +1077,7 @@ float SSpinBox<NumericType>::GetTextMinDesiredWidth() const
 template <typename NumericType>
 bool SSpinBox<NumericType>::IsCharacterValid(TCHAR InChar) const
 {
-	return Interface->IsCharacterValid(InChar);
+	return InterfaceAttr.Get()->IsCharacterValid(InChar);
 }
 
 template <typename NumericType>
@@ -1101,7 +1114,7 @@ void SSpinBox<NumericType>::ResetCachedValueString()
 { 
 	const NumericType CurrentValue = ValueAttribute.Get();
 	CachedExternalValue = CurrentValue;
-	CachedValueString = Interface->ToString(CachedExternalValue);
+	CachedValueString = InterfaceAttr.Get()->ToString(CachedExternalValue);
 }
 
 // Explicit instantiation of the numeric types valid for this template
diff --git a/Engine/Source/Runtime/Slate/Public/Widgets/Input/SSpinBox.h b/Engine/Source/Runtime/Slate/Public/Widgets/Input/SSpinBox.h
--- a/Engine/Source/Runtime/Slate/Public/Widgets/Input/SSpinBox.h
+++ b/Engine/Source/Runtime/Slate/Public/Widgets/Input/SSpinBox.h
@@ -172,7 +172,7 @@ public:
 		/** What keyboard to display. */
 		SLATE_ATTRIBUTE(EKeyboardType, KeyboardType)
 		/** Provide custom type conversion functionality to this spin box */
-		SLATE_ARGUMENT(TSharedPtr< INumericTypeInterface<NumericType> >, TypeInterface)
+		SLATE_ATTRIBUTE(TSharedPtr< INumericTypeInterface<NumericType> >, TypeInterface)
 		/** If refresh requests for the viewport should happen for all value changes **/
 		SLATE_ARGUMENT(bool, PreventThrottling)
 		/** Menu extender for the right-click context menu */
@@ -400,7 +400,7 @@ private:
 	TSharedPtr<SEditableText> EditableText;
 
 	/** Interface that defines conversion functionality for the templated type */
-	TSharedPtr< INumericTypeInterface<NumericType> > Interface;
+	TAttribute<TSharedPtr< INumericTypeInterface<NumericType> >> InterfaceAttr;
 
 	/** True when no range is specified, spinner can be spun indefinitely */
 	bool bUnlimitedSpinRange;
diff --git a/Engine/Source/Runtime/TimeManagement/Private/FrameNumberTimeEvaluator.cpp b/Engine/Source/Runtime/TimeManagement/Private/FrameNumberTimeEvaluator.cpp
--- a/Engine/Source/Runtime/TimeManagement/Private/FrameNumberTimeEvaluator.cpp
+++ b/Engine/Source/Runtime/TimeManagement/Private/FrameNumberTimeEvaluator.cpp
@@ -405,64 +405,4 @@ TValueOrError<FFrameTime, FExpressionError> FFrameNumberTimeEvaluator::EvaluateT
 	OutDirectlyParsed = false;
 	return MakeError(NSLOCTEXT("TimeManagement", "UnrecognizedTimeResult", "Unrecognized result returned from expression"));
 
-}
-
-TValueOrError<FFrameTime, FExpressionError> FFrameNumberTimeEvaluator::EvaluateBarBeat(const TCHAR* InExpression, const UFrameBasedMusicMap* MusicMap, FrameBasedMusicMap::EBarBeatFormat BarBeatFormat, bool& OutDirectlyParsed) const
-{
-	OutDirectlyParsed = false;
-
-	if (!MusicMap || MusicMap->IsEmpty())
-	{
-		return MakeError(NSLOCTEXT("TimeManagement", "BarBeatError_NoTempoMap", "No supplied TempoMap"));
-	}
-
-	FString InAsString(InExpression);
-
-	// find bar.beat boundary and replace with space...
-	int32 PeriodCount = 0;
-	for (TCHAR& Character : InAsString)
-	{
-		if (Character == '.')
-		{
-			PeriodCount++;
-			if (PeriodCount == 1)
-			{
-				Character = ' ';
-			}
-		}
-	}
-
-	if (PeriodCount != 2)
-	{
-		return MakeError(NSLOCTEXT("TimeManagement", "BarBeatError_BadFormat", "Format not recognized as a Bar.Beat"));
-	}
-
-	TValueOrError<TArray<FExpressionToken>, FExpressionError> LexResult = ExpressionParser::Lex(*InAsString, TimeTokenDefinitions);
-	if (!LexResult.IsValid())
-	{
-		return MakeError(LexResult.StealError());
-	}
-	TArray<FExpressionToken> Tokens = LexResult.StealValue();
-	if (Tokens.Num() != 2)
-	{
-		return MakeError(NSLOCTEXT("TimeManagement", "BarBeatError_BadFormat", "Format not recognized as a Bar.Beat"));
-	}
-
-	if (!Tokens[0].Node.Cast<double>() || !Tokens[1].Node.Cast<double>())
-	{
-		return MakeError(NSLOCTEXT("TimeManagement", "BarBeatError_BadFormat", "Format not recognized as a Bar.Beat"));
-	}
-
-	int32 Bar = static_cast<int32>(*Tokens[0].Node.Cast<double>());
-	float BeatInBar = static_cast<float>(*Tokens[1].Node.Cast<double>());
-	if (BeatInBar < 1.0f && BarBeatFormat == FrameBasedMusicMap::EBarBeatFormat::Position)
-	{
-		return MakeError(NSLOCTEXT("TimeManagement", "BarBeatError_InvalidBeatValue", "Beats are one based, so the specified beat must be >= 1.0"));
-	}
-	
-	// in user land bars and beats are 1 based, so make them 0 based.
-	OutDirectlyParsed = true;
-	int32 BarCorrection = BarBeatFormat == FrameBasedMusicMap::EBarBeatFormat::Position ? 1 : 0;
-	float BeatCorrection = BarBeatFormat == FrameBasedMusicMap::EBarBeatFormat::Position ? 1.0f : 0.0f;
-	return MakeValue(MusicMap->GetFrameTime(Bar - BarCorrection, BeatInBar - BeatCorrection));
-}
+}
\ No newline at end of file
diff --git a/Engine/Source/Runtime/TimeManagement/Public/FrameNumberDisplayFormat.h b/Engine/Source/Runtime/TimeManagement/Public/FrameNumberDisplayFormat.h
--- a/Engine/Source/Runtime/TimeManagement/Public/FrameNumberDisplayFormat.h
+++ b/Engine/Source/Runtime/TimeManagement/Public/FrameNumberDisplayFormat.h
@@ -18,7 +18,7 @@ enum class EFrameNumberDisplayFormats : uint8
 	
 	Frames,
 
-	BarsBeats UMETA(DisplayName = "Bars : Beats"),
+	Custom UMETA(Hidden),
 
 	MAX_Count UMETA(Hidden)
 };
\ No newline at end of file
diff --git a/Engine/Source/Runtime/TimeManagement/Public/FrameNumberNumericInterface.h b/Engine/Source/Runtime/TimeManagement/Public/FrameNumberNumericInterface.h
--- a/Engine/Source/Runtime/TimeManagement/Public/FrameNumberNumericInterface.h
+++ b/Engine/Source/Runtime/TimeManagement/Public/FrameNumberNumericInterface.h
@@ -6,7 +6,6 @@
 #include "FrameNumberDisplayFormat.h"
 #include "Widgets/Input/NumericTypeInterface.h"
 #include "Misc/FrameRate.h"
-#include "FrameBasedMusicMap.h"
 #include "FrameNumberTimeEvaluator.h"
 #include "Misc/Timecode.h"
 #include "Misc/ExpressionParserTypes.h"
@@ -30,14 +29,12 @@
 */
 struct FFrameNumberInterface : public INumericTypeInterface<double>
 {
-	FFrameNumberInterface(const TAttribute<EFrameNumberDisplayFormats>& InDisplayFormatAttr, const TAttribute<uint8>& InOnGetZeroPadFrameNumber, const TAttribute<FFrameRate>& InTickResolutionAttr, const TAttribute<FFrameRate>& InDisplayRateAttr, const TAttribute<const UFrameBasedMusicMap*>& InMusicMap = TAttribute<const UFrameBasedMusicMap*>(), FrameBasedMusicMap::EBarBeatFormat InBarBeatFormat = FrameBasedMusicMap::EBarBeatFormat::Position)
+	FFrameNumberInterface(const TAttribute<EFrameNumberDisplayFormats>& InDisplayFormatAttr, const TAttribute<uint8>& InOnGetZeroPadFrameNumber, const TAttribute<FFrameRate>& InTickResolutionAttr, const TAttribute<FFrameRate>& InDisplayRateAttr)
 		: DisplayFormatAttr(InDisplayFormatAttr)
 		, TickResolutionAttr(InTickResolutionAttr)
 		, DisplayRateAttr(InDisplayRateAttr)
 		, ZeroPadFramesAttr(InOnGetZeroPadFrameNumber)
-		, MusicMapAttr(InMusicMap)
 		, SubframeIndicator(TEXT("*"))
-		, BarBeatFormat(InBarBeatFormat)
 	{
 		check(InDisplayFormatAttr.IsSet());
 		check(InTickResolutionAttr.IsSet());
@@ -87,13 +84,6 @@ struct FFrameNumberInterface : public INumericTypeInterface<double>
 			Format = EFrameNumberDisplayFormats::NonDropFrameTimecode;
 		}
 
-		// Don't show bars:beats under various conditions...
-		if (Format == EFrameNumberDisplayFormats::BarsBeats &&
-			(!MusicMapAttr.IsSet() || !MusicMapAttr.Get() || MusicMapAttr.Get()->IsEmpty() || BarBeatFormat == FrameBasedMusicMap::EBarBeatFormat::NotAllowed))
-		{
-			Format = EFrameNumberDisplayFormats::Seconds;
-		}
-
 		switch (Format)
 		{
 		case EFrameNumberDisplayFormats::Frames:
@@ -126,11 +116,6 @@ struct FFrameNumberInterface : public INumericTypeInterface<double>
 			FTimecode AsNonDropTimecode = FTimecode::FromFrameNumber(DisplayTime.FloorToFrame(), DestinationFrameRate, bIsDropTimecode);
 			return FString::Printf(TEXT("[%s%s]"), *AsNonDropTimecode.ToString(false), *SubframeSuffix);
 		}
-		case EFrameNumberDisplayFormats::BarsBeats:
-		{
-			check(MusicMapAttr.IsSet() && MusicMapAttr.Get());
-			return MusicMapAttr.Get()->MakeBarBeatStringForFrame(Value, BarBeatFormat);
-		}
 		default:
 			return FString(TEXT("Unsupported Format"));
 		}
@@ -141,14 +126,6 @@ struct FFrameNumberInterface : public INumericTypeInterface<double>
 		FFrameRate SourceFrameRate = DisplayRateAttr.Get();
 		FFrameRate DestinationFrameRate = TickResolutionAttr.Get();
 		EFrameNumberDisplayFormats FallbackFormat = DisplayFormatAttr.Get();
-		const UFrameBasedMusicMap* MusicMap = MusicMapAttr.Get();
-
-		// If we are currently displaying bars.beats but this specific property field does not allow bars.beats 
-		// display/entry, force the fallback here to be seconds...
-		if (FallbackFormat == EFrameNumberDisplayFormats::BarsBeats && BarBeatFormat == FrameBasedMusicMap::EBarBeatFormat::NotAllowed)
-		{
-			FallbackFormat = EFrameNumberDisplayFormats::Seconds;
-		}
 
 		// We allow input in any format (time, frames or timecode) and we just convert it into into the internal sequence resolution.
 		// The user's input can be ambiguous though (does "5" mean 5 frames or 5 seconds?) so when we check each possible result we
@@ -162,14 +139,12 @@ struct FFrameNumberInterface : public INumericTypeInterface<double>
 		TValueOrError<FFrameTime, FExpressionError> FrameResult = Eval.EvaluateFrame(*InString, SourceFrameRate, DestinationFrameRate, /*Out*/ bWasFrameText);
 		bool bWasTimeText;
 		TValueOrError<FFrameTime, FExpressionError> TimeResult = Eval.EvaluateTime(*InString, DestinationFrameRate, /*Out*/ bWasTimeText);
-		bool bWasBarBeatText;
-		TValueOrError<FFrameTime, FExpressionError> BarBeatResult = Eval.EvaluateBarBeat(*InString, MusicMap, BarBeatFormat, /*Out*/ bWasBarBeatText);
 
 
 		// All three formats support ambiguous conversion where the user can enter "5" and wants it in the logical unit based on the current
 		// display format. This means 5 -> 5f, 5 ->5s, and 5 -> 5 frames (in timecode). We also support specifically specifying in a different
 		// format than your current display, ie if your display is in frames and you enter 1s, you get 1s in frames or 30 (at 30fps play rate).
-		if (!bWasTimecodeText && !bWasFrameText && !bWasTimeText && !bWasBarBeatText)
+		if (!bWasTimecodeText && !bWasFrameText && !bWasTimeText)
 		{
 			// They've entered an ambiguous number, so we'll check the display format and see if we did successfully parse as that type. If it
 			// was able to parse the input for the given display format, we return that.
@@ -231,10 +206,6 @@ struct FFrameNumberInterface : public INumericTypeInterface<double>
 		{
 			return TOptional<double>(FrameResult.GetValue().GetFrame().Value);
 		}
-		else if (bWasBarBeatText)
-		{
-			return TOptional<double>(BarBeatResult.GetValue().GetFrame().Value);
-		}
 
 		// We're not sure what they typed in
 		return TOptional<double>();
@@ -267,10 +238,8 @@ private:
 	TAttribute<FFrameRate> TickResolutionAttr;
 	TAttribute<FFrameRate> DisplayRateAttr;
 	TAttribute<uint8> ZeroPadFramesAttr;
-	TAttribute<const UFrameBasedMusicMap*> MusicMapAttr;
 
 	FString SubframeIndicator;
-	FrameBasedMusicMap::EBarBeatFormat BarBeatFormat;
 
 	FOnSettingChanged OnSettingChangedEvent;
 };
diff --git a/Engine/Source/Runtime/TimeManagement/Public/FrameNumberTimeEvaluator.h b/Engine/Source/Runtime/TimeManagement/Public/FrameNumberTimeEvaluator.h
--- a/Engine/Source/Runtime/TimeManagement/Public/FrameNumberTimeEvaluator.h
+++ b/Engine/Source/Runtime/TimeManagement/Public/FrameNumberTimeEvaluator.h
@@ -3,7 +3,6 @@
 #pragma once
 
 #include "CoreTypes.h"
-#include "FrameBasedMusicMap.h"
 #include "Math/BasicMathExpressionEvaluator.h"
 #include "Misc/ExpressionParserTypes.h"
 #include "Misc/Optional.h"
@@ -31,7 +30,6 @@ struct FFrameNumberTimeEvaluator
 	TIMEMANAGEMENT_API TValueOrError<FFrameTime, FExpressionError> EvaluateTimecode(const TCHAR* InExpression, const FFrameRate& InDisplayFrameRate, const FFrameRate& InTickResolution, bool& OutDirectlyParsed) const;
 	TIMEMANAGEMENT_API TValueOrError<FFrameTime, FExpressionError> EvaluateFrame(const TCHAR* InExpression, const FFrameRate& InDisplayFrameRate, const FFrameRate& InTickResolution, bool& OutDirectlyParsed) const;
 	TIMEMANAGEMENT_API TValueOrError<FFrameTime, FExpressionError> EvaluateTime(const TCHAR* InExpression, FFrameRate InFrameRate, bool& OutDirectlyParsed) const;
-	TIMEMANAGEMENT_API TValueOrError<FFrameTime, FExpressionError> EvaluateBarBeat(const TCHAR* InExpression, const UFrameBasedMusicMap* MusicMap, FrameBasedMusicMap::EBarBeatFormat BarBeatFormat, bool& OutDirectlyParsed) const;
 
 	FTokenDefinitions TimecodeTokenDefinitions;
 	FTokenDefinitions FrameTokenDefinitions;
