Sequencer: Added missing CurveValueType typedefs, SupportsDefaults, and EvaluateChannel

#rnx

diff --git a/Engine/Plugins/MovieScene/MovieSceneTextTrack/Source/MovieSceneTextTrack/Public/MovieSceneTextChannel.h b/Engine/Plugins/MovieScene/MovieSceneTextTrack/Source/MovieSceneTextTrack/Public/MovieSceneTextChannel.h
--- a/Engine/Plugins/MovieScene/MovieSceneTextTrack/Source/MovieSceneTextTrack/Public/MovieSceneTextChannel.h
+++ b/Engine/Plugins/MovieScene/MovieSceneTextTrack/Source/MovieSceneTextTrack/Public/MovieSceneTextChannel.h
@@ -21,6 +21,8 @@ struct MOVIESCENETEXTTRACK_API FMovieSceneTextChannel : public FMovieSceneChanne
 {
 	GENERATED_BODY()
 
+	typedef FText CurveValueType;
+
 #if WITH_EDITOR
 	void SetPackage(UPackage* InPackage);
 
diff --git a/Engine/Source/Runtime/MovieScene/Public/Channels/MovieSceneBoolChannel.h b/Engine/Source/Runtime/MovieScene/Public/Channels/MovieSceneBoolChannel.h
--- a/Engine/Source/Runtime/MovieScene/Public/Channels/MovieSceneBoolChannel.h
+++ b/Engine/Source/Runtime/MovieScene/Public/Channels/MovieSceneBoolChannel.h
@@ -32,6 +32,8 @@ struct FMovieSceneBoolChannel : public FMovieSceneChannel
 {
 	GENERATED_BODY()
 
+	typedef bool CurveValueType;
+
 	FMovieSceneBoolChannel()
 		: PreInfinityExtrap(RCCE_Constant), PostInfinityExtrap(RCCE_Constant), DefaultValue(false), bHasDefaultValue(false)
 	{}
diff --git a/Engine/Source/Runtime/MovieScene/Public/Channels/MovieSceneByteChannel.h b/Engine/Source/Runtime/MovieScene/Public/Channels/MovieSceneByteChannel.h
--- a/Engine/Source/Runtime/MovieScene/Public/Channels/MovieSceneByteChannel.h
+++ b/Engine/Source/Runtime/MovieScene/Public/Channels/MovieSceneByteChannel.h
@@ -33,6 +33,8 @@ struct FMovieSceneByteChannel : public FMovieSceneChannel
 {
 	GENERATED_BODY()
 
+	typedef uint8 CurveValueType;
+
 	FMovieSceneByteChannel()
 		: PreInfinityExtrap(RCCE_Constant)
 		, PostInfinityExtrap(RCCE_Constant)
diff --git a/Engine/Source/Runtime/MovieScene/Public/Channels/MovieSceneIntegerChannel.h b/Engine/Source/Runtime/MovieScene/Public/Channels/MovieSceneIntegerChannel.h
--- a/Engine/Source/Runtime/MovieScene/Public/Channels/MovieSceneIntegerChannel.h
+++ b/Engine/Source/Runtime/MovieScene/Public/Channels/MovieSceneIntegerChannel.h
@@ -33,6 +33,8 @@ struct FMovieSceneIntegerChannel : public FMovieSceneChannel
 {
 	GENERATED_BODY()
 
+	typedef int32 CurveValueType;
+
 	FMovieSceneIntegerChannel()
 		: PreInfinityExtrap(RCCE_Constant)
 		, PostInfinityExtrap(RCCE_Constant)
diff --git a/Engine/Source/Runtime/MovieScene/Public/Channels/MovieSceneObjectPathChannel.h b/Engine/Source/Runtime/MovieScene/Public/Channels/MovieSceneObjectPathChannel.h
--- a/Engine/Source/Runtime/MovieScene/Public/Channels/MovieSceneObjectPathChannel.h
+++ b/Engine/Source/Runtime/MovieScene/Public/Channels/MovieSceneObjectPathChannel.h
@@ -71,6 +71,8 @@ struct FMovieSceneObjectPathChannel : public FMovieSceneChannel
 {
 	GENERATED_BODY()
 
+	typedef FMovieSceneObjectPathChannelKeyValue CurveValueType;
+
 	/** Default constructor */
 	FMovieSceneObjectPathChannel()
 		: PropertyClass(nullptr)
@@ -146,6 +148,11 @@ public:
 		DefaultValue = InDefaultValue;
 	}
 
+	FORCEINLINE void SetDefault(FMovieSceneObjectPathChannelKeyValue InDefaultValue)
+	{
+		DefaultValue = InDefaultValue;
+	}
+
 	/**
 	 * Get this channel's default value that will be used when no keys are present
 	 *
@@ -193,3 +200,11 @@ struct TMovieSceneChannelTraits<FMovieSceneObjectPathChannel> : TMovieSceneChann
 
 #endif
 };
+
+inline bool EvaluateChannel(const FMovieSceneObjectPathChannel* InChannel, FFrameTime InTime, FMovieSceneObjectPathChannelKeyValue& OutValue)
+{
+	UObject* Value = nullptr;
+	bool bSuccess = InChannel->Evaluate(InTime, Value);
+	OutValue = FMovieSceneObjectPathChannelKeyValue(Value);
+	return bSuccess;
+}
diff --git a/Engine/Source/Runtime/MovieScene/Public/Channels/MovieSceneTimeWarpChannel.h b/Engine/Source/Runtime/MovieScene/Public/Channels/MovieSceneTimeWarpChannel.h
--- a/Engine/Source/Runtime/MovieScene/Public/Channels/MovieSceneTimeWarpChannel.h
+++ b/Engine/Source/Runtime/MovieScene/Public/Channels/MovieSceneTimeWarpChannel.h
@@ -2,6 +2,7 @@
 
 #pragma once
 
+#include "Channels/MovieSceneChannelTraits.h"
 #include "Channels/MovieSceneDoubleChannel.h"
 
 #include "MovieSceneTimeWarpChannel.generated.h"
@@ -58,4 +59,9 @@ inline void AssignValue(FMovieSceneTimeWarpChannel* InChannel, FKeyHandle InKeyH
 inline void AssignValue(FMovieSceneTimeWarpChannel* InChannel, FKeyHandle InKeyHandle, float InValue)
 {
 	AssignValue(static_cast<FMovieSceneDoubleChannel*>(InChannel), InKeyHandle, InValue);
-}
\ No newline at end of file
+}
+template<>
+struct TMovieSceneChannelTraits<FMovieSceneTimeWarpChannel> : TMovieSceneChannelTraitsBase<FMovieSceneTimeWarpChannel>
+{
+	enum { SupportsDefaults = false };
+};
diff --git a/Engine/Source/Runtime/MovieSceneTracks/Public/Channels/MovieSceneCameraShakeSourceTriggerChannel.h b/Engine/Source/Runtime/MovieSceneTracks/Public/Channels/MovieSceneCameraShakeSourceTriggerChannel.h
--- a/Engine/Source/Runtime/MovieSceneTracks/Public/Channels/MovieSceneCameraShakeSourceTriggerChannel.h
+++ b/Engine/Source/Runtime/MovieSceneTracks/Public/Channels/MovieSceneCameraShakeSourceTriggerChannel.h
@@ -63,6 +63,8 @@ struct FMovieSceneCameraShakeSourceTriggerChannel : public FMovieSceneChannel
 {
 	GENERATED_BODY()
 
+	typedef FMovieSceneCameraShakeSourceTrigger CurveValueType;
+
 	FORCEINLINE TMovieSceneChannelData<FMovieSceneCameraShakeSourceTrigger> GetData()
 	{
 		return TMovieSceneChannelData<FMovieSceneCameraShakeSourceTrigger>(&KeyTimes, &KeyValues, this, &KeyHandles);
diff --git a/Engine/Source/Runtime/MovieSceneTracks/Public/Channels/MovieSceneEventChannel.h b/Engine/Source/Runtime/MovieSceneTracks/Public/Channels/MovieSceneEventChannel.h
--- a/Engine/Source/Runtime/MovieSceneTracks/Public/Channels/MovieSceneEventChannel.h
+++ b/Engine/Source/Runtime/MovieSceneTracks/Public/Channels/MovieSceneEventChannel.h
@@ -26,6 +26,8 @@ struct FMovieSceneEventChannel : public FMovieSceneChannel
 {
 	GENERATED_BODY()
 
+	typedef FMovieSceneEvent CurveValueType;
+
 	/**
 	 * Access a mutable interface for this channel's data
 	 *
diff --git a/Engine/Source/Runtime/MovieSceneTracks/Public/Channels/MovieSceneStringChannel.h b/Engine/Source/Runtime/MovieSceneTracks/Public/Channels/MovieSceneStringChannel.h
--- a/Engine/Source/Runtime/MovieSceneTracks/Public/Channels/MovieSceneStringChannel.h
+++ b/Engine/Source/Runtime/MovieSceneTracks/Public/Channels/MovieSceneStringChannel.h
@@ -31,6 +31,8 @@ struct FMovieSceneStringChannel : public FMovieSceneChannel
 {
 	GENERATED_BODY()
 
+	typedef FString CurveValueType;
+
 	FMovieSceneStringChannel()
 		: DefaultValue(), bHasDefaultValue(false)
 	{}
diff --git a/Engine/Source/Runtime/MovieSceneTracks/Public/Sections/MovieSceneActorReferenceSection.h b/Engine/Source/Runtime/MovieSceneTracks/Public/Sections/MovieSceneActorReferenceSection.h
--- a/Engine/Source/Runtime/MovieSceneTracks/Public/Sections/MovieSceneActorReferenceSection.h
+++ b/Engine/Source/Runtime/MovieSceneTracks/Public/Sections/MovieSceneActorReferenceSection.h
@@ -73,6 +73,8 @@ struct FMovieSceneActorReferenceData : public FMovieSceneChannel
 {
 	GENERATED_BODY()
 
+	typedef FMovieSceneActorReferenceKey CurveValueType;
+
 	FMovieSceneActorReferenceData()
 		: DefaultValue()
 	{}
