Sequencer: Fixed property handlers not correctly initializing initial values for interrogation entities

diff --git a/Engine/Source/Runtime/MovieScene/Public/EntitySystem/MovieScenePropertyComponentHandler.h b/Engine/Source/Runtime/MovieScene/Public/EntitySystem/MovieScenePropertyComponentHandler.h
--- a/Engine/Source/Runtime/MovieScene/Public/EntitySystem/MovieScenePropertyComponentHandler.h
+++ b/Engine/Source/Runtime/MovieScene/Public/EntitySystem/MovieScenePropertyComponentHandler.h
@@ -90,8 +90,8 @@ struct TInitialValueProcessorImpl<PropertyTraits, TPropertyMetaData<MetaDataType
 
 	virtual void PopulateFilter(FEntityComponentFilter& OutFilter) const override
 	{
-		// These initializers must have bound objects
-		OutFilter.All({ BuiltInComponents->BoundObject });
+		// These initializers must have bound objects or an interrogation key
+		OutFilter.Any({ BuiltInComponents->BoundObject, BuiltInComponents->Interrogation.OutputKey });
 	}
 
 	virtual void Process(const FEntityAllocation* Allocation, const FComponentMask& AllocationType) override
