Sequencer: Changed task scheduler to use dynamically-sized sparse bitsets over fixed-size

- Under interrogation there can be an explosion of tasks which ended up overflowing our fixed size sparse bitset of 1024 bits.

- To rectify this it now uses a dynamically resizing sparse bitset that is essentially a binary tree of fixed size bitsets. This allows us to have very efficient iteration of bitsets as well as saving memory, which being able to support any number of tasks.

- Child tasks are now always stored immediately after parent tasks in the task array, even after shuffling. This means we don't even need a bitset to keep track of children.

- Also fixed the potential for a stack overflow when flushing tasks with NoThreading. While the stack overflow described in UE-232394 was actually caused by a cyclic dependency due to overflowing the previously-fixed-size-bitset, there is still a possiblility that a very deep callstack caused by recursively running tasks could stack overflow. I've fixed this by always using the GameThreadTaskQueue even when threading is disabled, just without signalling the game thread.

- Finally, I have increased the strictness of prerequisite counting in the task scheduler, which would have detected this issue earlier. Previously we would always trigger susbsequent tasks when their prerequisite count was <= 0. This allows us to greacefully handle cases where the prereq count becomes negative, which should be considered an error. We now handle and report this error, but still signal the game thread in case this occurs on a background thread.

diff --git a/Engine/Source/Editor/SequencerCore/Private/MVVM/ICastable.cpp b/Engine/Source/Editor/SequencerCore/Private/MVVM/ICastable.cpp
--- a/Engine/Source/Editor/SequencerCore/Private/MVVM/ICastable.cpp
+++ b/Engine/Source/Editor/SequencerCore/Private/MVVM/ICastable.cpp
@@ -58,7 +58,7 @@ FCastableTypeTable* FCastableTypeTable::FCastableTypeTableGenerator::Commit(uint
 	// There are 2 sets of data that live on the end of the struct: the sparse bitset buckets, and the type offsets.
 	// Its layout in memory looks like this:
 	//	|												sizeof(FCastableTypeTable)										|		sizeof(uint16)*NumStaticTypes		|		sizeof(uint16)*TypeMask.NumBuckets()	|
-	//	|		int16*			|		TSparseBitSet		|		uint32		|		uint32		|		FName		|			uint16[NumStaticTypes]			|		uint16[TypeMask.NumBuckets()]			|
+	//	|		int16*			|	  TFixedSparseBitSet	|		uint32		|		uint32		|		FName		|			uint16[NumStaticTypes]			|		uint16[TypeMask.NumBuckets()]			|
 	//	|	StaticTypeOffsets	|		  TypeMask			|	NumStaticTypes	|	ThisTypeID		|	ThisTypeName	|											|		 										|
 	//				\_________________________\________________________________________________________________________/										   /
 	//										   \______________________________________________TypeMask.Buckets.Storage____________________________________________/
diff --git a/Engine/Source/Editor/SequencerCore/Public/MVVM/CastableTypeTable.h b/Engine/Source/Editor/SequencerCore/Public/MVVM/CastableTypeTable.h
--- a/Engine/Source/Editor/SequencerCore/Public/MVVM/CastableTypeTable.h
+++ b/Engine/Source/Editor/SequencerCore/Public/MVVM/CastableTypeTable.h
@@ -229,7 +229,7 @@ private:
 	private:
 
 		/* Accumulator sparse bitset of types that are implemented on the type that Generate was called with */
-		MovieScene::TSparseBitSet<uint32> TypeMask;
+		MovieScene::TFixedSparseBitSet<uint32> TypeMask;
 		/* Typeoffsets from T to the ptr for each type within TypeMask */
 		TArray<int16> StaticTypeOffsets;
 	};
@@ -239,7 +239,7 @@ private:
 	/** Ptr to the static type offsets allocated at the end of this struct. Size = NumStaticTypes */
 	int16* StaticTypeOffsets;
 	/** Fixed sparse bitset of types that this table contains. Storage is allocated at the end of this struct. */
-	MovieScene::TSparseBitSet<uint32, MovieScene::TFixedSparseBitSetBucketStorage<uint8>> TypeMask;
+	MovieScene::TFixedSparseBitSet<uint32, MovieScene::TFixedSparseBitSetBucketStorage<uint8>> TypeMask;
 	/** The number of static types implemented within this type table. */
 	uint32 NumStaticTypes;
 	/** Unique static type ID for this type. */
diff --git a/Engine/Source/Runtime/MovieScene/Private/EntitySystem/MovieSceneTaskScheduler.cpp b/Engine/Source/Runtime/MovieScene/Private/EntitySystem/MovieSceneTaskScheduler.cpp
--- a/Engine/Source/Runtime/MovieScene/Private/EntitySystem/MovieSceneTaskScheduler.cpp
+++ b/Engine/Source/Runtime/MovieScene/Private/EntitySystem/MovieSceneTaskScheduler.cpp
@@ -46,7 +46,6 @@ FScheduledTask::FScheduledTask(FEntityAllocationWriteContext InWriteContextOffse
 
 FScheduledTask::FScheduledTask(FScheduledTask&& InTask)
 	: ComputedSubsequents(MoveTemp(InTask.ComputedSubsequents))
-	, ChildTasks(MoveTemp(InTask.ChildTasks))
 	, TaskFunction(InTask.TaskFunction)
 	, TaskContext(InTask.TaskContext)
 #if !(UE_BUILD_SHIPPING || UE_BUILD_TEST)
@@ -54,11 +53,12 @@ FScheduledTask::FScheduledTask(FScheduledTask&& InTask)
 #endif
 	, StatId(InTask.StatId)
 	, WriteContextOffset(InTask.WriteContextOffset)
-	, LockedComponentData(InTask.LockedComponentData)
+	, LockedComponentData(MoveTemp(InTask.LockedComponentData))
 	, NumPrerequisites(InTask.NumPrerequisites)
 	, WaitCount(InTask.WaitCount.Load(EEntityThreadingModel::NoThreading))
 	, ChildCompleteCount(InTask.ChildCompleteCount.Load(EEntityThreadingModel::NoThreading))
 	, Parent(InTask.Parent)
+	, NumChildren(InTask.NumChildren)
 	, TaskFunctionType(InTask.TaskFunctionType)
 {
 	bForceGameThread = InTask.bForceGameThread;
@@ -120,7 +120,7 @@ void FScheduledTask::Run(const FEntitySystemScheduler* Scheduler, FTaskExecution
 				check(LockedComponentData.AllocationIndex != MAX_uint16);
 				FEntityAllocationProxy Allocation = FEntityAllocationProxy::MakeInstance(Scheduler->GetEntityManager(), LockedComponentData.AllocationIndex);
 
-				TArrayView<const FPreLockedDataPtr> PreLockedData(LockedComponentData.PreLockedDataStart.Resolve(Scheduler->GetPrelockedDataBuffer()), LockedComponentData.PrelockedDataSize);
+				TArrayView<const FPreLockedDataPtr> PreLockedData(LockedComponentData.PreLockedComponentData);
 				(TaskFunction.PreLockedAllocationItem)(Allocation, PreLockedData, TaskContext.Get(), ThisWriteContext);
 			}
 			break;
@@ -130,7 +130,7 @@ void FScheduledTask::Run(const FEntitySystemScheduler* Scheduler, FTaskExecution
 	// Now the task is finished, schedule any children to run, or any subsequents
 	// If we are a parent we do not call CompleteTask until _all_ our children have finished:
 	//     this happens in FEntitySystemScheduler::CompleteTask if Parent is valid
-	if (!ChildTasks.IsEmpty())
+	if (NumChildren > 0)
 	{
 		// Increment the child completion count to protect CompleteTask being called for _this_ parent task
 		// while the loop over ChildTasks is running. This prevents a race condition where the final child can end up being
@@ -140,9 +140,10 @@ void FScheduledTask::Run(const FEntitySystemScheduler* Scheduler, FTaskExecution
 		// Once our loop has finished we check the child complete count to see if this was the last one
 		ChildCompleteCount.Add(Scheduler->GetEntityManager()->GetThreadingModel(), 1);
 
-		for (int32 Index : ChildTasks)
+		const FScheduledTask* Children = this + 1;
+		for (uint16 ChildIndex = 0; ChildIndex < NumChildren; ++ChildIndex)
 		{
-			Scheduler->PrerequisiteCompleted(FTaskID(Index), nullptr);
+			Scheduler->PrerequisiteCompleted(&Children[ChildIndex], nullptr);
 		}
 
 		// Subtract our count added on ln 205. If this is the last count, complete this task (all children have completed)
@@ -186,7 +187,6 @@ FTaskID FEntitySystemScheduler::CreateForkedAllocationTask(const FTaskParams& In
 
 	FTaskID LastTaskID, ParentTaskID;
 
-	FTaskBitSet AllocationApplicability;
 	for (FEntityAllocationIteratorItem Allocation : EntityManager->Iterate(&Filter))
 	{
 		// If we haven't created a parent yet, create that now
@@ -245,16 +245,7 @@ FTaskID FEntitySystemScheduler::CreateForkedAllocationTask(const FTaskParams& In
 		FScheduledTask::FLockedComponentData LockedComponentData;
 		LockedComponentData.AllocationIndex = static_cast<uint16>(AllocationIndex);
 
-		int32 ComponentStart = PreLockedComponentData.Num();
-
-		InPreLockFunc(Allocation, PreLockedComponentData);
-
-		const int32 PrelockedDataSize = PreLockedComponentData.Num() - ComponentStart;
-		if (PrelockedDataSize > 0)
-		{
-			LockedComponentData.PreLockedDataStart.Reset(PreLockedComponentData.GetData(), &PreLockedComponentData[ComponentStart]);
-			LockedComponentData.PrelockedDataSize = static_cast<uint16>(PrelockedDataSize);
-		}
+		InPreLockFunc(Allocation, LockedComponentData.PreLockedComponentData);
 
 		// Create the task
 
@@ -266,7 +257,7 @@ FTaskID FEntitySystemScheduler::CreateForkedAllocationTask(const FTaskParams& In
 		NewTask.TaskContext = InTaskContext;
 		NewTask.bForceGameThread = InParams.bForceGameThread;
 		NewTask.NumPrerequisites = 1; // +1 Because the parent triggers us as well when it starts
-		NewTask.LockedComponentData = LockedComponentData;
+		NewTask.LockedComponentData = MoveTemp(LockedComponentData);
 #if !(UE_BUILD_SHIPPING || UE_BUILD_TEST)
 		if (InParams.DebugName)
 		{
@@ -280,7 +271,7 @@ FTaskID FEntitySystemScheduler::CreateForkedAllocationTask(const FTaskParams& In
 	#endif
 #endif
 
-		Tasks[ParentTaskID.Index].ChildTasks.SetBit(ChildTaskIndex);
+		++Tasks[ParentTaskID.Index].NumChildren;
 
 		// If we're forking these tasks, the parent schedules this allocation task as soon as possible
 		// In this case that means that anything that writes to the same components on the same allocation is finished
@@ -315,6 +306,8 @@ FTaskID FEntitySystemScheduler::CreateForkedAllocationTask(const FTaskParams& In
 		LastTaskID = ThisTask;
 	}
 
+	check(!ParentTaskID || ParentTaskID.Index + Tasks[ParentTaskID.Index].NumChildren == Tasks.Num()-1);
+
 	return ParentTaskID;
 }
 
@@ -332,18 +325,20 @@ void FEntitySystemScheduler::AddPrerequisite(FTaskID Prerequisite, FTaskID Subse
 
 void FEntitySystemScheduler::AddChildFront(FTaskID Parent, FTaskID Child)
 {
-	if (Parent && Child && !Tasks[Parent.Index].ChildTasks.IsBitSet(Child.Index))
+	if (Parent && Child && ensure(Child.Index > Parent.Index && Child.Index == Parent.Index + Tasks[Parent.Index].NumChildren + 1))
 	{
 		// Children can't be prerequisites of their parent or visa-versa
 		check(!Tasks[Parent.Index].ComputedSubsequents.IsBitSet(Child.Index) && !Tasks[Child.Index].ComputedSubsequents.IsBitSet(Parent.Index));
 
 		// All existing children must come after this
-		for (int32 Index : Tasks[Parent.Index].ChildTasks)
+		const int32 FirstChild = Parent.Index+1;
+		const int32 ChildrenEnd = FirstChild + Tasks[Parent.Index].NumChildren;
+		for (int32 ChildIndex = FirstChild; ChildIndex < ChildrenEnd; ++ChildIndex)
 		{
-			AddPrerequisite(Child, FTaskID(Index));
+			AddPrerequisite(Child, FTaskID(ChildIndex));
 		}
 
-		Tasks[Parent.Index].ChildTasks.SetBit(Child.Index);
+		++Tasks[Parent.Index].NumChildren;
 
 		// The parent will signal this task when it is run
 		Tasks[Child.Index].NumPrerequisites += 1;
@@ -353,18 +348,21 @@ void FEntitySystemScheduler::AddChildFront(FTaskID Parent, FTaskID Child)
 
 void FEntitySystemScheduler::AddChildBack(FTaskID Parent, FTaskID Child)
 {
-	if (Parent && Child && !Tasks[Parent.Index].ChildTasks.IsBitSet(Child.Index))
+	// Child tasks must be immediately reported
+	if (Parent && Child && ensure(Child.Index > Parent.Index && Child.Index == Parent.Index + Tasks[Parent.Index].NumChildren + 1))
 	{
 		// Children can't be prerequisites of their parent or visa-versa
 		check(!Tasks[Parent.Index].ComputedSubsequents.IsBitSet(Child.Index) && !Tasks[Child.Index].ComputedSubsequents.IsBitSet(Parent.Index));
 
 		// All existing children must come before this
-		for (int32 Index : Tasks[Parent.Index].ChildTasks)
+		const int32 FirstChild = Parent.Index+1;
+		const int32 ChildrenEnd = FirstChild + Tasks[Parent.Index].NumChildren;
+		for (int32 ChildIndex = FirstChild; ChildIndex < ChildrenEnd; ++ChildIndex)
 		{
-			AddPrerequisite(FTaskID(Index), Child);
+			AddPrerequisite(FTaskID(ChildIndex), Child);
 		}
 
-		Tasks[Parent.Index].ChildTasks.SetBit(Child.Index);
+		++Tasks[Parent.Index].NumChildren;
 
 		// The parent will signal this task when it is run
 		Tasks[Child.Index].NumPrerequisites += 1;
@@ -428,14 +426,44 @@ FTaskID FEntitySystemScheduler::AddNullTask()
 
 void FEntitySystemScheduler::ShuffleTasks()
 {
-	TArray<int32> ShuffledIndices;
-	ShuffledIndices.SetNum(Tasks.Num());
+	TArray<int32> RemainingIndices;
+
+	// We have to keep children contiguous in memory with their parents so only shuffle parents
 	for (int32 Index = 0; Index < Tasks.Num(); ++Index)
 	{
-		ShuffledIndices[Index] = Index;
+		RemainingIndices.Emplace(Index);
+		Index += Tasks[Index].NumChildren;
 	}
 
-	Algo::RandomShuffle(ShuffledIndices);
+	Algo::RandomShuffle(RemainingIndices);
+
+	TArray<int32> ReverseShuffledIndices;
+	ReverseShuffledIndices.SetNum(Tasks.Num());
+	for (int32 Index = 0, ShuffledIndex = 0; Index < RemainingIndices.Num(); ++Index)
+	{
+		int32 TaskIndex = RemainingIndices[Index];
+
+		ReverseShuffledIndices[ShuffledIndex] = TaskIndex;
+
+		const int32 ShuffeldChildStart   = ShuffledIndex + 1;
+		const int32 UnshuffeldChildStart = TaskIndex + 1;
+
+		const int32 NumChildren = Tasks[TaskIndex].NumChildren;
+
+		ShuffledIndex += 1 + NumChildren;
+		for (int32 ChildIndex = 0; ChildIndex < NumChildren; ++ChildIndex)
+		{
+			ReverseShuffledIndices[ShuffeldChildStart + ChildIndex] = UnshuffeldChildStart + ChildIndex;
+		}
+	}
+
+	TArray<int32> ShuffledIndices;
+	ShuffledIndices.SetNum(Tasks.Num());
+	for (int32 Index = 0; Index < ReverseShuffledIndices.Num(); ++Index)
+	{
+		const int32 OriginalTaskIndex = ReverseShuffledIndices[Index];
+		ShuffledIndices[OriginalTaskIndex] = Index;
+	}
 
 	auto RedirectMask = [&ShuffledIndices](FTaskBitSet& InOutBitSet)
 	{
@@ -455,8 +483,8 @@ void FEntitySystemScheduler::ShuffleTasks()
 		{
 			Task.Parent = FTaskID(ShuffledIndices[Task.Parent.Index]);
 		}
+
 		RedirectMask(Task.ComputedSubsequents);
-		RedirectMask(Task.ChildTasks);
 	}
 
 	RedirectMask(InitialTasks);
@@ -464,15 +492,9 @@ void FEntitySystemScheduler::ShuffleTasks()
 	TArray<FScheduledTask> OldTasks;
 	Swap(Tasks, OldTasks);
 
-	TArray<int32> SwapSourceIndices;
-	SwapSourceIndices.SetNum(ShuffledIndices.Num());
-	for (int32 Index = 0; Index < ShuffledIndices.Num(); ++Index)
-	{
-		SwapSourceIndices[ShuffledIndices[Index]] = Index;
-	}
-	for (int32 Index = 0; Index < ShuffledIndices.Num(); ++Index)
+	for (int32 Index = 0; Index < OldTasks.Num(); ++Index)
 	{
-		Tasks.Emplace(MoveTemp(OldTasks[SwapSourceIndices[Index]]));
+		Tasks.Emplace(MoveTemp(OldTasks[ReverseShuffledIndices[Index]]));
 	}
 }
 
@@ -501,7 +523,12 @@ void FEntitySystemScheduler::ExecuteTasks()
 		TRACE_CPUPROFILER_EVENT_SCOPE_STR("Run Scheduled Tasks");
 		for (int32 Index : InitialTasks)
 		{
-			Tasks[Index].Run(this, FTaskExecutionFlags());
+			GameThreadTaskList.Push(const_cast<FScheduledTask*>(&Tasks[Index]));
+		}
+
+		while (const FScheduledTask* Task = GameThreadTaskList.Pop())
+		{
+			Task->Run(this, FTaskExecutionFlags());
 		}
 
 		check(NumTasksRemaining.Load(ThreadingModel) == 0);
@@ -577,7 +604,7 @@ void FEntitySystemScheduler::CompleteTask(const FScheduledTask* Task, FTaskExecu
 {
 	// Reset the WaitCount ready for the next run
 	const int32 PreviousWaitCount = Task->WaitCount.Exchange(ThreadingModel, Task->NumPrerequisites);
-	const int32 PreviousChildCount = Task->ChildCompleteCount.Exchange(ThreadingModel, Task->ChildTasks.CountSetBits());
+	const int32 PreviousChildCount = Task->ChildCompleteCount.Exchange(ThreadingModel, Task->NumChildren);
 
 	checkSlow(PreviousWaitCount == 0 && PreviousChildCount == 0);
 
@@ -611,48 +638,64 @@ void FEntitySystemScheduler::CompleteTask(const FScheduledTask* Task, FTaskExecu
 }
 
 void FEntitySystemScheduler::PrerequisiteCompleted(FTaskID TaskID, int32* OptRunInlineIndex) const
+{
+	PrerequisiteCompleted(&Tasks[TaskID.Index], OptRunInlineIndex);
+}
+
+void FEntitySystemScheduler::PrerequisiteCompleted(const FScheduledTask* Task, int32* OptRunInlineIndex) const
 {
 	// We either need to not be using threading, or have a valid game thread signal event!
 	check((ThreadingModel == EEntityThreadingModel::NoThreading) || (GameThreadSignal != nullptr));
 
-	int32 PreviousWaitCount = Tasks[TaskID.Index].WaitCount.Sub(ThreadingModel, 1);
+	int32 PreviousWaitCount = Task->WaitCount.Sub(ThreadingModel, 1);
 	checkSlow(PreviousWaitCount >= 1);
 	if (PreviousWaitCount <= 1)
 	{
-		if (ThreadingModel == EEntityThreadingModel::NoThreading)
+		if (PreviousWaitCount <= 0)
 		{
-			Tasks[TaskID.Index].Run(this, FTaskExecutionFlags());
+			// This is an error
+			ensureMsgf(false, TEXT("Sequencer Task Prerequisite Count underflow!"));
+			if (GameThreadSignal)
+			{
+				// Trigger the game thread to wake up if necessary
+				GameThreadSignal->Trigger();
+			}
+		}
+		else if (ThreadingModel == EEntityThreadingModel::NoThreading)
+		{
+			GameThreadTaskList.Push(const_cast<FScheduledTask*>(Task));
 		}
-		else if (Tasks[TaskID.Index].bForceInline)
+		else if (Task->bForceInline)
 		{
 			FTaskExecutionFlags Flags;
 			// Don't let the completion of this task inline any others
 			// to prevent cascades of inlined tasks suffocating the dispatch of others
 			Flags.bCanInlineSubsequents = false;
 
-			Tasks[TaskID.Index].Run(this, Flags);
+			Task->Run(this, Flags);
 		}
-		else if (Tasks[TaskID.Index].bForceGameThread)
+		else if (Task->bForceGameThread)
 		{
 			// Push this onto the GT list even if we are already on the game thread.
 			// This ensures other subsequent tasks being processed in the same loop
 			// have a chance to be scheduled before we do any potentially time-consuming task work.
-			GameThreadTaskList.Push(const_cast<FScheduledTask*>(&Tasks[TaskID.Index]));
+			GameThreadTaskList.Push(const_cast<FScheduledTask*>(Task));
 			GameThreadSignal->Trigger();
 		}
 		else if (OptRunInlineIndex && *OptRunInlineIndex == INDEX_NONE)
 		{
-			*OptRunInlineIndex = TaskID.Index;
+			const int32 TaskIndex = (Task - Tasks.GetData());
+			*OptRunInlineIndex = TaskIndex;
 		}
 		else if (GameThreadTaskList.IsEmpty())
 		{
-			GameThreadTaskList.Push(const_cast<FScheduledTask*>(&Tasks[TaskID.Index]));
+			GameThreadTaskList.Push(const_cast<FScheduledTask*>(Task));
 			GameThreadSignal->Trigger();
 		}
 		else
 		{
-			UE::Tasks::Launch(TEXT("MovieSceneTask"), [this, Index = TaskID.Index](){
-				this->Tasks[Index].Run(this, FTaskExecutionFlags());
+			UE::Tasks::Launch(TEXT("MovieSceneTask"), [this, Task](){
+				Task->Run(this, FTaskExecutionFlags());
 			}, UE::Tasks::ETaskPriority::High);
 		}
 	}
@@ -678,7 +721,6 @@ void FEntitySystemScheduler::BeginConstruction()
 	SystemSerialIncrement = EntityManager->GetSystemSerial();
 
 	Tasks.Reset();
-	PreLockedComponentData.Reset();
 	InitialTasks = FTaskBitSet();
 }
 
@@ -715,7 +757,6 @@ void FEntitySystemScheduler::EndConstruction()
 	}
 
 	Tasks.Shrink();
-	PreLockedComponentData.Shrink();
 
 	AllPrerequisites.Empty();
 	ComponentWriteDepedenciesByAllocation.Empty();
@@ -727,7 +768,7 @@ void FEntitySystemScheduler::EndConstruction()
 	{
 		FScheduledTask& Task = Tasks[Index];
 		Task.WaitCount.Exchange(ThreadingModel, Task.NumPrerequisites);
-		Task.ChildCompleteCount.Exchange(ThreadingModel, Task.ChildTasks.CountSetBits());
+		Task.ChildCompleteCount.Exchange(ThreadingModel, Task.NumChildren);
 
 		if (Task.NumPrerequisites == 0)
 		{
@@ -768,13 +809,14 @@ FString FEntitySystemScheduler::ToString() const
 		SanitizedName.ReplaceCharInline('<', ' ');
 		SanitizedName.ReplaceCharInline('>', ' ');
 
-		const int32 NumChildren = Task.ChildTasks.CountSetBits();
-		if (NumChildren != 0)
+		if (Task.NumChildren != 0)
 		{
 			ParentTasks[Index] = true;
 			GraphString.Appendf(TEXT("\tsubgraph cluster_%d{ \n\t label=\"Parent Task %d: %s\";\n"), Index, Index, *SanitizedName);
 			GraphString.Appendf(TEXT("\t\ttask_%d[label=\"[All Children]\" style=invis];\n"), Index);
-			for (int32 ChildIndex : Task.ChildTasks)
+
+			const int32 FirstChild = Index+1;
+			for (int32 ChildIndex = FirstChild; ChildIndex < FirstChild + Task.NumChildren; ++ChildIndex)
 			{
 				ChildTasks[ChildIndex] = true;
 				GraphString.Appendf(TEXT("\t\ttask_%d;\n"), ChildIndex);
diff --git a/Engine/Source/Runtime/MovieScene/Private/EntitySystem/MovieSceneTaskScheduler.h b/Engine/Source/Runtime/MovieScene/Private/EntitySystem/MovieSceneTaskScheduler.h
--- a/Engine/Source/Runtime/MovieScene/Private/EntitySystem/MovieSceneTaskScheduler.h
+++ b/Engine/Source/Runtime/MovieScene/Private/EntitySystem/MovieSceneTaskScheduler.h
@@ -20,7 +20,7 @@ namespace UE::MovieScene
 struct FTaskExecutionFlags;
 class FEntitySystemScheduler;
 
-using FTaskBitSet = TSparseBitSet<uint64, TDynamicSparseBitSetBucketStorage<uint16, 4>>;
+using FTaskBitSet = TDynamicSparseBitSet<uint32, TDynamicSparseBitSetBucketStorage<uint16, 0>>; // Buckets of 512 task bits
 
 /**
  * Structure used for tracking task dependencies that must be propagated from system to system
@@ -106,9 +106,7 @@ struct FScheduledTask
 	struct FLockedComponentData
 	{
 		/** Relative pointer to the start of our prelocked data inside FEntitySystemScheduler::PreLockedComponentData */
-		TRelativePtr<FPreLockedDataPtr> PreLockedDataStart;
-		/** Number of entries within our prelocked data */
-		uint16 PrelockedDataSize = 0;
+		TArray<FPreLockedDataPtr> PreLockedComponentData;
 		/** Allocation index within FEntityManager::EntityAllocations */
 		uint16 AllocationIndex = MAX_uint16;
 	};
@@ -157,6 +155,8 @@ struct FScheduledTask
 	/** 4 Bytes - This task's parent (or None() if it is not a child task) */
 	FTaskID Parent;
 
+	uint16 NumChildren = 0;
+
 	/** 1 Byte - The type of the TaskFunction function pointer */
 	FScheduledTaskFuncionPtr::EType TaskFunctionType;
 
@@ -287,11 +287,6 @@ public:
 		return EntityManager;
 	}
 
-	const FPreLockedDataPtr* GetPrelockedDataBuffer() const
-	{
-		return PreLockedComponentData.GetData();
-	}
-
 	FEntityAllocationWriteContext GetWriteContextOffset() const
 	{
 		return WriteContextBase;
@@ -315,6 +310,7 @@ public:
 	 * @param OptRunInlineIndex    (Optional) When non-null, can be assigned a new task to run inline after this function call to avoid scheduling overhead
 	 */
 	void PrerequisiteCompleted(FTaskID TaskID, int32* OptRunInlineIndex) const;
+	void PrerequisiteCompleted(const FScheduledTask* Task, int32* OptRunInlineIndex) const;
 
 	/**
 	 * Called when all tasks have been completed
@@ -331,9 +327,6 @@ private:
 	/** Array of task data. Constant once EndConstruction has been called */
 	TArray<FScheduledTask> Tasks;
 
-	/** Array of pre-locked component data for each individual task. FLockedComponentData::PreLockedDataStart defines the start ptr relative to GetData, and PrelockedDataSize denotes the size */
-	TArray<FPreLockedDataPtr> PreLockedComponentData;
-
 	/** Pointer to the current node's system prerequisites if any. Only valid during construction. */
 	FTaskPrerequisiteCache* CurrentPrerequisites = nullptr;
 
diff --git a/Engine/Source/Runtime/MovieScene/Private/Tests/MovieSceneSparseBitSetTests.cpp b/Engine/Source/Runtime/MovieScene/Private/Tests/MovieSceneSparseBitSetTests.cpp
--- a/Engine/Source/Runtime/MovieScene/Private/Tests/MovieSceneSparseBitSetTests.cpp
+++ b/Engine/Source/Runtime/MovieScene/Private/Tests/MovieSceneSparseBitSetTests.cpp
@@ -8,6 +8,44 @@
 #include "Misc/GeneratedTypeName.h"
 
 #if WITH_DEV_AUTOMATION_TESTS
+	
+	template<typename T>
+	struct TBitIt
+	{
+		TBitIt()
+			: RemainingBits(0)
+		{}
+
+		TBitIt(T InRemainingBits)
+			: RemainingBits(InRemainingBits)
+		{}
+
+		FORCEINLINE friend bool operator==(const TBitIt& A, const TBitIt& B)
+		{
+			return A.RemainingBits == B.RemainingBits;
+		}
+
+		explicit operator bool() const
+		{
+			return RemainingBits != 0;
+		}
+
+		uint32 operator*()
+		{
+			return UE::MovieScene::Private::CountTrailingZeros(RemainingBits);
+		}
+
+		TBitIt& operator++()
+		{
+			// Clear the lowest 1 bit
+			RemainingBits = RemainingBits & (RemainingBits - 1);
+			return *this;
+		}
+	private:
+		T RemainingBits;
+	};
+
+
 
 namespace UE::MovieScene::Tests
 {
@@ -16,21 +54,21 @@ namespace UE::MovieScene::Tests
 	void TestBitSet(FAutomationTestBase* Test, SparseBitSet& InBitSet)
 	{
 		// Set random bits in the bitset
-		const int32 MaxBitIndex = static_cast<int32>(InBitSet.GetMaxNumBits());
+		const uint32 MaxBitIndex = InBitSet.GetMaxNumBits();
 
 		// Static int for re-running tests using the same seed as a previous run
 		static int32 SeedOverride = 0;
 		const int32 InitialSeed = SeedOverride == 0 ? FMath::Rand() : SeedOverride;
 		FRandomStream Random(InitialSeed);
 
-		UE_LOG(LogMovieScene, Display, TEXT("Running TSparseBitSet<> tests for %s with a random seed %i..."), GetGeneratedTypeName<SparseBitSet>(), InitialSeed);
+		UE_LOG(LogMovieScene, Display, TEXT("Running tests for %s with a random seed %i..."), GetGeneratedTypeName<SparseBitSet>(), InitialSeed);
 
 		TSet<uint32> SetIndices;
-		const int32 Num = Random.RandHelper(MaxBitIndex);
+		const int32 Num = Random.RandHelper(static_cast<int32>(FMath::Min(InBitSet.GetMaxNumBits(), 128u)));
 
 		for (int32 Index = 0; Index < Num; ++Index)
 		{
-			const uint32 Bit = static_cast<uint32>(Random.RandHelper(MaxBitIndex));
+			const uint32 Bit = Random.GetUnsignedInt() % MaxBitIndex;
 			InBitSet.SetBit(Bit);
 			SetIndices.Add(Bit);
 		}
@@ -38,7 +76,7 @@ namespace UE::MovieScene::Tests
 		Test->TestEqual(TEXT("Num Set Bits"), InBitSet.CountSetBits(), static_cast<uint32>(SetIndices.Num()));
 
 		// Test all bits are the correct state
-		for (uint32 BitIndex = 0; BitIndex < static_cast<uint32>(MaxBitIndex); ++BitIndex)
+		for (uint32 BitIndex = 0; BitIndex < FMath::Min(MaxBitIndex, 64u*64u); ++BitIndex)
 		{
 			Test->TestEqual(FString::Printf(TEXT("Bit index %d"), BitIndex), InBitSet.IsBitSet(BitIndex), SetIndices.Contains(BitIndex));
 		}
@@ -69,74 +107,126 @@ bool FMovieSceneSparseBitSetTest::RunTest(const FString& Parameters)
 
 	{
 		// 8 bit buckets
-		TSparseBitSet<uint8, TDynamicSparseBitSetBucketStorage<uint8, 4>> BitSet8_8;
-		TestEqual(TEXT("TSparseBitSet<uint8, ...<uint8>>::GetMaxNumBits"), BitSet8_8.GetMaxNumBits(), 8u*8u);
+		TFixedSparseBitSet<uint8, TDynamicSparseBitSetBucketStorage<uint8, 4>> BitSet8_8;
+		TestEqual(TEXT("TFixedSparseBitSet<uint8, ...<uint8>>::GetMaxNumBits"), BitSet8_8.GetMaxNumBits(), 8u*8u);
 		TestBitSet(this, BitSet8_8);
 
-		TSparseBitSet<uint16, TDynamicSparseBitSetBucketStorage<uint8, 4>> BitSet16_8;
-		TestEqual(TEXT("TSparseBitSet<uint16, ...<uint8>>::GetMaxNumBits"), BitSet16_8.GetMaxNumBits(), 16u*8u);
+		TFixedSparseBitSet<uint16, TDynamicSparseBitSetBucketStorage<uint8, 4>> BitSet16_8;
+		TestEqual(TEXT("TFixedSparseBitSet<uint16, ...<uint8>>::GetMaxNumBits"), BitSet16_8.GetMaxNumBits(), 16u*8u);
 		TestBitSet(this, BitSet16_8);
 
-		TSparseBitSet<uint32, TDynamicSparseBitSetBucketStorage<uint8, 4>> BitSet32_8;
-		TestEqual(TEXT("TSparseBitSet<uint32, ...<uint8>>::GetMaxNumBits"), BitSet32_8.GetMaxNumBits(), 32u*8u);
+		TFixedSparseBitSet<uint32, TDynamicSparseBitSetBucketStorage<uint8, 4>> BitSet32_8;
+		TestEqual(TEXT("TFixedSparseBitSet<uint32, ...<uint8>>::GetMaxNumBits"), BitSet32_8.GetMaxNumBits(), 32u*8u);
 		TestBitSet(this, BitSet32_8);
 
-		TSparseBitSet<uint64, TDynamicSparseBitSetBucketStorage<uint8, 4>> BitSet64_8;
-		TestEqual(TEXT("TSparseBitSet<uint64, ...<uint8>>::GetMaxNumBits"), BitSet64_8.GetMaxNumBits(), 64u*8u);
+		TFixedSparseBitSet<uint64, TDynamicSparseBitSetBucketStorage<uint8, 4>> BitSet64_8;
+		TestEqual(TEXT("TFixedSparseBitSet<uint64, ...<uint8>>::GetMaxNumBits"), BitSet64_8.GetMaxNumBits(), 64u*8u);
 		TestBitSet(this, BitSet64_8);
 
 		// 16 bit buckets
-		TSparseBitSet<uint8, TDynamicSparseBitSetBucketStorage<uint16, 4>> BitSet8_16;
-		TestEqual(TEXT("TSparseBitSet<uint8, ...<uint16>>::GetMaxNumBits"), BitSet8_16.GetMaxNumBits(), 8u*16u);
+		TFixedSparseBitSet<uint8, TDynamicSparseBitSetBucketStorage<uint16, 4>> BitSet8_16;
+		TestEqual(TEXT("TFixedSparseBitSet<uint8, ...<uint16>>::GetMaxNumBits"), BitSet8_16.GetMaxNumBits(), 8u*16u);
 		TestBitSet(this, BitSet8_16);
 
-		TSparseBitSet<uint16, TDynamicSparseBitSetBucketStorage<uint16, 4>> BitSet16_16;
-		TestEqual(TEXT("TSparseBitSet<uint16, ...<uint16>>::GetMaxNumBits"), BitSet16_16.GetMaxNumBits(), 16u*16u);
+		TFixedSparseBitSet<uint16, TDynamicSparseBitSetBucketStorage<uint16, 4>> BitSet16_16;
+		TestEqual(TEXT("TFixedSparseBitSet<uint16, ...<uint16>>::GetMaxNumBits"), BitSet16_16.GetMaxNumBits(), 16u*16u);
 		TestBitSet(this, BitSet16_16);
 
-		TSparseBitSet<uint32, TDynamicSparseBitSetBucketStorage<uint16, 4>> BitSet32_16;
-		TestEqual(TEXT("TSparseBitSet<uint32, ...<uint16>>::GetMaxNumBits"), BitSet32_16.GetMaxNumBits(), 32u*16u);
+		TFixedSparseBitSet<uint32, TDynamicSparseBitSetBucketStorage<uint16, 4>> BitSet32_16;
+		TestEqual(TEXT("TFixedSparseBitSet<uint32, ...<uint16>>::GetMaxNumBits"), BitSet32_16.GetMaxNumBits(), 32u*16u);
 		TestBitSet(this, BitSet32_16);
 
-		TSparseBitSet<uint64, TDynamicSparseBitSetBucketStorage<uint16, 4>> BitSet64_16;
-		TestEqual(TEXT("TSparseBitSet<uint64, ...<uint16>>::GetMaxNumBits"), BitSet64_16.GetMaxNumBits(), 64u*16u);
+		TFixedSparseBitSet<uint64, TDynamicSparseBitSetBucketStorage<uint16, 4>> BitSet64_16;
+		TestEqual(TEXT("TFixedSparseBitSet<uint64, ...<uint16>>::GetMaxNumBits"), BitSet64_16.GetMaxNumBits(), 64u*16u);
 		TestBitSet(this, BitSet64_16);
 
 		// 32 bit buckets
-		TSparseBitSet<uint8, TDynamicSparseBitSetBucketStorage<uint32, 4>> BitSet8_32;
-		TestEqual(TEXT("TSparseBitSet<uint8, ...<uint32>>::GetMaxNumBits"), BitSet8_32.GetMaxNumBits(), 8u*32u);
+		TFixedSparseBitSet<uint8, TDynamicSparseBitSetBucketStorage<uint32, 4>> BitSet8_32;
+		TestEqual(TEXT("TFixedSparseBitSet<uint8, ...<uint32>>::GetMaxNumBits"), BitSet8_32.GetMaxNumBits(), 8u*32u);
 		TestBitSet(this, BitSet8_32);
 
-		TSparseBitSet<uint16, TDynamicSparseBitSetBucketStorage<uint32, 4>> BitSet16_32;
-		TestEqual(TEXT("TSparseBitSet<uint16, ...<uint32>>::GetMaxNumBits"), BitSet16_32.GetMaxNumBits(), 16u*32u);
+		TFixedSparseBitSet<uint16, TDynamicSparseBitSetBucketStorage<uint32, 4>> BitSet16_32;
+		TestEqual(TEXT("TFixedSparseBitSet<uint16, ...<uint32>>::GetMaxNumBits"), BitSet16_32.GetMaxNumBits(), 16u*32u);
 		TestBitSet(this, BitSet16_32);
 
-		TSparseBitSet<uint32, TDynamicSparseBitSetBucketStorage<uint32, 4>> BitSet32_32;
-		TestEqual(TEXT("TSparseBitSet<uint32, ...<uint32>>::GetMaxNumBits"), BitSet32_32.GetMaxNumBits(), 32u*32u);
+		TFixedSparseBitSet<uint32, TDynamicSparseBitSetBucketStorage<uint32, 4>> BitSet32_32;
+		TestEqual(TEXT("TFixedSparseBitSet<uint32, ...<uint32>>::GetMaxNumBits"), BitSet32_32.GetMaxNumBits(), 32u*32u);
 		TestBitSet(this, BitSet32_32);
 
-		TSparseBitSet<uint64, TDynamicSparseBitSetBucketStorage<uint32, 4>> BitSet64_32;
-		TestEqual(TEXT("TSparseBitSet<uint64, ...<uint32>>::GetMaxNumBits"), BitSet64_32.GetMaxNumBits(), 64u*32u);
+		TFixedSparseBitSet<uint64, TDynamicSparseBitSetBucketStorage<uint32, 4>> BitSet64_32;
+		TestEqual(TEXT("TFixedSparseBitSet<uint64, ...<uint32>>::GetMaxNumBits"), BitSet64_32.GetMaxNumBits(), 64u*32u);
 		TestBitSet(this, BitSet64_32);
 
 		// 64 bit buckets
-		TSparseBitSet<uint8, TDynamicSparseBitSetBucketStorage<uint64, 4>> BitSet8_64;
-		TestEqual(TEXT("TSparseBitSet<uint8, ...<uint64>>::GetMaxNumBits"), BitSet8_64.GetMaxNumBits(), 8u*64u);
+		TFixedSparseBitSet<uint8, TDynamicSparseBitSetBucketStorage<uint64, 4>> BitSet8_64;
+		TestEqual(TEXT("TFixedSparseBitSet<uint8, ...<uint64>>::GetMaxNumBits"), BitSet8_64.GetMaxNumBits(), 8u*64u);
 		TestBitSet(this, BitSet8_64);
 
-		TSparseBitSet<uint16, TDynamicSparseBitSetBucketStorage<uint64, 4>> BitSet16_64;
-		TestEqual(TEXT("TSparseBitSet<uint16, ...<uint64>>::GetMaxNumBits"), BitSet16_64.GetMaxNumBits(), 16u*64u);
+		TFixedSparseBitSet<uint16, TDynamicSparseBitSetBucketStorage<uint64, 4>> BitSet16_64;
+		TestEqual(TEXT("TFixedSparseBitSet<uint16, ...<uint64>>::GetMaxNumBits"), BitSet16_64.GetMaxNumBits(), 16u*64u);
 		TestBitSet(this, BitSet16_64);
 
-		TSparseBitSet<uint32, TDynamicSparseBitSetBucketStorage<uint64, 4>> BitSet32_64;
-		TestEqual(TEXT("TSparseBitSet<uint32, ...<uint64>>::GetMaxNumBits"), BitSet32_64.GetMaxNumBits(), 32u*64u);
+		TFixedSparseBitSet<uint32, TDynamicSparseBitSetBucketStorage<uint64, 4>> BitSet32_64;
+		TestEqual(TEXT("TFixedSparseBitSet<uint32, ...<uint64>>::GetMaxNumBits"), BitSet32_64.GetMaxNumBits(), 32u*64u);
 		TestBitSet(this, BitSet32_64);
 
-		TSparseBitSet<uint64, TDynamicSparseBitSetBucketStorage<uint64, 4>> BitSet64_64;
-		TestEqual(TEXT("TSparseBitSet<uint64, ...<uint64>>::GetMaxNumBits"), BitSet64_64.GetMaxNumBits(), 64u*64u);
+		TFixedSparseBitSet<uint64, TDynamicSparseBitSetBucketStorage<uint64, 4>> BitSet64_64;
+		TestEqual(TEXT("TFixedSparseBitSet<uint64, ...<uint64>>::GetMaxNumBits"), BitSet64_64.GetMaxNumBits(), 64u*64u);
 		TestBitSet(this, BitSet64_64);
+	}
+
+	{
+		// 8 bit buckets
+		TDynamicSparseBitSet<uint8, TDynamicSparseBitSetBucketStorage<uint8, 4>> BitSet8_8;
+		TestBitSet(this, BitSet8_8);
+
+		TDynamicSparseBitSet<uint16, TDynamicSparseBitSetBucketStorage<uint8, 4>> BitSet16_8;
+		TestBitSet(this, BitSet16_8);
+
+		TDynamicSparseBitSet<uint32, TDynamicSparseBitSetBucketStorage<uint8, 4>> BitSet32_8;
+		TestBitSet(this, BitSet32_8);
+
+		TDynamicSparseBitSet<uint64, TDynamicSparseBitSetBucketStorage<uint8, 4>> BitSet64_8;
+		TestBitSet(this, BitSet64_8);
+
+		// 16 bit buckets
+		TDynamicSparseBitSet<uint8, TDynamicSparseBitSetBucketStorage<uint16, 4>> BitSet8_16;
+		TestBitSet(this, BitSet8_16);
+
+		TDynamicSparseBitSet<uint16, TDynamicSparseBitSetBucketStorage<uint16, 4>> BitSet16_16;
+		TestBitSet(this, BitSet16_16);
+
+		TDynamicSparseBitSet<uint32, TDynamicSparseBitSetBucketStorage<uint16, 4>> BitSet32_16;
+		TestBitSet(this, BitSet32_16);
+
+		TDynamicSparseBitSet<uint64, TDynamicSparseBitSetBucketStorage<uint16, 4>> BitSet64_16;
+		TestBitSet(this, BitSet64_16);
+
+		// 32 bit buckets
+		TDynamicSparseBitSet<uint8, TDynamicSparseBitSetBucketStorage<uint32, 4>> BitSet8_32;
+		TestBitSet(this, BitSet8_32);
 
+		TDynamicSparseBitSet<uint16, TDynamicSparseBitSetBucketStorage<uint32, 4>> BitSet16_32;
+		TestBitSet(this, BitSet16_32);
+
+		TDynamicSparseBitSet<uint32, TDynamicSparseBitSetBucketStorage<uint32, 4>> BitSet32_32;
+		TestBitSet(this, BitSet32_32);
 
+		TDynamicSparseBitSet<uint64, TDynamicSparseBitSetBucketStorage<uint32, 4>> BitSet64_32;
+		TestBitSet(this, BitSet64_32);
+
+		// 64 bit buckets
+		TDynamicSparseBitSet<uint8, TDynamicSparseBitSetBucketStorage<uint64, 4>> BitSet8_64;
+		TestBitSet(this, BitSet8_64);
+
+		TDynamicSparseBitSet<uint16, TDynamicSparseBitSetBucketStorage<uint64, 4>> BitSet16_64;
+		TestBitSet(this, BitSet16_64);
+
+		TDynamicSparseBitSet<uint32, TDynamicSparseBitSetBucketStorage<uint64, 4>> BitSet32_64;
+		TestBitSet(this, BitSet32_64);
+
+		TDynamicSparseBitSet<uint64, TDynamicSparseBitSetBucketStorage<uint64, 4>> BitSet64_64;
+		TestBitSet(this, BitSet64_64);
 	}
 	return true;
 }
diff --git a/Engine/Source/Runtime/MovieScene/Public/Containers/SparseBitSet.h b/Engine/Source/Runtime/MovieScene/Public/Containers/SparseBitSet.h
--- a/Engine/Source/Runtime/MovieScene/Public/Containers/SparseBitSet.h
+++ b/Engine/Source/Runtime/MovieScene/Public/Containers/SparseBitSet.h
@@ -16,8 +16,10 @@ template<typename T>
 struct TFixedSparseBitSetBucketStorage;
 
 template<typename HashType, typename BucketStorage = TDynamicSparseBitSetBucketStorage<uint8, 4>>
-struct TSparseBitSet;
+struct TFixedSparseBitSet;
 
+template<typename HashType, typename BucketStorage = TDynamicSparseBitSetBucketStorage<uint8, 4>>
+struct TDynamicSparseBitSet;
 
 enum class ESparseBitSetBitResult
 {
@@ -74,33 +76,33 @@ namespace Private
  * This container is specialized for relatively large bitfields that have small numbers of set bits (ie, needles in haystacks) as they will provide the best memory vs CPU tradeoffs.
  */
 template<typename HashType, typename BucketStorage>
-struct TSparseBitSet
+struct TFixedSparseBitSet
 {
 	using BucketType = typename BucketStorage::BucketType;
 	static constexpr uint32 HashSize    = sizeof(HashType)*8;
 	static constexpr uint32 BucketSize  = sizeof(typename BucketStorage::BucketType)*8;
 	static constexpr uint32 MaxNumBits  = HashSize * BucketSize;
 
-	explicit TSparseBitSet()
+	explicit TFixedSparseBitSet()
 		: BucketHash(0)
 	{}
 
 	template<typename ...StorageArgs>
-	explicit TSparseBitSet(StorageArgs&& ...Storage)
+	explicit TFixedSparseBitSet(StorageArgs&& ...Storage)
 		: Buckets(Forward<StorageArgs>(Storage)...)
 		, BucketHash(0)
 	{}
 
-	TSparseBitSet(const TSparseBitSet&) = default;
-	TSparseBitSet& operator=(const TSparseBitSet&) = default;
+	TFixedSparseBitSet(const TFixedSparseBitSet&) = default;
+	TFixedSparseBitSet& operator=(const TFixedSparseBitSet&) = default;
 
-	TSparseBitSet(TSparseBitSet&&) = default;
-	TSparseBitSet& operator=(TSparseBitSet&&) = default;
+	TFixedSparseBitSet(TFixedSparseBitSet&&) = default;
+	TFixedSparseBitSet& operator=(TFixedSparseBitSet&&) = default;
 
 	template<typename OtherHashType, typename OtherStorageType>
-	void CopyTo(TSparseBitSet<OtherHashType, OtherStorageType>& Other) const
+	void CopyTo(TFixedSparseBitSet<OtherHashType, OtherStorageType>& Other) const
 	{
-		static_assert(TSparseBitSet<OtherHashType, OtherStorageType>::BucketSize == BucketSize, "Cannot copy sparse bitsets of different bucket sizes");
+		static_assert(TFixedSparseBitSet<OtherHashType, OtherStorageType>::BucketSize == BucketSize, "Cannot copy sparse bitsets of different bucket sizes");
 
 		// Copy the buckets
 		const uint32 NumBuckets = FMath::CountBits(Other.BucketHash);
@@ -110,9 +112,9 @@ struct TSparseBitSet
 
 	/** Copy this bitset to another without resizing the destination's bucket storage. Bucket storage must be >= this size. */
 	template<typename OtherHashType, typename OtherStorageType>
-	void CopyToUnsafe(TSparseBitSet<OtherHashType, OtherStorageType>& Other, uint32 OtherBucketCapacity) const
+	void CopyToUnsafe(TFixedSparseBitSet<OtherHashType, OtherStorageType>& Other, uint32 OtherBucketCapacity) const
 	{
-		static_assert(TSparseBitSet<OtherHashType, OtherStorageType>::BucketSize == BucketSize, "Cannot copy sparse bitsets of different bucket sizes");
+		static_assert(TFixedSparseBitSet<OtherHashType, OtherStorageType>::BucketSize == BucketSize, "Cannot copy sparse bitsets of different bucket sizes");
 
 		const uint32 ThisNumBuckets = this->NumBuckets();
 		checkf(OtherBucketCapacity >= ThisNumBuckets, TEXT("Attempting to copy a sparse bitset without enough capacity in the destination (%d, required %d)"), OtherBucketCapacity, ThisNumBuckets);
@@ -124,7 +126,7 @@ struct TSparseBitSet
 		FMemory::Memcpy(Other.Buckets.GetData(), this->Buckets.Storage.GetData(), sizeof(typename BucketStorage::BucketType)*ThisNumBuckets);
 	}
 
-	TSparseBitSet& operator|=(const TSparseBitSet& Other)
+	TFixedSparseBitSet& operator|=(const TFixedSparseBitSet& Other)
 	{
 		using namespace Private;
 
@@ -274,7 +276,15 @@ struct TSparseBitSet
 
 	struct FIterator
 	{
-		static FIterator Begin(const TSparseBitSet<HashType, BucketStorage>* InBitSet)
+		FIterator()
+			: BitSet(nullptr)
+			, BucketBitIndex(0)
+			, IndexWithinBucket(0)
+			, CurrentBucket(0)
+		{
+		}
+
+		static FIterator Begin(const TFixedSparseBitSet<HashType, BucketStorage>* InBitSet)
 		{
 			FIterator It;
 			It.BitSet = InBitSet;
@@ -293,7 +303,7 @@ struct TSparseBitSet
 			}
 			return It;
 		}
-		static FIterator End(const TSparseBitSet<HashType, BucketStorage>* InBitSet)
+		static FIterator End(const TFixedSparseBitSet<HashType, BucketStorage>* InBitSet)
 		{
 			FIterator It;
 			It.BitSet = InBitSet;
@@ -307,10 +317,14 @@ struct TSparseBitSet
 		{
 			using namespace Private;
 
-			CurrentBucket &= ~(BucketType(1u)<<IndexWithinBucket);
-			IndexWithinBucket = CountTrailingZeros(CurrentBucket);
+			// Clear the lowest 1 bit
+			CurrentBucket = CurrentBucket & (CurrentBucket - 1);
 
-			if (IndexWithinBucket == BucketSize || CurrentBucket == 0)
+			if (CurrentBucket != 0)
+			{
+				IndexWithinBucket = CountTrailingZeros(CurrentBucket);
+			}
+			else
 			{
 				// If this was the last bit, reset the iterator to end()
 				if (BucketBitIndex == HashSize-1)
@@ -347,31 +361,30 @@ struct TSparseBitSet
 			return BucketBitIndex < HashSize;
 		}
 
-		friend bool operator==(FIterator A, FIterator B)
+		friend bool operator==(const FIterator& A, const FIterator& B)
 		{
 			return A.BitSet == B.BitSet && A.BucketBitIndex == B.BucketBitIndex && A.IndexWithinBucket == B.IndexWithinBucket;
 		}
-		friend bool operator!=(FIterator A, FIterator B)
+		friend bool operator!=(const FIterator& A, const FIterator& B)
 		{
 			return !(A == B);
 		}
 private:
-		FIterator() = default;
 
-		const TSparseBitSet<HashType, BucketStorage>* BitSet;
+		const TFixedSparseBitSet<HashType, BucketStorage>* BitSet;
 		uint8 BucketBitIndex;
 		uint8 IndexWithinBucket;
 
 		BucketType CurrentBucket;
 	};
 
-	friend FIterator begin(const TSparseBitSet<HashType, BucketStorage>& In) { return FIterator::Begin(&In); }
-	friend FIterator end(const TSparseBitSet<HashType, BucketStorage>& In)   { return FIterator::End(&In); }
+	friend FIterator begin(const TFixedSparseBitSet<HashType, BucketStorage>& In) { return FIterator::Begin(&In); }
+	friend FIterator end(const TFixedSparseBitSet<HashType, BucketStorage>& In)   { return FIterator::End(&In); }
 
 private:
 
 	template<typename, typename>
-	friend struct TSparseBitSet;
+	friend struct TFixedSparseBitSet;
 
 	FORCEINLINE void CheckIndex(uint32 BitIndex) const
 	{
@@ -464,4 +477,263 @@ struct TFixedSparseBitSetBucketStorage
 	BucketType  Get(int32 Index) const { return Storage[Index]; }
 };
 
+/**
+ * NOTE: This class is currently considered internal only, and should only be used by engine code.
+ * A dynamically sized sparse bitset comprising multiple TFixedSparseBitSets.
+ * 
+ * In theory this class supports the full integer range, it is optimized for small numbers of set bits within a large range, ideally when they occupy the same adjacent space.
+ */
+template<typename HashType, typename BucketStorage>
+struct TDynamicSparseBitSet
+{
+	/**
+	 * Get the maximum number of bits that this bitset supports
+	 */
+	uint32 GetMaxNumBits() const
+	{
+		return MAX_uint32;
+	}
+
+
+	/**
+	 * Set the bit at the specified index.
+	 * Any bits between Num and BitIndex will be considered 0.
+	 *
+	 * @return true if the bit was previously considered 0 and is now set, false if it was already set.
+	 */
+	ESparseBitSetBitResult SetBit(uint32 Bit)
+	{
+		const uint32 Bucket = Bit / NumBitsInBucket;
+
+		Bit -= Bucket*NumBitsInBucket;
+
+		FEntry* EntrPtr = Entries.GetData();
+
+		const int32 Num = Entries.Num();
+		for (int32 EntryIndex = 0; EntryIndex < Num; ++EntryIndex)
+		{
+			if (EntrPtr[EntryIndex].Offset == Bucket)
+			{
+				return EntrPtr[EntryIndex].Bits.SetBit(Bit);
+			}
+			else if (EntrPtr[EntryIndex].Offset > Bucket)
+			{
+				Entries.InsertUninitialized(EntryIndex);
+				new(&Entries[EntryIndex]) FEntry(Bucket, Bit);
+				return ESparseBitSetBitResult::NewlySet;
+			}
+		}
+
+		Entries.Emplace(Bucket, Bit);
+		return ESparseBitSetBitResult::NewlySet;
+	}
+
+
+	/**
+	 * Check whether this container has any bits set
+	 */
+	bool IsEmpty() const
+	{
+		return Entries.Num() == 0;
+	}
+
+
+	/**
+	 * Check whether the specified bit index is set
+	 */
+	bool IsBitSet(uint32 Bit) const
+	{
+		const uint32 Bucket = Bit / NumBitsInBucket;
+
+		const FEntry* EntrPtr = Entries.GetData();
+
+		const int32 Num = Entries.Num();
+		for (int32 EntryIndex = 0; EntryIndex < Num; ++EntryIndex)
+		{
+			if (EntrPtr[EntryIndex].Offset == Bucket)
+			{
+				return EntrPtr[EntryIndex].Bits.IsBitSet(Bit);
+			}
+			if (EntrPtr[EntryIndex].Offset > Bucket)
+			{
+				return false;
+			}
+		}
+
+		return false;
+	}
+
+	/**
+	 * Count the total number of set bits in this container
+	 */
+	uint32 CountSetBits() const
+	{
+		uint32 SetBits = 0;
+		for (const FEntry& Entry : Entries)
+		{
+			SetBits += Entry.Bits.CountSetBits();
+		}
+		return SetBits;
+	}
+
+
+	TDynamicSparseBitSet<HashType, BucketStorage>& operator|=(const TDynamicSparseBitSet<HashType, BucketStorage>& Other)
+	{
+		if (Other.Entries.Num() == 0)
+		{
+			return *this;
+		}
+
+		if (Entries.Num() == 0)
+		{
+			*this = Other;
+			return *this;
+		}
+
+
+		int32 ThisIndex = 0;
+		int32 OtherIndex = 0;
+
+		while (OtherIndex < Other.Entries.Num() && Other.Entries[OtherIndex].Offset < Entries[0].Offset)
+		{
+			++OtherIndex;
+		}
+
+		if (OtherIndex > 0)
+		{
+			Entries.Insert(Other.Entries.GetData(), OtherIndex, 0);
+		}
+
+		while (OtherIndex < Other.Entries.Num() && ThisIndex < Entries.Num())
+		{
+			if (Other.Entries[OtherIndex].Offset < Entries[ThisIndex].Offset)
+			{
+				Entries.Insert(Other.Entries[OtherIndex], ThisIndex);
+				++OtherIndex;
+			}
+			else if (Other.Entries[OtherIndex].Offset == Entries[ThisIndex].Offset)
+			{
+				Entries[ThisIndex].Bits |= Other.Entries[OtherIndex].Bits;
+				++OtherIndex;
+				++ThisIndex;
+			}
+			else
+			{
+				++ThisIndex;
+			}
+		}
+
+		return *this;
+	}
+
+
+private:
+
+	using FBucketBitSet = TFixedSparseBitSet<HashType, BucketStorage>;
+
+	static constexpr uint32 NumBitsInBucket = FBucketBitSet::MaxNumBits;
+
+	struct FEntry
+	{
+		FEntry(uint32 InOffset)
+			: Offset(InOffset)
+		{
+		}
+		FEntry(uint32 InOffset, uint32 InBit)
+			: Offset(InOffset)
+		{
+			checkSlow(InBit < FBucketBitSet::MaxNumBits);
+			Bits.SetBit(InBit);
+		}
+
+		FBucketBitSet Bits;
+		uint32 Offset;
+	};
+
+public:
+
+	struct FIterator
+	{
+		static FIterator Begin(const TDynamicSparseBitSet<HashType, BucketStorage>* InBitSet)
+		{
+			FIterator It;
+			It.Entries = InBitSet->Entries.GetData();
+			It.NumEntries = InBitSet->Entries.Num();
+			It.EntryIndex = 0;
+			It.CurrentOffsetInBits = 0;
+
+			if (It.NumEntries != 0)
+			{
+				It.CurrentOffsetInBits = It.Entries[0].Offset * NumBitsInBucket;
+				It.BucketIt = BucketIterator::Begin(&It.Entries[0].Bits);
+			}
+			return It;
+		}
+		static FIterator End(const TDynamicSparseBitSet<HashType, BucketStorage>* InBitSet)
+		{
+			FIterator It;
+			It.Entries = InBitSet->Entries.GetData();
+			It.NumEntries = InBitSet->Entries.Num();
+			It.EntryIndex = It.NumEntries;
+			It.CurrentOffsetInBits = 0;
+			return It;
+		}
+
+		void operator++()
+		{
+			using namespace Private;
+
+			++BucketIt;
+			if (!BucketIt)
+			{
+				++EntryIndex;
+				if (EntryIndex < NumEntries)
+				{
+					CurrentOffsetInBits = Entries[EntryIndex].Offset * NumBitsInBucket;
+					BucketIt = BucketIterator::Begin(&Entries[EntryIndex].Bits);
+				}
+				else
+				{
+					CurrentOffsetInBits = 0;
+					BucketIt = BucketIterator();
+				}
+			}
+		}
+
+		int32 operator*() const
+		{
+			return CurrentOffsetInBits + *BucketIt;
+		}
+
+		explicit operator bool() const
+		{
+			return EntryIndex < NumEntries;
+		}
+
+		friend bool operator==(const FIterator& A, const FIterator& B)
+		{
+			return A.Entries == B.Entries && A.EntryIndex == B.EntryIndex && A.BucketIt == B.BucketIt;
+		}
+		friend bool operator!=(const FIterator& A, const FIterator& B)
+		{
+			return !(A == B);
+		}
+private:
+		FIterator() = default;
+
+		using BucketIterator = TFixedSparseBitSet<HashType, BucketStorage>::FIterator;
+
+		const typename TDynamicSparseBitSet<HashType, BucketStorage>::FEntry* Entries;
+		BucketIterator BucketIt;
+		int32 NumEntries;
+		int32 EntryIndex;
+		int32 CurrentOffsetInBits;
+	};
+
+	friend FIterator begin(const TDynamicSparseBitSet<HashType, BucketStorage>& In) { return FIterator::Begin(&In); }
+	friend FIterator end(const TDynamicSparseBitSet<HashType, BucketStorage>& In)   { return FIterator::End(&In); }
+
+	TArray<FEntry> Entries;
+};
+
 } // namespace UE::MovieScene
