Android CIS fix

diff --git a/Engine/Source/Runtime/MovieScene/Public/Containers/SparseBitSet.h b/Engine/Source/Runtime/MovieScene/Public/Containers/SparseBitSet.h
--- a/Engine/Source/Runtime/MovieScene/Public/Containers/SparseBitSet.h
+++ b/Engine/Source/Runtime/MovieScene/Public/Containers/SparseBitSet.h
@@ -721,7 +721,7 @@ public:
 private:
 		FIterator() = default;
 
-		using BucketIterator = TFixedSparseBitSet<HashType, BucketStorage>::FIterator;
+		using BucketIterator = typename TFixedSparseBitSet<HashType, BucketStorage>::FIterator;
 
 		const typename TDynamicSparseBitSet<HashType, BucketStorage>::FEntry* Entries;
 		BucketIterator BucketIt;
