Remove debug check that is now handled gracefully

diff --git a/Engine/Source/Runtime/MovieScene/Private/EntitySystem/MovieSceneTaskScheduler.cpp b/Engine/Source/Runtime/MovieScene/Private/EntitySystem/MovieSceneTaskScheduler.cpp
--- a/Engine/Source/Runtime/MovieScene/Private/EntitySystem/MovieSceneTaskScheduler.cpp
+++ b/Engine/Source/Runtime/MovieScene/Private/EntitySystem/MovieSceneTaskScheduler.cpp
@@ -648,7 +648,6 @@ void FEntitySystemScheduler::PrerequisiteCompleted(const FScheduledTask* Task, i
 	check((ThreadingModel == EEntityThreadingModel::NoThreading) || (GameThreadSignal != nullptr));
 
 	int32 PreviousWaitCount = Task->WaitCount.Sub(ThreadingModel, 1);
-	checkSlow(PreviousWaitCount >= 1);
 	if (PreviousWaitCount <= 1)
 	{
 		if (PreviousWaitCount <= 0)
