FMovieSceneTimeController: Since StartPlay REALLY wants the clock to start running immediately, AND is also called when the scene loops (part of a stop / start pair), added PrepareToPlay and IsReadyToPlay calls so that Time Controllers can...

1) Do 'one-time setup' things that it can't / shouldn't do when looping.

2) Delay the actual start of playback by returning false from IsReadyToPlay until such time as future calls to RequestCurrentTime will reveal a "moving" clock.

diff --git a/Engine/Source/Runtime/MovieScene/Private/MovieSceneSequencePlayer.cpp b/Engine/Source/Runtime/MovieScene/Private/MovieSceneSequencePlayer.cpp
--- a/Engine/Source/Runtime/MovieScene/Private/MovieSceneSequencePlayer.cpp
+++ b/Engine/Source/Runtime/MovieScene/Private/MovieSceneSequencePlayer.cpp
@@ -175,6 +175,7 @@ UMovieSceneSequencePlayer::UMovieSceneSequencePlayer(const FObjectInitializer& I
 	, Status(EMovieScenePlayerStatus::Stopped)
 	, bReversePlayback(false)
 	, bPendingOnStartedPlaying(false)
+	, bIsTimeControllerReadyToPlay(true)
 	, bIsAsyncUpdate(false)
 	, bSkipNextUpdate(false)
 	, bUpdateNetSync(false)
@@ -315,7 +316,7 @@ void UMovieSceneSequencePlayer::PlayInternal()
 		return;
 	}
 
-	if (!IsPlaying() && Sequence && CanPlay())
+	if (!IsPlaying() && Sequence && CanPlay() && bIsTimeControllerReadyToPlay)
 	{
 		const FString SequenceName = GetSequenceName(true);
 		UE_LOG(LogMovieScene, Verbose, TEXT("PlayInternal - %s (current status: %s)"), *SequenceName, *UEnum::GetValueAsString(Status));
@@ -347,52 +348,15 @@ void UMovieSceneSequencePlayer::PlayInternal()
 			RootTemplateInstance.EnableGlobalPreAnimatedStateCapture();
 		}
 
-		bPendingOnStartedPlaying = true;
-		Status = EMovieScenePlayerStatus::Playing;
-		TimeController->StartPlaying(GetCurrentTime());
-		
-		if (PlayPosition.GetEvaluationType() == EMovieSceneEvaluationType::FrameLocked)
+		TimeController->PrepareToPlay(GetCurrentTime());
+		if (TimeController->IsReadyToPlay())
 		{
-			if (!OldMaxTickRate.IsSet())
-			{
-				OldMaxTickRate = GEngine->GetMaxFPS();
-			}
-
-			GEngine->SetMaxFPS(1.f / PlayPosition.GetInputRate().AsInterval());
-
-			if (GSequencerApplyDisplayRateToDynResFrameTimeBudget)
-			{
-				const float DyResFrameTimeBudget = PlayPosition.GetInputRate().AsInterval() * 1000.0f;
-
-				static IConsoleVariable* CVarDynResFrameTimeBudget = IConsoleManager::Get().FindConsoleVariable(TEXT("r.DynamicRes.FrameTimeBudget"));
-				CVarDynResFrameTimeBudget->Set(DyResFrameTimeBudget, ECVF_SetByCode);
-
-				bOverridingDynResFrameTimeBudget = true;
-			}
-		}
-
-		if (!PlayPosition.GetLastPlayEvalPostition().IsSet() || PlayPosition.GetLastPlayEvalPostition() != PlayPosition.GetCurrentPosition())
-		{
-			EPlayDirection PreferredDirection = bReversePlayback ? EPlayDirection::Backwards : EPlayDirection::Forwards;
-			UpdateMovieSceneInstance(PlayPosition.PlayTo(PlayPosition.GetCurrentPosition(), PreferredDirection), EMovieScenePlayerStatus::Playing);
-		}
-
-		RunLatentActions();
-		UpdateNetworkSyncProperties();
-		
-		if (bReversePlayback)
-		{
-			if (OnPlayReverse.IsBound())
-			{
-				OnPlayReverse.Broadcast();
-			}
+			StartTimeControllerAndBroadcastPlayState();
+			bIsTimeControllerReadyToPlay = true;
 		}
 		else
 		{
-			if (OnPlay.IsBound())
-			{
-				OnPlay.Broadcast();
-			}
+			bIsTimeControllerReadyToPlay = false;
 		}
 	}
 }
@@ -490,6 +454,11 @@ void UMovieSceneSequencePlayer::StopInternal(FFrameTime TimeToResetTo)
 	const FString SequenceName = GetSequenceName(true);
 	UE_LOG(LogMovieScene, Verbose, TEXT("StopInternal - %s (at: %s, current status: %s)"), *SequenceName, *LexToString(TimeToResetTo), *UEnum::GetValueAsString(Status));
 
+	// We may have been waiting for our time controller to tell us it was ready to play,
+	// so bIsTimeControllerReadyToPlay may be false. Here we reset that back to its
+	// initialized state of true.
+	bIsTimeControllerReadyToPlay = true;
+
 	if (IsPlaying() || IsPaused())
 	{
 		Status = EMovieScenePlayerStatus::Stopped;
@@ -1096,7 +1065,13 @@ void UMovieSceneSequencePlayer::Update(const float DeltaSeconds)
 
 	UpdateNetworkSync();
 
-	if (IsPlaying())
+	if (!bIsTimeControllerReadyToPlay && TimeController->IsReadyToPlay())
+	{
+		StartTimeControllerAndBroadcastPlayState();
+		bIsTimeControllerReadyToPlay = true;
+	}
+
+	if (bIsTimeControllerReadyToPlay && IsPlaying())
 	{
 		// Delta seconds has already been multiplied by GetEffectiveTimeDilation at this point, so don't pass that through to Tick
 		float PlayRate = bReversePlayback ? -PlaybackSettings.PlayRate : PlaybackSettings.PlayRate;
@@ -1176,6 +1151,57 @@ FMovieSceneEvaluationRange UpdatePlayPosition(FMovieScenePlaybackPosition& InOut
 	return InOutPlayPosition.JumpTo(NewTime, PreferredDirection);
 }
 
+void UMovieSceneSequencePlayer::StartTimeControllerAndBroadcastPlayState()
+{
+	bPendingOnStartedPlaying = true;
+	Status = EMovieScenePlayerStatus::Playing;
+	TimeController->StartPlaying(GetCurrentTime());
+
+	if (PlayPosition.GetEvaluationType() == EMovieSceneEvaluationType::FrameLocked)
+	{
+		if (!OldMaxTickRate.IsSet())
+		{
+			OldMaxTickRate = GEngine->GetMaxFPS();
+		}
+
+		GEngine->SetMaxFPS(1.f / PlayPosition.GetInputRate().AsInterval());
+
+		if (GSequencerApplyDisplayRateToDynResFrameTimeBudget)
+		{
+			const float DyResFrameTimeBudget = PlayPosition.GetInputRate().AsInterval() * 1000.0f;
+
+			static IConsoleVariable* CVarDynResFrameTimeBudget = IConsoleManager::Get().FindConsoleVariable(TEXT("r.DynamicRes.FrameTimeBudget"));
+			CVarDynResFrameTimeBudget->Set(DyResFrameTimeBudget, ECVF_SetByCode);
+
+			bOverridingDynResFrameTimeBudget = true;
+		}
+	}
+
+	if (!PlayPosition.GetLastPlayEvalPostition().IsSet() || PlayPosition.GetLastPlayEvalPostition() != PlayPosition.GetCurrentPosition())
+	{
+		EPlayDirection PreferredDirection = bReversePlayback ? EPlayDirection::Backwards : EPlayDirection::Forwards;
+		UpdateMovieSceneInstance(PlayPosition.PlayTo(PlayPosition.GetCurrentPosition(), PreferredDirection), EMovieScenePlayerStatus::Playing);
+	}
+
+	RunLatentActions();
+	UpdateNetworkSyncProperties();
+
+	if (bReversePlayback)
+	{
+		if (OnPlayReverse.IsBound())
+		{
+			OnPlayReverse.Broadcast();
+		}
+	}
+	else
+	{
+		if (OnPlay.IsBound())
+		{
+			OnPlay.Broadcast();
+		}
+	}
+}
+
 void UMovieSceneSequencePlayer::UpdateTimeCursorPosition_Internal(FFrameTime NewPosition, EUpdatePositionMethod Method, bool bHasJumpedOverride)
 {
 	EMovieScenePlayerStatus::Type StatusOverride = UpdateMethodToStatus(Method);
diff --git a/Engine/Source/Runtime/MovieScene/Private/MovieSceneTimeController.cpp b/Engine/Source/Runtime/MovieScene/Private/MovieSceneTimeController.cpp
--- a/Engine/Source/Runtime/MovieScene/Private/MovieSceneTimeController.cpp
+++ b/Engine/Source/Runtime/MovieScene/Private/MovieSceneTimeController.cpp
@@ -46,6 +46,18 @@ void FMovieSceneTimeController::PlayerStatusChanged(EMovieScenePlayerStatus::Typ
 	}
 }
 
+void FMovieSceneTimeController::PrepareToPlay(const FQualifiedFrameTime& InStartTime)
+{
+	UE_LOG(LogMovieScene, VeryVerbose, TEXT("TimeController PrepareToPlay: Sequence started preparing: frame %d, subframe %f. Frame rate: %f fps."), InStartTime.Time.FrameNumber.Value, InStartTime.Time.GetSubFrame(), InStartTime.Rate.AsDecimal());
+
+	OnPrepareToPlay(InStartTime);
+}
+
+bool FMovieSceneTimeController::IsReadyToPlay()
+{
+	return OnIsReadyToPlay();
+}
+
 void FMovieSceneTimeController::StartPlaying(const FQualifiedFrameTime& InStartTime)
 {
 	UE_LOG(LogMovieScene, VeryVerbose, TEXT("TimeController Start: Sequence started: frame %d, subframe %f. Frame rate: %f fps."), InStartTime.Time.FrameNumber.Value, InStartTime.Time.GetSubFrame(), InStartTime.Rate.AsDecimal());
diff --git a/Engine/Source/Runtime/MovieScene/Public/MovieSceneSequencePlayer.h b/Engine/Source/Runtime/MovieScene/Public/MovieSceneSequencePlayer.h
--- a/Engine/Source/Runtime/MovieScene/Public/MovieSceneSequencePlayer.h
+++ b/Engine/Source/Runtime/MovieScene/Public/MovieSceneSequencePlayer.h
@@ -601,6 +601,7 @@ protected:
 	virtual void OnStopped() {}
 	
 private:
+	MOVIESCENE_API void StartTimeControllerAndBroadcastPlayState();
 
 	MOVIESCENE_API void UpdateTimeCursorPosition_Internal(FFrameTime NewPosition, EUpdatePositionMethod Method, bool bHasJumpedOverride);
 
@@ -670,6 +671,9 @@ protected:
 	/** Set to true to invoke OnStartedPlaying on first update tick for started playing */
 	uint32 bPendingOnStartedPlaying : 1;
 
+	/** set to true when Time Controller has been told to PrepareToPlay but has not yet become IsReadyToPlay */
+	uint32 bIsTimeControllerReadyToPlay : 1;
+
 	/** Set to true when the player is currently in the main level update */
 	uint32 bIsAsyncUpdate : 1;
 
diff --git a/Engine/Source/Runtime/MovieScene/Public/MovieSceneTimeController.h b/Engine/Source/Runtime/MovieScene/Public/MovieSceneTimeController.h
--- a/Engine/Source/Runtime/MovieScene/Public/MovieSceneTimeController.h
+++ b/Engine/Source/Runtime/MovieScene/Public/MovieSceneTimeController.h
@@ -16,6 +16,16 @@ public:
 
 	virtual ~FMovieSceneTimeController() {}
 
+	/**
+	 * Called whenever a sequence starts or resumes playback from a non-playing state
+	 */
+	MOVIESCENE_API void PrepareToPlay(const FQualifiedFrameTime& InStartTime);
+
+	/**
+	 * Called whenever a sequence starts or resumes playback from a non-playing state
+	*/
+	MOVIESCENE_API bool IsReadyToPlay();
+
 	/**
 	 * Called whenever a sequence starts or resumes playback from a non-playing state
 	 */
@@ -56,6 +66,8 @@ public:
 protected:
 
 	virtual void OnTick(float DeltaSeconds, float InPlayRate){}
+	virtual void OnPrepareToPlay(const FQualifiedFrameTime& InStartTime) {}
+	virtual bool OnIsReadyToPlay() { return true; }
 	virtual void OnStartPlaying(const FQualifiedFrameTime& InStartTime){}
 	virtual void OnStopPlaying(const FQualifiedFrameTime& InStopTime){}
 	virtual FFrameTime OnRequestCurrentTime(const FQualifiedFrameTime& InCurrentTime, float InPlayRate) = 0;
