Sequencer: Changed atomic list to be optionally atomic to support autortfm transaction scopes

diff --git a/Engine/Source/Runtime/MovieScene/Private/EntitySystem/MovieSceneTaskScheduler.cpp b/Engine/Source/Runtime/MovieScene/Private/EntitySystem/MovieSceneTaskScheduler.cpp
--- a/Engine/Source/Runtime/MovieScene/Private/EntitySystem/MovieSceneTaskScheduler.cpp
+++ b/Engine/Source/Runtime/MovieScene/Private/EntitySystem/MovieSceneTaskScheduler.cpp
@@ -35,6 +35,60 @@ struct FTaskExecutionFlags
 	uint8 bCanInlineSubsequents : 1;
 };
 
+FEntitySystemScheduler::FOptionallyAtomicList::FOptionallyAtomicList()
+	: GameThreadTaskList(TInPlaceType<FArrayList>())
+{
+}
+
+void FEntitySystemScheduler::FOptionallyAtomicList::SetThreadingModel(EEntityThreadingModel InThreadingModel)
+{
+	ensureMsgf(IsEmpty(), TEXT("Should not be trying to change the threading model on a task list that is populated!"));
+
+	if (InThreadingModel == EEntityThreadingModel::NoThreading && !GameThreadTaskList.IsType<FArrayList>())
+	{
+		GameThreadTaskList.Emplace<FArrayList>();
+	}
+	else if (InThreadingModel != EEntityThreadingModel::NoThreading && !GameThreadTaskList.IsType<FLockFreeList>())
+	{
+		GameThreadTaskList.Emplace<FLockFreeList>();
+	}
+}
+
+void FEntitySystemScheduler::FOptionallyAtomicList::Push(FScheduledTask* InTask)
+{
+	switch (GameThreadTaskList.GetIndex())
+	{
+	case 0:  GameThreadTaskList.Get<FLockFreeList>().Push(InTask);       break;
+	case 1:  GameThreadTaskList.Get<FArrayList>().Push(InTask);          break;
+	}
+}
+
+FScheduledTask* FEntitySystemScheduler::FOptionallyAtomicList::Pop()
+{
+	if (FLockFreeList* LockFreeList = GameThreadTaskList.TryGet<FLockFreeList>())
+	{
+		return LockFreeList->Pop();
+	}
+	else
+	{
+		FArrayList* ArrayList = GameThreadTaskList.TryGet<FArrayList>();
+		return ArrayList->Num() != 0 ? ArrayList->Pop() : nullptr;
+	}
+}
+
+bool FEntitySystemScheduler::FOptionallyAtomicList::IsEmpty() const
+{
+	switch (GameThreadTaskList.GetIndex())
+	{
+	case 0:  return GameThreadTaskList.Get<FLockFreeList>().IsEmpty();
+	case 1:  return GameThreadTaskList.Get<FArrayList>().Num() == 0;
+	}
+
+	checkf(false, TEXT("Invalid index encountered in variant"));
+	return true;
+}
+
+
 FScheduledTask::FScheduledTask(FEntityAllocationWriteContext InWriteContextOffset)
 	: StatId(GET_STATID(MovieSceneEval_AnonymousTask))
 	, WriteContextOffset(InWriteContextOffset)
@@ -511,6 +565,8 @@ void FEntitySystemScheduler::ExecuteTasks()
 	check(PreviousNumRemaining == 0);
 
 	ThreadingModel = ThreadingModelToUse;
+	GameThreadTaskList.SetThreadingModel(ThreadingModelToUse);
+
 	// In a transaction we can only support the no-threading mode.
 	check(!AutoRTFM::IsTransactional() || (EEntityThreadingModel::NoThreading == ThreadingModel));
 
diff --git a/Engine/Source/Runtime/MovieScene/Private/EntitySystem/MovieSceneTaskScheduler.h b/Engine/Source/Runtime/MovieScene/Private/EntitySystem/MovieSceneTaskScheduler.h
--- a/Engine/Source/Runtime/MovieScene/Private/EntitySystem/MovieSceneTaskScheduler.h
+++ b/Engine/Source/Runtime/MovieScene/Private/EntitySystem/MovieSceneTaskScheduler.h
@@ -8,6 +8,7 @@
 #include "Tasks/Task.h"
 #include "Containers/LockFreeList.h"
 #include "Containers/StringView.h"
+#include "Misc/TVariant.h"
 #include "EntitySystem/RelativePtr.h"
 #include "EntitySystem/IMovieSceneTaskScheduler.h"
 #include "EntitySystem/MovieSceneEntityManager.h"
@@ -346,7 +347,22 @@ private:
 	mutable FEntitySystemMaybeAtomicInt32 NumTasksRemaining = 0;
 
 	FEvent* GameThreadSignal = nullptr;
-	mutable TLockFreePointerListFIFO<FScheduledTask, PLATFORM_CACHE_LINE_SIZE> GameThreadTaskList;
+
+	struct FOptionallyAtomicList
+	{
+		FOptionallyAtomicList();
+		void SetThreadingModel(EEntityThreadingModel ThreadingModel);
+		void Push(FScheduledTask* InTask);
+		FScheduledTask* Pop();
+		bool IsEmpty() const;
+	private:
+
+		using FLockFreeList = TLockFreePointerListFIFO<FScheduledTask, PLATFORM_CACHE_LINE_SIZE>;
+		using FArrayList    = TArray<FScheduledTask*>;
+		TVariant<FLockFreeList, FArrayList> GameThreadTaskList;
+	};
+
+	mutable FOptionallyAtomicList GameThreadTaskList;
 	FEntityAllocationWriteContext WriteContextBase = FEntityAllocationWriteContext::NewAllocation();
 	uint32 SystemSerialIncrement = 0;
 	EEntityThreadingModel ThreadingModel = EEntityThreadingModel::NoThreading;
diff --git a/Engine/Source/Runtime/MovieScene/Private/Tests/AutoRTFM/AutoRTFMMovieSceneTests.cpp b/Engine/Source/Runtime/MovieScene/Private/Tests/AutoRTFM/AutoRTFMMovieSceneTests.cpp
--- a/Engine/Source/Runtime/MovieScene/Private/Tests/AutoRTFM/AutoRTFMMovieSceneTests.cpp
+++ b/Engine/Source/Runtime/MovieScene/Private/Tests/AutoRTFM/AutoRTFMMovieSceneTests.cpp
@@ -54,7 +54,20 @@ bool FAutoRTFMMovieSceneTests::RunTest(const FString & Parameters)
 		TStatId StatId;
 
 		Scheduler.BeginConstruction();
-		Scheduler.AddTask<FMyTask>(FTaskParams(StatId));
+
+		// Create a chain of tasks that depend on each other (A -> B -> C ... etc) to ensures
+		// that we hit the codepath in FEntitySystemScheduler::PrerequisiteCompleted
+		FTaskID LastTask;
+		for (int32 TaskIndex = 0; TaskIndex < 10; ++TaskIndex)
+		{
+			FTaskID ThisTask = Scheduler.AddTask<FMyTask>(FTaskParams(StatId));
+			if (LastTask)
+			{
+				Scheduler.AddPrerequisite(LastTask, ThisTask);
+			}
+			LastTask = ThisTask;
+		}
+
 		Scheduler.EndConstruction();
 
 		AutoRTFM::ETransactionResult Result = AutoRTFM::Transact([&]
@@ -72,7 +85,7 @@ bool FAutoRTFMMovieSceneTests::RunTest(const FString & Parameters)
 			});
 
 		TEST_CHECK_TRUE(AutoRTFM::ETransactionResult::Committed == Result);
-		TEST_CHECK_TRUE(1 == FMyTask::WasHit);//-V547
+		TEST_CHECK_TRUE(10 == FMyTask::WasHit);//-V547
 	}
 
 	{
