Revert CL 38770712 and usage of GameThreadTaskList when not using threading

diff --git a/Engine/Source/Runtime/MovieScene/Private/EntitySystem/MovieSceneTaskScheduler.cpp b/Engine/Source/Runtime/MovieScene/Private/EntitySystem/MovieSceneTaskScheduler.cpp
--- a/Engine/Source/Runtime/MovieScene/Private/EntitySystem/MovieSceneTaskScheduler.cpp
+++ b/Engine/Source/Runtime/MovieScene/Private/EntitySystem/MovieSceneTaskScheduler.cpp
@@ -35,60 +35,6 @@ struct FTaskExecutionFlags
 	uint8 bCanInlineSubsequents : 1;
 };
 
-FEntitySystemScheduler::FOptionallyAtomicList::FOptionallyAtomicList()
-	: GameThreadTaskList(TInPlaceType<FArrayList>())
-{
-}
-
-void FEntitySystemScheduler::FOptionallyAtomicList::SetThreadingModel(EEntityThreadingModel InThreadingModel)
-{
-	ensureMsgf(IsEmpty(), TEXT("Should not be trying to change the threading model on a task list that is populated!"));
-
-	if (InThreadingModel == EEntityThreadingModel::NoThreading && !GameThreadTaskList.IsType<FArrayList>())
-	{
-		GameThreadTaskList.Emplace<FArrayList>();
-	}
-	else if (InThreadingModel != EEntityThreadingModel::NoThreading && !GameThreadTaskList.IsType<FLockFreeList>())
-	{
-		GameThreadTaskList.Emplace<FLockFreeList>();
-	}
-}
-
-void FEntitySystemScheduler::FOptionallyAtomicList::Push(FScheduledTask* InTask)
-{
-	switch (GameThreadTaskList.GetIndex())
-	{
-	case 0:  GameThreadTaskList.Get<FLockFreeList>().Push(InTask);       break;
-	case 1:  GameThreadTaskList.Get<FArrayList>().Push(InTask);          break;
-	}
-}
-
-FScheduledTask* FEntitySystemScheduler::FOptionallyAtomicList::Pop()
-{
-	if (FLockFreeList* LockFreeList = GameThreadTaskList.TryGet<FLockFreeList>())
-	{
-		return LockFreeList->Pop();
-	}
-	else
-	{
-		FArrayList* ArrayList = GameThreadTaskList.TryGet<FArrayList>();
-		return ArrayList->Num() != 0 ? ArrayList->Pop() : nullptr;
-	}
-}
-
-bool FEntitySystemScheduler::FOptionallyAtomicList::IsEmpty() const
-{
-	switch (GameThreadTaskList.GetIndex())
-	{
-	case 0:  return GameThreadTaskList.Get<FLockFreeList>().IsEmpty();
-	case 1:  return GameThreadTaskList.Get<FArrayList>().Num() == 0;
-	}
-
-	checkf(false, TEXT("Invalid index encountered in variant"));
-	return true;
-}
-
-
 FScheduledTask::FScheduledTask(FEntityAllocationWriteContext InWriteContextOffset)
 	: StatId(GET_STATID(MovieSceneEval_AnonymousTask))
 	, WriteContextOffset(InWriteContextOffset)
@@ -565,8 +511,6 @@ void FEntitySystemScheduler::ExecuteTasks()
 	check(PreviousNumRemaining == 0);
 
 	ThreadingModel = ThreadingModelToUse;
-	GameThreadTaskList.SetThreadingModel(ThreadingModelToUse);
-
 	// In a transaction we can only support the no-threading mode.
 	check(!AutoRTFM::IsTransactional() || (EEntityThreadingModel::NoThreading == ThreadingModel));
 
@@ -579,12 +523,7 @@ void FEntitySystemScheduler::ExecuteTasks()
 		TRACE_CPUPROFILER_EVENT_SCOPE_STR("Run Scheduled Tasks");
 		for (int32 Index : InitialTasks)
 		{
-			GameThreadTaskList.Push(const_cast<FScheduledTask*>(&Tasks[Index]));
-		}
-
-		while (const FScheduledTask* Task = GameThreadTaskList.Pop())
-		{
-			Task->Run(this, FTaskExecutionFlags());
+			Tasks[Index].Run(this, FTaskExecutionFlags());
 		}
 
 		check(NumTasksRemaining.Load(ThreadingModel) == 0);
@@ -718,7 +657,7 @@ void FEntitySystemScheduler::PrerequisiteCompleted(const FScheduledTask* Task, i
 		}
 		else if (ThreadingModel == EEntityThreadingModel::NoThreading)
 		{
-			GameThreadTaskList.Push(const_cast<FScheduledTask*>(Task));
+			Task->Run(this, FTaskExecutionFlags());
 		}
 		else if (Task->bForceInline)
 		{
diff --git a/Engine/Source/Runtime/MovieScene/Private/EntitySystem/MovieSceneTaskScheduler.h b/Engine/Source/Runtime/MovieScene/Private/EntitySystem/MovieSceneTaskScheduler.h
--- a/Engine/Source/Runtime/MovieScene/Private/EntitySystem/MovieSceneTaskScheduler.h
+++ b/Engine/Source/Runtime/MovieScene/Private/EntitySystem/MovieSceneTaskScheduler.h
@@ -8,7 +8,6 @@
 #include "Tasks/Task.h"
 #include "Containers/LockFreeList.h"
 #include "Containers/StringView.h"
-#include "Misc/TVariant.h"
 #include "EntitySystem/RelativePtr.h"
 #include "EntitySystem/IMovieSceneTaskScheduler.h"
 #include "EntitySystem/MovieSceneEntityManager.h"
@@ -347,22 +346,7 @@ private:
 	mutable FEntitySystemMaybeAtomicInt32 NumTasksRemaining = 0;
 
 	FEvent* GameThreadSignal = nullptr;
-
-	struct FOptionallyAtomicList
-	{
-		FOptionallyAtomicList();
-		void SetThreadingModel(EEntityThreadingModel ThreadingModel);
-		void Push(FScheduledTask* InTask);
-		FScheduledTask* Pop();
-		bool IsEmpty() const;
-	private:
-
-		using FLockFreeList = TLockFreePointerListFIFO<FScheduledTask, PLATFORM_CACHE_LINE_SIZE>;
-		using FArrayList    = TArray<FScheduledTask*>;
-		TVariant<FLockFreeList, FArrayList> GameThreadTaskList;
-	};
-
-	mutable FOptionallyAtomicList GameThreadTaskList;
+	mutable TLockFreePointerListFIFO<FScheduledTask, PLATFORM_CACHE_LINE_SIZE> GameThreadTaskList;
 	FEntityAllocationWriteContext WriteContextBase = FEntityAllocationWriteContext::NewAllocation();
 	uint32 SystemSerialIncrement = 0;
 	EEntityThreadingModel ThreadingModel = EEntityThreadingModel::NoThreading;
