Sequencer: Fixed component write operations not depending on previously issued read operations

- Previously, component write tasks would not depend on all previously scheduled read taks creating the potential for them to be scheduled too early.
- This really only showed up when using transform tracks and skeletal animation together with root motion, since that situation requires that the facade transform task waits for all transform reads.

diff --git a/Engine/Source/Runtime/MovieScene/Private/EntitySystem/MovieSceneTaskScheduler.cpp b/Engine/Source/Runtime/MovieScene/Private/EntitySystem/MovieSceneTaskScheduler.cpp
--- a/Engine/Source/Runtime/MovieScene/Private/EntitySystem/MovieSceneTaskScheduler.cpp
+++ b/Engine/Source/Runtime/MovieScene/Private/EntitySystem/MovieSceneTaskScheduler.cpp
@@ -182,8 +182,8 @@ FTaskID FEntitySystemScheduler::CreateForkedAllocationTask(const FTaskParams& In
 
 	const int32 StartingNumTasks = Tasks.Num();
 
-	// Visit all the allocations
-	FComponentMask ReadWritePrerequisites = FComponentMask::BitwiseOR(ReadDeps, WriteDeps, EBitwiseOperatorFlags::MaxSize);
+	// We should never encounter both read and write dependencies for the same component
+	ensure(FComponentMask::BitwiseAND(ReadDeps, WriteDeps, EBitwiseOperatorFlags::MinSize).NumComponents() == 0);
 
 	FTaskID LastTaskID, ParentTaskID;
 
@@ -274,27 +274,41 @@ FTaskID FEntitySystemScheduler::CreateForkedAllocationTask(const FTaskParams& In
 		++Tasks[ParentTaskID.Index].NumChildren;
 
 		// If we're forking these tasks, the parent schedules this allocation task as soon as possible
-		// In this case that means that anything that writes to the same components on the same allocation is finished
+		// In this case that means as soon as everything upstream that writes to the same components on the same allocation is finished
 		// This guarantees we can never read/write from/to component data that is being written to on another thread
-		for (FComponentMaskIterator It = ReadWritePrerequisites.Iterate(); It; ++It)
+
+		// Component reads only depend up upstream writes
+		for (FComponentMaskIterator It = ReadDeps.Iterate(); It; ++It)
 		{
 			FComponentTypeID Type = FComponentTypeID::FromBitIndex(It.GetIndex());
 
 			// Otherwise we can be scheduled as soon as the last write task to this allocation is done
-			if (const FTaskBitSet* UpstreamTask = ComponentWriteDepedenciesByAllocation.Find(MakeTuple(AllocationIndex, Type)))
+			FComponentDependencies& Dependencies = ComponentDepedenciesByAllocation.FindOrAdd(MakeTuple(AllocationIndex, Type));
+			if (Dependencies.WriteTask)
 			{
-				for (int32 Dep : *UpstreamTask)
-				{
-					AddPrerequisite(FTaskID(Dep), FTaskID(AllocationTask));
-				}
+				AddPrerequisite(Dependencies.WriteTask, ThisTask);
 			}
+
+			Dependencies.ReadTasks.SetBit(ThisTask.Index);
 		}
 
-		// We are now the new dependency for anything that we write to for this allocation.
+		// Component writes depend up upstream reads and writes, and become the new write dependency for anything downstream
 		for (FComponentMaskIterator It = WriteDeps.Iterate(); It; ++It)
 		{
 			FComponentTypeID Type = FComponentTypeID::FromBitIndex(It.GetIndex());
-			ComponentWriteDepedenciesByAllocation.FindOrAdd(MakeTuple(AllocationIndex, Type)).SetBit(ThisTask.Index);
+
+			FComponentDependencies& Dependencies = ComponentDepedenciesByAllocation.FindOrAdd(MakeTuple(AllocationIndex, Type));
+			for (int32 Dep : Dependencies.ReadTasks)
+			{
+				AddPrerequisite(FTaskID(Dep), ThisTask);
+			}
+			if (Dependencies.WriteTask)
+			{
+				AddPrerequisite(Dependencies.WriteTask, ThisTask);
+			}
+
+			Dependencies.ReadTasks = FTaskBitSet();
+			Dependencies.WriteTask = ThisTask;
 		}
 
 		// If the tasks are serial, we depend on the last child we made
@@ -316,7 +330,7 @@ void FEntitySystemScheduler::AddPrerequisite(FTaskID Prerequisite, FTaskID Subse
 	// Take care not to add the same dependency multiple times.
 	// This will happen if the same upstream task writes to many components
 	// that this task needs.
-	if (Prerequisite && Subsequent && !Tasks[Prerequisite.Index].ComputedSubsequents.IsBitSet(Subsequent.Index))
+	if (Prerequisite && Subsequent && Prerequisite.Index != Subsequent.Index && !Tasks[Prerequisite.Index].ComputedSubsequents.IsBitSet(Subsequent.Index))
 	{
 		Tasks[Prerequisite.Index].ComputedSubsequents.SetBit(Subsequent.Index);
 		Tasks[Subsequent.Index].NumPrerequisites += 1;
@@ -753,7 +767,7 @@ void FEntitySystemScheduler::EndConstruction()
 	Tasks.Shrink();
 
 	AllPrerequisites.Empty();
-	ComponentWriteDepedenciesByAllocation.Empty();
+	ComponentDepedenciesByAllocation.Empty();
 
 	CurrentSubsequents.Reset();
 	CurrentPrerequisites = nullptr;
diff --git a/Engine/Source/Runtime/MovieScene/Private/EntitySystem/MovieSceneTaskScheduler.h b/Engine/Source/Runtime/MovieScene/Private/EntitySystem/MovieSceneTaskScheduler.h
--- a/Engine/Source/Runtime/MovieScene/Private/EntitySystem/MovieSceneTaskScheduler.h
+++ b/Engine/Source/Runtime/MovieScene/Private/EntitySystem/MovieSceneTaskScheduler.h
@@ -336,8 +336,13 @@ private:
 	/** Sparse bit set of all the tasks that have no prerequisites. Only valid after EndConstruction has been called. */
 	FTaskBitSet InitialTasks;
 
-	/** Map that defines tasks that write to specific components on specific allocations. */
-	TMap<TPair<int32, FComponentTypeID>, FTaskBitSet> ComponentWriteDepedenciesByAllocation;
+	struct FComponentDependencies
+	{
+		FTaskBitSet ReadTasks;
+		FTaskID WriteTask;
+	};
+	/** Map that defines tasks that read from/write to specific components on specific allocations. */
+	TMap<TPair<int32, FComponentTypeID>, FComponentDependencies> ComponentDepedenciesByAllocation;
 
 	TMap<uint16, FTaskPrerequisiteCache> AllPrerequisites;
 
diff --git a/Engine/Source/Runtime/MovieScene/Public/EntitySystem/MovieSceneComponentAccessors.h b/Engine/Source/Runtime/MovieScene/Public/EntitySystem/MovieSceneComponentAccessors.h
--- a/Engine/Source/Runtime/MovieScene/Public/EntitySystem/MovieSceneComponentAccessors.h
+++ b/Engine/Source/Runtime/MovieScene/Public/EntitySystem/MovieSceneComponentAccessors.h
@@ -532,14 +532,12 @@ inline void PopulateReadWriteDependencies(const FOptionalReadAccess* In, FCompon
 inline void PopulateReadWriteDependencies(const FWriteAccess* In, FComponentMask& OutReadDependencies, FComponentMask& OutWriteDependencies)
 {
 	checkSlow(In->ComponentType);
-	OutReadDependencies.Set(In->ComponentType);
 	OutWriteDependencies.Set(In->ComponentType);
 }
 inline void PopulateReadWriteDependencies(const FOptionalWriteAccess* In, FComponentMask& OutReadDependencies, FComponentMask& OutWriteDependencies)
 {
 	if (In->ComponentType)
 	{
-		OutReadDependencies.Set(In->ComponentType);
 		OutWriteDependencies.Set(In->ComponentType);
 	}
 }
