Sequencer: Increase the maximum number of component types

diff --git a/Engine/Source/Runtime/MovieScene/Public/EntitySystem/MovieSceneEntityIDs.h b/Engine/Source/Runtime/MovieScene/Public/EntitySystem/MovieSceneEntityIDs.h
--- a/Engine/Source/Runtime/MovieScene/Public/EntitySystem/MovieSceneEntityIDs.h
+++ b/Engine/Source/Runtime/MovieScene/Public/EntitySystem/MovieSceneEntityIDs.h
@@ -27,7 +27,7 @@ template<typename T> struct TComponentTypeID;
  *
  * Component masks are stored as fixed size inline bit arrays, so considerations should be made before increasing this number.
  */
-static constexpr int32 MaximumNumComponentsSupported = 256;
+static constexpr int32 MaximumNumComponentsSupported = 384;
 
 /** Defines the number of DWORDs that we need to store for component masks */
 static constexpr int32 ComponentMaskDWORDSize = (MaximumNumComponentsSupported + 31) / 32;
