Sequencer: Fixed rare crash caused by stale IMovieScenePlaybackClient pointers being held in the caabilities struct when playback context changes

diff --git a/Engine/Source/Runtime/MovieScene/Private/Evaluation/MovieSceneEvaluationTemplateInstance.cpp b/Engine/Source/Runtime/MovieScene/Private/Evaluation/MovieSceneEvaluationTemplateInstance.cpp
--- a/Engine/Source/Runtime/MovieScene/Private/Evaluation/MovieSceneEvaluationTemplateInstance.cpp
+++ b/Engine/Source/Runtime/MovieScene/Private/Evaluation/MovieSceneEvaluationTemplateInstance.cpp
@@ -11,6 +11,7 @@
 #include "EntitySystem/MovieSceneEntitySystemRunner.h"
 #include "EntitySystem/BuiltInComponentTypes.h"
 #include "Evaluation/Instances/MovieSceneTrackEvaluator.h"
+#include "IMovieScenePlaybackClient.h"
 
 #include "Sections/MovieSceneSubSection.h"
 
@@ -454,6 +455,13 @@ void FMovieSceneRootEvaluationTemplateInstance::PlaybackContextChanged(IMovieSce
 	{
 		EntitySystemLinker->CleanupInvalidBoundObjects();
 
+		// Update the playback client from the player - in some cases this points to an actor that could be destroyed when the playback context changes
+		//   so we overwrite it with the latest value (which could be nullptr, but that's fine)
+		if (SharedPlaybackState->HasCapability<IMovieScenePlaybackClient>())
+		{
+			SharedPlaybackState->GetCapabilities().OverwriteCapabilityRaw(Player.GetPlaybackClient());
+		}
+
 		TSharedPtr<FMovieSceneEntitySystemRunner> Runner = SharedPlaybackState->GetRunner();
 		if (Runner)
 		{
