Allow multiple flags for UIFrameDisplayAs metadata on FFrameNumber displays so we can properly show nonmusical durations

#rnx

diff --git a/Engine/Source/Editor/Sequencer/Private/ISequencerNumericTypeInterface.cpp b/Engine/Source/Editor/Sequencer/Private/ISequencerNumericTypeInterface.cpp
--- a/Engine/Source/Editor/Sequencer/Private/ISequencerNumericTypeInterface.cpp
+++ b/Engine/Source/Editor/Sequencer/Private/ISequencerNumericTypeInterface.cpp
@@ -12,7 +12,7 @@ int32 FSequencerNumericTypeInterface::GetRelevancyScore(const ISequencer& Sequen
 	{
 		const FString& MetaUIDisplayAsString = Property->GetMetaData(TEXT("UIFrameDisplayAs"));
 
-		const bool bIsDurationProperty = MetaUIDisplayAsString.Equals(TEXT("Duration"), ESearchCase::IgnoreCase);
+		const bool bIsDurationProperty = MetaUIDisplayAsString.Contains(TEXT("Duration"), ESearchCase::IgnoreCase);
 
 		// Don't use non-duration intents for duration properties and visa-versa
 		if ((Intent == ENumericIntent::Duration) == bIsDurationProperty)
diff --git a/Engine/Source/Runtime/MovieScene/Public/MovieSceneSection.h b/Engine/Source/Runtime/MovieScene/Public/MovieSceneSection.h
--- a/Engine/Source/Runtime/MovieScene/Public/MovieSceneSection.h
+++ b/Engine/Source/Runtime/MovieScene/Public/MovieSceneSection.h
@@ -757,11 +757,11 @@ public:
 private:
 
 	/** The amount of time to prepare this section for evaluation before it actually starts. */
-	UPROPERTY(EditAnywhere, AdvancedDisplay, Category="Section", meta=(UIMin=0, UIFrameDisplayAs=NonMusical))
+	UPROPERTY(EditAnywhere, AdvancedDisplay, Category="Section", meta=(UIMin=0, UIFrameDisplayAs="NonMusical|Duration"))
 	FFrameNumber PreRollFrames;
 
 	/** The amount of time to continue 'postrolling' this section for after evaluation has ended. */
-	UPROPERTY(EditAnywhere, AdvancedDisplay, Category="Section", meta=(UIMin=0, UIFrameDisplayAs=NonMusical))
+	UPROPERTY(EditAnywhere, AdvancedDisplay, Category="Section", meta=(UIMin=0, UIFrameDisplayAs="NonMusical|Duration"))
 	FFrameNumber PostRollFrames;
 
 	/** The row index that this section sits on */
diff --git a/Engine/Source/Runtime/MovieSceneTracks/Public/Sections/MovieSceneAudioSection.h b/Engine/Source/Runtime/MovieSceneTracks/Public/Sections/MovieSceneAudioSection.h
--- a/Engine/Source/Runtime/MovieSceneTracks/Public/Sections/MovieSceneAudioSection.h
+++ b/Engine/Source/Runtime/MovieSceneTracks/Public/Sections/MovieSceneAudioSection.h
@@ -214,7 +214,7 @@ private:
 	TObjectPtr<USoundBase> Sound;
 
 	/** The offset into the beginning of the audio clip */
-	UPROPERTY(EditAnywhere, Category="Audio", meta=(UIFrameDisplayAs=NonMusical))
+	UPROPERTY(EditAnywhere, Category="Audio", meta=(UIFrameDisplayAs="NonMusical|Position"))
 	FFrameNumber StartFrameOffset;
 
 	/** The offset into the beginning of the audio clip */
