Sequencer: Refactored grouping system to better handle relinked entities

Previously, the grouping system worked on the assumption that entites that had been grouped would always have NeedsUnlink called on them if the group ID component was ever removed. This was sadly not the case for reimported entities, where the group component would be silently thrown away, and the NeedsLink tag added back onto the entity. This was problematic for the control rig parameter track, so I previously made Group components 'preserved' to ensure that they didn't get thrown away.

This was a bit of a hack though: really, the grouping system should be able to handle re-linking an entity and correctly grouping entities regardless of their previous group state.

This was exacerbated by the MPC grouping code that I recently submitted, because now relinked entities were getting added to groups multiple times, and the invariants in the grouping system were being broken.

Really, the grouping system should be more robust against these kinds of issues, so I've reworked it a bit to properly keep track of which entity is in which group (by using a sparse array map of Entity ID -> group ID in place of the previous bit array), and handlers maintain a reverse mapping from group index to group key in order to facilitate the proper removal of group keys when a group index is freed.

I've also removed the superfluous Pre/Post task code that was unused.

diff --git a/Engine/Source/Runtime/MovieScene/Private/EntitySystem/BuiltInComponentTypes.cpp b/Engine/Source/Runtime/MovieScene/Private/EntitySystem/BuiltInComponentTypes.cpp
--- a/Engine/Source/Runtime/MovieScene/Private/EntitySystem/BuiltInComponentTypes.cpp
+++ b/Engine/Source/Runtime/MovieScene/Private/EntitySystem/BuiltInComponentTypes.cpp
@@ -287,7 +287,7 @@ FBuiltInComponentTypes::FBuiltInComponentTypes()
 	ComponentRegistry->NewComponentType(&EvalTime,              TEXT("Eval Time"));
 	ComponentRegistry->NewComponentType(&EvalSeconds,           TEXT("Eval Seconds"));
 
-	ComponentRegistry->NewComponentType(&Group,					TEXT("Entity Group"), EComponentTypeFlags::Preserved);
+	ComponentRegistry->NewComponentType(&Group,					TEXT("Entity Group"));
 
 	ComponentRegistry->NewComponentType(&BoundObjectKey,        TEXT("Bound Object Key"));
 	// Intentionally hidden from the reference graph because they are always accompanied by a BoundObjectKey which is used for garbage collection
diff --git a/Engine/Source/Runtime/MovieScene/Private/EntitySystem/MovieSceneEntityGroupingSystem.cpp b/Engine/Source/Runtime/MovieScene/Private/EntitySystem/MovieSceneEntityGroupingSystem.cpp
--- a/Engine/Source/Runtime/MovieScene/Private/EntitySystem/MovieSceneEntityGroupingSystem.cpp
+++ b/Engine/Source/Runtime/MovieScene/Private/EntitySystem/MovieSceneEntityGroupingSystem.cpp
@@ -21,45 +21,81 @@ FEntityGroupBuilder::FEntityGroupBuilder(UMovieSceneEntityGroupingSystem* InOwne
 
 void FEntityGroupBuilder::AddEntityToGroup(const FMovieSceneEntityID& InEntity, const FEntityGroupID& InNewGroupID)
 {
+	const int32 EntityIndex = InEntity.AsIndex();
+
 	if (ensure(InNewGroupID.HasGroup()))
 	{
-		// Add the entity to the group.
+		const bool bAlreadyHasGroup = Owner->EntityIDToGroup.IsValidIndex(EntityIndex);
+		if (bAlreadyHasGroup)
+		{
+			// Remove it from the old group
+			FEntityGroupID OldGroupID = Owner->EntityIDToGroup[EntityIndex];
+			if (OldGroupID == InNewGroupID)
+			{
+				// Already part of this group - nothing to do
+				return;
+			}
+			RemoveEntityFromGroup(InEntity, OldGroupID);
+		}
+
 		UMovieSceneEntityGroupingSystem::FEntityGroupInfo& GroupInfo = Owner->Groups.FindOrAdd(InNewGroupID);
-		GroupInfo.Entities.Add(InEntity);
+
+		++GroupInfo.NumEntities;
+		Owner->EntityIDToGroup.Insert(EntityIndex, InNewGroupID);
 	}
 }
 
-bool FEntityGroupBuilder::RemoveEntityFromGroup(const FMovieSceneEntityID& InEntity, const FEntityGroupID& InPreviousGroupID)
+void FEntityGroupBuilder::RemoveEntityFromGroup(const FMovieSceneEntityID& InEntity, const FEntityGroupID& InPreviousGroupID)
 {
-	if (ensure(InPreviousGroupID.HasGroup()))
+	const int32 EntityIndex = InEntity.AsIndex();
+
+	if (ensure(InPreviousGroupID.HasGroup()) && Owner->EntityIDToGroup.IsValidIndex(EntityIndex))
 	{
+		ensureMsgf(InPreviousGroupID == Owner->EntityIDToGroup[EntityIndex], TEXT("Attempting to remove an entity from the wrong group!"));
+
+		Owner->EntityIDToGroup.RemoveAt(InEntity.AsIndex());
+
 		// Remove the entity from the group. We should find that group, and find that entity inside it.
 		UMovieSceneEntityGroupingSystem::FEntityGroupInfo* PreviousGroup = Owner->Groups.Find(InPreviousGroupID);
-		if (ensure(PreviousGroup))
+		if (ensureAlways(PreviousGroup))
 		{
 			using FEntityGroupingHandlerInfo = UMovieSceneEntityGroupingSystem::FEntityGroupingHandlerInfo;
 
-			const int32 NumRemoved = PreviousGroup->Entities.Remove(InEntity);
-			ensureAlwaysMsgf(NumRemoved == 1, TEXT("Expected 1 entry to be removed from group by got %d instead."), NumRemoved);
-			if (PreviousGroup->Entities.Num() == 0)
+			--PreviousGroup->NumEntities;
+			if (PreviousGroup->NumEntities <= 0)
 			{
-				// The group is now empty! Let's remove it.
-				Owner->Groups.Remove(InPreviousGroupID);
-				return true;
+				ensure(PreviousGroup->NumEntities == 0);
+
+				// The group is now empty, we can re-use its index for a new group... but we don't
+				// want to re-use it until later, because we could end up in two situations we
+				// want to avoid:
+				//
+				// 1) At the same time that a group is emptied of all its entities, new entities
+				//    come in that belong to that group because they generate the exact same key.
+				//    We want that group to effectively "persist" with the same index.
+				//
+				// 2) If we freed the index right away, we couldn't tell the difference between
+				//	  the above situation, and a brand new group that just happens to re-use the
+				//	  recently freed index.
+				//
+				Owner->EmptyGroupIndices.PadToNum(InPreviousGroupID.GroupIndex + 1, false);
+				Owner->EmptyGroupIndices[InPreviousGroupID.GroupIndex] = true;
 			}
 		}
 	}
-	return false;
 }
 
-int32 FEntityGroupBuilder::AllocateGroupIndex()
+void FEntityGroupBuilder::ReportUsedGroupIndex(int32 GroupIndex)
 {
-	return Owner->AllocateGroupIndex();
+	if (Owner->EmptyGroupIndices.IsValidIndex(GroupIndex))
+	{
+		Owner->EmptyGroupIndices[GroupIndex] = false;
+	}
 }
 
-void FEntityGroupBuilder::FreeGroupIndex(int32 InGroupIndex)
+int32 FEntityGroupBuilder::AllocateGroupIndex()
 {
-	return Owner->FreeGroupIndex(InGroupIndex);
+	return Owner->AllocateGroupIndex(PolicyKey);
 }
 
 
@@ -78,7 +114,7 @@ struct FAddGroupMutation : IMovieSceneEntityMutation
 	{
 		for (const UMovieSceneEntityGroupingSystem::FEntityGroupingHandlerInfo& HandlerInfo : System->GroupHandlers)
 		{
-			if (InOutEntityComponentTypes->ContainsAll(HandlerInfo.ComponentMask))
+			if (HandlerInfo.ComponentFilter.Match(*InOutEntityComponentTypes))
 			{
 				InOutEntityComponentTypes->Set(BuiltInComponents->Group);
 				break;
@@ -96,7 +132,7 @@ struct FAddGroupMutation : IMovieSceneEntityMutation
 		FEntityGroupingPolicyKey PolicyKey;
 		for (auto It = System->GroupHandlers.CreateConstIterator(); It; ++It)
 		{
-			if (AllocationType.ContainsAll(It->ComponentMask))
+			if (It->ComponentFilter.Match(AllocationType))
 			{
 				PolicyKey = FEntityGroupingPolicyKey(It.GetIndex());
 				break;
@@ -128,23 +164,6 @@ struct FUpdateGroupsTask
 	{
 	}
 
-	void PreTask()
-	{
-		BuiltInComponents = FBuiltInComponentTypes::Get();
-
-		// Run the handlers' pre-task callback.
-		
-		for (auto It = System->GroupHandlers.CreateIterator(); It; ++It)
-		{
-			const UMovieSceneEntityGroupingSystem::FEntityGroupingHandlerInfo& HandlerInfo = *It;
-			if (HandlerInfo.Handler.IsValid())
-			{
-				FEntityGroupBuilder Builder(System, FEntityGroupingPolicyKey{ It.GetIndex() });
-				HandlerInfo.Handler->PreTask(&Builder);
-			}
-		}
-	}
-
 	void ForEachAllocation(FEntityAllocationIteratorItem Item, FReadEntityIDs EntityIDs, TWrite<FEntityGroupID> GroupComponents)
 	{
 		TBitArray<> MatchingHandlers;
@@ -171,7 +190,7 @@ struct FUpdateGroupsTask
 	{
 		for (auto It = System->GroupHandlers.CreateConstIterator(); It; ++It)
 		{
-			if (AllocationType.ContainsAll(It->ComponentMask))
+			if (It->ComponentFilter.Match(AllocationType))
 			{
 				int32 Index = It.GetIndex();
 				OutMatchingHandlers.PadToNum(Index + 1, false);
@@ -182,15 +201,9 @@ struct FUpdateGroupsTask
 
 	void PostTask()
 	{
-		// Run the handlers' post-task callback.
-		for (auto It = System->GroupHandlers.CreateIterator(); It; ++It)
+		if (bFreeGroupIDs)
 		{
-			const UMovieSceneEntityGroupingSystem::FEntityGroupingHandlerInfo& HandlerInfo = *It;
-			if (HandlerInfo.Handler.IsValid())
-			{
-				FEntityGroupBuilder Builder(System, FEntityGroupingPolicyKey{ It.GetIndex() });
-				HandlerInfo.Handler->PostTask(bFreeGroupIDs, &Builder);
-			}
+			System->FreeEmptyGroups();
 		}
 	}
 };
@@ -230,27 +243,31 @@ bool UMovieSceneEntityGroupingSystem::IsRelevantImpl(UMovieSceneEntitySystemLink
 	return !GroupHandlers.IsEmpty();
 }
 
-int32 UMovieSceneEntityGroupingSystem::AllocateGroupIndex()
+int32 UMovieSceneEntityGroupingSystem::AllocateGroupIndex(FEntityGroupingPolicyKey InPolicy)
 {
-	// Find an existing free index
-	int32 NewGroupIndex = AllocatedGroupIndices.Find(false);
-	if (NewGroupIndex == INDEX_NONE)
-	{
-		NewGroupIndex = AllocatedGroupIndices.Add(true);
-	}
-	else
-	{
-		AllocatedGroupIndices[NewGroupIndex] = true;
-	}
-	return NewGroupIndex;
+	return AllocatedGroupIndices.Emplace(InPolicy);
 }
 
-void UMovieSceneEntityGroupingSystem::FreeGroupIndex(int32 InGroupIndex)
+void UMovieSceneEntityGroupingSystem::FreeEmptyGroups()
 {
-	if (ensure(AllocatedGroupIndices.IsValidIndex(InGroupIndex) && AllocatedGroupIndices[InGroupIndex]))
+	if (EmptyGroupIndices.Find(true) == INDEX_NONE)
 	{
-		AllocatedGroupIndices[InGroupIndex] = false;
+		return;
 	}
+
+	// Free the indices we don't use anymore, and free the group key lookup entry too.
+	for (TConstSetBitIterator<> It(EmptyGroupIndices); It; ++It)
+	{
+		const int32 GroupIndex = It.GetIndex();
+		FEntityGroupingPolicyKey Policy = AllocatedGroupIndices[GroupIndex];
+
+		GroupHandlers[Policy.Index].Handler->OnGroupIndexFreed(GroupIndex);
+		
+		// The group is now empty! Let's remove it.
+		Groups.Remove(FEntityGroupID{ Policy, GroupIndex });
+		AllocatedGroupIndices.RemoveAt(GroupIndex);
+	}
+	EmptyGroupIndices.Empty();
 }
 
 void UMovieSceneEntityGroupingSystem::OnRun(FSystemTaskPrerequisites& InPrerequisites, FSystemSubsequentTasks& Subsequents)
@@ -273,16 +290,18 @@ void UMovieSceneEntityGroupingSystem::OnRun(FSystemTaskPrerequisites& InPrerequi
 	.FilterAny({ BuiltInComponents->Tags.NeedsLink, BuiltInComponents->Tags.NeedsUnlink })
 	.RunInline_PerAllocation(&Linker->EntityManager, GroupTask);
 
-#if !UE_BUILD_SHIPPING
-	for (const TPair<FEntityGroupID, FEntityGroupInfo>& Pair : Groups)
+#if DO_CHECK && !(UE_BUILD_SHIPPING || UE_BUILD_TEST)
+	for (int32 EntityIndex = 0; EntityIndex < EntityIDToGroup.GetMaxIndex(); ++EntityIndex)
 	{
-		for (FMovieSceneEntityID EntityID : Pair.Value.Entities)
+		if (!EntityIDToGroup.IsValidIndex(EntityIndex))
 		{
-			TReadOptional<FEntityGroupID> GroupComponent = Linker->EntityManager.ReadComponent(EntityID, BuiltInComponents->Group);
-			ensureMsgf(
-					GroupComponent.IsValid() && *GroupComponent == Pair.Key,
-					TEXT("Found mismatch between group cache and group component!"));
+			continue;
 		}
+
+		TReadOptional<FEntityGroupID> GroupComponent = Linker->EntityManager.ReadComponent(FMovieSceneEntityID::FromIndex(EntityIndex), BuiltInComponents->Group);
+		ensureMsgf(
+				GroupComponent.IsValid() && *GroupComponent == EntityIDToGroup[EntityIndex],
+				TEXT("Found mismatch between group cache and group component!"));
 	}
 #endif
 }
@@ -360,13 +379,3 @@ void UMovieSceneEntityGroupingSystem::RemoveGrouping(FEntityGroupingPolicyKey In
 	// Remove the handler itself.
 	GroupHandlers.RemoveAt(InPolicyKey.Index);
 }
-
-TArrayView<const UE::MovieScene::FMovieSceneEntityID> UMovieSceneEntityGroupingSystem::GetGroup(const FEntityGroupID& InGroupID) const
-{
-	if (const FEntityGroupInfo* GroupInfo = Groups.Find(InGroupID))
-	{
-		return MakeArrayView(GroupInfo->Entities);
-	}
-	return TArrayView<const UE::MovieScene::FMovieSceneEntityID>();
-}
-
diff --git a/Engine/Source/Runtime/MovieScene/Public/EntitySystem/MovieSceneEntityGroupingSystem.h b/Engine/Source/Runtime/MovieScene/Public/EntitySystem/MovieSceneEntityGroupingSystem.h
--- a/Engine/Source/Runtime/MovieScene/Public/EntitySystem/MovieSceneEntityGroupingSystem.h
+++ b/Engine/Source/Runtime/MovieScene/Public/EntitySystem/MovieSceneEntityGroupingSystem.h
@@ -46,10 +46,11 @@ struct MOVIESCENE_API FEntityGroupBuilder
 	/** Add the entity to the given group. The entity must already have the group ID component. */
 	void AddEntityToGroup(const FMovieSceneEntityID& InEntity, const FEntityGroupID& InNewGroupID);
 	/** Remove the entity from the given group. The entity must already have the group ID component. */
-	bool RemoveEntityFromGroup(const FMovieSceneEntityID& InEntity, const FEntityGroupID& InPreviousGroupID);
+	void RemoveEntityFromGroup(const FMovieSceneEntityID& InEntity, const FEntityGroupID& InPreviousGroupID);
 
 	int32 AllocateGroupIndex();
-	void FreeGroupIndex(int32 InGroupIndex);
+
+	void ReportUsedGroupIndex(int32 GroupIndex);
 
 private:
 	UMovieSceneEntityGroupingSystem* Owner;
@@ -62,9 +63,8 @@ private:
 struct IEntityGroupingHandler
 {
 	virtual ~IEntityGroupingHandler() {}
-	virtual void PreTask(FEntityGroupBuilder* Builder) {}
 	virtual void ProcessAllocation(FEntityAllocationIteratorItem Item, FReadEntityIDs EntityIDs, TWrite<FEntityGroupID> GroupIDs, FEntityGroupBuilder* Builder) = 0;
-	virtual void PostTask(bool bFreeGroupIDs, FEntityGroupBuilder* Builder) {}
+	virtual void OnGroupIndexFreed(int32 InGroupIndex) {}
 
 #if WITH_EDITOR
 	virtual void OnObjectsReplaced(const TMap<UObject*, UObject*>& ReplacementMap) = 0;
@@ -104,16 +104,14 @@ struct TEntityGroupingHandlerBase
 			// This group key isn't known to us... let's allocate a new group index for it.
 			// Try to find an available index first. Otherwise use a new high index.
 			GroupIndex = Builder->AllocateGroupIndex();
+			IndexToGroupKey.FindOrAdd(GroupIndex, InGroupKey);
 		}
 		else
 		{
 			// We know this group key, so we'll return the group index we already have
 			// associated with it. We just need to "revive" it in case it was scheduled
 			// for being freed.
-			if (EmptyGroupIndices.IsValidIndex(GroupIndex))
-			{
-				EmptyGroupIndices[GroupIndex] = false;
-			}
+			Builder->ReportUsedGroupIndex(GroupIndex);
 		}
 		return GroupIndex;
 	}
@@ -122,9 +120,7 @@ protected:
 
 	/** The group keys that we know about, mapped to their corresponding group index */
 	TMap<GroupKeyType, int32> GroupKeyToIndex;
-
-	/** The transient list of groups freed this frame */
-	TBitArray<> EmptyGroupIndices;
+	TMap<int32, GroupKeyType> IndexToGroupKey;
 };
 
 
@@ -147,38 +143,9 @@ struct TEntityGroupingHandlerImpl<GroupingPolicy, TIntegerSequence<int, Componen
 	{
 	}
 
-	/** Callback on the grouping policy invoked before all grouping happens this frame */
-	static void PreTaskImpl(void*, ...) {}
-	template <typename T> static void PreTaskImpl(T* InPolicy, FEntityGroupBuilder* Builder, decltype(&T::PreTask)* = 0)
-	{
-		InPolicy->PreTask(Builder);
-	}
-
-	virtual void PreTask(FEntityGroupBuilder* Builder) override
-	{
-		PreTaskImpl(&Policy, Builder);
-	}
-
-	/**
-	 * Callback on the grouping policy invoked before grouping happens for a given allocation.
-	 * If the callback returns false, this allocation will be skipped.
-	 */
-	static bool PreAllocationCallbackImpl(void*, ...) { return true; }
-	template <typename T> static bool PreAllocationCallbackImpl(T* InPolicy, FEntityAllocationIteratorItem Item, decltype(&T::PreProcessGroups)* = 0)
-	{
-		return InPolicy->PreProcessGroups(Item);
-	}
-
 	/** Process an allocation and group the entities found therein */
 	virtual void ProcessAllocation(FEntityAllocationIteratorItem Item,  FReadEntityIDs EntityIDs, TWrite<FEntityGroupID> GroupIDs, FEntityGroupBuilder* Builder) override
 	{
-		// Check if we should process this allocation.
-		const bool bProcessAllocation = PreAllocationCallbackImpl(&Policy, Item);
-		if (!bProcessAllocation)
-		{
-			return;
-		}
-
 		const FBuiltInComponentTypes* BuiltInComponents = FBuiltInComponentTypes::Get();
 		const FComponentMask& AllocationType = Item.GetAllocationType();
 		const bool bNeedsLink = AllocationType.Contains(BuiltInComponents->Tags.NeedsLink);
@@ -193,8 +160,6 @@ struct TEntityGroupingHandlerImpl<GroupingPolicy, TIntegerSequence<int, Componen
 		{
 			VisitUnlinkedEntities(Item, EntityIDs, GroupIDs, Builder);
 		}
-
-		PostAllocationCallbackImpl(&Policy, Item); //-V510
 	}
 	
 	void VisitLinkedEntities(FEntityAllocationIteratorItem Item, FReadEntityIDs EntityIDs, TWrite<FEntityGroupID> GroupIDs, FEntityGroupBuilder* Builder)
@@ -253,24 +218,7 @@ struct TEntityGroupingHandlerImpl<GroupingPolicy, TIntegerSequence<int, Componen
 
 			if (GroupID.HasGroup())
 			{
-				const bool bIsGroupEmpty = Builder->RemoveEntityFromGroup(EntityID, GroupID);
-				if (bIsGroupEmpty)
-				{
-					// The group is now empty, we can re-use its index for a new group... but we don't
-					// want to re-use it until later, because we could end up in two situations we
-					// want to avoid:
-					//
-					// 1) At the same time that a group is emptied of all its entities, new entities
-					//    come in that belong to that group because they generate the exact same key.
-					//    We want that group to effectively "persist" with the same index.
-					//
-					// 2) If we freed the index right away, we couldn't tell the difference between
-					//	  the above situation, and a brand new group that just happens to re-use the
-					//	  recently freed index.
-					//
-					this->EmptyGroupIndices.PadToNum(GroupID.GroupIndex + 1, false);
-					this->EmptyGroupIndices[GroupID.GroupIndex] = true;
-				}
+				Builder->RemoveEntityFromGroup(EntityID, GroupID);
 				// Leave the GroupID on the entity so that downstream systems can use it to track
 				// that this entity is leaving its group, but flag it so we don't re-free it.
 				ensure(!EnumHasAllFlags(GroupID.Flags , EEntityGroupFlags::RemovedFromGroup));
@@ -278,48 +226,12 @@ struct TEntityGroupingHandlerImpl<GroupingPolicy, TIntegerSequence<int, Componen
 			}
 		}
 	}
-	
-	/** Callback on the grouping policy invoked after a given allocation has been processed */
-	static void PostAllocationCallbackImpl(void*, ...) {}
-	template <typename T> static void PostAllocationCallbackImpl(T* InPolicy, FEntityAllocationIteratorItem Item, decltype(&T::PostProcessGroups)* = 0)
-	{
-		InPolicy->PostProcessGroups(Item);
-	}
 
-	static void PostTaskImpl(void*, ...){}
-	template <typename T> static void PostTaskImpl(T* InPolicy, FEntityGroupBuilder* Builder, bool bInFreeGroupIDs, decltype(&T::PostTask)* = 0)
+	virtual void OnGroupIndexFreed(int32 GroupIndex) override
 	{
-		InPolicy->PostTask(Builder, bInFreeGroupIDs);
-	}
-	
-	virtual void PostTask(bool bInFreeGroupIDs, FEntityGroupBuilder* Builder) override
-	{
-		if (bInFreeGroupIDs && this->EmptyGroupIndices.Find(true) != INDEX_NONE)
-		{
-			// Build a reverse lookup map to figure out which group keys we don't need anymore
-			// based on the group indices we have freed.
-			TMap<int32, GroupKeyType> GroupIndexToKey;
-			for (const TPair<GroupKeyType, int32>& Pair : this->GroupKeyToIndex)
-			{
-				GroupIndexToKey.Add(Pair.Value, Pair.Key);
-			}
-
-			// Free the indices we don't use anymore, and free the group key lookup entry too.
-			for (TConstSetBitIterator<> It(this->EmptyGroupIndices); It; ++It)
-			{
-				const GroupKeyType* GroupKey = GroupIndexToKey.Find(It.GetIndex());
-				if (ensure(GroupKey))
-				{
-					const int NumRemoved = this->GroupKeyToIndex.Remove(*GroupKey);
-					ensure(NumRemoved == 1);
-				}
-
-				Builder->FreeGroupIndex(It.GetIndex());
-			}
-			this->EmptyGroupIndices.Reset();
-		}
-
-		PostTaskImpl(&Policy, Builder, bInFreeGroupIDs);
+		GroupKeyType Key = this->IndexToGroupKey.FindChecked(GroupIndex);
+		this->IndexToGroupKey.Remove(GroupIndex);
+		this->GroupKeyToIndex.Remove(Key);
 	}
 
 #if WITH_EDITOR
@@ -443,6 +355,12 @@ public:
 	 */
 	template<typename GroupingPolicy, typename ...ComponentTypes>
 	FEntityGroupingPolicyKey AddGrouping(GroupingPolicy&& InPolicy, TComponentTypeID<ComponentTypes>... InComponents)
+	{
+		return AddGrouping(Forward<GroupingPolicy>(InPolicy), UE::MovieScene::FEntityComponentFilter(), InComponents...);
+	}
+
+	template<typename GroupingPolicy, typename ...ComponentTypes>
+	FEntityGroupingPolicyKey AddGrouping(GroupingPolicy&& InPolicy, UE::MovieScene::FEntityComponentFilter&& InComponentFilter, TComponentTypeID<ComponentTypes>... InComponents)
 	{
 		using namespace UE::MovieScene;
 
@@ -453,7 +371,8 @@ public:
 		const int32 HandlerIndex = GroupHandlers.Emplace();
 		FEntityGroupingHandlerInfo& HandlerInfo = GroupHandlers[HandlerIndex];
 		HandlerInfo.Handler.Emplace<NewGroupHandlerType>(MoveTemp(InPolicy), InComponents...);
-		HandlerInfo.ComponentMask = FComponentMask({ InComponents... });
+		HandlerInfo.ComponentFilter = MoveTemp(InComponentFilter);
+		HandlerInfo.ComponentFilter.All({ InComponents... });
 
 		FEntityGroupingPolicyKey NewPolicyKey{ HandlerIndex };
 		return NewPolicyKey;
@@ -474,22 +393,11 @@ public:
 	 */
 	void RemoveGrouping(FEntityGroupingPolicyKey InPolicyKey);
 
-	/**
-	 * Get the list of entities in the given group.
-	 */
-	TArrayView<const UE::MovieScene::FMovieSceneEntityID> GetGroup(const FEntityGroupID& InGroupID) const;
-
 	/**
 	 * Allocate a new group index used to uniquely identify a collection of entities that animate the same target.
 	 * Group indices are globally unique within this system, regardless of the 'type' of the target.
 	 */
-	int32 AllocateGroupIndex();
-
-
-	/**
-	 * Free a previously allocated group index. Must not be in use or present on any components to be freed.
-	 */
-	void FreeGroupIndex(int32 Index);
+	int32 AllocateGroupIndex(FEntityGroupingPolicyKey InPolicy);
 
 
 	/**
@@ -497,9 +405,11 @@ public:
 	 */
 	int32 NumGroups() const
 	{
-		return AllocatedGroupIndices.Num();
+		return AllocatedGroupIndices.GetMaxIndex();
 	}
 
+	void FreeEmptyGroups();
+
 private:
 
 	virtual bool IsRelevantImpl(UMovieSceneEntitySystemLinker* InLinker) const override;
@@ -518,19 +428,24 @@ private:
 
 
 	/** The list of group indices in use */
-	TBitArray<> AllocatedGroupIndices;
+	TSparseArray<FEntityGroupingPolicyKey> AllocatedGroupIndices;
 
 	struct FEntityGroupInfo
 	{
-		TArray<UE::MovieScene::FMovieSceneEntityID> Entities;
+		int32 NumEntities = 0;
 	};
 
 	TMap<FEntityGroupID, FEntityGroupInfo> Groups;
 
+	TSparseArray<FEntityGroupID> EntityIDToGroup;
+
+	/** The transient list of groups freed this frame */
+	TBitArray<> EmptyGroupIndices;
+
 	struct FEntityGroupingHandlerInfo
 	{
 		TInlineValue<IEntityGroupingHandler, 256> Handler;
-		FComponentMask ComponentMask;
+		UE::MovieScene::FEntityComponentFilter ComponentFilter;
 	};
 
 	TSparseArray<FEntityGroupingHandlerInfo> GroupHandlers;
