Sequencer: Added the ability to add modular 'decorations' to moviescene, track and section data

- This new API allows optional meta-data-like objects, or behaviors, to be added to whole sequences, tracks or sections in order to embellish their runtime behavior or to store sparse meta-data on such objects.
- IMovieSceneDecoration, IMovieSceneTrackDecoration, or IMovieSceneSectionDecoration can be added to decoration objects in order to implement added/removed events
- ISequencerCoreModule now implements an API for registering functors to automatically construct ViewModels in the MVVM framework given a specific UObject class.
- This API will be expanded in future to allow even greater extensibility and customizability.

diff --git a/Engine/Source/Editor/Sequencer/Private/MVVM/SectionModelStorageExtension.cpp b/Engine/Source/Editor/Sequencer/Private/MVVM/SectionModelStorageExtension.cpp
--- a/Engine/Source/Editor/Sequencer/Private/MVVM/SectionModelStorageExtension.cpp
+++ b/Engine/Source/Editor/Sequencer/Private/MVVM/SectionModelStorageExtension.cpp
@@ -4,12 +4,23 @@
 
 #include "MVVM/ViewModels/SectionModel.h"
 #include "MovieSceneSection.h"
+#include "ISequencerCoreModule.h"
+#include "Modules/ModuleManager.h"
 
-namespace UE
+namespace UE::Sequencer
 {
-namespace Sequencer
+
+namespace SectionModelStorage
 {
 
+	ISequencerCoreModule& GetSequencerCoreModule()
+	{
+		static ISequencerCoreModule* Module = &FModuleManager::Get().LoadModuleChecked<ISequencerCoreModule>("SequencerCore");
+		return *Module;
+	}
+
+}// namespace SectionModelStorage
+
 FSectionModelStorageExtension::FSectionModelStorageExtension()
 {
 }
@@ -34,10 +45,25 @@ TSharedPtr<FSectionModel> FSectionModelStorageExtension::CreateModelForSection(U
 		return Existing;
 	}
 
-	TSharedPtr<FSectionModel> NewModel = MakeShared<FSectionModel>(InSection, SectionInterface);
-	SectionToModel.Add(SectionAsKey, NewModel);
+	TSharedPtr<FSectionModel> SectionModel;
+
+	if (TSharedPtr<FViewModel> ViewModel = SectionModelStorage::GetSequencerCoreModule().FactoryNewModel(InSection))
+	{
+		SectionModel = ViewModel->CastThisShared<FSectionModel>();
+		ensureMsgf(SectionModel, TEXT("Section model type for Section Object was not an FSectionModel! %s (type: %s)"), *InSection->GetPathName(), *InSection->GetClass()->GetName());
+	}
+	
+	if (!SectionModel)
+	{
+		SectionModel = MakeShared<FSectionModel>();
+		SectionModel->InitializeObject(InSection);
+	}
+
+	SectionModel->InitializeSection(SectionInterface);
+
+	SectionToModel.Add(SectionAsKey, SectionModel);
 
-	return NewModel;
+	return SectionModel;
 }
 
 TSharedPtr<FSectionModel> FSectionModelStorageExtension::FindModelForSection(const UMovieSceneSection* InSection) const
@@ -46,6 +72,5 @@ TSharedPtr<FSectionModel> FSectionModelStorageExtension::FindModelForSection(con
 	return SectionToModel.FindRef(SectionAsKey).Pin();
 }
 
-} // namespace Sequencer
-} // namespace UE
+} // namespace UE::Sequencer
 
diff --git a/Engine/Source/Editor/Sequencer/Private/MVVM/ViewModels/SectionModel.cpp b/Engine/Source/Editor/Sequencer/Private/MVVM/ViewModels/SectionModel.cpp
--- a/Engine/Source/Editor/Sequencer/Private/MVVM/ViewModels/SectionModel.cpp
+++ b/Engine/Source/Editor/Sequencer/Private/MVVM/ViewModels/SectionModel.cpp
@@ -37,15 +37,9 @@ struct FSectionModel_StretchParams : IDynamicExtension
 
 UE_SEQUENCER_DEFINE_VIEW_MODEL_TYPE_ID(FSectionModel_StretchParams);
 
-FSectionModel::FSectionModel(UMovieSceneSection* InSection, TSharedPtr<ISequencerSection> InSectionInterface)
+FSectionModel::FSectionModel()
 	: ChannelList(FTrackModel::GetTopLevelChannelType() | EViewModelListType::Generic)
-	, SectionInterface(InSectionInterface)
-	, WeakSection(InSection)
 {
-	InSection->EventHandlers.Link(this);
-	InSection->UMovieSceneSignedObject::EventHandlers.Link(this);
-
-	UpdateCachedData();
 	RegisterChildList(&ChannelList);
 }
 
@@ -57,6 +51,28 @@ FSectionModel::~FSectionModel()
 }
 PRAGMA_ENABLE_DEPRECATION_WARNINGS
 
+void FSectionModel::InitializeSection(TSharedPtr<ISequencerSection> InSectionInterface)
+{
+	SectionInterface = InSectionInterface;
+}
+
+void FSectionModel::InitializeObject(TWeakObjectPtr<> InWeakObject)
+{
+	UMovieSceneSection* Section = CastChecked<UMovieSceneSection>(InWeakObject.Get());
+
+	WeakSection = Section;
+
+	Section->EventHandlers.Link(this);
+	Section->UMovieSceneSignedObject::EventHandlers.Link(this);
+
+	UpdateCachedData();
+}
+
+UObject* FSectionModel::GetObject() const
+{
+	return WeakSection.Get();
+}
+
 TRange<FFrameNumber> FSectionModel::GetRange() const
 {
 	UMovieSceneSection* Section = WeakSection.Get();
diff --git a/Engine/Source/Editor/Sequencer/Private/MVVM/ViewModels/SequenceModel.cpp b/Engine/Source/Editor/Sequencer/Private/MVVM/ViewModels/SequenceModel.cpp
--- a/Engine/Source/Editor/Sequencer/Private/MVVM/ViewModels/SequenceModel.cpp
+++ b/Engine/Source/Editor/Sequencer/Private/MVVM/ViewModels/SequenceModel.cpp
@@ -2,6 +2,7 @@
 
 #include "MVVM/ViewModels/SequenceModel.h"
 #include "MVVM/Extensions/IObjectBindingExtension.h"
+#include "MVVM/Extensions/IObjectModelExtension.h"
 #include "MVVM/Extensions/ILockableExtension.h"
 #include "MVVM/Extensions/IMutableExtension.h"
 #include "MVVM/Extensions/ISoloableExtension.h"
@@ -21,6 +22,7 @@
 
 #include "Sequencer.h"
 #include "SequencerOutlinerItemDragDropOp.h"
+#include "ISequencerCoreModule.h"
 
 #include "ScopedTransaction.h"
 
@@ -35,15 +37,23 @@ FOnInitializeSequenceModel FSequenceModel::CreateExtensionsEvent;
 
 FSequenceModel::FSequenceModel(TWeakPtr<FSequencerEditorViewModel> InEditorViewModel)
 	: RootOutlinerItems(EViewModelListType::Outliner)
+	, DecorationModelList(GetDecorationModelListType())
 {
 	WeakEditor = InEditorViewModel;
 
 	RegisterChildList(&RootOutlinerItems);
+	RegisterChildList(&DecorationModelList);
 
 	TSharedPtr<FSharedViewModelData> NewSharedData = MakeShared<FEditorSharedViewModelData>(WeakEditor.Pin().ToSharedRef());
 	SetSharedDataFromConstructor(NewSharedData);
 }
 
+EViewModelListType FSequenceModel::GetDecorationModelListType()
+{
+	static EViewModelListType DecorationModelList = RegisterCustomModelListType();
+	return DecorationModelList;
+}
+
 void FSequenceModel::InitializeExtensions()
 {
 	TSharedPtr<FSharedViewModelData> ThisSharedData = GetSharedData();
@@ -79,6 +89,29 @@ void FSequenceModel::InitializeExtensions()
 	PostInitializeExtensions();
 }
 
+void FSequenceModel::OnDecorationAdded(UObject* AddedDecoration)
+{
+	static ISequencerCoreModule& SequencerCore = FModuleManager::Get().LoadModuleChecked<ISequencerCoreModule>("SequencerCore");
+
+	TSharedPtr<FViewModel> NewModel = SequencerCore.FactoryNewModel(AddedDecoration);
+	if (NewModel && NewModel->IsA<IObjectModelExtension>())
+	{
+		GetChildrenForList(&DecorationModelList).AddChild(NewModel);
+	}
+}
+
+void FSequenceModel::OnDecorationRemoved(UObject* RemovedDecoration)
+{
+	for (TViewModelPtr<IObjectModelExtension> Object : GetChildrenForList(&DecorationModelList).IterateSubList<IObjectModelExtension>())
+	{
+		if (Object->GetObject() == RemovedDecoration)
+		{
+			Object.AsModel()->RemoveFromParent();
+			break;
+		}
+	}
+}
+
 FMovieSceneSequenceID FSequenceModel::GetSequenceID() const
 {
 	return SequenceID;
@@ -113,14 +146,25 @@ void FSequenceModel::SetSequence(UMovieSceneSequence* InSequence, FMovieSceneSeq
 
 		SequenceEventHandler.Unlink();
 		MovieSceneEventHandler.Unlink();
+		MovieSceneDataEventHandler.Unlink();
 
 		if (InSequence)
 		{
 			InSequence->EventHandlers.Link(SequenceEventHandler, this);
 			InSequence->GetMovieScene()->UMovieSceneSignedObject::EventHandlers.Link(MovieSceneEventHandler, this);
+			InSequence->GetMovieScene()->EventHandlers.Link(MovieSceneDataEventHandler, this);
 		}
 
 		ReinitializeExtensions();
+
+		GetChildrenForList(&DecorationModelList).Empty();
+		if (InSequence)
+		{
+			for (UObject* Decoration : InSequence->GetMovieScene()->GetDecorations())
+			{
+				OnDecorationAdded(Decoration);
+			}
+		}
 	}
 }
 
diff --git a/Engine/Source/Editor/Sequencer/Public/MVVM/ViewModels/SectionModel.h b/Engine/Source/Editor/Sequencer/Public/MVVM/ViewModels/SectionModel.h
--- a/Engine/Source/Editor/Sequencer/Public/MVVM/ViewModels/SectionModel.h
+++ b/Engine/Source/Editor/Sequencer/Public/MVVM/ViewModels/SectionModel.h
@@ -9,6 +9,7 @@
 #include "MVVM/Extensions/ILayerBarExtension.h"
 #include "MVVM/Extensions/ITrackLaneExtension.h"
 #include "MVVM/Extensions/IGeometryExtension.h"
+#include "MVVM/Extensions/IObjectModelExtension.h"
 #include "MVVM/Extensions/ISelectableExtension.h"
 #include "MVVM/Extensions/ISnappableExtension.h"
 #include "MVVM/Extensions/IDraggableTrackAreaExtension.h"
@@ -38,6 +39,7 @@ struct FOverlappingSections;
 class SEQUENCER_API FSectionModel
 	: public FViewModel
 	, public FLinkedOutlinerExtension
+	, public IObjectModelExtension
 	, public ILayerBarExtension
 	, public ITrackLaneExtension
 	, public ISelectableExtension
@@ -52,6 +54,7 @@ public:
 
 	UE_SEQUENCER_DECLARE_CASTABLE(FSectionModel, FViewModel
 		, FLinkedOutlinerExtension
+		, IObjectModelExtension
 		, ILayerBarExtension
 		, ITrackLaneExtension
 		, ISelectableExtension
@@ -66,9 +69,11 @@ public:
 	UE_DEPRECATED(5.5, "This member is no longer supported, please subscribe to UMovieSceneSignedObject::OnModifiedDirectly.")
 	FOnModelUpdated OnUpdated;
 
-	explicit FSectionModel(UMovieSceneSection* InSection, TSharedPtr<ISequencerSection> InSectionInterface);
+	FSectionModel();
 	~FSectionModel();
 
+	void InitializeSection(TSharedPtr<ISequencerSection> InSectionInterface);
+
 	/** Gets the time range of the section */
 	TRange<FFrameNumber> GetRange() const;
 
@@ -112,6 +117,10 @@ public:
 
 public:
 
+	/*~ IObjectModelExtension Interface */
+	virtual void InitializeObject(TWeakObjectPtr<> InWeakObject) override;
+	virtual UObject* GetObject() const override;
+
 	/*~ ILayerBarExtension Interface */
 	TRange<FFrameNumber> GetLayerBarRange() const override;
 	void OffsetLayerBar(FFrameNumber Amount) override;
diff --git a/Engine/Source/Editor/Sequencer/Public/MVVM/ViewModels/SequenceModel.h b/Engine/Source/Editor/Sequencer/Public/MVVM/ViewModels/SequenceModel.h
--- a/Engine/Source/Editor/Sequencer/Public/MVVM/ViewModels/SequenceModel.h
+++ b/Engine/Source/Editor/Sequencer/Public/MVVM/ViewModels/SequenceModel.h
@@ -11,6 +11,7 @@
 
 #include "MovieSceneSequenceID.h"
 #include "EventHandlers/ISignedObjectEventHandler.h"
+#include "EventHandlers/ISequenceDataEventHandler.h"
 
 #include "SequencerCoreFwd.h"
 #include "MVVM/ViewModels/ViewModel.h"
@@ -41,6 +42,7 @@ class SEQUENCER_API FSequenceModel
 	, public ISortableExtension
 	, public IOutlinerDropTargetOutlinerExtension
 	, public UE::MovieScene::ISignedObjectEventHandler
+	, public UE::MovieScene::ISequenceDataEventHandler
 {
 public:
 
@@ -73,13 +75,21 @@ public:
 	FSortingKey GetSortingKey() const override;
 	void SetCustomOrder(int32 InCustomOrder) override;
 
+	/*~ Begin ISequenceDataEventHandler */
+	virtual void OnDecorationAdded(UObject* AddedDecoration) override;
+	virtual void OnDecorationRemoved(UObject* RemovedDecoration) override;
+	/*~ End ISequenceDataEventHandler */
+
 	/*~ IOutlinerDropTargetOutlinerExtension */
 	TOptional<EItemDropZone> CanAcceptDrop(const FViewModelPtr& TargetModel, const FDragDropEvent& DragDropEvent, EItemDropZone InItemDropZone) override;
 	void PerformDrop(const FViewModelPtr& TargetModel, const FDragDropEvent& DragDropEvent, EItemDropZone InItemDropZone) override;
 
+	static EViewModelListType GetDecorationModelListType();
+
 private:
 
 	FViewModelListHead RootOutlinerItems;
+	FViewModelListHead DecorationModelList;
 
 	TWeakObjectPtr<UMovieSceneSequence> WeakSequence;
 	TWeakPtr<FSequencerEditorViewModel> WeakEditor;
@@ -89,6 +99,7 @@ private:
 
 	MovieScene::TNonIntrusiveEventHandler<MovieScene::ISignedObjectEventHandler> SequenceEventHandler;
 	MovieScene::TNonIntrusiveEventHandler<MovieScene::ISignedObjectEventHandler> MovieSceneEventHandler;
+	MovieScene::TNonIntrusiveEventHandler<MovieScene::ISequenceDataEventHandler> MovieSceneDataEventHandler;
 };
 
 } // namespace Sequencer
diff --git a/Engine/Source/Editor/SequencerCore/Private/MVVM/ViewModelTypeIDs.cpp b/Engine/Source/Editor/SequencerCore/Private/MVVM/ViewModelTypeIDs.cpp
--- a/Engine/Source/Editor/SequencerCore/Private/MVVM/ViewModelTypeIDs.cpp
+++ b/Engine/Source/Editor/SequencerCore/Private/MVVM/ViewModelTypeIDs.cpp
@@ -13,6 +13,7 @@
 #include "MVVM/Extensions/IKeyExtension.h"
 #include "MVVM/Extensions/ILayerBarExtension.h"
 #include "MVVM/Extensions/IMutableExtension.h"
+#include "MVVM/Extensions/IObjectModelExtension.h"
 #include "MVVM/Extensions/IOutlinerExtension.h"
 #include "MVVM/Extensions/IPinnableExtension.h"
 #include "MVVM/Extensions/IRecyclableExtension.h"
@@ -55,6 +56,7 @@ UE_SEQUENCER_DEFINE_VIEW_MODEL_TYPE_ID(ILayerBarExtension);
 UE_SEQUENCER_DEFINE_VIEW_MODEL_TYPE_ID(IMutableExtension);
 UE_SEQUENCER_DEFINE_VIEW_MODEL_TYPE_ID(ILockableExtension);
 UE_SEQUENCER_DEFINE_VIEW_MODEL_TYPE_ID(IOutlinerDropTargetOutlinerExtension);
+UE_SEQUENCER_DEFINE_VIEW_MODEL_TYPE_ID(IObjectModelExtension);
 UE_SEQUENCER_DEFINE_VIEW_MODEL_TYPE_ID(IHierarchicalCache);
 UE_SEQUENCER_DEFINE_VIEW_MODEL_TYPE_ID(IOutlinerExtension);
 UE_SEQUENCER_DEFINE_VIEW_MODEL_TYPE_ID(IPinnableExtension);
diff --git a/Engine/Source/Editor/SequencerCore/Private/SequencerCoreModule.cpp b/Engine/Source/Editor/SequencerCore/Private/SequencerCoreModule.cpp
--- a/Engine/Source/Editor/SequencerCore/Private/SequencerCoreModule.cpp
+++ b/Engine/Source/Editor/SequencerCore/Private/SequencerCoreModule.cpp
@@ -1,8 +1,11 @@
 // Copyright Epic Games, Inc. All Rights Reserved.
 
-#include "SequencerCoreModule.h"
+#include "ISequencerCoreModule.h"
 #include "Modules/ModuleManager.h"
 #include "SequencerCoreLog.h"
+#include "MVVM/Extensions/IObjectModelExtension.h"
+#include "UObject/ObjectKey.h"
+#include "UObject/Class.h"
 
 namespace UE
 {
@@ -16,10 +19,72 @@ namespace Sequencer
 class FSequencerCoreModule
 	: public ISequencerCoreModule
 {
+	virtual FDelegateHandle RegisterModelType(TWeakObjectPtr<UClass> WeakClass, TFunction<TSharedPtr<FViewModel>(UObject*)> FactoryFunctor) override
+	{
+		TObjectKey<UClass> Key = WeakClass.Get();
+		ensureMsgf(!ModelFactories.Contains(Key), TEXT("Attempting to register a duplicate model factory type"));
 
+		FDelegateHandle Handle(FDelegateHandle::GenerateNewHandle);
+		
+		ModelFactories.Add(Key, FFactoryEntry{ FactoryFunctor, Handle });
+		return Handle;
+	}
+
+	virtual void UnRegisterModelType(FDelegateHandle Handle) override
+	{
+		for (auto It = ModelFactories.CreateIterator(); It; ++It)
+		{
+			if (It->Value.Handle == Handle)
+			{
+				It.RemoveCurrent();
+				return;
+			}
+		}
+	}
+
+	virtual TSharedPtr<FViewModel> FactoryNewModel(UObject* Object) const override
+	{
+		check(Object);
+
+		UClass* Class = Object->GetClass();
+		while (Class && Class != UObject::StaticClass())
+		{
+			if (const FFactoryEntry* Factory = ModelFactories.Find(Class))
+			{
+				TSharedPtr<FViewModel> NewModel = Factory->Functor(Object);
+				if (NewModel)
+				{
+					return NewModel;
+				}
+			}
+
+			Class = Class->GetSuperClass();
+		}
+
+		return nullptr;
+	}
+
+private:
+
+	struct FFactoryEntry
+	{
+		TFunction<TSharedPtr<FViewModel>(UObject*)> Functor;
+		FDelegateHandle Handle;
+	};
+
+	TMap<TObjectKey<UClass>, FFactoryEntry> ModelFactories;
 };
 
 
+void ISequencerCoreModule::InitializeObjectModel(const FViewModelPtr& ViewModel, UObject* Object)
+{
+	if (TViewModelPtr<IObjectModelExtension> ObjectModel = ViewModel.ImplicitCast())
+	{
+		ObjectModel->InitializeObject(Object);
+	}
+}
+
+
 } // namespace Sequencer
 } // namespace UE
 
diff --git a/Engine/Source/Editor/SequencerCore/Private/SequencerCoreModule.h b/Engine/Source/Editor/SequencerCore/Private/SequencerCoreModule.h
deleted file mode 100644
diff --git a/Engine/Source/Editor/SequencerCore/Public/ISequencerCoreModule.h b/Engine/Source/Editor/SequencerCore/Public/ISequencerCoreModule.h
new file mode 100644
--- /dev/null
+++ b/Engine/Source/Editor/SequencerCore/Public/ISequencerCoreModule.h
@@ -0,0 +1,83 @@
+// Copyright Epic Games, Inc. All Rights Reserved.
+
+#pragma once
+
+#include "Delegates/IDelegateInstance.h"
+#include "Modules/ModuleInterface.h"
+#include "MVVM/ViewModelPtr.h"
+#include "Templates/Function.h"
+
+
+namespace UE::Sequencer
+{
+
+class FViewModel;
+
+/**
+ * Interface for the Sequencer module.
+ */
+class ISequencerCoreModule
+	: public IModuleInterface
+{
+public:
+
+	/**
+	 * Register a new MVVM view model type factory functor that will be used for any UObjects of the specified class type
+	 * @note It is a violation to attempt to register multiple functors for the same class type
+	 * 
+	 * @param WeakClass       A weak pointer to the class type that should invoke this factory
+	 * @param FactoryFunctor  The functor that defines how to create a view model for the object
+	 * 
+	 * @return A delegate handle that can be passed to UnRegisterModelType to unregister this factory
+	 */
+	virtual FDelegateHandle RegisterModelType(TWeakObjectPtr<UClass> WeakClass, TFunction<TSharedPtr<FViewModel>(UObject*)> FactoryFunctor) = 0;
+
+
+	/**
+	 * Unregister a previously registered view model type
+	 *
+	 * @param Handle  The handle to the factory functor that was provided from calling RegisterModelType
+	 */
+	virtual void UnRegisterModelType(FDelegateHandle Handle) = 0;
+
+
+	/**
+	 * Attempt to create a new model type for the specified object
+	 *
+	 * @param Object  The object to create a view model for
+	 * @return A pointer to a new view-model, or nullptr if the object didn't match any registered class.
+	 */
+	virtual TSharedPtr<FViewModel> FactoryNewModel(UObject* Object) const = 0;
+
+public:
+
+	/**
+	 * Register a new MVVM view model type from a view model type ID that will be used for any UObjects of the specified class type
+	 *
+	 * @param Object  The object to create a view model for
+	 * @param Type    The view model type ID to create for UObjects of type WeakClass
+	 *
+	 * @return A delegate handle that can be passed to UnRegisterModelType to unregister this factory
+	 */
+	template<typename T>
+	FDelegateHandle RegisterModelType(TWeakObjectPtr<UClass> WeakClass, TViewModelTypeID<T> Type)
+	{
+		auto Factory = [this, Type](UObject* Object)
+		{
+			TSharedPtr<T> NewViewModel = MakeShared<T>();
+
+			this->InitializeObjectModel(NewViewModel, Object);
+
+			return NewViewModel;
+		};
+		return RegisterModelType(WeakClass, Factory);
+	}
+
+private:
+
+
+	SEQUENCERCORE_API void InitializeObjectModel(const FViewModelPtr& ViewModel, UObject* Object);
+};
+
+
+} // namespace UE::Sequencer
\ No newline at end of file
diff --git a/Engine/Source/Editor/SequencerCore/Public/MVVM/Extensions/IObjectModelExtension.h b/Engine/Source/Editor/SequencerCore/Public/MVVM/Extensions/IObjectModelExtension.h
new file mode 100644
--- /dev/null
+++ b/Engine/Source/Editor/SequencerCore/Public/MVVM/Extensions/IObjectModelExtension.h
@@ -0,0 +1,27 @@
+// Copyright Epic Games, Inc. All Rights Reserved.
+
+#pragma once
+
+#include "CoreTypes.h"
+#include "MVVM/ViewModelTypeID.h"
+#include "UObject/WeakObjectPtrFwd.h"
+
+namespace UE::Sequencer
+{
+
+/**
+ * An extension that is used for view models that represent UObject data
+ */
+class SEQUENCERCORE_API IObjectModelExtension
+{
+public:
+
+	UE_SEQUENCER_DECLARE_VIEW_MODEL_TYPE_ID(IObjectModelExtension)
+
+	virtual ~IObjectModelExtension(){}
+
+	virtual void InitializeObject(TWeakObjectPtr<> InWeakObject) = 0;
+	virtual UObject* GetObject() const = 0;
+};
+
+} // namespace UE::Sequencer
\ No newline at end of file
diff --git a/Engine/Source/Runtime/MovieScene/Private/Compilation/MovieSceneCompiledDataManager.cpp b/Engine/Source/Runtime/MovieScene/Private/Compilation/MovieSceneCompiledDataManager.cpp
--- a/Engine/Source/Runtime/MovieScene/Private/Compilation/MovieSceneCompiledDataManager.cpp
+++ b/Engine/Source/Runtime/MovieScene/Private/Compilation/MovieSceneCompiledDataManager.cpp
@@ -13,6 +13,7 @@
 #include "Tracks/MovieSceneSubTrack.h"
 #include "Tracks/MovieSceneTimeWarpTrack.h"
 #include "Channels/MovieSceneTimeWarpChannel.h"
+#include "Decorations/IMovieSceneDecoration.h"
 #include "IMovieSceneModule.h"
 #include "MovieSceneTimeHelpers.h"
 #include "MovieSceneTransformTypes.h"
@@ -883,6 +884,14 @@ void UMovieSceneCompiledDataManager::Compile(FMovieSceneCompiledDataID DataID, U
 	UMovieScene* MovieScene = Sequence->GetMovieScene();
 	if (ensure(MovieScene))
 	{
+		for (TObjectPtr<UObject> DecorationObject : MovieScene->GetDecorations())
+		{
+			if (IMovieSceneDecoration* Decoration = Cast<IMovieSceneDecoration>(DecorationObject))
+			{
+				Decoration->OnPreDecorationCompiled();
+			}
+		}
+
 		MovieScene->ResetGeneratedConditions();
 	}
 
@@ -1008,6 +1017,14 @@ void UMovieSceneCompiledDataManager::Compile(FMovieSceneCompiledDataID DataID, U
 	CompiledDataEntries[DataID.Value] = Entry;
 	++ReallocationVersion;
 
+	for (TObjectPtr<UObject> DecorationObject : Sequence->GetMovieScene()->GetDecorations())
+	{
+		if (IMovieSceneDecoration* Decoration = Cast<IMovieSceneDecoration>(DecorationObject))
+		{
+			Decoration->OnPostDecorationCompiled();
+		}
+	}
+
 #if 0
 #if !NO_LOGGING
 	if (bHasHierarchy)
@@ -1037,6 +1054,20 @@ void UMovieSceneCompiledDataManager::Gather(const FMovieSceneCompiledDataEntry&
 
 	if (ensure(MovieScene))
 	{
+		// Allow decorations on the movie scene to define entities in the field
+		if (OutCompilerData->EntityField)
+		{
+			FMovieSceneEntityComponentFieldBuilder FieldBuilder(OutCompilerData->EntityField);
+			for (TObjectPtr<UObject> DecorationObject : MovieScene->GetDecorations())
+			{
+				if (IMovieSceneEntityProvider* Provider = Cast<IMovieSceneEntityProvider>(DecorationObject))
+				{
+					FMovieSceneEvaluationFieldEntityMetaData MetaData;
+					Provider->PopulateEvaluationField(TRange<FFrameNumber>::All(), MetaData, &FieldBuilder);
+				}
+			}
+		}
+
 		if (UMovieSceneTrack* Track = MovieScene->GetCameraCutTrack())
 		{
 			GatherTrack(nullptr, Track, Params, TrackTemplate, OutCompilerData);
@@ -1451,6 +1482,19 @@ void UMovieSceneCompiledDataManager::GatherTrack(const FMovieSceneBinding* Objec
 			FieldBuilder.GetSharedMetaData().ObjectBindingID = ObjectBinding->GetObjectGuid();
 		}
 
+		for (UObject* Decoration : Track->GetDecorations())
+		{
+			if (IMovieSceneEntityProvider* Provider = Cast<IMovieSceneEntityProvider>(Decoration))
+			{
+				FMovieSceneEvaluationFieldEntityMetaData MetaData(PreCompileResult.DefaultMetaData);
+				MetaData.bEvaluateInSequencePreRoll  = Track->EvalOptions.bEvaluateInPreroll;
+				MetaData.bEvaluateInSequencePostRoll = Track->EvalOptions.bEvaluateInPostroll;
+				MetaData.Condition = Track->ConditionContainer.Condition;
+
+				Provider->PopulateEvaluationField(Params.LocalClampRange, MetaData, &FieldBuilder);
+			}
+		}
+
 		IMovieSceneEntityProvider* TrackEntityProvider = Cast<IMovieSceneEntityProvider>(Track);
 
 		// If the track is an entity provider, allow it to add entries first
diff --git a/Engine/Source/Runtime/MovieScene/Private/Decorations/MovieSceneDecorationContainer.cpp b/Engine/Source/Runtime/MovieScene/Private/Decorations/MovieSceneDecorationContainer.cpp
new file mode 100644
--- /dev/null
+++ b/Engine/Source/Runtime/MovieScene/Private/Decorations/MovieSceneDecorationContainer.cpp
@@ -0,0 +1,95 @@
+// Copyright Epic Games, Inc. All Rights Reserved.
+
+#include "Decorations/MovieSceneDecorationContainer.h"
+
+#include "MovieScene.h"
+#include "MovieSceneTrack.h"
+#include "MovieSceneSection.h"
+
+
+UObject* FMovieSceneDecorationContainer::FindDecoration(const TSubclassOf<UObject>& InClass) const
+{
+	if (UClass* Class = InClass.Get())
+	{
+		// Do an exact match - we intentionally do not support derived decorations
+		for (UObject* Decoration : Decorations)
+		{
+			if (Decoration->IsA(Class))
+			{
+				return Decoration;
+			}
+		}
+	}
+	return nullptr;
+}
+
+void FMovieSceneDecorationContainer::AddDecoration(UObject* InDecoration, UObject* Outer, TFunctionRef<void(UObject*)> Event)
+{
+	check(InDecoration && Outer);
+
+	if (!ensureMsgf(FindDecoration(InDecoration->GetClass()) == nullptr, TEXT("Attempting to add a decoration when one of the same type already exists. This request will be ignored.")))
+	{
+		return;
+	}
+
+	if (!ensureMsgf(InDecoration->IsIn(Outer->GetOutermost()), TEXT("Attempting to add a decoration from a different pacakge - this is not allowed.")))
+	{
+		return;
+	}
+
+	Decorations.Add(InDecoration);
+	Event(InDecoration);
+}
+
+UObject* FMovieSceneDecorationContainer::GetOrCreateDecoration(const TSubclassOf<UObject>& InClass, UObject* Outer, TFunctionRef<void(UObject*)> Event)
+{
+	UObject* Found = FindDecoration(InClass);
+	if (!Found)
+	{
+		Found = NewObject<UObject>(Outer, InClass);
+		Decorations.Add(Found);
+		Event(Found);
+	}
+	return Found;
+}
+
+void FMovieSceneDecorationContainer::RemoveDecoration(const TSubclassOf<UObject>& InClass, TFunctionRef<void(UObject*)> Event)
+{
+	if (UClass* Class = InClass.Get())
+	{
+		for (int32 Index = Decorations.Num()-1; Index >= 0; --Index)
+		{
+			if (Decorations[Index]->IsA(Class))
+			{
+				Event(Decorations[Index]);
+				Decorations.RemoveAtSwap(Index);
+			}
+		}
+	}
+}
+
+TArrayView<const TObjectPtr<UObject>> FMovieSceneDecorationContainer::GetDecorations() const
+{
+	return Decorations;
+}
+
+void UMovieSceneDecorationContainerObject::AddDecoration(UObject* InDecoration)
+{
+	Decorations.AddDecoration(InDecoration, this, [this](UObject* Decoration){
+		this->OnDecorationAdded(Decoration);
+	});
+}
+
+UObject* UMovieSceneDecorationContainerObject::GetOrCreateDecoration(const TSubclassOf<UObject>& InClass)
+{
+	return Decorations.GetOrCreateDecoration(InClass, this, [this](UObject* Decoration){
+		this->OnDecorationAdded(Decoration);
+	});
+}
+
+void UMovieSceneDecorationContainerObject::RemoveDecoration(const TSubclassOf<UObject>& InClass)
+{
+	Decorations.RemoveDecoration(InClass, [this](UObject* Decoration){
+		this->OnDecorationRemoved(Decoration);
+	});
+}
\ No newline at end of file
diff --git a/Engine/Source/Runtime/MovieScene/Private/EntitySystem/MovieSceneEntityLedger.cpp b/Engine/Source/Runtime/MovieScene/Private/EntitySystem/MovieSceneEntityLedger.cpp
--- a/Engine/Source/Runtime/MovieScene/Private/EntitySystem/MovieSceneEntityLedger.cpp
+++ b/Engine/Source/Runtime/MovieScene/Private/EntitySystem/MovieSceneEntityLedger.cpp
@@ -5,6 +5,7 @@
 #include "EntitySystem/MovieSceneEntitySystemLinker.h"
 #include "EntitySystem/BuiltInComponentTypes.h"
 #include "EntitySystem/MovieSceneEntityMutations.h"
+#include "EntitySystem/IMovieSceneEntityDecorator.h"
 
 #include "Evaluation/MovieSceneEvaluationField.h"
 #include "Conditions/MovieSceneCondition.h"
@@ -172,6 +173,17 @@ void FEntityLedger::UpdateConditionalEntities(UMovieSceneEntitySystemLinker* Lin
 
 			if (!ImportedEntity.IsEmpty())
 			{
+				if (UMovieSceneDecorationContainerObject* DecorationContainer = Cast<UMovieSceneDecorationContainerObject>(EntityOwner))
+				{
+					for (UObject* Decoration : DecorationContainer->GetDecorations())
+					{
+						if (IMovieSceneEntityDecorator* EntityDecorator = Cast<IMovieSceneEntityDecorator>(Decoration))
+						{
+							EntityDecorator->ExtendEntity(Linker, Params, &ImportedEntity);
+						}
+					}
+				}
+
 				if (UMovieSceneSection* Section = Cast<UMovieSceneSection>(EntityOwner))
 				{
 					Section->BuildDefaultComponents(Linker, Params, &ImportedEntity);
diff --git a/Engine/Source/Runtime/MovieScene/Private/MovieScene.cpp b/Engine/Source/Runtime/MovieScene/Private/MovieScene.cpp
--- a/Engine/Source/Runtime/MovieScene/Private/MovieScene.cpp
+++ b/Engine/Source/Runtime/MovieScene/Private/MovieScene.cpp
@@ -13,6 +13,7 @@
 #include "CommonFrameRates.h"
 #include "EventHandlers/ISequenceDataEventHandler.h"
 #include "Misc/FrameRate.h"
+#include "Decorations/IMovieSceneDecoration.h"
 #include "UObject/ObjectSaveContext.h"
 #include "UObject/UObjectHash.h"
 #include "MovieSceneClock.h"
@@ -313,12 +314,17 @@ void UMovieScene::AddSpawnable(const FMovieSceneSpawnable& InNewSpawnable, const
 	InsertSorted(Spawnables, MoveTemp(NewSpawnable));
 
 	FMovieSceneBinding NewBinding = InNewBinding;
-	for (auto Track : NewBinding.GetTracks())
+	for (UMovieSceneTrack* Track : NewBinding.GetTracks())
 	{
 		Track->Rename(nullptr, this, REN_DontCreateRedirectors);
 	}
 	int32 NewBindingIndex = InsertSorted(ObjectBindings, MoveTemp(NewBinding));
 
+	for (UMovieSceneTrack* Track : ObjectBindings[NewBindingIndex].GetTracks())
+	{
+		Track->OnAddedToMovieScene(this);
+	}
+
 	EventHandlers.Trigger(&UE::MovieScene::ISequenceDataEventHandler::OnBindingAdded, ObjectBindings[NewBindingIndex]);
 }
 
@@ -400,6 +406,11 @@ void UMovieScene::AddPossessable(const FMovieScenePossessable& InNewPossessable,
 	int32 NewBindingIndex = InsertSorted(ObjectBindings, MoveTemp(NewBinding));
 	check(ObjectBindings.IsValidIndex(NewBindingIndex));
 
+	for (UMovieSceneTrack* Track : ObjectBindings[NewBindingIndex].GetTracks())
+	{
+		Track->OnAddedToMovieScene(this);
+	}
+
 	EventHandlers.Trigger(&UE::MovieScene::ISequenceDataEventHandler::OnBindingAdded, ObjectBindings[NewBindingIndex]);
 }
 
@@ -1165,6 +1176,8 @@ bool UMovieScene::RemoveTrack(UMovieSceneTrack& Track)
 	{
 		Modify();
 
+		Track.OnRemovedFromMovieScene();
+
 		const bool bRemoved = Tracks.RemoveSingle(&Track) != 0;
 		if (bRemoved)
 		{
@@ -1242,6 +1255,7 @@ UMovieSceneTrack* UMovieScene::AddTrack( TSubclassOf<UMovieSceneTrack> TrackClas
 	UMovieSceneTrack* CreatedType = NewObject<UMovieSceneTrack>(this, TrackClass, NAME_None, RF_Transactional);
 	Tracks.Add( CreatedType );
 
+	CreatedType->OnAddedToMovieScene(this);
 	EventHandlers.Trigger(&UE::MovieScene::ISequenceDataEventHandler::OnTrackAdded, CreatedType);
 
 	return CreatedType;
@@ -1263,6 +1277,7 @@ bool UMovieScene::AddGivenTrack(UMovieSceneTrack* InTrack)
 		Tracks.Add(InTrack);
 		InTrack->Rename(nullptr, this, REN_DontCreateRedirectors);
 
+		InTrack->OnAddedToMovieScene(this);
 		EventHandlers.Trigger(&UE::MovieScene::ISequenceDataEventHandler::OnTrackAdded, InTrack);
 
 		return true;
@@ -1299,6 +1314,7 @@ UMovieSceneTrack* UMovieScene::AddCameraCutTrack( TSubclassOf<UMovieSceneTrack>
 		Modify();
 		CameraCutTrack = NewObject<UMovieSceneTrack>(this, TrackClass, NAME_None, RF_Transactional);
 
+		CameraCutTrack->OnAddedToMovieScene(this);
 		EventHandlers.Trigger(&UE::MovieScene::ISequenceDataEventHandler::OnTrackAdded, CameraCutTrack);
 	}
 
@@ -1341,6 +1357,7 @@ void UMovieScene::SetCameraCutTrack(UMovieSceneTrack* InTrack)
 		EventHandlers.Trigger(&UE::MovieScene::ISequenceDataEventHandler::OnTrackRemoved, OldCameraCutTrack);
 	}
 
+	CameraCutTrack->OnAddedToMovieScene(this);
 	EventHandlers.Trigger(&UE::MovieScene::ISequenceDataEventHandler::OnTrackAdded, CameraCutTrack);
 }
 
@@ -1536,6 +1553,24 @@ void UMovieScene::PreSave(FObjectPreSaveContext ObjectSaveContext)
 #endif
 }
 
+void UMovieScene::OnDecorationAdded(UObject* NewDecoration)
+{
+	if (IMovieSceneDecoration* DecorationInterface = Cast<IMovieSceneDecoration>(NewDecoration))
+	{
+		DecorationInterface->OnDecorationAdded(this);
+	}
+
+	EventHandlers.Trigger(&UE::MovieScene::ISequenceDataEventHandler::OnDecorationAdded, NewDecoration);
+}
+
+void UMovieScene::OnDecorationRemoved(UObject* Decoration)
+{
+	if (IMovieSceneDecoration* DecorationInterface = Cast<IMovieSceneDecoration>(Decoration))
+	{
+		DecorationInterface->OnDecorationRemoved();
+	}
+	EventHandlers.Trigger(&UE::MovieScene::ISequenceDataEventHandler::OnDecorationRemoved, Decoration);
+}
 
 /* UMovieScene implementation
  *****************************************************************************/
@@ -1548,6 +1583,10 @@ void UMovieScene::RemoveBinding(const FGuid& Guid)
 	int32 Index = IndexOfBinding(Guid);
 	if (Index != INDEX_NONE)
 	{
+		for (UMovieSceneTrack* Track : ObjectBindings[Index].GetTracks())
+		{
+			Track->OnRemovedFromMovieScene();
+		}
 		ObjectBindings.RemoveAt(Index);
 	}
 }
diff --git a/Engine/Source/Runtime/MovieScene/Private/MovieSceneBinding.cpp b/Engine/Source/Runtime/MovieScene/Private/MovieSceneBinding.cpp
--- a/Engine/Source/Runtime/MovieScene/Private/MovieSceneBinding.cpp
+++ b/Engine/Source/Runtime/MovieScene/Private/MovieSceneBinding.cpp
@@ -29,6 +29,7 @@ bool FMovieSceneBinding::RemoveTrack(UMovieSceneTrack& Track, UMovieScene* Owner
 {
 	if (Tracks.RemoveSingle(&Track) != 0)
 	{
+		Track.OnRemovedFromMovieScene();
 		Owner->EventHandlers.Trigger(&UE::MovieScene::ISequenceDataEventHandler::OnTrackRemovedFromBinding, &Track, ObjectGuid);
 		return true;
 	}
@@ -53,6 +54,7 @@ TArray<UMovieSceneTrack*> FMovieSceneBinding::StealTracks(UMovieScene* Owner)
 	{
 		for (UMovieSceneTrack* Track : Tracks)
 		{
+			Track->OnRemovedFromMovieScene();
 			Owner->EventHandlers.Trigger(&UE::MovieScene::ISequenceDataEventHandler::OnTrackRemovedFromBinding, Track, ObjectGuid);
 		}
 	}
@@ -77,6 +79,7 @@ void FMovieSceneBinding::SetTracks(TArray<UMovieSceneTrack*>&& InTracks, UMovieS
 		UMovieSceneTrack* Track = Tracks[Index];
 		if (!NewTracks.Contains(Track))
 		{
+			Track->OnRemovedFromMovieScene();
 			Tracks.RemoveAt(Index, EAllowShrinking::No);
 			if (Owner)
 			{
@@ -94,6 +97,7 @@ void FMovieSceneBinding::SetTracks(TArray<UMovieSceneTrack*>&& InTracks, UMovieS
 		Tracks.Add(Track);
 		if (Owner)
 		{
+			Track->OnAddedToMovieScene(Owner);
 			Owner->EventHandlers.Trigger(&UE::MovieScene::ISequenceDataEventHandler::OnTrackAddedToBinding, Track, ObjectGuid);
 		}
 	}
diff --git a/Engine/Source/Runtime/MovieScene/Private/MovieSceneSection.cpp b/Engine/Source/Runtime/MovieScene/Private/MovieSceneSection.cpp
--- a/Engine/Source/Runtime/MovieScene/Private/MovieSceneSection.cpp
+++ b/Engine/Source/Runtime/MovieScene/Private/MovieSceneSection.cpp
@@ -6,6 +6,7 @@
 #include "MovieSceneSequence.h"
 #include "MovieSceneCommonHelpers.h"
 #include "MovieSceneTimeHelpers.h"
+#include "Decorations/IMovieSceneSectionDecoration.h"
 #include "Evaluation/MovieSceneEvalTemplate.h"
 #include "EntitySystem/MovieSceneEntityManager.h"
 #include "Generators/MovieSceneEasingCurves.h"
@@ -126,6 +127,25 @@ void UMovieSceneSection::Serialize(FArchive& Ar)
 	}
 }
 
+void UMovieSceneSection::OnDecorationAdded(UObject* NewDecoration)
+{
+	if (IMovieSceneSectionDecoration* DecorationInterface = Cast<IMovieSceneSectionDecoration>(NewDecoration))
+	{
+		DecorationInterface->OnDecorationAdded(this);
+	}
+
+	EventHandlers.Trigger(&UE::MovieScene::ISectionEventHandler::OnDecorationAdded, NewDecoration);
+}
+
+void UMovieSceneSection::OnDecorationRemoved(UObject* Decoration)
+{
+	if (IMovieSceneSectionDecoration* DecorationInterface = Cast<IMovieSceneSectionDecoration>(Decoration))
+	{
+		DecorationInterface->OnDecorationRemoved();
+	}
+	EventHandlers.Trigger(&UE::MovieScene::ISectionEventHandler::OnDecorationRemoved, Decoration);
+}
+
 #if WITH_EDITORONLY_DATA
 void UMovieSceneSection::DeclareConstructClasses(TArray<FTopLevelAssetPath>& OutConstructClasses, const UClass* SpecificSubclass)
 {
diff --git a/Engine/Source/Runtime/MovieScene/Private/MovieSceneTrack.cpp b/Engine/Source/Runtime/MovieScene/Private/MovieSceneTrack.cpp
--- a/Engine/Source/Runtime/MovieScene/Private/MovieSceneTrack.cpp
+++ b/Engine/Source/Runtime/MovieScene/Private/MovieSceneTrack.cpp
@@ -12,6 +12,8 @@
 #include "Compilation/MovieSceneEvaluationTreePopulationRules.h"
 #include "Compilation/IMovieSceneTemplateGenerator.h"
 
+#include "Decorations/IMovieSceneTrackDecoration.h"
+
 #include "Evaluation/MovieSceneEvaluationTrack.h"
 #include "Evaluation/MovieSceneEvaluationTemplate.h"
 
@@ -91,6 +93,32 @@ void UMovieSceneTrack::PostLoad()
 	}
 }
 
+void UMovieSceneTrack::OnAddedToMovieScene(UMovieScene* MovieScene)
+{
+	OnAddedToMovieSceneImpl(MovieScene);
+}
+
+void UMovieSceneTrack::OnRemovedFromMovieScene()
+{
+	OnRemovedFromMovieSceneImpl();
+}
+
+void UMovieSceneTrack::OnDecorationAdded(UObject* NewDecoration)
+{
+	if (IMovieSceneTrackDecoration* DecorationInterface = Cast<IMovieSceneTrackDecoration>(NewDecoration))
+	{
+		DecorationInterface->OnDecorationAdded(this);
+	}
+}
+
+void UMovieSceneTrack::OnDecorationRemoved(UObject* Decoration)
+{
+	if (IMovieSceneTrackDecoration* DecorationInterface = Cast<IMovieSceneTrackDecoration>(Decoration))
+	{
+		DecorationInterface->OnDecorationRemoved();
+	}
+}
+
 bool UMovieSceneTrack::IsPostLoadThreadSafe() const
 {
 	return true;
diff --git a/Engine/Source/Runtime/MovieScene/Public/Decorations/IMovieSceneDecoration.h b/Engine/Source/Runtime/MovieScene/Public/Decorations/IMovieSceneDecoration.h
new file mode 100644
--- /dev/null
+++ b/Engine/Source/Runtime/MovieScene/Public/Decorations/IMovieSceneDecoration.h
@@ -0,0 +1,64 @@
+// Copyright Epic Games, Inc. All Rights Reserved.
+
+#pragma once
+
+#include "CoreTypes.h"
+#include "UObject/Interface.h"
+#include "IMovieSceneDecoration.generated.h"
+
+class UMovieScene;
+
+UINTERFACE(MinimalAPI)
+class UMovieSceneDecoration : public UInterface
+{
+public:
+	GENERATED_BODY()
+};
+
+
+/** 
+ * Optional interface that can be added to any UObject that is used as a decoration on UMovieScene objects
+ *        to receive decoration and compilation events.
+ * 
+ * @note: This interface has no effect to objects used as UMovieSceneTrack or UMovieSceneSection decorations:
+ *        use IMovieSceneTrackDecoration or IMovieSceneSectionDecoration for these cases.
+
+ * @see UMovieScene::GetOrCreateDecoration, UMovieScene::AddDecoration
+ */
+class IMovieSceneDecoration
+{
+public:
+	GENERATED_BODY()
+
+	/**
+	 * Called when this decoration is first added directly to a UMovieScene
+	 */
+	virtual void OnDecorationAdded(UMovieScene* MovieScene)
+	{
+	}
+
+
+	/**
+	 * Called when this decoration is removed from a UMovieScene
+	 */
+	virtual void OnDecorationRemoved()
+	{
+	}
+
+
+	/**
+	 * Called before the movie scene this decoration exists on is compiled
+	 */
+	virtual void OnPreDecorationCompiled()
+	{
+	}
+
+
+	/**
+	 * Called after the movie scene this decoration exists on is compiled
+	 */
+	virtual void OnPostDecorationCompiled()
+	{
+	}
+};
+
diff --git a/Engine/Source/Runtime/MovieScene/Public/Decorations/IMovieSceneSectionDecoration.h b/Engine/Source/Runtime/MovieScene/Public/Decorations/IMovieSceneSectionDecoration.h
new file mode 100644
--- /dev/null
+++ b/Engine/Source/Runtime/MovieScene/Public/Decorations/IMovieSceneSectionDecoration.h
@@ -0,0 +1,46 @@
+// Copyright Epic Games, Inc. All Rights Reserved.
+
+#pragma once
+
+#include "CoreTypes.h"
+#include "UObject/Interface.h"
+#include "IMovieSceneSectionDecoration.generated.h"
+
+class UMovieSceneSection;
+
+UINTERFACE(MinimalAPI)
+class UMovieSceneSectionDecoration : public UInterface
+{
+public:
+	GENERATED_BODY()
+};
+
+
+/** 
+ * Optional interface that can be added to any UObject that is used as a decoration on UMovieSceneSection objects
+ *        to receive decoration events.
+
+ * @see UMovieSceneTrack::GetOrCreateDecoration, UMovieSceneTrack::AddDecoration
+ */
+class IMovieSceneSectionDecoration
+{
+public:
+	GENERATED_BODY()
+
+
+	/**
+	 * Called when this decoration is first added directly to a UMovieSceneSection
+	 */
+	virtual void OnDecorationAdded(UMovieSceneSection* Section)
+	{
+	}
+
+
+	/**
+	 * Called when this decoration is removed from a UMovieSceneSection
+	 */
+	virtual void OnDecorationRemoved()
+	{
+	}
+};
+
diff --git a/Engine/Source/Runtime/MovieScene/Public/Decorations/IMovieSceneTrackDecoration.h b/Engine/Source/Runtime/MovieScene/Public/Decorations/IMovieSceneTrackDecoration.h
new file mode 100644
--- /dev/null
+++ b/Engine/Source/Runtime/MovieScene/Public/Decorations/IMovieSceneTrackDecoration.h
@@ -0,0 +1,46 @@
+// Copyright Epic Games, Inc. All Rights Reserved.
+
+#pragma once
+
+#include "CoreTypes.h"
+#include "UObject/Interface.h"
+#include "IMovieSceneTrackDecoration.generated.h"
+
+class UMovieSceneTrack;
+
+UINTERFACE(MinimalAPI)
+class UMovieSceneTrackDecoration : public UInterface
+{
+public:
+	GENERATED_BODY()
+};
+
+
+/** 
+ * Optional interface that can be added to any UObject that is used as a decoration on UMovieSceneTrack objects
+ *        to receive decoration events.
+
+ * @see UMovieSceneTrack::GetOrCreateDecoration, UMovieSceneTrack::AddDecoration
+ */
+class IMovieSceneTrackDecoration
+{
+public:
+	GENERATED_BODY()
+
+
+	/**
+	 * Called when this decoration is first added directly to a UMovieSceneTrack
+	 */
+	virtual void OnDecorationAdded(UMovieSceneTrack* Track)
+	{
+	}
+
+
+	/**
+	 * Called when this decoration is removed from a UMovieSceneTrack
+	 */
+	virtual void OnDecorationRemoved()
+	{
+	}
+};
+
diff --git a/Engine/Source/Runtime/MovieScene/Public/Decorations/MovieSceneDecorationContainer.h b/Engine/Source/Runtime/MovieScene/Public/Decorations/MovieSceneDecorationContainer.h
new file mode 100644
--- /dev/null
+++ b/Engine/Source/Runtime/MovieScene/Public/Decorations/MovieSceneDecorationContainer.h
@@ -0,0 +1,182 @@
+// Copyright Epic Games, Inc. All Rights Reserved.
+
+#pragma once
+
+#include "CoreTypes.h"
+#include "UObject/Interface.h"
+#include "MovieSceneSignedObject.h"
+#include "Templates/FunctionFwd.h"
+#include "MovieSceneDecorationContainer.generated.h"
+
+template <typename T> class TSubclassOf;
+
+USTRUCT()
+struct FMovieSceneDecorationContainer
+{
+public:
+
+	GENERATED_BODY()
+
+	/**
+	 * Find meta-data of a particular type for this level sequence instance.
+	 * @param InClass - Class that you wish to find the metadata object for.
+	 * @return An instance of this class if it already exists as metadata on this Level Sequence, otherwise null.
+	 */
+	MOVIESCENE_API UObject* FindDecoration(const TSubclassOf<UObject>& InClass) const;
+
+	/**
+	 * Find meta-data of a particular type for this level sequence instance, adding it if it doesn't already exist.
+	 * @param InClass - Class that you wish to find or create the metadata object for.
+	 * @return An instance of this class as metadata on this Level Sequence.
+	 */
+	MOVIESCENE_API void AddDecoration(UObject* InDecoration, UObject* Outer, TFunctionRef<void(UObject*)> Event);
+
+	/**
+	 * Find meta-data of a particular type for this level sequence instance, adding it if it doesn't already exist.
+	 * @param InClass - Class that you wish to find or create the metadata object for.
+	 * @return An instance of this class as metadata on this Level Sequence.
+	 */
+	MOVIESCENE_API UObject* GetOrCreateDecoration(const TSubclassOf<UObject>& InClass, UObject* Outer, TFunctionRef<void(UObject*)> Event);
+
+	/**
+	* Remove meta-data of a particular type for this level sequence instance, if it exists
+	* @param InClass - The class type that you wish to remove the metadata for
+	*/
+	MOVIESCENE_API void RemoveDecoration(const TSubclassOf<UObject>& InClass, TFunctionRef<void(UObject*)> Event);
+
+	/**
+	* Retrieve all modular Decorations
+	*/
+	MOVIESCENE_API TArrayView<const TObjectPtr<UObject>> GetDecorations() const;
+
+	/**
+	 * Find meta-data of a particular type for this level sequence instance
+	 */
+	template<typename DecorationType>
+	DecorationType* FindDecoration() const
+	{
+		UObject* Found = nullptr;
+
+		if constexpr (TIsIInterface<DecorationType>::Value)
+		{
+			Found = FindDecoration(DecorationType::UClassType::StaticClass());
+		}
+		else
+		{
+			Found = FindDecoration(DecorationType::StaticClass());
+		}
+
+		return CastChecked<DecorationType>(Found, ECastCheckedType::NullAllowed);
+	}
+
+	/**
+	 * Find meta-data of a particular type for this level sequence instance, adding one if it was not found.
+	 * Meta-data may implement the ILevelSequenceMetaData interface in order to hook into default ULevelSequence functionality.
+	 */
+	template<typename DecorationType>
+	DecorationType* GetOrCreateDecoration(UObject* Outer, TFunctionRef<void(UObject*)> Event)
+	{
+		UObject* Found = GetOrCreateDecoration(DecorationType::StaticClass(), Event);
+		return CastChecked<DecorationType>(Found);
+	}
+
+	/**
+	 * Remove meta-data of a particular type for this level sequence instance, if it exists
+	 */
+	template<typename DecorationType>
+	void RemoveDecoration()
+	{
+		RemoveDecoration(DecorationType::StaticClass());
+	}
+
+protected:
+
+	UPROPERTY()
+	TArray<TObjectPtr<UObject>> Decorations;
+};
+
+
+UCLASS(DefaultToInstanced, MinimalAPI)
+class UMovieSceneDecorationContainerObject
+	: public UMovieSceneSignedObject
+{
+	GENERATED_BODY()
+
+public:
+
+	/**
+	 * Find meta-data of a particular type for this level sequence instance.
+	 * @param InClass - Class that you wish to find the metadata object for.
+	 * @return An instance of this class if it already exists as metadata on this Level Sequence, otherwise null.
+	 */
+	UObject* FindDecoration(const TSubclassOf<UObject>& InClass) const
+	{
+		return Decorations.FindDecoration(InClass);
+	}
+
+	MOVIESCENE_API void AddDecoration(UObject* InDecoration);
+
+	/**
+	 * Find meta-data of a particular type for this level sequence instance, adding it if it doesn't already exist.
+	 * @param InClass - Class that you wish to find or create the metadata object for.
+	 * @return An instance of this class as metadata on this Level Sequence.
+	 */
+	MOVIESCENE_API UObject* GetOrCreateDecoration(const TSubclassOf<UObject>& InClass);
+
+	/**
+	* Remove meta-data of a particular type for this level sequence instance, if it exists
+	* @param InClass - The class type that you wish to remove the metadata for
+	*/
+	MOVIESCENE_API void RemoveDecoration(const TSubclassOf<UObject>& InClass);
+
+
+	/**
+	* Retrieve all modular Decorations
+	*/
+	TArrayView<const TObjectPtr<UObject>> GetDecorations() const
+	{
+		return Decorations.GetDecorations();
+	}
+
+	/**
+	 * Find meta-data of a particular type for this level sequence instance
+	 */
+	template<typename DecorationType>
+	DecorationType* FindDecoration() const
+	{
+		return Decorations.FindDecoration<DecorationType>();
+	}
+
+	/**
+	 * Find meta-data of a particular type for this level sequence instance, adding one if it was not found.
+	 * Meta-data may implement the ILevelSequenceMetaData interface in order to hook into default ULevelSequence functionality.
+	 */
+	template<typename DecorationType>
+	DecorationType* GetOrCreateDecoration()
+	{
+		return CastChecked<DecorationType>(GetOrCreateDecoration(DecorationType::StaticClass()));
+	}
+
+	/**
+	 * Remove meta-data of a particular type for this level sequence instance, if it exists
+	 */
+	template<typename DecorationType>
+	void RemoveDecoration()
+	{
+		return RemoveDecoration(DecorationType::StaticClass());
+	}
+
+protected:
+
+	virtual void OnDecorationAdded(UObject* Decoration)
+	{}
+
+	virtual void OnDecorationRemoved(UObject* Decoration)
+	{}
+
+protected:
+
+	/** Array of decorations for this movie scene */
+	UPROPERTY()
+	FMovieSceneDecorationContainer Decorations;
+};
diff --git a/Engine/Source/Runtime/MovieScene/Public/EntitySystem/IMovieSceneEntityDecorator.h b/Engine/Source/Runtime/MovieScene/Public/EntitySystem/IMovieSceneEntityDecorator.h
new file mode 100644
--- /dev/null
+++ b/Engine/Source/Runtime/MovieScene/Public/EntitySystem/IMovieSceneEntityDecorator.h
@@ -0,0 +1,37 @@
+// Copyright Epic Games, Inc. All Rights Reserved.
+
+#pragma once
+
+#include "EntitySystem/IMovieSceneEntityProvider.h"
+#include "IMovieSceneEntityDecorator.generated.h"
+
+
+UINTERFACE(MinimalAPI)
+class UMovieSceneEntityDecorator : public UInterface
+{
+public:
+	GENERATED_BODY()
+};
+
+
+/** 
+ * Optional interface that can be added to any UMovieSceneSection decoration in order to decorate the ECS entities it produces at runtime
+ */
+class IMovieSceneEntityDecorator
+{
+public:
+
+	using FEntityImportParams        = UE::MovieScene::FEntityImportParams;
+	using FImportedEntity            = UE::MovieScene::FImportedEntity;
+
+	GENERATED_BODY()
+
+	void ExtendEntity(UMovieSceneEntitySystemLinker* EntityLinker, const FEntityImportParams& Params, FImportedEntity* OutImportedEntity)
+	{
+		ExtendEntityImpl(EntityLinker, Params, OutImportedEntity);
+	}
+
+private:
+
+	virtual void ExtendEntityImpl(UMovieSceneEntitySystemLinker* EntityLinker, const FEntityImportParams& Params, FImportedEntity* OutImportedEntity) = 0;
+};
diff --git a/Engine/Source/Runtime/MovieScene/Public/EventHandlers/ISectionEventHandler.h b/Engine/Source/Runtime/MovieScene/Public/EventHandlers/ISectionEventHandler.h
--- a/Engine/Source/Runtime/MovieScene/Public/EventHandlers/ISectionEventHandler.h
+++ b/Engine/Source/Runtime/MovieScene/Public/EventHandlers/ISectionEventHandler.h
@@ -18,6 +18,10 @@ class ISectionEventHandler
 public:
 
 	virtual void OnRowChanged(UMovieSceneSection*) {}
+
+	virtual void OnDecorationAdded(UObject* AddedDecoration) {}
+
+	virtual void OnDecorationRemoved(UObject* RemovedDecoration) {}
 };
 
 } // namespace MovieScene
diff --git a/Engine/Source/Runtime/MovieScene/Public/EventHandlers/ISequenceDataEventHandler.h b/Engine/Source/Runtime/MovieScene/Public/EventHandlers/ISequenceDataEventHandler.h
--- a/Engine/Source/Runtime/MovieScene/Public/EventHandlers/ISequenceDataEventHandler.h
+++ b/Engine/Source/Runtime/MovieScene/Public/EventHandlers/ISequenceDataEventHandler.h
@@ -44,6 +44,10 @@ public:
 
 	virtual void OnClockChanged(EUpdateClockSource NewClock, UMovieSceneClock* NewCustomClock) {}
 
+	virtual void OnDecorationAdded(UObject* AddedDecoration) {}
+
+	virtual void OnDecorationRemoved(UObject* RemovedDecoration) {}
+
 };
 
 } // namespace MovieScene
diff --git a/Engine/Source/Runtime/MovieScene/Public/MovieScene.h b/Engine/Source/Runtime/MovieScene/Public/MovieScene.h
--- a/Engine/Source/Runtime/MovieScene/Public/MovieScene.h
+++ b/Engine/Source/Runtime/MovieScene/Public/MovieScene.h
@@ -11,6 +11,7 @@
 #include "CoreTypes.h"
 #include "Delegates/Delegate.h"
 #include "EventHandlers/MovieSceneDataEventContainer.h"
+#include "Decorations/MovieSceneDecorationContainer.h"
 #include "HAL/PlatformCrt.h"
 #include "Internationalization/Text.h"
 #include "Math/Color.h"
@@ -340,7 +341,7 @@ public:
  */
 UCLASS(DefaultToInstanced, MinimalAPI)
 class UMovieScene
-	: public UMovieSceneSignedObject
+	: public UMovieSceneDecorationContainerObject
 {
 	GENERATED_UCLASS_BODY()
 
@@ -721,7 +722,7 @@ public:
 	* @return			Pointer to the binding, otherwise nullptr.
 	*/
 	MOVIESCENE_API const FMovieSceneBinding* FindBinding(const FGuid& ForGuid) const;
-	
+
 public:
 
 	// @todo sequencer: the following methods really shouldn't be here
@@ -1241,6 +1242,8 @@ protected:
 
 	/** Called before this object is being deserialized. */
 	MOVIESCENE_API virtual void PreSave(FObjectPreSaveContext ObjectSaveContext) override;
+	MOVIESCENE_API virtual void OnDecorationAdded(UObject* NewDecoration) override;
+	MOVIESCENE_API virtual void OnDecorationRemoved(UObject* NewDecoration) override;
 
 	/** Perform legacy upgrade of time ranges */
 	MOVIESCENE_API void UpgradeTimeRanges();
diff --git a/Engine/Source/Runtime/MovieScene/Public/MovieSceneSection.h b/Engine/Source/Runtime/MovieScene/Public/MovieSceneSection.h
--- a/Engine/Source/Runtime/MovieScene/Public/MovieSceneSection.h
+++ b/Engine/Source/Runtime/MovieScene/Public/MovieSceneSection.h
@@ -14,6 +14,7 @@
 #include "Evaluation/MovieSceneEvaluationCustomVersion.h"
 #include "EventHandlers/ISectionEventHandler.h"
 #include "EventHandlers/MovieSceneDataEventContainer.h"
+#include "Decorations/MovieSceneDecorationContainer.h"
 #include "HAL/PlatformCrt.h"
 #include "Math/Range.h"
 #include "Math/RangeBound.h"
@@ -233,7 +234,7 @@ enum class EMovieSceneTransformChannel : uint32
  */
 UCLASS(abstract, DefaultToInstanced, MinimalAPI, BlueprintType)
 class UMovieSceneSection
-	: public UMovieSceneSignedObject
+	: public UMovieSceneDecorationContainerObject
 {
 	GENERATED_UCLASS_BODY()
 
@@ -719,6 +720,9 @@ protected:
 	MOVIESCENE_API virtual void PostEditImport() override;
 	MOVIESCENE_API virtual void Serialize(FArchive& Ar) override;
 
+	MOVIESCENE_API virtual void OnDecorationAdded(UObject* NewDecoration) override;
+	MOVIESCENE_API virtual void OnDecorationRemoved(UObject* NewDecoration) override;
+
 	virtual void OnMoved(int32 DeltaTime) {}
 	virtual void OnDilated(float DilationFactor, FFrameNumber Origin) {}
 
diff --git a/Engine/Source/Runtime/MovieScene/Public/MovieSceneTrack.h b/Engine/Source/Runtime/MovieScene/Public/MovieSceneTrack.h
--- a/Engine/Source/Runtime/MovieScene/Public/MovieSceneTrack.h
+++ b/Engine/Source/Runtime/MovieScene/Public/MovieSceneTrack.h
@@ -9,6 +9,7 @@
 #include "CoreTypes.h"
 #include "Evaluation/Blending/MovieSceneBlendType.h"
 #include "Evaluation/MovieSceneEvaluationField.h"
+#include "Decorations/MovieSceneDecorationContainer.h"
 #include "HAL/Platform.h"
 #include "Internationalization/Text.h"
 #include "Math/Color.h"
@@ -185,7 +186,7 @@ struct MOVIESCENE_API FMovieSceneTrackRowMetadata
  */
 UCLASS(abstract, DefaultToInstanced, MinimalAPI, BlueprintType)
 class UMovieSceneTrack
-	: public UMovieSceneSignedObject
+	: public UMovieSceneDecorationContainerObject
 {
 	GENERATED_BODY()
 
@@ -235,6 +236,10 @@ public:
 	 */
 	MOVIESCENE_API virtual void UpdateEasing();
 
+
+	void OnAddedToMovieScene(UMovieScene* MovieScene);
+	void OnRemovedFromMovieScene();
+
 protected:
 
 	//~ UObject interface
@@ -242,6 +247,9 @@ protected:
 	MOVIESCENE_API virtual void PostLoad() override;
 	MOVIESCENE_API virtual bool IsPostLoadThreadSafe() const override;
 
+	MOVIESCENE_API virtual void OnDecorationAdded(UObject* NewDecoration) override;
+	MOVIESCENE_API virtual void OnDecorationRemoved(UObject* Decoration) override;
+
 	/** Intentionally not a UPROPERTY so this isn't serialized */
 	FMovieSceneBlendTypeField SupportedBlendTypes;
 
@@ -307,6 +315,14 @@ private:
 	{
 	}
 
+	virtual void OnAddedToMovieSceneImpl(UMovieScene* InMovieScene)
+	{
+	}
+
+	virtual void OnRemovedFromMovieSceneImpl()
+	{
+	}
+
 private:
 
 	void UpdateEvaluationTree();
