PR #12760: Fixed animation asset "RateScale" not being respected

diff --git a/Engine/Source/Runtime/MovieScene/Private/Sections/MovieSceneSectionTimingParameters.cpp b/Engine/Source/Runtime/MovieScene/Private/Sections/MovieSceneSectionTimingParameters.cpp
--- a/Engine/Source/Runtime/MovieScene/Private/Sections/MovieSceneSectionTimingParameters.cpp
+++ b/Engine/Source/Runtime/MovieScene/Private/Sections/MovieSceneSectionTimingParameters.cpp
@@ -27,6 +27,7 @@ FMovieSceneSequenceTransform FMovieSceneSectionTimingParametersSeconds::MakeTran
 	// ----------------------------------------------------------------------------
 	// Time warp
 	AddPlayRate(Result);
+	Result.Add(0, FMovieSceneTimeWarpVariant(InnerPlayRate));
 
 	// ----------------------------------------------------------------------------
 	// FrameRate conversion to seconds
@@ -40,7 +41,7 @@ FMovieSceneSequenceTransform FMovieSceneSectionTimingParametersSeconds::MakeTran
 	double StartOffset = InnerStartOffset + FirstLoopStartOffset;
 
 	// Accommodate negative play rates by playing from the end of the clip
-	if (PlayRate.GetType() == EMovieSceneTimeWarpType::FixedPlayRate && PlayRate.AsFixedPlayRate() < 0.0)
+	if (PlayRate.GetType() == EMovieSceneTimeWarpType::FixedPlayRate && (PlayRate.AsFixedPlayRate() * InnerPlayRate) < 0.0)
 	{
 		StartOffset += Duration;
 	}
