[AutoRTFM] Fix a validation error in static initialization.

This is just a running counter used for generating hash values, and doesn't need to be rolled back if a transaction is aborted.

diff --git a/Engine/Source/Runtime/MovieScene/Private/Evaluation/MovieSceneAnimTypeID.cpp b/Engine/Source/Runtime/MovieScene/Private/Evaluation/MovieSceneAnimTypeID.cpp
--- a/Engine/Source/Runtime/MovieScene/Private/Evaluation/MovieSceneAnimTypeID.cpp
+++ b/Engine/Source/Runtime/MovieScene/Private/Evaluation/MovieSceneAnimTypeID.cpp
@@ -27,7 +27,7 @@ namespace
 	static uint64 RunningCount = 0;
 	struct FUniqueTypeID : FMovieSceneAnimTypeID
 	{
-		FUniqueTypeID()
+		UE_AUTORTFM_ALWAYS_OPEN FUniqueTypeID()
 		{
 			ID = GenerateHash(&RunningCount, ++RunningCount);
 		}
