Sequencer: Ensure that root entities have been instantiated before entity grouping can run

diff --git a/Engine/Source/Runtime/MovieScene/Private/EntitySystem/MovieSceneEntityGroupingSystem.cpp b/Engine/Source/Runtime/MovieScene/Private/EntitySystem/MovieSceneEntityGroupingSystem.cpp
--- a/Engine/Source/Runtime/MovieScene/Private/EntitySystem/MovieSceneEntityGroupingSystem.cpp
+++ b/Engine/Source/Runtime/MovieScene/Private/EntitySystem/MovieSceneEntityGroupingSystem.cpp
@@ -6,6 +6,7 @@
 #include "EntitySystem/BuiltInComponentTypes.h"
 #include "EntitySystem/MovieSceneBoundObjectInstantiator.h"
 #include "EntitySystem/MovieSceneBoundSceneComponentInstantiator.h"
+#include "EntitySystem/MovieSceneRootInstantiatorSystem.h"
 #include "EntitySystem/MovieSceneEntityIDs.h"
 #include "EntitySystem/MovieSceneEntityMutations.h"
 #include "EntitySystem/MovieSceneEntitySystemLinker.h"
@@ -234,6 +235,7 @@ UMovieSceneEntityGroupingSystem::UMovieSceneEntityGroupingSystem(const FObjectIn
 		// to group things using bound objects and/or bound scene components, so let's run after the
 		// systems that set those up.
 		DefineComponentConsumer(GetClass(), BuiltInComponents->BoundObject);
+		DefineImplicitPrerequisite(UMovieSceneRootInstantiatorSystem::StaticClass(), GetClass());
 	}
 }
 
