Fixed uninitialized member variable

diff --git a/Engine/Source/Runtime/MovieScene/Public/EntitySystem/IMovieSceneTaskScheduler.h b/Engine/Source/Runtime/MovieScene/Public/EntitySystem/IMovieSceneTaskScheduler.h
--- a/Engine/Source/Runtime/MovieScene/Public/EntitySystem/IMovieSceneTaskScheduler.h
+++ b/Engine/Source/Runtime/MovieScene/Public/EntitySystem/IMovieSceneTaskScheduler.h
@@ -53,6 +53,7 @@ struct FTaskParams
 		bSerialTasks = false;
 		bForcePropagateDownstream = false;
 		bForceConsumeUpstream = false;
+		bForcePrePostTask = false;
 	}
 
 	explicit FTaskParams(const TCHAR* InDebugName, const TStatId& InStatId = TStatId())
@@ -66,6 +67,7 @@ struct FTaskParams
 		bSerialTasks = false;
 		bForcePropagateDownstream = false;
 		bForceConsumeUpstream = false;
+		bForcePrePostTask = false;
 	}
 
 	/**
