Sequencer: Initial prototype of stretching and scaling sequences based on localized audio durations

UMovieSceneScalingAnchors is a new decoration type that can be added to UMovieScene that acts as a timewarp curve. IMovieSceneScalingDrivers can be added to it to define anchors that are able to stretch time based on their position and duration.
The scaling anchors class generates a play rate curve based on an algorithm that analysis the initial (reference), and current positions and durations of each anchor, accounting for overlapping regions.

The audio track has been updated to implement the IMovieSceneScalingDriver interface, allowing it to provide scaling information to the main decoration on the moviescene based on the length of localized audio. Any track or section can follow a similar pattern to implement their own scaling logic.

As part of this work, a new ability to scale the track area on a per-section basis using non-linear view-spaces has been added. This allows us to accurately paint section boundaries and channel keys at their scaled positions for previewability.

diff --git a/Engine/Plugins/Experimental/PropertyAnimator/Source/PropertyAnimatorEditor/PropertyAnimatorEditor.Build.cs b/Engine/Plugins/Experimental/PropertyAnimator/Source/PropertyAnimatorEditor/PropertyAnimatorEditor.Build.cs
--- a/Engine/Plugins/Experimental/PropertyAnimator/Source/PropertyAnimatorEditor/PropertyAnimatorEditor.Build.cs
+++ b/Engine/Plugins/Experimental/PropertyAnimator/Source/PropertyAnimatorEditor/PropertyAnimatorEditor.Build.cs
@@ -28,6 +28,7 @@ public class PropertyAnimatorEditor : ModuleRules
                 "PropertyAnimatorCore",
                 "PropertyAnimatorCoreEditor",
                 "Sequencer",
+                "SequencerCore",
                 "Slate",
                 "SlateCore",
                 "ToolMenus",
diff --git a/Engine/Source/Editor/MovieSceneTools/MovieSceneTools.Build.cs b/Engine/Source/Editor/MovieSceneTools/MovieSceneTools.Build.cs
--- a/Engine/Source/Editor/MovieSceneTools/MovieSceneTools.Build.cs
+++ b/Engine/Source/Editor/MovieSceneTools/MovieSceneTools.Build.cs
@@ -61,6 +61,7 @@ public class MovieSceneTools : ModuleRules
 				"RenderCore",
 				"RHI",
 				"SequenceRecorder",
+				"SequencerCore",
 				"TimeManagement",
 				"ToolWidgets",
                 "AnimationCore",
diff --git a/Engine/Source/Editor/MovieSceneTools/Private/MVVM/MovieSceneToolsViewModelTypeIDs.cpp b/Engine/Source/Editor/MovieSceneTools/Private/MVVM/MovieSceneToolsViewModelTypeIDs.cpp
--- a/Engine/Source/Editor/MovieSceneTools/Private/MVVM/MovieSceneToolsViewModelTypeIDs.cpp
+++ b/Engine/Source/Editor/MovieSceneTools/Private/MVVM/MovieSceneToolsViewModelTypeIDs.cpp
@@ -6,6 +6,7 @@
 #include "MVVM/ViewModels/CinematicShotTrackModel.h"
 #include "MVVM/ViewModels/BindingLifetimeTrackModel.h"
 #include "MVVM/ViewModels/TimeWarpChannelModel.h"
+#include "MVVM/ViewModels/ScalingAnchorsModel.h"
 
 
 namespace UE
@@ -18,6 +19,7 @@ UE_SEQUENCER_DEFINE_CASTABLE(FCameraCutTrackModel);
 UE_SEQUENCER_DEFINE_CASTABLE(FCinematicShotTrackModel);
 UE_SEQUENCER_DEFINE_CASTABLE(FBindingLifetimeTrackModel);
 UE_SEQUENCER_DEFINE_CASTABLE(FTimeWarpChannelModel);
+UE_SEQUENCER_DEFINE_CASTABLE(FScalingAnchorsModel);
 
 } // namespace Sequencer
 } // namespace UE
diff --git a/Engine/Source/Editor/MovieSceneTools/Private/MVVM/ViewModels/ScalingAnchorsModel.cpp b/Engine/Source/Editor/MovieSceneTools/Private/MVVM/ViewModels/ScalingAnchorsModel.cpp
new file mode 100644
--- /dev/null
+++ b/Engine/Source/Editor/MovieSceneTools/Private/MVVM/ViewModels/ScalingAnchorsModel.cpp
@@ -0,0 +1,208 @@
+// Copyright Epic Games, Inc. All Rights Reserved.
+
+#include "MVVM/ViewModels/ScalingAnchorsModel.h"
+#include "EntitySystem/BuiltInComponentTypes.h"
+#include "EntitySystem/MovieSceneEntityIDs.h"
+#include "EntitySystem/MovieSceneEntitySystemLinker.h"
+#include "Evaluation/MovieSceneEvaluationTemplateInstance.h"
+#include "ISequencer.h"
+#include "MovieScene.h"
+#include "MVVM/Extensions/IViewSpaceClientExtension.h"
+#include "MVVM/SectionModelStorageExtension.h"
+#include "MVVM/Selection/Selection.h"
+#include "MVVM/ViewModels/EditorViewModel.h"
+#include "MVVM/ViewModels/SectionModel.h"
+#include "MVVM/ViewModels/SequenceModel.h"
+#include "MVVM/ViewModels/SequencerEditorViewModel.h"
+#include "MVVM/ViewModels/TrackAreaViewModel.h"
+#include "MVVM/ViewModels/TrackAreaViewSpace.h"
+#include "SequencerToolMenuContext.h"
+#include "ToolMenus.h"
+#include "Decorations/MovieSceneScalingAnchors.h"
+
+
+namespace UE::Sequencer
+{
+
+void FScalingAnchorsModel::InitializeObject(TWeakObjectPtr<> InWeakObject)
+{
+	WeakAnchors = CastChecked<UMovieSceneScalingAnchors>(InWeakObject.Get());
+
+	FToolMenuOwnerScoped OwnerScoped(this);
+
+	UToolMenu* ToolMenu = UToolMenus::Get()->ExtendMenu("Sequencer.SectionContextMenu");
+	ToolMenu->AddDynamicSection(
+		NAME_None,
+		FNewToolMenuDelegate::CreateSP(this, &FScalingAnchorsModel::ExtendSectionMenu),
+		FToolMenuInsert(NAME_None, EToolMenuInsertType::After)
+	);
+}
+
+UObject* FScalingAnchorsModel::GetObject() const
+{
+	return WeakAnchors.Get();
+}
+
+void FScalingAnchorsModel::ExtendSectionMenu(UToolMenu* InMenu)
+{
+	USequencerToolMenuContext* ContextObject = InMenu->FindContext<USequencerToolMenuContext>();
+	TWeakPtr<ISequencer>       WeakSequencer = ContextObject ? ContextObject->WeakSequencer : nullptr;
+	TSharedPtr<ISequencer>     Sequencer     = WeakSequencer.Pin();
+
+	if (!Sequencer || Sequencer->GetViewModel()->GetRootModel() != FViewModelPtr(FindAncestorOfType<FSequenceModel>()))
+	{
+		return;
+	}
+
+	FToolMenuSection& ScalingSection = InMenu->AddSection(TEXT("Scaling"));
+
+	ScalingSection.AddMenuEntry(
+		NAME_None,
+		NSLOCTEXT("FScalingAnchorsModel", "CreateScalingGroup", "Create Scaling Group"),
+		FText(),
+		FSlateIcon(),
+		FUIAction(
+			FExecuteAction::CreateSP(this, &FScalingAnchorsModel::CreateScalingGroup, WeakSequencer) 
+		)
+	);
+}
+
+void FScalingAnchorsModel::CreateScalingGroup(TWeakPtr<ISequencer> InWeakSequencer)
+{
+	UMovieSceneScalingAnchors* Anchors   = WeakAnchors.Get();
+	TSharedPtr<ISequencer>     Sequencer = InWeakSequencer.Pin();
+
+	if (Sequencer && Anchors)
+	{
+		FGuid NewGuid = FGuid::NewGuid();
+
+		FMovieSceneAnchorsScalingGroup& Group = Anchors->GetOrCreateScalingGroup(NewGuid);
+		for (TSharedPtr<FSectionModel> Section : Sequencer->GetViewModel()->GetSelection()->TrackArea.Filter<FSectionModel>())
+		{
+			Group.Sections.Add(Section->GetSection());
+		}
+	}
+}
+
+void FScalingAnchorsModel::UpdateViewSpaces(FTrackAreaViewModel& TrackAreaViewModel)
+{
+	using namespace UE::MovieScene;
+
+	FSectionModelStorageExtension* SectionModelStorage = TrackAreaViewModel.GetEditor()->GetRootModel()->CastDynamic<FSectionModelStorageExtension>();
+
+	UMovieSceneScalingAnchors* Anchors = WeakAnchors.Get();
+	if (Anchors)
+	{
+		class FAnchorsSpace : public FTrackAreaViewSpace
+		{
+		public:
+			FAnchorsSpace(TWeakObjectPtr<UMovieSceneScalingAnchors> InWeakAnchors, TWeakPtr<ISequencer> InWeakSequencer)
+				: WeakAnchors(InWeakAnchors)
+				, WeakSequencer(InWeakSequencer)
+			{
+			}
+
+#if 0
+			UMovieSceneTimeWarpGetter* FindTimeWarp() const
+			{
+				UMovieSceneEntitySystemLinker* Linker = WeakLinker.Get();
+				if (!Linker || AnchorEntity.IsEmpty() || !AnchorEntity.IsValid(Linker->EntityManager))
+				{
+					TSharedPtr<ISequencer> Sequencer = WeakSequencer.Pin();
+					if (!Sequencer)
+					{
+						WeakLinker = nullptr;
+						AnchorEntity = FEntityHandle();
+						return nullptr;
+					}
+
+					FMovieSceneRootEvaluationTemplateInstance& Root = Sequencer->GetEvaluationTemplate();
+					Linker = Root.GetEntitySystemLinker();
+
+					const FSequenceInstance* SequenceInstance = Root.FindInstance(Sequencer->GetFocusedTemplateID());
+					if (SequenceInstance)
+					{
+						FMovieSceneEntityID Entity = SequenceInstance->FindEntity(WeakAnchors.Get(), UMovieSceneScalingAnchors::DefaultTimeWarpEntityID);
+						if (Entity)
+						{
+							AnchorEntity = Linker->EntityManager.GetEntityHandle(Entity);
+							WeakLinker = Linker;
+						}
+					}
+				}
+
+				if (AnchorEntity.IsValid(Linker->EntityManager))
+				{
+					TOptionalComponentReader<TObjectPtr<UMovieSceneTimeWarpGetter>> TimeWarp = Linker->EntityManager.ReadComponent(AnchorEntity.GetEntityID(), FBuiltInComponentTypes::Get()->TimeWarp);
+					return TimeWarp ? *TimeWarp : TObjectPtr<UMovieSceneTimeWarpGetter>();
+				}
+				return nullptr;
+			}
+#endif
+
+			virtual double SourceToView(double SourceTime) const
+			{
+				using namespace UE::MovieScene;
+
+				UMovieSceneScalingAnchors* Anchors = WeakAnchors.Get();
+				UMovieSceneTimeWarpGetter* DefaultTimeWarp = Anchors;
+				if (!DefaultTimeWarp)
+				{
+					return SourceTime;
+				}
+
+				FFrameRate TickResolution = Anchors->GetTypedOuter<UMovieScene>()->GetTickResolution();
+
+				TOptional<FFrameTime> Time = DefaultTimeWarp->InverseRemapTimeCycled(SourceTime * TickResolution, 0, FInverseTransformTimeParams());
+				return Time ? Time.GetValue() / TickResolution : SourceTime;
+			}
+
+			virtual double ViewToSource(double ViewTime) const
+			{
+				UMovieSceneScalingAnchors* Anchors = WeakAnchors.Get();
+				UMovieSceneTimeWarpGetter* DefaultTimeWarp = Anchors;
+				if (!DefaultTimeWarp)
+				{
+					return ViewTime;
+				}
+				FFrameRate TickResolution = Anchors->GetTypedOuter<UMovieScene>()->GetTickResolution();
+				return DefaultTimeWarp->RemapTime(ViewTime * TickResolution) / TickResolution;
+			}
+
+			mutable TWeakObjectPtr<UMovieSceneEntitySystemLinker> WeakLinker;
+			TWeakObjectPtr<UMovieSceneScalingAnchors> WeakAnchors;
+			TWeakPtr<ISequencer> WeakSequencer;
+			mutable FEntityHandle AnchorEntity;
+		};
+
+
+		TSharedPtr<ISequencer> Sequencer = TrackAreaViewModel.GetEditor()->CastThis<FSequencerEditorViewModel>()->GetSequencer();
+		TSharedPtr<FTrackAreaViewSpace> DefaultViewSpace = MakeShared<FAnchorsSpace>(WeakAnchors, Sequencer);
+		TrackAreaViewModel.AddViewSpace(FGuid(), DefaultViewSpace);
+
+#if 0
+		if (SectionModelStorage)
+		{
+			for (const TPair<FGuid, FMovieSceneAnchorsScalingGroup>& Pair : Anchors->GetScalingGroups())
+			{
+				// @todo: make a way for the view space to be well defined
+				TSharedPtr<FTrackAreaViewSpace> NewViewSpace = MakeShared<FAnchorsSpace>(WeakAnchors, Sequencer);
+				TrackAreaViewModel.AddViewSpace(Pair.Key, NewViewSpace);
+
+				for (TObjectPtr<UMovieSceneSection> Section : Pair.Value.Sections)
+				{
+					TSharedPtr<FSectionModel> SectionModel = SectionModelStorage->FindModelForSection(Section);
+					if (SectionModel)
+					{
+						FViewSpaceClientExtensionShim& ViewSpaceExt = SectionModel->AddDynamicExtension<FViewSpaceClientExtensionShim>(Pair.Key);
+						ViewSpaceExt.SetViewSpaceID(Pair.Key);
+					}
+				}
+			}
+		}
+#endif
+	}
+}
+
+} // namespace UE::Sequencer
+
diff --git a/Engine/Source/Editor/MovieSceneTools/Private/MVVM/ViewModels/ScalingAnchorsModel.h b/Engine/Source/Editor/MovieSceneTools/Private/MVVM/ViewModels/ScalingAnchorsModel.h
new file mode 100644
--- /dev/null
+++ b/Engine/Source/Editor/MovieSceneTools/Private/MVVM/ViewModels/ScalingAnchorsModel.h
@@ -0,0 +1,50 @@
+// Copyright Epic Games, Inc. All Rights Reserved.
+
+#pragma once
+
+#include "MVVM/Extensions/IObjectModelExtension.h"
+#include "MVVM/Extensions/ITrackAreaViewSpaceProviderExtension.h"
+#include "MVVM/ViewModels/ViewModel.h"
+#include "UObject/WeakObjectPtr.h"
+#include "UObject/WeakObjectPtrTemplates.h"
+
+class ISequencer;
+class UMovieSceneScalingAnchors;
+class UToolMenu;
+
+namespace UE::Sequencer
+{
+
+class FTrackAreaViewModel;
+
+class MOVIESCENETOOLS_API FScalingAnchorsModel
+	: public FViewModel
+	, public IObjectModelExtension
+	, public ITrackAreaViewSpaceProviderExtension
+{
+public:
+
+	UE_SEQUENCER_DECLARE_CASTABLE(FScalingAnchorsModel, FViewModel, IObjectModelExtension, ITrackAreaViewSpaceProviderExtension);
+
+	/*~ Begin IObjectModelExtension */
+	virtual void InitializeObject(TWeakObjectPtr<> InWeakObject) override;
+	virtual UObject* GetObject() const override;
+	/*~ End IObjectModelExtension */
+
+	/*~ Begin ITrackAreaViewSpaceProviderExtension */
+	virtual void UpdateViewSpaces(FTrackAreaViewModel& TrackAreaViewModel) override;
+	/*~ End ITrackAreaViewSpaceProviderExtension */
+
+private:
+
+	void ExtendSectionMenu(UToolMenu* InMenu);
+	void CreateScalingGroup(TWeakPtr<ISequencer> InWeakSequencer);
+
+private:
+
+	TWeakObjectPtr<UMovieSceneScalingAnchors> WeakAnchors;
+};
+
+
+} // namespace UE::Sequencer
+
diff --git a/Engine/Source/Editor/MovieSceneTools/Private/MovieSceneToolsModule.cpp b/Engine/Source/Editor/MovieSceneTools/Private/MovieSceneToolsModule.cpp
--- a/Engine/Source/Editor/MovieSceneTools/Private/MovieSceneToolsModule.cpp
+++ b/Engine/Source/Editor/MovieSceneTools/Private/MovieSceneToolsModule.cpp
@@ -5,6 +5,7 @@
 #include "Editor.h"
 #include "Modules/ModuleManager.h"
 #include "Curves/RichCurve.h"
+#include "ISequencerCoreModule.h"
 #include "ISequencerModule.h"
 #include "ICurveEditorModule.h"
 #include "MovieSceneToolsProjectSettingsCustomization.h"
@@ -63,6 +64,9 @@
 #include "MVVM/ViewModels/CameraCutTrackModel.h"
 #include "MVVM/ViewModels/CinematicShotTrackModel.h"
 #include "MVVM/ViewModels/BindingLifetimeTrackModel.h"
+#include "MVVM/ViewModels/ScalingAnchorsModel.h"
+
+#include "Decorations/MovieSceneScalingAnchors.h"
 
 #include "MovieSceneBuiltInEasingFunctionCustomization.h"
 #include "MovieSceneAlphaBlendOptionCustomization.h"
@@ -262,6 +266,9 @@ void FMovieSceneToolsModule::StartupModule()
 				return SNew(SCurveEditorKeyBarView, WeakCurveEditor);
 			}
 		));
+
+		ISequencerCoreModule& SequencerCoreModule = FModuleManager::Get().LoadModuleChecked<ISequencerCoreModule>("SequencerCore");
+		SequencerCoreModule.RegisterModelType(UMovieSceneScalingAnchors::StaticClass(), FScalingAnchorsModel::ID);
 	}
 
 	FixupDynamicBindingPayloadParameterNameHandle = UMovieScene::FixupDynamicBindingPayloadParameterNameEvent.AddStatic(FixupPayloadParameterNameForDynamicBinding);
diff --git a/Engine/Source/Editor/MovieSceneTools/Private/TrackEditors/AudioTrackEditor.cpp b/Engine/Source/Editor/MovieSceneTools/Private/TrackEditors/AudioTrackEditor.cpp
--- a/Engine/Source/Editor/MovieSceneTools/Private/TrackEditors/AudioTrackEditor.cpp
+++ b/Engine/Source/Editor/MovieSceneTools/Private/TrackEditors/AudioTrackEditor.cpp
@@ -1,6 +1,11 @@
 // Copyright Epic Games, Inc. All Rights Reserved.
 
 #include "TrackEditors/AudioTrackEditor.h"
+#include "Decorations/MovieSceneLanguagePreviewDecoration.h"
+#include "Internationalization/Internationalization.h"
+#include "Internationalization/TextLocalizationManager.h"
+#include "Internationalization/PackageLocalizationManager.h"
+#include "Internationalization/Culture.h"
 #include "Textures/SlateTextureData.h"
 #include "Rendering/RenderingCommon.h"
 #include "Rendering/DrawElements.h"
@@ -16,8 +21,10 @@
 #include "Sound/DialogueWave.h"
 #include "Framework/Application/SlateApplication.h"
 #include "Widgets/Layout/SBox.h"
+#include "Widgets/Input/SComboButton.h"
 #include "SequencerSectionPainter.h"
 #include "MVVM/ViewModels/SequencerEditorViewModel.h"
+#include "MVVM/ViewModels/TrackModel.h"
 #include "Styling/AppStyle.h"
 #include "Editor/UnrealEdEngine.h"
 #include "Sound/SoundCue.h"
@@ -32,6 +39,8 @@
 #include "IContentBrowserSingleton.h"
 #include "ContentBrowserModule.h"
 #include "MVVM/Views/ViewUtilities.h"
+#include "MVVM/Extensions/ITrackExtension.h"
+#include "MVVM/ViewModels/OutlinerColumns/OutlinerColumnTypes.h"
 #include "AssetRegistry/AssetRegistryModule.h"
 #include "ISectionLayoutBuilder.h"
 #include "ISequencer.h"
@@ -52,13 +61,142 @@ namespace AnimatableAudioEditorConstants
 	const uint32 MaxSamplesPerPixel = 60;
 }
 
+namespace UE::Sequencer::AudioTrack
+{
+	class SLanguagePreviewComboWidget : public SCompoundWidget
+	{
+	public:
+
+		SLATE_BEGIN_ARGS(SLanguagePreviewComboWidget){}
+		SLATE_END_ARGS()
+
+		void Construct(const FArguments& InArgs, TWeakObjectPtr<UMovieScene> InWeakMovieScene, TWeakObjectPtr<UMovieSceneTrack> InWeakAudioTrack)
+		{
+			WeakMovieScene = InWeakMovieScene;
+			WeakAudioTrack = InWeakAudioTrack;
+
+			ChildSlot
+			[
+				SNew(SComboButton)
+				.Visibility(this, &SLanguagePreviewComboWidget::GetLanguagePreviewVisibility)
+				.OnGetMenuContent(this, &SLanguagePreviewComboWidget::GetMenuContent)
+				.ButtonContent()
+				[
+					SNew(STextBlock)
+					.Text(this, &SLanguagePreviewComboWidget::GetCurrentPreviewLanguage)
+				]
+			];
+		}
+
+	private:
+
+		TSharedRef<SWidget> GetMenuContent()
+		{
+			FMenuBuilder MenuBuilder(true, nullptr);
+
+			// Get a list of available localization cultures
+
+			constexpr bool bIncludeDerivedCultures = false;
+
+			const TArray<FString> LocalizedCultureNames = FTextLocalizationManager::Get().GetLocalizedCultureNames(ELocalizationLoadFlags::Game);
+			TArray<FCultureRef> Cultures = FInternationalization::Get().GetAvailableCultures(LocalizedCultureNames, bIncludeDerivedCultures);
+
+			for (FCultureRef Culture : Cultures)
+			{
+				MenuBuilder.AddMenuEntry(
+					FText::FromString(Culture->GetName()),
+					FText(),
+					FSlateIcon(),
+					FUIAction(
+						FExecuteAction::CreateSP(this, &SLanguagePreviewComboWidget::SetPreviewLanguage, Culture->GetName())
+					)
+				);
+			}
+
+			return MenuBuilder.MakeWidget();
+		}
+
+		EVisibility GetLanguagePreviewVisibility() const
+		{
+			UMovieSceneTrack* AudioTrack = WeakAudioTrack.Get();
+			if (AudioTrack)
+			{
+				const TArray<FString> LocalizedCultureNames = FTextLocalizationManager::Get().GetLocalizedCultureNames(ELocalizationLoadFlags::Game);
+				TArray<FCultureRef> Cultures = FInternationalization::Get().GetAvailableCultures(LocalizedCultureNames, false);
+
+				for (FCultureRef Culture : Cultures)
+
+				for (UMovieSceneSection* Section : AudioTrack->GetAllSections())
+				{
+					UMovieSceneAudioSection* AudioSection = CastChecked<UMovieSceneAudioSection>(Section);
+					if (AudioSection)
+					{
+						if (USoundBase* Sound = AudioSection->GetSound())
+						{
+							FName LocalizedName = FPackageLocalizationManager::Get().FindLocalizedPackageNameForCulture(*Sound->GetOutermost()->GetPathName(), Culture->GetName());
+							if (LocalizedName != NAME_None)
+							{
+								return EVisibility::Visible;
+							}
+						}
+					}
+				}
+			}
+
+			return EVisibility::Collapsed;
+		}
+
+		FText GetCurrentPreviewLanguage() const
+		{
+			UMovieScene* MovieScene = WeakMovieScene.Get();
+			if (MovieScene)
+			{
+				UMovieSceneLanguagePreviewDecoration* PreviewLanguage = MovieScene->FindDecoration<UMovieSceneLanguagePreviewDecoration>();
+				if (PreviewLanguage && PreviewLanguage->PreviewLanguage.Len() > 0)
+				{
+					return FText::FromString(PreviewLanguage->PreviewLanguage);
+				}
+			}
+			return LOCTEXT("PreviewLanguage", "Preview Language");
+		}
+
+		void SetPreviewLanguage(FString InCultureName)
+		{
+			UMovieScene* MovieScene = WeakMovieScene.Get();
+			if (MovieScene)
+			{
+				UMovieSceneLanguagePreviewDecoration* Decoration = MovieScene->GetOrCreateDecoration<UMovieSceneLanguagePreviewDecoration>();
+				Decoration->SetFlags(RF_Transient);
+				Decoration->PreviewLanguage = InCultureName;
+
+				// Mark as changed to force an evaluation
+				MovieScene->MarkAsChanged();
+			}
+		}
+
+		FString FindCurrentPreviewCulture() const
+		{
+			UMovieScene* MovieScene = WeakMovieScene.Get();
+			UMovieSceneLanguagePreviewDecoration* PreviewLanguage = MovieScene->FindDecoration<UMovieSceneLanguagePreviewDecoration>();
+			if (PreviewLanguage)
+			{
+				return PreviewLanguage->PreviewLanguage;
+			}
+			return FString();
+		}
+
+		TWeakObjectPtr<UMovieScene> WeakMovieScene;
+		TWeakObjectPtr<UMovieSceneTrack> WeakAudioTrack;
+	};
+}
+
 
 /** These utility functions should go away once we start handling sound cues properly */
 USoundWave* DeriveSoundWave(UMovieSceneAudioSection* AudioSection)
 {
 	USoundWave* SoundWave = NULL;
 	
-	USoundBase* Sound = AudioSection->GetSound();
+	USoundBase* Sound = AudioSection->GetPlaybackSound();
 	if (!Sound)
 	{
 		return SoundWave;
@@ -103,7 +241,7 @@ USoundWave* DeriveSoundWave(UMovieSceneAudioSection* AudioSection)
 
 float DeriveUnloopedDuration(UMovieSceneAudioSection* AudioSection)
 {
-	USoundBase* Sound = AudioSection->GetSound();
+	USoundBase* Sound = AudioSection->GetPlaybackSound();
 	if (Sound && Sound->GetDuration() != INDEFINITELY_LOOPING_DURATION)
 	{
 		return Sound->GetDuration();
@@ -674,10 +812,10 @@ FText FAudioSection::GetSectionTitle() const
 	UMovieSceneAudioSection* AudioSection = Cast<UMovieSceneAudioSection>(&Section);
 	if (AudioSection)
 	{
-		if (AudioSection->GetSound())
+		if (USoundBase* Sound = AudioSection->GetPlaybackSound())
 		{
 			// Return the asset name if it exists
-			return FText::FromString(AudioSection->GetSound()->GetName());
+			return FText::FromString(Sound->GetName());
 		}
 		else
 		{
@@ -803,8 +941,6 @@ int32 FAudioSection::OnPaintSection( FSequencerSectionPainter& Painter ) const
 
 	const ESlateDrawEffect DrawEffects = Painter.bParentEnabled ? ESlateDrawEffect::None : ESlateDrawEffect::DisabledEffect;
 
-	const FTimeToPixel& TimeToPixelConverter = Painter.GetTimeConverter();
-
 	static const FSlateBrush* GenericDivider = FAppStyle::GetBrush("Sequencer.GenericDivider");
 
 	if (!Section.HasStartFrame() || !Section.HasEndFrame())
@@ -812,14 +948,16 @@ int32 FAudioSection::OnPaintSection( FSequencerSectionPainter& Painter ) const
 		return LayerId;
 	}
 
+	TimeToPixel = Painter.GetTimeConverter();
+
 	UMovieSceneAudioSection* AudioSection = Cast<UMovieSceneAudioSection>(&Section);
-	if (!AudioSection || !AudioSection->GetSound())
+	if (!AudioSection || !AudioSection->GetPlaybackSound())
 	{
 		return LayerId;
 	}
 
 	// Add lines where the animation starts and ends/loops
-	FFrameRate TickResolution = TimeToPixelConverter.GetTickResolution();
+	FFrameRate TickResolution = TimeToPixel->GetTickResolution();
 	const float AudioDuration = DeriveUnloopedDuration(AudioSection);
 
 	if (AudioDuration > KINDA_SMALL_NUMBER)
@@ -831,7 +969,7 @@ int32 FAudioSection::OnPaintSection( FSequencerSectionPainter& Painter ) const
 		float OffsetTime = AudioDuration - StartOffsetTime;
 		while (OffsetTime < MaxOffset)
 		{
-			float OffsetPixel = TimeToPixelConverter.SecondsToPixel(SectionStartTime + OffsetTime);
+			float OffsetPixel = TimeToPixel->SecondsToPixel(OffsetTime);
 
 			FSlateDrawElement::MakeBox(
 				Painter.DrawElements,
@@ -868,7 +1006,7 @@ void FAudioSection::Tick( const FGeometry& AllottedGeometry, const FGeometry& Pa
 	
 	EMovieScenePlayerStatus::Type PlaybackState = SequencerPin->GetPlaybackStatus();
 
-	if (PlaybackState == EMovieScenePlayerStatus::Playing || PlaybackState == EMovieScenePlayerStatus::Scrubbing)
+	if (!TimeToPixel.IsSet() || PlaybackState == EMovieScenePlayerStatus::Playing || PlaybackState == EMovieScenePlayerStatus::Scrubbing)
 	{
 		return;
 	}
@@ -887,20 +1025,20 @@ void FAudioSection::Tick( const FGeometry& AllottedGeometry, const FGeometry& Pa
 			FSlateRect(FVector2D(0, 0), ParentGeometry.GetLocalSize())
 		);
 
+		// Always draw audio with linear space
+		FTimeToPixel LinearTimeToPixel = *TimeToPixel;
+		LinearTimeToPixel.NonLinearTransform = nullptr;
+
 		const float LeftMostVisiblePixel = FMath::Max(ParentRect.Left, 0.f);
 		const float RightMostVisiblePixel = FMath::Min(ParentRect.Right, AllottedGeometry.GetLocalSize().X);
 
-		FFrameRate   TickResolution = AudioSection->GetTypedOuter<UMovieScene>()->GetTickResolution();
-		FTimeToPixel TimeToPixel( AllottedGeometry, AudioSection->GetRange() / TickResolution, TickResolution );
-
-		TRange<float> DrawRange = TRange<float>(
-			TimeToPixel.PixelToSeconds(LeftMostVisiblePixel),
-			TimeToPixel.PixelToSeconds(RightMostVisiblePixel)
-			);
+		float StartTime = TimeToPixel->PixelToSeconds(LeftMostVisiblePixel);
+		float EndTime   = StartTime + (LinearTimeToPixel.PixelToSeconds(RightMostVisiblePixel) - LinearTimeToPixel.PixelToSeconds(LeftMostVisiblePixel));
+		TRange<float> DrawRange = TRange<float>(StartTime, EndTime);
 
 		// generate texture x offset and x size
-		int32 XOffset = LeftMostVisiblePixel;//PixelRange.GetLowerBoundValue() - TimeToPixelConverter.TimeToPixel(SectionRange.GetLowerBoundValue());
-		int32 XSize = RightMostVisiblePixel - LeftMostVisiblePixel;//PixelRange.Size<int32>();
+		int32 XOffset = LeftMostVisiblePixel;
+		int32 XSize = RightMostVisiblePixel - LeftMostVisiblePixel;
 
 		if (!FMath::IsNearlyEqual(DrawRange.GetLowerBoundValue(), StoredDrawRange.GetLowerBoundValue()) ||
 			!FMath::IsNearlyEqual(DrawRange.GetUpperBoundValue(), StoredDrawRange.GetUpperBoundValue()) ||
@@ -1031,7 +1169,7 @@ void FAudioSection::RegenerateWaveforms(TRange<float> DrawRange, int32 XOffset,
 	StoredSectionHeight = GetSectionHeight(SequencerPin->GetViewModel()->GetViewDensity());
 	bStoredLooping = AudioSection->GetLooping();
 
-	if (DrawRange.IsDegenerate() || DrawRange.IsEmpty() || AudioSection->GetSound() == NULL)
+	if (DrawRange.IsDegenerate() || DrawRange.IsEmpty() || AudioSection->GetPlaybackSound() == NULL)
 	{
 		WaveformThumbnail.Reset();
 	}
@@ -1260,9 +1398,38 @@ TSharedRef<ISequencerSection> FAudioTrackEditor::MakeSectionInterface( UMovieSce
 	return MakeShareable( new FAudioSection(SectionObject, GetSequencer()) );
 }
 
-TSharedPtr<SWidget> FAudioTrackEditor::BuildOutlinerEditWidget(const FGuid& ObjectBinding, UMovieSceneTrack* Track, const FBuildEditWidgetParams& Params)
+TSharedPtr<SWidget> FAudioTrackEditor::BuildOutlinerColumnWidget(const FBuildColumnWidgetParams& Params, const FName& ColumnName)
 {
-	return UE::Sequencer::MakeAddButton(LOCTEXT("AudioText", "Audio"), FOnGetContent::CreateSP(this, &FAudioTrackEditor::BuildAudioSubMenu, FOnAssetSelected::CreateRaw(this, &FAudioTrackEditor::OnAudioAssetSelected, Track), FOnAssetEnterPressed::CreateRaw(this, &FAudioTrackEditor::OnAudioAssetEnterPressed, Track)), Params.ViewModel);
+	using namespace UE::Sequencer;
+
+	UMovieSceneTrack* Track = Params.TrackModel->GetTrack();
+	if (!Track)
+	{
+		return SNullWidget::NullWidget;
+	}
+
+	if (ColumnName == FCommonOutlinerNames::Edit)
+	{
+		if (Params.TrackModel.AsModel()->IsA<FTrackModel>())
+		{
+			return SNew(AudioTrack::SLanguagePreviewComboWidget, Track->GetTypedOuter<UMovieScene>(), Track);
+		}
+	}
+
+	if (ColumnName == FCommonOutlinerNames::Add)
+	{
+		return UE::Sequencer::MakeAddButton(
+			LOCTEXT("AudioText", "Audio"),
+			FOnGetContent::CreateSP(
+				this, &FAudioTrackEditor::BuildAudioSubMenu,
+				FOnAssetSelected::CreateRaw(this, &FAudioTrackEditor::OnAudioAssetSelected, Track),
+				FOnAssetEnterPressed::CreateRaw(this, &FAudioTrackEditor::OnAudioAssetEnterPressed, Track)
+			),
+			Params.ViewModel
+		);
+	}
+
+	return FMovieSceneTrackEditor::BuildOutlinerColumnWidget(Params, ColumnName);
 }
 
 bool FAudioTrackEditor::HandleAssetAdded(UObject* Asset, const FGuid& TargetObjectGuid)
diff --git a/Engine/Source/Editor/MovieSceneTools/Private/TrackEditors/TimeWarpTrackEditor.cpp b/Engine/Source/Editor/MovieSceneTools/Private/TrackEditors/TimeWarpTrackEditor.cpp
--- a/Engine/Source/Editor/MovieSceneTools/Private/TrackEditors/TimeWarpTrackEditor.cpp
+++ b/Engine/Source/Editor/MovieSceneTools/Private/TrackEditors/TimeWarpTrackEditor.cpp
@@ -182,7 +182,7 @@ struct STimeWarpScrubber : public SLeafWidget
 
 		TSharedPtr<ISequencer> Sequencer = WeakSequencer.Pin();
 
-		if (Sequencer && TimeWarpTrack && TimeWarpTrack->bIsActiveTimeWarp)
+		if (Sequencer && TimeWarpTrack && TimeWarpTrack->IsTimeWarpActive())
 		{
 			FLinearColor TimeWarpColor = FStyleColors::AccentOrange.GetSpecifiedColor();
 			if (IsDirectlyHovered())
@@ -243,7 +243,7 @@ struct STimeWarpScrubber : public SLeafWidget
 		UMovieSceneTimeWarpTrack*   TimeWarpTrack   = TimeWarpSection ? TimeWarpSection->GetTypedOuter<UMovieSceneTimeWarpTrack>() : nullptr;
 
 		const bool bIsVisible = (Sequencer && Sequencer->GetSequencerSettings()->GetTimeWarpDisplayMode() == ESequencerTimeWarpDisplay::Both)
-			&& TimeWarpTrack && TimeWarpTrack->bIsActiveTimeWarp;
+			&& TimeWarpTrack && TimeWarpTrack->IsTimeWarpActive();
 
 		return bIsVisible ? EVisibility::Visible : EVisibility::Collapsed;
 	}
@@ -299,7 +299,7 @@ UE_SEQUENCER_DEFINE_VIEW_MODEL_TYPE_ID(FTimeWarpTrackExtension)
 bool FTimeWarpTrackModel::IsActiveTimeWarp() const
 {
 	UMovieSceneTimeWarpTrack* Track = Cast<UMovieSceneTimeWarpTrack>(GetTrack());
-	return Track && Track->bIsActiveTimeWarp;
+	return Track && Track->IsTimeWarpActive();
 }
 
 void FTimeWarpTrackModel::OnConstruct()
diff --git a/Engine/Source/Editor/MovieSceneTools/Public/TrackEditors/AudioTrackEditor.h b/Engine/Source/Editor/MovieSceneTools/Public/TrackEditors/AudioTrackEditor.h
--- a/Engine/Source/Editor/MovieSceneTools/Public/TrackEditors/AudioTrackEditor.h
+++ b/Engine/Source/Editor/MovieSceneTools/Public/TrackEditors/AudioTrackEditor.h
@@ -13,6 +13,7 @@
 #include "MovieSceneTrackEditor.h"
 #include "IContentBrowserSingleton.h"
 #include "Containers/Map.h"
+#include "TimeToPixel.h"
 
 struct FAssetData;
 class FAudioThumbnail;
@@ -57,7 +58,7 @@ public:
 	virtual void OnRelease() override;
 	virtual void BuildAddTrackMenu(FMenuBuilder& MenuBuilder) override;
 	virtual void BuildObjectBindingTrackMenu(FMenuBuilder& MenuBuilder, const TArray<FGuid>& ObjectBindings, const UClass* ObjectClass) override;
-	virtual TSharedPtr<SWidget> BuildOutlinerEditWidget(const FGuid& ObjectBinding, UMovieSceneTrack* Track, const FBuildEditWidgetParams& Params) override;
+	virtual TSharedPtr<SWidget> BuildOutlinerColumnWidget(const FBuildColumnWidgetParams& Params, const FName& ColumnName) override;
 	virtual bool HandleAssetAdded(UObject* Asset, const FGuid& TargetObjectGuid) override;
 	virtual TSharedRef<ISequencerSection> MakeSectionInterface(UMovieSceneSection& SectionObject, UMovieSceneTrack& Track, FGuid ObjectBinding) override;
 	virtual bool SupportsType(TSubclassOf<UMovieSceneTrack> Type) const override;
@@ -164,6 +165,8 @@ private:
 	/** The section we are visualizing. */
 	UMovieSceneSection& Section;
 
+	mutable TOptional<FTimeToPixel> TimeToPixel;
+
 	/** The waveform thumbnail render object. */
 	TSharedPtr<class FAudioThumbnail> WaveformThumbnail;
 
diff --git a/Engine/Source/Editor/Sequencer/Private/MVVM/ViewModels/SectionModel.cpp b/Engine/Source/Editor/Sequencer/Private/MVVM/ViewModels/SectionModel.cpp
--- a/Engine/Source/Editor/Sequencer/Private/MVVM/ViewModels/SectionModel.cpp
+++ b/Engine/Source/Editor/Sequencer/Private/MVVM/ViewModels/SectionModel.cpp
@@ -5,6 +5,7 @@
 #include "MVVM/ViewModels/ViewModelIterators.h"
 #include "MVVM/ViewModels/SequencerEditorViewModel.h"
 #include "MVVM/Extensions/ITrackExtension.h"
+#include "MVVM/Extensions/IViewSpaceClientExtension.h"
 #include "IKeyArea.h"
 #include "Sequencer.h"
 #include "SSequencerSection.h"
@@ -127,7 +128,11 @@ FTrackLaneVirtualAlignment FSectionModel::ArrangeVirtualTrackLaneView() const
 {
 	using namespace UE::MovieScene;
 
-	return FTrackLaneVirtualAlignment::Proportional(SectionRange, 1.f);
+	// Use CastThisShared to ensure that we cast both native and dynamic extensions
+	TSharedPtr<const IViewSpaceClientExtension> ViewSpaceClient = CastThisShared<const IViewSpaceClientExtension>();
+
+	FGuid ViewSpaceID = ViewSpaceClient ? ViewSpaceClient->GetViewSpaceID() : FGuid();
+	return FTrackLaneVirtualAlignment::Proportional(SectionRange, 1.f, VAlign_Center, ViewSpaceID);
 }
 
 TRange<FFrameNumber> FSectionModel::GetLayerBarRange() const
diff --git a/Engine/Source/Editor/Sequencer/Private/MVVM/Views/SChannelView.cpp b/Engine/Source/Editor/Sequencer/Private/MVVM/Views/SChannelView.cpp
--- a/Engine/Source/Editor/Sequencer/Private/MVVM/Views/SChannelView.cpp
+++ b/Engine/Source/Editor/Sequencer/Private/MVVM/Views/SChannelView.cpp
@@ -172,6 +172,11 @@ EViewDependentCacheFlags FChannelViewKeyCachedState::CompareTo(const FChannelVie
 		}
 	}
 
+	if (Other.NonLinearTransform.Pin() != NonLinearTransform.Pin())
+	{
+		Flags |= EViewDependentCacheFlags::ViewChanged;
+	}
+
 	if (ValidPlayRangeMin != Other.ValidPlayRangeMin || ValidPlayRangeMax != Other.ValidPlayRangeMax)
 	{
 		// The valid key ranges for the data has changed
@@ -524,6 +529,7 @@ int32 SChannelView::OnPaint(const FPaintArgs& Args, const FGeometry& AllottedGeo
 	FChannelViewKeyCachedState NewCachedState(VisibleRange, TrackArea->GetHotspot(), Model, Sequencer.Get());
 	NewCachedState.bIsChannelHovered = IsHovered();
 	NewCachedState.KeySizePx = SectionInterface->GetKeySize();
+	NewCachedState.NonLinearTransform = RelativeTimeToPixel.NonLinearTransform;
 
 	FKeyBatchParameters Params(RelativeTimeToPixel);
 	Params.CacheState = KeyRendererCache.IsSet() ? NewCachedState.CompareTo(KeyRendererCache.GetValue()) : EViewDependentCacheFlags::All;
diff --git a/Engine/Source/Editor/Sequencer/Private/SSequencer.cpp b/Engine/Source/Editor/Sequencer/Private/SSequencer.cpp
--- a/Engine/Source/Editor/Sequencer/Private/SSequencer.cpp
+++ b/Engine/Source/Editor/Sequencer/Private/SSequencer.cpp
@@ -222,12 +222,12 @@ void SSequencer::Construct(const FArguments& InArgs, TSharedRef<FSequencer> InSe
 	{
 		TimeSliderArgs.ViewRange = InArgs._ViewRange;
 		TimeSliderArgs.ClampRange = InArgs._ClampRange;
-		TimeSliderArgs.PlaybackRange = InArgs._PlaybackRange;
+		TimeSliderArgs.PlaybackRange = MakeAttributeSP(this, &SSequencer::GetViewSpacePlaybackRange, InArgs._PlaybackRange);
 		TimeSliderArgs.TimeBounds = InArgs._TimeBounds;
 		TimeSliderArgs.DisplayRate = TAttribute<FFrameRate>(InSequencer, &FSequencer::GetFocusedDisplayRate);
 		TimeSliderArgs.TickResolution = TAttribute<FFrameRate>(InSequencer, &FSequencer::GetFocusedTickResolution);
 		TimeSliderArgs.SelectionRange = InArgs._SelectionRange;
-		TimeSliderArgs.OnPlaybackRangeChanged = InArgs._OnPlaybackRangeChanged;
+		TimeSliderArgs.OnPlaybackRangeChanged = FOnFrameRangeChanged::CreateSP(this, &SSequencer::OnViewSpacePlaybackRangeChanged, InArgs._OnPlaybackRangeChanged);
 		TimeSliderArgs.OnPlaybackRangeBeginDrag = OnPlaybackRangeBeginDrag;
 		TimeSliderArgs.OnPlaybackRangeEndDrag = OnPlaybackRangeEndDrag;
 		TimeSliderArgs.OnSelectionRangeChanged = InArgs._OnSelectionRangeChanged;
@@ -1167,6 +1167,34 @@ void SSequencer::NotifyPostChange(const FPropertyChangedEvent& PropertyChangedEv
 /* SSequencer implementation
  *****************************************************************************/
 
+TRange<FFrameNumber> SSequencer::GetViewSpacePlaybackRange(TAttribute<TRange<FFrameNumber>> RangeAttribute) const
+{
+	TRange<FFrameNumber> Range = RangeAttribute.Get();
+
+	TSharedPtr<FTimeToPixel> TimeToPixel = TrackArea->GetTimeToPixel();
+	if (TimeToPixel->NonLinearTransform)
+	{
+		double       EndSeconds = Range.GetUpperBoundValue() / TimeToPixel->GetTickResolution();
+		FFrameNumber EndFrame   = (TimeToPixel->NonLinearTransform->SourceToView(EndSeconds) * TimeToPixel->GetTickResolution()).RoundToFrame();
+		Range.SetUpperBoundValue(EndFrame);
+	}
+
+	return Range;
+}
+
+void SSequencer::OnViewSpacePlaybackRangeChanged(TRange<FFrameNumber> NewRange, FOnFrameRangeChanged OnPlaybackRangeChanged)
+{
+	TSharedPtr<FTimeToPixel> TimeToPixel = TrackArea->GetTimeToPixel();
+	if (TimeToPixel->NonLinearTransform)
+	{
+		double       EndSeconds = NewRange.GetUpperBoundValue() / TimeToPixel->GetTickResolution();
+		FFrameNumber EndFrame   = (TimeToPixel->NonLinearTransform->ViewToSource(EndSeconds) * TimeToPixel->GetTickResolution()).RoundToFrame();
+		NewRange.SetUpperBoundValue(EndFrame);
+	}
+
+	OnPlaybackRangeChanged.ExecuteIfBound(NewRange);
+}
+
 void SSequencer::UpdateOutlinerViewColumns()
 {
 	using namespace UE::Sequencer;
diff --git a/Engine/Source/Editor/Sequencer/Private/SSequencer.h b/Engine/Source/Editor/Sequencer/Private/SSequencer.h
--- a/Engine/Source/Editor/Sequencer/Private/SSequencer.h
+++ b/Engine/Source/Editor/Sequencer/Private/SSequencer.h
@@ -456,6 +456,9 @@ public:
 
 private:
 
+	TRange<FFrameNumber> GetViewSpacePlaybackRange(TAttribute<TRange<FFrameNumber>> Range) const;
+	void OnViewSpacePlaybackRangeChanged(TRange<FFrameNumber> NewRange, FOnFrameRangeChanged OnPlaybackRangeChanged);
+
 	/**
 	 * Called when any outliner column's visibily is modified.
 	 * Updates SequencerSettings and visible outliner columns in Outliner View.
diff --git a/Engine/Source/Editor/Sequencer/Private/SSequencerSection.cpp b/Engine/Source/Editor/Sequencer/Private/SSequencerSection.cpp
--- a/Engine/Source/Editor/Sequencer/Private/SSequencerSection.cpp
+++ b/Engine/Source/Editor/Sequencer/Private/SSequencerSection.cpp
@@ -63,17 +63,6 @@ struct FEasingCurvePoint
 	FLinearColor Color;
 };
 
-FTimeToPixel ConstructTimeConverterForSection(const FGeometry& InSectionGeometry, const UMovieSceneSection& InSection, FSequencer& InSequencer)
-{
-	TRange<double> ViewRange       = InSequencer.GetViewRange();
-
-	FFrameRate     TickResolution  = InSection.GetTypedOuter<UMovieScene>()->GetTickResolution();
-	double         LowerTime       = InSection.HasStartFrame() ? InSection.GetInclusiveStartFrame() / TickResolution : ViewRange.GetLowerBoundValue();
-	double         UpperTime       = InSection.HasEndFrame()   ? InSection.GetExclusiveEndFrame()   / TickResolution : ViewRange.GetUpperBoundValue();
-
-	return FTimeToPixel(InSectionGeometry, TRange<double>(LowerTime, UpperTime), TickResolution);
-}
-
 static bool IsSectionToKey(const UMovieSceneTrack* Track, UMovieSceneSection* SectionObject)
 {
 	if (const IMovieSceneSectionsToKey* MultipleSectionsToKey = Cast<IMovieSceneSectionsToKey>(Track))
@@ -874,10 +863,7 @@ void SSequencerSection::Construct( const FArguments& InArgs, TSharedPtr<FSequenc
 
 		+ SOverlay::Slot(FCreateSectionViewWidgetParams::CompoundTrackLaneViewOrder)
 		[
-			SAssignNew(ChildLaneWidgets, SCompoundTrackLaneView)
-			.TimeToPixel(FGetTimeToPixel::CreateLambda([this](const FGeometry& AllottedGeometry){
-				return *this->TimeToPixel;
-			}))
+			SAssignNew(ChildLaneWidgets, SCompoundTrackLaneView, SharedThis(this))
 		]
 
 		+ SOverlay::Slot(FCreateSectionViewWidgetParams::ChannelViewOrder)
@@ -1036,7 +1022,7 @@ void SSequencerSection::ReportParentGeometry(const FGeometry& InParentGeometry)
 	ParentGeometry = InParentGeometry;
 }
 
-FTrackLaneScreenAlignment SSequencerSection::GetAlignment(const FTimeToPixel& InTimeToPixel, const FGeometry& InParentGeometry) const
+FTrackLaneScreenAlignment SSequencerSection::GetAlignment(const ITrackLaneWidgetSpace& InScreenSpace, const FGeometry& InParentGeometry) const
 {
 	using namespace UE::MovieScene;
 	TSharedPtr<FSectionModel> SectionModel = WeakSectionModel.Pin();
@@ -1046,7 +1032,7 @@ FTrackLaneScreenAlignment SSequencerSection::GetAlignment(const FTimeToPixel& In
 	}
 
 	FTrackLaneVirtualAlignment VirtualAlignment = SectionModel->ArrangeVirtualTrackLaneView();
-	FTrackLaneScreenAlignment  ScreenAlignment  = VirtualAlignment.ToScreen(InTimeToPixel, InParentGeometry);
+	FTrackLaneScreenAlignment  ScreenAlignment  = VirtualAlignment.ToScreen(InScreenSpace.GetScreenSpace(VirtualAlignment.ViewSpaceID), InParentGeometry);
 
 	if (TOptional<FFrameNumber> FiniteLength = VirtualAlignment.GetFiniteLength())
 	{
@@ -1416,7 +1402,6 @@ int32 SSequencerSection::OnPaint( const FPaintArgs& Args, const FGeometry& Allot
 	const ESlateDrawEffect DrawEffects = bEnabled ? ESlateDrawEffect::None : ESlateDrawEffect::DisabledEffect;
 
 	FGeometry SectionGeometry = MakeHandlesGeometry(AllottedGeometry);
-	*TimeToPixel = ConstructTimeConverterForSection(SectionGeometry, *SectionObject, GetSequencer()); 
 
 	FSequencerSectionPainterImpl Painter(GetSequencer(), TrackAreaViewModel, SectionModel, OutDrawElements, SectionGeometry, *this);
 
@@ -1902,10 +1887,11 @@ void SSequencerSection::DrawSectionHandles( const FGeometry& AllottedGeometry, F
 
 void SSequencerSection::Tick( const FGeometry& AllottedGeometry, const double InCurrentTime, const float InDeltaTime )
 {
-	UMovieSceneSection*    Section      = SectionInterface->GetSectionObject();
-	TSharedPtr<FSequencer> SequencerPtr = Sequencer.Pin();
+	UMovieSceneSection*       Section      = SectionInterface->GetSectionObject();
+	TSharedPtr<FSequencer>    SequencerPtr = Sequencer.Pin();
+	TSharedPtr<FSectionModel> SectionModel = WeakSectionModel.Pin();
 
-	if (!Section || !SequencerPtr)
+	if (!Section || !SequencerPtr || !SectionModel)
 	{
 		return;
 	}
@@ -1930,12 +1916,23 @@ void SSequencerSection::Tick( const FGeometry& AllottedGeometry, const double In
 		}
 
 		FGeometry SectionGeometry = MakeSectionGeometryWithoutHandles(AllottedGeometry);
-		*TimeToPixel = ConstructTimeConverterForSection(SectionGeometry, *Section, *SequencerPtr);
+
+		FTrackLaneVirtualAlignment VirtualAlignment = SectionModel->ArrangeVirtualTrackLaneView();
+		*TimeToPixel = *WeakOwningTrackLane.Pin()->GetTrackAreaView()->GetTimeToPixel(VirtualAlignment.ViewSpaceID);
+		if (Section->HasStartFrame())
+		{
+			*TimeToPixel = TimeToPixel->RelativeTo(Section->GetInclusiveStartFrame());
+		}
 		SectionInterface->Tick(SectionGeometry, ParentGeometry, InCurrentTime, InDeltaTime);
 	}
 	else
 	{
-		*TimeToPixel = ConstructTimeConverterForSection(AllottedGeometry, *Section, *SequencerPtr);
+		FTrackLaneVirtualAlignment VirtualAlignment = SectionModel->ArrangeVirtualTrackLaneView();
+		*TimeToPixel = *WeakOwningTrackLane.Pin()->GetTrackAreaView()->GetTimeToPixel(VirtualAlignment.ViewSpaceID);
+		if (Section->HasStartFrame())
+		{
+			*TimeToPixel = TimeToPixel->RelativeTo(Section->GetInclusiveStartFrame());
+		}
 	}
 }
 
@@ -1957,6 +1954,33 @@ TSharedRef<FTimeToPixel> SSequencerSection::GetTimeToPixel() const
 	return TimeToPixel.ToSharedRef();
 }
 
+FTimeToPixel SSequencerSection::GetScreenSpace(const FGuid& ViewSpaceID) const
+{
+	using namespace UE::MovieScene;
+
+	TSharedPtr<STrackLane>     TrackLane = WeakOwningTrackLane.Pin();
+	TSharedPtr<STrackAreaView> TrackArea = TrackLane ? TrackLane->GetTrackAreaView() : nullptr;
+	if (TrackArea)
+	{
+		TSharedPtr<FTimeToPixel> Space = TrackArea->GetTimeToPixel(ViewSpaceID);
+
+		TSharedPtr<FSectionModel> Section = WeakSectionModel.Pin();
+		if (Section)
+		{
+			TRange<FFrameNumber> Range = Section->GetRange();
+			if (Range.GetLowerBound().IsClosed())
+			{
+				return Space->RelativeTo(DiscreteInclusiveLower(Range));
+			}
+		}
+
+		return *Space;
+	}
+
+	return *TimeToPixel;
+}
+
+
 FReply SSequencerSection::OnMouseButtonDown(const FGeometry& MyGeometry, const FPointerEvent& MouseEvent)
 {
 	return FReply::Unhandled();
diff --git a/Engine/Source/Editor/Sequencer/Private/SSequencerSection.h b/Engine/Source/Editor/Sequencer/Private/SSequencerSection.h
--- a/Engine/Source/Editor/Sequencer/Private/SSequencerSection.h
+++ b/Engine/Source/Editor/Sequencer/Private/SSequencerSection.h
@@ -30,7 +30,7 @@ class FSectionModel;
 class FTrackAreaViewModel;
 struct ITrackAreaHotspot;
 
-class SSequencerSection : public SCompoundWidget, public ITrackLaneWidget, public UE::Sequencer::ISectionView, public UE::MovieScene::ISignedObjectEventHandler
+class SSequencerSection : public SCompoundWidget, public ITrackLaneWidget, public UE::Sequencer::ISectionView, public UE::Sequencer::ITrackLaneWidgetSpace, public UE::MovieScene::ISignedObjectEventHandler
 {
 public:
 	SLATE_BEGIN_ARGS( SSequencerSection )
@@ -90,7 +90,7 @@ private:
 	virtual void OnMouseLeave( const FPointerEvent& MouseEvent ) override;
 
 	/*~ ITrackLaneWidget interface */
-	virtual FTrackLaneScreenAlignment GetAlignment(const FTimeToPixel& TimeToPixel, const FGeometry& InParentGeometry) const override;
+	virtual FTrackLaneScreenAlignment GetAlignment(const ITrackLaneWidgetSpace& ScreenSpace, const FGeometry& InParentGeometry) const override;
 	virtual int32 GetOverlapPriority() const override;
 	virtual void ReportParentGeometry(const FGeometry& InParentGeometry) override;
 	virtual TSharedRef<const SWidget> AsWidget() const override { return AsShared(); }
@@ -103,6 +103,9 @@ private:
 	/*~ ISectionView interface */
 	virtual TSharedRef<FTimeToPixel> GetTimeToPixel() const override;
 
+	/*~ ITrackLaneWidgetSpace interface */
+	virtual FTimeToPixel GetScreenSpace(const FGuid& ViewSpaceID = FGuid()) const override;
+
 	/**
 	 * Paint the easing handles for this section
 	 */
diff --git a/Engine/Source/Editor/Sequencer/Private/SequencerContextMenus.cpp b/Engine/Source/Editor/Sequencer/Private/SequencerContextMenus.cpp
--- a/Engine/Source/Editor/Sequencer/Private/SequencerContextMenus.cpp
+++ b/Engine/Source/Editor/Sequencer/Private/SequencerContextMenus.cpp
@@ -3,6 +3,7 @@
 #include "SequencerContextMenus.h"
 #include "Modules/ModuleManager.h"
 #include "Styling/AppStyle.h"
+#include "Decorations/MovieSceneSectionAnchorsDecoration.h"
 #include "SequencerCommonHelpers.h"
 #include "SequencerCommands.h"
 #include "SSequencer.h"
@@ -403,6 +404,14 @@ void FSectionContextMenu::PopulateMenu(FMenuBuilder& MenuBuilder, TSharedPtr<FEx
 			);
 		}
 
+		if (SelectionSupportsScaling())
+		{
+			MenuBuilder.AddSubMenu(
+				LOCTEXT("ScalingSection", "Scaling"),
+				LOCTEXT("ScalingSectionTooltip", "Options for scaling this section"),
+				FNewMenuDelegate::CreateLambda([=](FMenuBuilder& InMenuBuilder) { Shared->AddScalingMenu(InMenuBuilder); }));
+		}
+
 		MenuBuilder.AddSubMenu(
 			LOCTEXT("EditSection", "Edit"),
 			LOCTEXT("EditSectionTooltip", "Edit section"),
@@ -498,6 +507,98 @@ void FSectionContextMenu::PopulateMenu(FMenuBuilder& MenuBuilder, TSharedPtr<FEx
 	MenuBuilder.EndSection(); // SequencerSections
 }
 
+bool FSectionContextMenu::SelectionSupportsScaling() const
+{
+	TSharedPtr<FSequencer> Sequencer = WeakSequencer.Pin();
+	if (!Sequencer)
+	{
+		return false;
+	}
+
+	TSet<UClass*> CompatibleDecorations;
+	for (UMovieSceneSection* Section : Sequencer->GetViewModel()->GetSelection()->GetSelectedSections())
+	{
+		if (Section)
+		{
+			if (Cast<IMovieSceneScalingDriver>(Section))
+			{
+				return true;
+			}
+
+			Section->GetCompatibleUserDecorations(CompatibleDecorations);
+		}
+	}
+
+	return CompatibleDecorations.Contains(UMovieSceneSectionAnchorsDecoration::StaticClass());
+}
+
+void FSectionContextMenu::AddScalingMenu(FMenuBuilder& MenuBuilder)
+{
+	// Copy a reference to the context menu by value into each lambda handler to ensure the type stays alive until the menu is closed
+	TSharedRef<FSectionContextMenu> Shared = AsShared();
+
+	auto GetScalingDriverCheckState = [Shared]
+	{
+		TOptional<ECheckBoxState> CheckState;
+		TSharedPtr<FSequencer> Sequencer = Shared->WeakSequencer.Pin();
+		if (Sequencer)
+		{
+			for (UMovieSceneSection* Section : Sequencer->GetViewModel()->GetSelection()->GetSelectedSections())
+			{
+				ECheckBoxState ThisCheckState = (Section && Section->FindDecoration<UMovieSceneSectionAnchorsDecoration>() != nullptr)
+					? ECheckBoxState::Checked
+					: ECheckBoxState::Unchecked;
+
+				if (!CheckState.IsSet())
+				{
+					CheckState = ThisCheckState;
+				}
+				else if (CheckState.GetValue() != ThisCheckState)
+				{
+					return ECheckBoxState::Undetermined;
+				}
+			}
+		}
+		return CheckState.Get(ECheckBoxState::Unchecked);
+	};
+
+	MenuBuilder.AddMenuEntry(
+		LOCTEXT("ScalingDriver", "Scaling Driver"),
+		LOCTEXT("ScalingDriverTooltip", "Defines whether this section will rescale the sequence based on its start/end times when being played back."),
+		FSlateIcon(),
+		FUIAction(
+			FExecuteAction::CreateLambda([Shared, GetScalingDriverCheckState]
+			{
+				TSharedPtr<FSequencer> Sequencer = Shared->WeakSequencer.Pin();
+				if (Sequencer)
+				{
+					ECheckBoxState State = GetScalingDriverCheckState();
+					FScopedTransaction Transaction(LOCTEXT("ToggleScaling", "Toggle Scaling"));
+
+					for (UMovieSceneSection* Section : Sequencer->GetViewModel()->GetSelection()->GetSelectedSections())
+					{
+						if (Section)
+						{
+							Section->Modify();
+							if (State == ECheckBoxState::Checked)
+							{
+								Section->RemoveDecoration<UMovieSceneSectionAnchorsDecoration>();
+							}
+							else
+							{
+								Section->GetOrCreateDecoration<UMovieSceneSectionAnchorsDecoration>();
+							}
+						}
+					}
+				}
+			}),
+			FCanExecuteAction(),
+			FGetActionCheckState::CreateLambda(GetScalingDriverCheckState)
+		),
+		NAME_None,
+		EUserInterfaceActionType::ToggleButton
+	);
+}
 
 void FSectionContextMenu::AddEditMenu(FMenuBuilder& MenuBuilder)
 {
diff --git a/Engine/Source/Editor/Sequencer/Private/SequencerContextMenus.h b/Engine/Source/Editor/Sequencer/Private/SequencerContextMenus.h
--- a/Engine/Source/Editor/Sequencer/Private/SequencerContextMenus.h
+++ b/Engine/Source/Editor/Sequencer/Private/SequencerContextMenus.h
@@ -40,6 +40,9 @@ private:
 	/** Add edit menu for trim and split */
 	void AddEditMenu(FMenuBuilder& MenuBuilder);
 
+	void AddScalingMenu(FMenuBuilder& MenuBuilder);
+	bool SelectionSupportsScaling() const;
+
 	/** Add the Order sub-menu. */
 	void AddOrderMenu(FMenuBuilder& MenuBuilder);
 
diff --git a/Engine/Source/Editor/Sequencer/Private/SequencerSettings.cpp b/Engine/Source/Editor/Sequencer/Private/SequencerSettings.cpp
--- a/Engine/Source/Editor/Sequencer/Private/SequencerSettings.cpp
+++ b/Engine/Source/Editor/Sequencer/Private/SequencerSettings.cpp
@@ -72,6 +72,7 @@ USequencerSettings::USequencerSettings( const FObjectInitializer& ObjectInitiali
 	bShowTickLines = true;
 	bShowSequencerToolbar = true;
 	bShowMarkedFrames = true;
+	bShowScalingAnchors = true;
 	ViewDensity = "Relaxed";
 	AssetBrowserWidth = 500.f;
 	AssetBrowserHeight = 300.f;
@@ -712,6 +713,20 @@ void USequencerSettings::SetShowMarkedFrames(bool bInShowMarkedFrames)
 	}
 }
 
+bool USequencerSettings::GetShowScalingAnchors() const
+{
+	return bShowScalingAnchors;
+}
+
+void USequencerSettings::SetShowScalingAnchors(bool bInShowScalingAnchors)
+{
+	if (bShowScalingAnchors != bInShowScalingAnchors)
+	{
+		bShowScalingAnchors = bInShowScalingAnchors;
+		SaveConfig();
+	}
+}
+
 bool USequencerSettings::HasKeyAreaCurveExtents(const FString& ChannelName) const
 {
 	TArray<FString> ChannelsArray;
diff --git a/Engine/Source/Editor/Sequencer/Private/SequencerTimeSliderController.cpp b/Engine/Source/Editor/Sequencer/Private/SequencerTimeSliderController.cpp
--- a/Engine/Source/Editor/Sequencer/Private/SequencerTimeSliderController.cpp
+++ b/Engine/Source/Editor/Sequencer/Private/SequencerTimeSliderController.cpp
@@ -31,6 +31,7 @@
 #include "CommonFrameRates.h"
 #include "Sequencer.h"
 #include "Modules/ModuleManager.h"
+#include "Decorations/MovieSceneScalingAnchors.h"
 #include "MovieSceneSequence.h"
 #include "MovieScene.h"
 #include "Compilation/MovieSceneCompiledDataManager.h"
@@ -336,6 +337,56 @@ void FSequencerTimeSliderController::DrawLinearTicks(FSlateWindowElementList& Ou
 	}
 }
 
+int32 FSequencerTimeSliderController::DrawScalingAnchors(const FGeometry& AllottedGeometry, const FScrubRangeToScreen& RangeToScreen, FSlateWindowElementList& OutDrawElements, int32 LayerId, const ESlateDrawEffect& DrawEffects) const
+{
+	UMovieSceneScalingAnchors* Anchors = WeakSequencer.Pin()->GetFocusedMovieSceneSequence()->GetMovieScene()->FindDecoration<UMovieSceneScalingAnchors>();
+	if (!Anchors || !WeakSequencer.Pin()->GetSequencerSettings()->GetShowScalingAnchors())
+	{
+		return LayerId;
+	}
+
+	TArray<FVector2f> LinePoints;
+	for (const TPair<FGuid, FMovieSceneScalingAnchor>& Pair : Anchors->GetCurrentAnchors())
+	{
+		const double StartSeconds = Pair.Value.Position / GetTickResolution();
+		const float  StartLinePos = RangeToScreen.InputToLocalX(StartSeconds);
+
+		LinePoints.Reset(2);
+		LinePoints.Emplace(FVector2f(StartLinePos, 0.0f));
+		LinePoints.Emplace(FVector2f(StartLinePos, FMath::FloorToFloat(AllottedGeometry.Size.Y)));
+
+		FSlateDrawElement::MakeDashedLines(
+			OutDrawElements,
+			LayerId,
+			AllottedGeometry.ToPaintGeometry(),
+			MoveTemp(LinePoints),
+			DrawEffects,
+			FLinearColor::Yellow.CopyWithNewOpacity(0.25f)
+		);
+
+		if (Pair.Value.Duration > 0)
+		{
+			const double EndSeconds = (Pair.Value.Position + Pair.Value.Duration) / GetTickResolution();
+			const float  EndLinePos = RangeToScreen.InputToLocalX(EndSeconds);
+
+			LinePoints.Reset(2);
+			LinePoints.Emplace(FVector2f(EndLinePos, 0.0f));
+			LinePoints.Emplace(FVector2f(EndLinePos, FMath::FloorToFloat(AllottedGeometry.Size.Y)));
+
+			FSlateDrawElement::MakeDashedLines(
+				OutDrawElements,
+				LayerId,
+				AllottedGeometry.ToPaintGeometry(),
+				MoveTemp(LinePoints),
+				DrawEffects,
+				FLinearColor::Yellow.CopyWithNewOpacity(0.25f)
+			);
+		}
+	}
+
+	return LayerId + 1;
+}
+
 int32 FSequencerTimeSliderController::DrawMarkedFrames( const FGeometry& AllottedGeometry, const FScrubRangeToScreen& RangeToScreen, FSlateWindowElementList& OutDrawElements, int32 LayerId, const ESlateDrawEffect& DrawEffects, const FWidgetStyle& InWidgetStyle, bool bDrawLabels ) const
 {
 	if (!WeakSequencer.Pin()->GetSequencerSettings()->GetShowMarkedFrames())
@@ -1494,6 +1545,8 @@ int32 FSequencerTimeSliderController::OnPaintViewArea( const FGeometry& Allotted
 		LayerId = DrawMarkedFrames(AllottedGeometry, RangeToScreen, OutDrawElements, LayerId, DrawEffects, FWidgetStyle(), false);
 	}
 
+	LayerId = DrawScalingAnchors(AllottedGeometry, RangeToScreen, OutDrawElements, LayerId, DrawEffects);
+
 	LayerId = DrawVerticalFrames(AllottedGeometry, RangeToScreen, OutDrawElements, LayerId, DrawEffects);
 
 	FTimeWarpTrackExtension*   TimeWarpExtension   = Sequencer->GetViewModel()->GetRootSequenceModel()->GetSharedData()->CastDynamic<FTimeWarpTrackExtension>();
diff --git a/Engine/Source/Editor/Sequencer/Private/SequencerUtilities.cpp b/Engine/Source/Editor/Sequencer/Private/SequencerUtilities.cpp
--- a/Engine/Source/Editor/Sequencer/Private/SequencerUtilities.cpp
+++ b/Engine/Source/Editor/Sequencer/Private/SequencerUtilities.cpp
@@ -201,12 +201,15 @@ void FSequencerUtilities::PopulateTimeWarpSubMenu(FMenuBuilder& MenuBuilder, TFu
 		const UClass* IconClass = FClassIconFinder::GetIconClassForAssetData(AssetData);
 		const UClass* Class     = Cast<const UClass>(AssetData.FastGetAsset());
 
-		MenuBuilder.AddMenuEntry(
-			Class ? Class->GetDisplayNameText() : FText::FromName(ClassPath.GetAssetName()),
-			Class ? Class->GetToolTipText()     : FText(),
-			FSlateIconFinder::FindIconForClass(IconClass),
-			FUIAction(FExecuteAction::CreateLambda(HandleTimeWarpSelection, ClassPath))
-		);
+		if (!Class->HasMetaData("Hidden"))
+		{
+			MenuBuilder.AddMenuEntry(
+				Class ? Class->GetDisplayNameText() : FText::FromName(ClassPath.GetAssetName()),
+				Class ? Class->GetToolTipText()     : FText(),
+				FSlateIconFinder::FindIconForClass(IconClass),
+				FUIAction(FExecuteAction::CreateLambda(HandleTimeWarpSelection, ClassPath))
+			);
+		}
 	}
 
 	MenuBuilder.EndSection();
diff --git a/Engine/Source/Editor/Sequencer/Private/Widgets/SSequencerBindingLifetimeOverlay.cpp b/Engine/Source/Editor/Sequencer/Private/Widgets/SSequencerBindingLifetimeOverlay.cpp
--- a/Engine/Source/Editor/Sequencer/Private/Widgets/SSequencerBindingLifetimeOverlay.cpp
+++ b/Engine/Source/Editor/Sequencer/Private/Widgets/SSequencerBindingLifetimeOverlay.cpp
@@ -53,11 +53,11 @@ namespace UE
 			return AsShared();
 		}
 
-		FTrackLaneScreenAlignment SSequencerBindingLifetimeOverlay::GetAlignment(const FTimeToPixel& InTimeToPixel, const FGeometry& InParentGeometry) const
+		FTrackLaneScreenAlignment SSequencerBindingLifetimeOverlay::GetAlignment(const ITrackLaneWidgetSpace& InScreenSpace, const FGeometry& InParentGeometry) const
 		{
 			if (TSharedPtr<FBindingLifetimeOverlayModel> BindingLifetimeOverlayModel = GetBindingLifetimeOverlayModel())
 			{
-				return BindingLifetimeOverlayModel->ArrangeVirtualTrackLaneView().ToScreen(InTimeToPixel, InParentGeometry);
+				return BindingLifetimeOverlayModel->ArrangeVirtualTrackLaneView().ToScreen(InScreenSpace.GetScreenSpace(), InParentGeometry);
 			}
 			return FTrackLaneScreenAlignment{};
 		}
diff --git a/Engine/Source/Editor/Sequencer/Private/Widgets/SSequencerBindingLifetimeOverlay.h b/Engine/Source/Editor/Sequencer/Private/Widgets/SSequencerBindingLifetimeOverlay.h
--- a/Engine/Source/Editor/Sequencer/Private/Widgets/SSequencerBindingLifetimeOverlay.h
+++ b/Engine/Source/Editor/Sequencer/Private/Widgets/SSequencerBindingLifetimeOverlay.h
@@ -46,7 +46,7 @@ namespace UE
 
 			/*~ ITrackLaneWidget interface */
 			TSharedRef<const SWidget> AsWidget() const override;
-			FTrackLaneScreenAlignment GetAlignment(const FTimeToPixel& TimeToPixel, const FGeometry& InParentGeometry) const override;
+			FTrackLaneScreenAlignment GetAlignment(const ITrackLaneWidgetSpace& ScreenSpace, const FGeometry& InParentGeometry) const override;
 			int32 GetOverlapPriority() const override;
 
 			/*~ SCompoundWidget interface */
diff --git a/Engine/Source/Editor/Sequencer/Private/Widgets/SSequencerLayerBar.cpp b/Engine/Source/Editor/Sequencer/Private/Widgets/SSequencerLayerBar.cpp
--- a/Engine/Source/Editor/Sequencer/Private/Widgets/SSequencerLayerBar.cpp
+++ b/Engine/Source/Editor/Sequencer/Private/Widgets/SSequencerLayerBar.cpp
@@ -293,11 +293,11 @@ TSharedRef<const SWidget> SSequencerLayerBar::AsWidget() const
 	return AsShared();
 }
 
-FTrackLaneScreenAlignment SSequencerLayerBar::GetAlignment(const FTimeToPixel& InTimeToPixel, const FGeometry& InParentGeometry) const
+FTrackLaneScreenAlignment SSequencerLayerBar::GetAlignment(const ITrackLaneWidgetSpace& InScreenSpace, const FGeometry& InParentGeometry) const
 {
 	if (TSharedPtr<FLayerBarModel> LayerBar = WeakLayerBar.Pin())
 	{
-		return LayerBar->ArrangeVirtualTrackLaneView().ToScreen(InTimeToPixel, InParentGeometry);
+		return LayerBar->ArrangeVirtualTrackLaneView().ToScreen(InScreenSpace.GetScreenSpace(), InParentGeometry);
 	}
 	return FTrackLaneScreenAlignment{};
 }
diff --git a/Engine/Source/Editor/Sequencer/Private/Widgets/SSequencerLayerBar.h b/Engine/Source/Editor/Sequencer/Private/Widgets/SSequencerLayerBar.h
--- a/Engine/Source/Editor/Sequencer/Private/Widgets/SSequencerLayerBar.h
+++ b/Engine/Source/Editor/Sequencer/Private/Widgets/SSequencerLayerBar.h
@@ -45,7 +45,7 @@ private:
 
 	/*~ ITrackLaneWidget interface */
 	TSharedRef<const SWidget> AsWidget() const override;
-	FTrackLaneScreenAlignment GetAlignment(const FTimeToPixel& TimeToPixel, const FGeometry& InParentGeometry) const override;
+	FTrackLaneScreenAlignment GetAlignment(const ITrackLaneWidgetSpace& ScreenSpace, const FGeometry& InParentGeometry) const override;
 
 	/*~ SCompoundWidget interface */
 	int32 OnPaint( const FPaintArgs& Args, const FGeometry& AllottedGeometry, const FSlateRect& MyCullingRect, FSlateWindowElementList& OutDrawElements, int32 LayerId, const FWidgetStyle& InWidgetStyle, bool bParentEnabled ) const override;
diff --git a/Engine/Source/Editor/Sequencer/Public/MVVM/Views/SChannelView.h b/Engine/Source/Editor/Sequencer/Public/MVVM/Views/SChannelView.h
--- a/Engine/Source/Editor/Sequencer/Public/MVVM/Views/SChannelView.h
+++ b/Engine/Source/Editor/Sequencer/Public/MVVM/Views/SChannelView.h
@@ -31,7 +31,8 @@ struct FChannelViewKeyCachedState
 	TWeakPtr<ITrackAreaHotspot> WeakHotspot;
 	/** DO NOT DEREFERENCE. Raw ptr to the last used hotspot. Only used to differentiate an expired hotspot from an explicitly null one. */
 	ITrackAreaHotspot* RawHotspotDoNotUse = nullptr;
-
+	/** Weak ptr to the last used non-linear transform for rendering these keys */
+	TWeakPtr<INonLinearTimeTransform> NonLinearTransform;
 	/** The min/max tick value relating to the FMovieSceneSubSequenceData::ValidPlayRange bounds, or the current playback range */
 	FFrameNumber ValidPlayRangeMin, ValidPlayRangeMax;
 
diff --git a/Engine/Source/Editor/Sequencer/Public/SequencerSettings.h b/Engine/Source/Editor/Sequencer/Public/SequencerSettings.h
--- a/Engine/Source/Editor/Sequencer/Public/SequencerSettings.h
+++ b/Engine/Source/Editor/Sequencer/Public/SequencerSettings.h
@@ -369,6 +369,11 @@ public:
 	/** Set whether to show marked frames */
 	void SetShowMarkedFrames(bool bShowMarkedFrames);
 
+	/** @return true if showing scaling anchors */
+	bool GetShowScalingAnchors() const;
+	/** Set whether to show scaling anchors */
+	void SetShowScalingAnchors(bool bShowScalingAnchors);
+
 	/** @return Whether the given channel has curve extents */
 	bool HasKeyAreaCurveExtents(const FString& ChannelName) const;
 	/** @ Remove curve extents for the given channel */
@@ -736,6 +741,10 @@ protected:
 	UPROPERTY(config, EditAnywhere, Category = Timeline)
 	bool bShowMarkedFrames;
 
+	/** Enable or disable showing scaling anchors */
+	UPROPERTY(config, EditAnywhere, Category = Timeline)
+	bool bShowScalingAnchors;
+
 	/** The key area curve extents, stored per channel name */
 	UPROPERTY(config, EditAnywhere, Category = Timeline)
 	FString KeyAreaCurveExtents;
diff --git a/Engine/Source/Editor/Sequencer/Public/SequencerTimeSliderController.h b/Engine/Source/Editor/Sequencer/Public/SequencerTimeSliderController.h
--- a/Engine/Source/Editor/Sequencer/Public/SequencerTimeSliderController.h
+++ b/Engine/Source/Editor/Sequencer/Public/SequencerTimeSliderController.h
@@ -291,6 +291,14 @@ private:
 	 */
 	int32 DrawMarkedFrames(const FGeometry& AllottedGeometry, const FScrubRangeToScreen& RangeToScreen, FSlateWindowElementList& OutDrawElements, int32 LayerId, const ESlateDrawEffect& DrawEffects, const FWidgetStyle& InWidgetStyle, bool bDrawLabels) const;
 
+
+	/**
+	 * Draw any scaling anchors.
+	 *
+	 * @return the new layer ID
+	 */
+	int32 DrawScalingAnchors( const FGeometry& AllottedGeometry, const FScrubRangeToScreen& RangeToScreen, FSlateWindowElementList& OutDrawElements, int32 LayerId, const ESlateDrawEffect& DrawEffects ) const;
+
 private:
 
 	/**
diff --git a/Engine/Source/Editor/SequencerCore/Private/MVVM/Extensions/ITrackLaneExtension.cpp b/Engine/Source/Editor/SequencerCore/Private/MVVM/Extensions/ITrackLaneExtension.cpp
--- a/Engine/Source/Editor/SequencerCore/Private/MVVM/Extensions/ITrackLaneExtension.cpp
+++ b/Engine/Source/Editor/SequencerCore/Private/MVVM/Extensions/ITrackLaneExtension.cpp
@@ -94,6 +94,14 @@ FArrangedWidget FTrackLaneScreenAlignment::ArrangeWidget(TSharedRef<SWidget> InW
 	);
 }
 
+PRAGMA_DISABLE_DEPRECATION_WARNINGS
+FTrackLaneScreenAlignment ITrackLaneWidget::GetAlignment(const ITrackLaneWidgetSpace& ScreenSpace, const FGeometry& InParentGeometry) const
+{
+	ensureMsgf(false, TEXT("FTrackLaneScreenAlignment GetAlignment(const FTimeToPixel&, const FGeometry&) is deprecated. Please update your code to override the new method"));
+	return GetAlignment(ScreenSpace.GetScreenSpace(), InParentGeometry);
+}
+PRAGMA_ENABLE_DEPRECATION_WARNINGS
+
 } // namespace Sequencer
 } // namespace UE
 
diff --git a/Engine/Source/Editor/SequencerCore/Private/MVVM/ViewModelTypeIDs.cpp b/Engine/Source/Editor/SequencerCore/Private/MVVM/ViewModelTypeIDs.cpp
--- a/Engine/Source/Editor/SequencerCore/Private/MVVM/ViewModelTypeIDs.cpp
+++ b/Engine/Source/Editor/SequencerCore/Private/MVVM/ViewModelTypeIDs.cpp
@@ -27,12 +27,15 @@
 #include "MVVM/Extensions/ITrackAreaExtension.h"
 #include "MVVM/Extensions/ITrackLaneExtension.h"
 #include "MVVM/Extensions/LinkedOutlinerExtension.h"
+#include "MVVM/Extensions/ITrackAreaViewSpaceProviderExtension.h"
+#include "MVVM/Extensions/IViewSpaceClientExtension.h"
 #include "MVVM/ICastable.h"
 #include "MVVM/ViewModelTypeID.h"
 #include "MVVM/ViewModels/EditorViewModel.h"
 #include "MVVM/ViewModels/OutlinerSpacer.h"
 #include "MVVM/ViewModels/OutlinerViewModel.h"
 #include "MVVM/ViewModels/TrackAreaViewModel.h"
+#include "MVVM/ViewModels/TrackAreaViewSpace.h"
 #include "MVVM/ViewModels/ViewModel.h"
 #include "MVVM/ViewModels/EditorSharedViewModelData.h"
 #include "MVVM/Views/ITrackAreaHotspot.h"
@@ -56,6 +59,9 @@ UE_SEQUENCER_DEFINE_VIEW_MODEL_TYPE_ID(ILayerBarExtension);
 UE_SEQUENCER_DEFINE_VIEW_MODEL_TYPE_ID(IMutableExtension);
 UE_SEQUENCER_DEFINE_VIEW_MODEL_TYPE_ID(ILockableExtension);
 UE_SEQUENCER_DEFINE_VIEW_MODEL_TYPE_ID(IOutlinerDropTargetOutlinerExtension);
+UE_SEQUENCER_DEFINE_VIEW_MODEL_TYPE_ID(ITrackAreaViewSpaceProviderExtension);
+UE_SEQUENCER_DEFINE_VIEW_MODEL_TYPE_ID(IViewSpaceClientExtension);
+UE_SEQUENCER_DEFINE_VIEW_MODEL_TYPE_ID(FViewSpaceClientExtensionShim);
 UE_SEQUENCER_DEFINE_VIEW_MODEL_TYPE_ID(IObjectModelExtension);
 UE_SEQUENCER_DEFINE_VIEW_MODEL_TYPE_ID(IHierarchicalCache);
 UE_SEQUENCER_DEFINE_VIEW_MODEL_TYPE_ID(IOutlinerExtension);
@@ -85,6 +91,7 @@ UE_SEQUENCER_DEFINE_CASTABLE(FEditorViewModel);
 UE_SEQUENCER_DEFINE_CASTABLE(FOutlinerSpacer);
 UE_SEQUENCER_DEFINE_CASTABLE(FOutlinerViewModel);
 UE_SEQUENCER_DEFINE_CASTABLE(FTrackAreaViewModel);
+UE_SEQUENCER_DEFINE_CASTABLE(FTrackAreaViewSpace);
 UE_SEQUENCER_DEFINE_CASTABLE(FViewModel);
 UE_SEQUENCER_DEFINE_CASTABLE(FSharedViewModelData);
 UE_SEQUENCER_DEFINE_CASTABLE(FEditorSharedViewModelData);
diff --git a/Engine/Source/Editor/SequencerCore/Private/MVVM/ViewModels/EditorViewModel.cpp b/Engine/Source/Editor/SequencerCore/Private/MVVM/ViewModels/EditorViewModel.cpp
--- a/Engine/Source/Editor/SequencerCore/Private/MVVM/ViewModels/EditorViewModel.cpp
+++ b/Engine/Source/Editor/SequencerCore/Private/MVVM/ViewModels/EditorViewModel.cpp
@@ -6,6 +6,7 @@
 #include "MVVM/ViewModels/OutlinerViewModel.h"
 #include "MVVM/ViewModels/TrackAreaViewModel.h"
 #include "MVVM/Selection/SequencerCoreSelection.h"
+#include "MVVM/SharedViewModelData.h"
 #include "Scripting/SequencerScriptingLayer.h"
 #include "Misc/AssertionMacros.h"
 #include "Templates/TypeHash.h"
@@ -57,6 +58,7 @@ void FEditorViewModel::InitializeEditor()
 
 	ScriptingLayer.Reset(Scripting);
 
+	SetSharedData(MakeShared<FSharedViewModelData>());
 	InitializeEditorImpl();
 }
 
diff --git a/Engine/Source/Editor/SequencerCore/Private/MVVM/ViewModels/TrackAreaViewModel.cpp b/Engine/Source/Editor/SequencerCore/Private/MVVM/ViewModels/TrackAreaViewModel.cpp
--- a/Engine/Source/Editor/SequencerCore/Private/MVVM/ViewModels/TrackAreaViewModel.cpp
+++ b/Engine/Source/Editor/SequencerCore/Private/MVVM/ViewModels/TrackAreaViewModel.cpp
@@ -4,7 +4,10 @@
 
 #include "HAL/PlatformCrt.h"
 #include "ISequencerEditTool.h"
+#include "MVVM/Extensions/ITrackAreaViewSpaceProviderExtension.h"
+#include "MVVM/SharedViewModelData.h"
 #include "MVVM/ViewModels/EditorViewModel.h"
+#include "MVVM/ViewModels/ViewModelIterators.h"
 #include "MVVM/Views/ITrackAreaHotspot.h"
 #include "Misc/FrameRate.h"
 #include "Templates/UnrealTemplate.h"
@@ -24,6 +27,11 @@ FTrackAreaViewModel::~FTrackAreaViewModel()
 {
 }
 
+void FTrackAreaViewModel::OnConstruct()
+{
+	UpdateViewSpaces();
+}
+
 TSharedPtr<FEditorViewModel> FTrackAreaViewModel::GetEditor() const
 {
 	if (TSharedPtr<FViewModel> Parent = GetParent())
@@ -33,6 +41,23 @@ TSharedPtr<FEditorViewModel> FTrackAreaViewModel::GetEditor() const
 	return nullptr;
 }
 
+void FTrackAreaViewModel::UpdateViewSpaces()
+{
+	ViewSpaces.Empty();
+
+	FViewModelPtr RootModel = GetEditor()->GetRootModel();
+	FSimpleMulticastDelegate& Event = RootModel->GetSharedData()->SubscribeToHierarchyChanged(RootModel);
+	if (!Event.IsBoundToObject(this))
+	{
+		Event.AddSP(this, &FTrackAreaViewModel::UpdateViewSpaces);
+	}
+
+	for (TViewModelPtr<ITrackAreaViewSpaceProviderExtension> ViewSpace : RootModel->GetChildrenOfType<ITrackAreaViewSpaceProviderExtension>())
+	{
+		ViewSpace->UpdateViewSpaces(*this);
+	}
+}
+
 void FTrackAreaViewModel::AddEditTool(TSharedPtr<ISequencerEditTool> InNewTool)
 {
 	EditTools.Add(MoveTemp(InNewTool));
@@ -181,6 +206,31 @@ void FTrackAreaViewModel::UnlockEditTool()
 	bEditToolLocked = false;
 }
 
+void FTrackAreaViewModel::AddViewSpace(const FGuid& InID, TSharedPtr<FTrackAreaViewSpace> ViewSpace)
+{
+	TSharedPtr<FTrackAreaViewSpace> Existing = ViewSpaces.FindRef(InID);
+	if (Existing != ViewSpace)
+	{
+		ViewSpaces.Add(InID, ViewSpace);
+		OnViewSpaceAdded.Broadcast(InID, ViewSpace);
+	}
+}
+
+void FTrackAreaViewModel::RemoveViewSpace(const FGuid& InID)
+{
+	TSharedPtr<FTrackAreaViewSpace> Existing = ViewSpaces.FindRef(InID);
+	if (Existing)
+	{
+		ViewSpaces.Remove(InID);
+		OnViewSpaceRemoved.Broadcast(InID, Existing);
+	}
+}
+
+TSharedPtr<FTrackAreaViewSpace> FTrackAreaViewModel::FindViewSpace(const FGuid& InID) const
+{
+	return ViewSpaces.FindRef(InID);
+}
+
 } // namespace Sequencer
 } // namespace UE
 
diff --git a/Engine/Source/Editor/SequencerCore/Private/MVVM/ViewModels/TrackAreaViewSpace.cpp b/Engine/Source/Editor/SequencerCore/Private/MVVM/ViewModels/TrackAreaViewSpace.cpp
new file mode 100644
--- /dev/null
+++ b/Engine/Source/Editor/SequencerCore/Private/MVVM/ViewModels/TrackAreaViewSpace.cpp
@@ -0,0 +1,29 @@
+// Copyright Epic Games, Inc. All Rights Reserved.
+
+#include "MVVM/ViewModels/TrackAreaViewSpace.h"
+
+namespace UE::Sequencer
+{
+
+
+FTrackAreaViewSpace::FTrackAreaViewSpace()
+{
+}
+
+FTrackAreaViewSpace::~FTrackAreaViewSpace()
+{
+}
+
+double FTrackAreaViewSpace::SourceToView(double Seconds) const
+{
+	return Seconds;
+}
+
+double FTrackAreaViewSpace::ViewToSource(double Source) const
+{
+	return Source;
+}
+
+
+} // namespace UE::Sequencer
+
diff --git a/Engine/Source/Editor/SequencerCore/Private/MVVM/Views/SCompoundTrackLaneView.cpp b/Engine/Source/Editor/SequencerCore/Private/MVVM/Views/SCompoundTrackLaneView.cpp
--- a/Engine/Source/Editor/SequencerCore/Private/MVVM/Views/SCompoundTrackLaneView.cpp
+++ b/Engine/Source/Editor/SequencerCore/Private/MVVM/Views/SCompoundTrackLaneView.cpp
@@ -2,11 +2,10 @@
 
 #include "MVVM/Views/SCompoundTrackLaneView.h"
 #include "MVVM/Views/STrackLane.h"
+#include "MVVM/Views/STrackAreaView.h"
 #include "MVVM/Extensions/ITrackLaneExtension.h"
 #include "Widgets/SWeakWidget.h"
 
-#include "TimeToPixel.h"
-
 namespace UE::Sequencer
 {
 
@@ -28,14 +27,11 @@ SCompoundTrackLaneView::~SCompoundTrackLaneView()
 {
 }
 
-void SCompoundTrackLaneView::Construct(const FArguments& InArgs)
+void SCompoundTrackLaneView::Construct(const FArguments& InArgs, TSharedPtr<ITrackLaneWidgetSpace> InTrackLaneWidgetSpace)
 {
-	TimeToPixelDelegate = InArgs._TimeToPixel;
-
-	// Has to be bound dynamically (otherwise it wouldn't update when section ranges change)
-	check(TimeToPixelDelegate.IsBound());
-
 	SetClipping(EWidgetClipping::ClipToBounds);
+
+	WeakTrackLaneWidgetSpace = InTrackLaneWidgetSpace;
 }
 
 void SCompoundTrackLaneView::AddWeakWidget(TSharedPtr<ITrackLaneWidget> InWidget, TWeakPtr<STrackLane> InOwningLane)
@@ -56,7 +52,11 @@ void SCompoundTrackLaneView::AddStrongWidget(TSharedPtr<ITrackLaneWidget> InWidg
 
 void SCompoundTrackLaneView::OnArrangeChildren(const FGeometry& AllottedGeometry, FArrangedChildren& ArrangedChildren) const
 {
-	FTimeToPixel TimeToPixel = TimeToPixelDelegate.Execute(AllottedGeometry);
+	TSharedPtr<ITrackLaneWidgetSpace> WidgetSpace = WeakTrackLaneWidgetSpace.Pin();
+	if (!WidgetSpace)
+	{
+		return;
+	}
 
 	for (int32 WidgetIndex = 0; WidgetIndex < Children.Num(); ++WidgetIndex)
 	{
@@ -85,7 +85,7 @@ void SCompoundTrackLaneView::OnArrangeChildren(const FGeometry& AllottedGeometry
 				);
 			}
 
-			FTrackLaneScreenAlignment ScreenAlignment = Interface->GetAlignment(TimeToPixel, LaneGeometry);
+			FTrackLaneScreenAlignment ScreenAlignment = Interface->GetAlignment(*WidgetSpace, LaneGeometry);
 			if (ScreenAlignment.IsVisible())
 			{
 				FArrangedWidget ArrangedWidget = ScreenAlignment.ArrangeWidget(Widget, LaneGeometry);
diff --git a/Engine/Source/Editor/SequencerCore/Private/MVVM/Views/STrackAreaLaneView.cpp b/Engine/Source/Editor/SequencerCore/Private/MVVM/Views/STrackAreaLaneView.cpp
--- a/Engine/Source/Editor/SequencerCore/Private/MVVM/Views/STrackAreaLaneView.cpp
+++ b/Engine/Source/Editor/SequencerCore/Private/MVVM/Views/STrackAreaLaneView.cpp
@@ -3,6 +3,9 @@
 #include "MVVM/Views/STrackAreaLaneView.h"
 #include "MVVM/Views/STrackAreaView.h"
 #include "MVVM/ViewModels/TrackAreaViewModel.h"
+#include "MVVM/ViewModels/TrackAreaViewSpace.h"
+#include "MVVM/ViewModels/ViewModelIterators.h"
+#include "MVVM/Extensions/IViewSpaceClientExtension.h"
 #include "TimeToPixel.h"
 
 
@@ -13,7 +16,12 @@ void STrackAreaLaneView::Construct(const FArguments& InArgs, const FViewModelPtr
 {
 	WeakModel = InViewModel;
 	WeakTrackAreaView = InTrackAreaView;
-	TrackAreaTimeToPixel = InTrackAreaView->GetTimeToPixel();
+
+	// Use CastThisShared to ensure that we cast both native and dynamic extensions
+	TSharedPtr<const IViewSpaceClientExtension> ViewSpaceClient = InViewModel->FindAncestorOfType<const IViewSpaceClientExtension>(true);
+
+	FGuid ViewSpaceID = ViewSpaceClient ? ViewSpaceClient->GetViewSpaceID() : FGuid();
+	TrackAreaTimeToPixel = InTrackAreaView->GetTimeToPixel(ViewSpaceID);
 
 	ChildSlot
 	[
@@ -26,15 +34,13 @@ TSharedRef<const SWidget> STrackAreaLaneView::AsWidget() const
 	return AsShared();
 }
 
-FTrackLaneScreenAlignment STrackAreaLaneView::GetAlignment(const FTimeToPixel& InTimeToPixel, const FGeometry& InParentGeometry) const
+FTrackLaneScreenAlignment STrackAreaLaneView::GetAlignment(const ITrackLaneWidgetSpace& InScreenSpace, const FGeometry& InParentGeometry) const
 {
 	TSharedPtr<ITrackLaneExtension> TrackLaneExtension = WeakModel.ImplicitPin();
 	if (TrackLaneExtension)
 	{
 		FTrackLaneVirtualAlignment VirtualAlignment = TrackLaneExtension->ArrangeVirtualTrackLaneView();
-		FTrackLaneScreenAlignment  ScreenAlignment  = VirtualAlignment.ToScreen(InTimeToPixel, InParentGeometry);
-
-		return ScreenAlignment;
+		return VirtualAlignment.ToScreen(InScreenSpace.GetScreenSpace(VirtualAlignment.ViewSpaceID), InParentGeometry);
 	}
 	return FTrackLaneScreenAlignment();
 }
diff --git a/Engine/Source/Editor/SequencerCore/Private/MVVM/Views/STrackAreaView.cpp b/Engine/Source/Editor/SequencerCore/Private/MVVM/Views/STrackAreaView.cpp
--- a/Engine/Source/Editor/SequencerCore/Private/MVVM/Views/STrackAreaView.cpp
+++ b/Engine/Source/Editor/SequencerCore/Private/MVVM/Views/STrackAreaView.cpp
@@ -18,6 +18,7 @@
 #include "MVVM/Extensions/IOutlinerExtension.h"
 #include "MVVM/ViewModels/OutlinerViewModel.h"
 #include "MVVM/ViewModels/TrackAreaViewModel.h"
+#include "MVVM/ViewModels/TrackAreaViewSpace.h"
 #include "MVVM/Views/ITrackAreaHotspot.h"
 #include "MVVM/Views/SOutlinerView.h"
 #include "MVVM/Views/STrackLane.h"
@@ -84,6 +85,9 @@ void STrackAreaView::Construct(const FArguments& InArgs, TWeakPtr<FTrackAreaView
 	TSharedPtr<FTrackAreaViewModel> ViewModel = WeakViewModel.Pin();
 	check(ViewModel);
 
+	ViewModel->OnViewSpaceAdded.AddSP(this, &STrackAreaView::HandleViewSpaceAdded);
+	ViewModel->OnViewSpaceRemoved.AddSP(this, &STrackAreaView::HandleViewSpaceRemoved);
+
 	InputStack.OnBeginCapture.AddSP(ViewModel.Get(), &FTrackAreaViewModel::LockEditTool);
 	InputStack.OnEndCapture.AddSP(ViewModel.Get(), &FTrackAreaViewModel::UnlockEditTool);
 
@@ -103,11 +107,34 @@ TSharedPtr<FTrackAreaViewModel> STrackAreaView::GetViewModel() const
 	return WeakViewModel.Pin();
 }
 
-TSharedPtr<FTimeToPixel> STrackAreaView::GetTimeToPixel() const
+TSharedPtr<FTimeToPixel> STrackAreaView::GetFallbackTimeToPixel() const
 {
 	return TimeToPixel;
 }
 
+TSharedPtr<FTimeToPixel> STrackAreaView::GetTimeToPixel(const FGuid& InViewSpaceID) const
+{
+	TSharedPtr<FTimeToPixel> Result = ViewSpaceTimeToPixel.FindRef(InViewSpaceID);
+	return Result ? Result : TimeToPixel;
+}
+
+FTimeToPixel STrackAreaView::GetScreenSpace(const FGuid& InViewSpaceID) const
+{
+	return *GetTimeToPixel(InViewSpaceID);
+}
+
+void STrackAreaView::HandleViewSpaceAdded(const FGuid& InID, TSharedPtr<FTrackAreaViewSpace> ViewSpace)
+{
+	TSharedPtr<FTimeToPixel> TTP = MakeShared<FTimeToPixel>(TimeToPixel->CreateNonLinear(ViewSpace));
+	ensure(TTP->NonLinearTransform);
+	ViewSpaceTimeToPixel.Add(InID, TTP);
+}
+
+void STrackAreaView::HandleViewSpaceRemoved(const FGuid& InID, TSharedPtr<FTrackAreaViewSpace> ViewSpace)
+{
+	ViewSpaceTimeToPixel.Remove(InID);
+}
+
 FLinearColor STrackAreaView::BlendDefaultTrackColor(FLinearColor InColor)
 {
 	static FLinearColor BaseColor(FColor(71,71,71));
@@ -152,6 +179,11 @@ void STrackAreaView::OnArrangeChildren( const FGeometry& AllottedGeometry, FArra
 
 	*const_cast<STrackAreaView*>(this)->TimeToPixel = FTimeToPixel(TrackArea->GetTimeToPixel(AllottedGeometry.GetLocalSize().X));
 
+	for (TPair<FGuid, TSharedPtr<FTimeToPixel>>& Pair : ViewSpaceTimeToPixel)
+	{
+		*Pair.Value = TimeToPixel->CreateNonLinear(Pair.Value->NonLinearTransform);
+	}
+
 	for (int32 ChildIndex = 0; ChildIndex < Children.Num(); ++ChildIndex)
 	{
 		const FTrackAreaSlot& CurChild = Children[ChildIndex];
@@ -220,7 +252,11 @@ int32 STrackAreaView::OnPaint(const FPaintArgs& Args, const FGeometry& AllottedG
 	}
 
 	// Reassign the time <-> pixel space for this frame
-	*GetTimeToPixel() = TrackArea->GetTimeToPixel(AllottedGeometry.GetLocalSize().X);
+	*TimeToPixel = TrackArea->GetTimeToPixel(AllottedGeometry.GetLocalSize().X);
+	for (TPair<FGuid, TSharedPtr<FTimeToPixel>>& Pair : ViewSpaceTimeToPixel)
+	{
+		*Pair.Value = TimeToPixel->CreateNonLinear(Pair.Value->NonLinearTransform);
+	}
 
 	// TODO: TrackEditors are drawn in subclass, maybe should be done in viewmodel?
 
diff --git a/Engine/Source/Editor/SequencerCore/Private/MVVM/Views/STrackLane.cpp b/Engine/Source/Editor/SequencerCore/Private/MVVM/Views/STrackLane.cpp
--- a/Engine/Source/Editor/SequencerCore/Private/MVVM/Views/STrackLane.cpp
+++ b/Engine/Source/Editor/SequencerCore/Private/MVVM/Views/STrackLane.cpp
@@ -212,6 +212,8 @@ void STrackLane::RecreateWidgets()
 
 void STrackLane::OnArrangeChildren( const FGeometry& AllottedGeometry, FArrangedChildren& ArrangedChildren ) const
 {
+	TSharedPtr<STrackAreaView> TrackAreaView = WeakTrackAreaView.Pin();
+
 	for (int32 WidgetIndex = 0; WidgetIndex < Children.Num(); ++WidgetIndex)
 	{
 		const FSlot&        Slot   = Children[WidgetIndex];
@@ -220,7 +222,7 @@ void STrackLane::OnArrangeChildren( const FGeometry& AllottedGeometry, FArranged
 		EVisibility WidgetVisibility = Widget->GetVisibility();
 		if (ArrangedChildren.Accepts(WidgetVisibility))
 		{
-			FTrackLaneScreenAlignment ScreenAlignment = Slot.Interface->GetAlignment(*TimeToPixel, AllottedGeometry);
+			FTrackLaneScreenAlignment ScreenAlignment = Slot.Interface->GetAlignment(*TrackAreaView, AllottedGeometry);
 			if (ScreenAlignment.IsVisible())
 			{
 				FArrangedWidget ArrangedWidget = ScreenAlignment.ArrangeWidget(Widget, AllottedGeometry);
diff --git a/Engine/Source/Editor/SequencerCore/Private/TimeToPixel.cpp b/Engine/Source/Editor/SequencerCore/Private/TimeToPixel.cpp
new file mode 100644
--- /dev/null
+++ b/Engine/Source/Editor/SequencerCore/Private/TimeToPixel.cpp
@@ -0,0 +1,126 @@
+// Copyright Epic Games, Inc. All Rights Reserved.
+
+#include "TimeToPixel.h"
+
+
+namespace UE::Sequencer
+{
+
+
+
+FTimeToPixelSpace::FTimeToPixelSpace( const FGeometry& AllottedGeometry, const TRange<double>& InLocalViewRange, const FFrameRate& InTickResolution )
+	: FTimeToPixelSpace(AllottedGeometry.GetLocalSize().X, InLocalViewRange, InTickResolution)
+{}
+
+FTimeToPixelSpace::FTimeToPixelSpace( float WidthPx, const TRange<double>& InLocalViewRange, const FFrameRate& InTickResolution )
+	: TickResolution( InTickResolution )
+{
+	ensureMsgf(WidthPx >= 0, TEXT("Got negative pixel width!"));
+	const double VisibleWidth = InLocalViewRange.Size<double>();
+
+	const float MaxPixelsPerSecond = 1000.f;
+	PixelsPerSecond = VisibleWidth > 0 ? static_cast<float>(WidthPx / VisibleWidth) : MaxPixelsPerSecond;
+	PixelOffset     = -InLocalViewRange.GetLowerBoundValue() * PixelsPerSecond;
+}
+
+FTimeToPixelSpace FTimeToPixelSpace::CreateNonLinear(TSharedPtr<INonLinearTimeTransform> InNonLinearTransform) const
+{
+	FTimeToPixelSpace Copy = *this;
+	Copy.NonLinearTransform = InNonLinearTransform;
+	return Copy;
+}
+
+float FTimeToPixelSpace::SecondsToPixel( double Time ) const
+{
+	if (NonLinearTransform)
+	{
+		Time = NonLinearTransform->SourceToView(Time);
+	}
+
+	return static_cast<float>(Time * PixelsPerSecond) + PixelOffset;
+}
+
+
+float FTimeToPixelSpace::SecondsDeltaToPixel( double TimeDelta ) const
+{
+	return static_cast<float>(TimeDelta * PixelsPerSecond);
+}
+
+
+double FTimeToPixelSpace::PixelToSeconds( float PixelX ) const
+{
+	PixelX -= PixelOffset;
+	double Result = (PixelX/PixelsPerSecond);
+	if (NonLinearTransform)
+	{
+		Result = NonLinearTransform->ViewToSource(Result);
+	}
+	return Result;
+}
+
+
+float FTimeToPixelSpace::FrameToPixel( const FFrameTime& Time ) const
+{
+	return SecondsToPixel(Time / TickResolution);
+}
+
+
+float FTimeToPixelSpace::FrameDeltaToPixel( const FFrameTime& TimeDelta) const
+{
+	return static_cast<float>((TimeDelta / TickResolution) * PixelsPerSecond);
+}
+
+
+FFrameTime FTimeToPixelSpace::PixelToFrame( float PixelX ) const
+{
+	return PixelToSeconds(PixelX) * TickResolution;
+}
+
+
+FFrameTime FTimeToPixelSpace::PixelDeltaToFrame( float PixelDelta ) const
+{
+	return ( PixelDelta / PixelsPerSecond ) * TickResolution;
+}
+
+
+double FTimeToPixelSpace::PixelDeltaToSeconds( float PixelDelta ) const
+{
+	return ( PixelDelta / PixelsPerSecond );
+}
+
+
+FFrameRate FTimeToPixelSpace::GetTickResolution() const
+{
+	return TickResolution;
+}
+
+
+FTimeToPixelSpace FTimeToPixelSpace::RelativeTo(const FFrameTime& FrameZero) const
+{
+	FTimeToPixelSpace Copy = *this;
+
+	double ViewRangeStart = FrameZero / TickResolution;
+	if (Copy.NonLinearTransform)
+	{
+		ViewRangeStart = Copy.NonLinearTransform->SourceToView(ViewRangeStart);
+	}
+
+	Copy.PixelOffset = -ViewRangeStart*PixelsPerSecond;
+	return Copy;
+}
+
+
+} // namespace UE::Sequencer
+
+FTimeToPixel FTimeToPixel::RelativeTo(const FFrameTime& FrameZero) const
+{
+	FTimeToPixel Copy = *this;
+	double ViewRangeStart = FrameZero / TickResolution;
+	if (Copy.NonLinearTransform)
+	{
+		ViewRangeStart = Copy.NonLinearTransform->SourceToView(ViewRangeStart);
+	}
+
+	Copy.PixelOffset = -ViewRangeStart*PixelsPerSecond;
+	return Copy;
+}
\ No newline at end of file
diff --git a/Engine/Source/Editor/SequencerCore/Public/MVVM/Extensions/ITrackAreaViewSpaceProviderExtension.h b/Engine/Source/Editor/SequencerCore/Public/MVVM/Extensions/ITrackAreaViewSpaceProviderExtension.h
new file mode 100644
--- /dev/null
+++ b/Engine/Source/Editor/SequencerCore/Public/MVVM/Extensions/ITrackAreaViewSpaceProviderExtension.h
@@ -0,0 +1,28 @@
+// Copyright Epic Games, Inc. All Rights Reserved.
+
+#pragma once
+
+#include "CoreTypes.h"
+#include "MVVM/ViewModelTypeID.h"
+
+namespace UE::Sequencer
+{
+
+class FTrackAreaViewModel;
+
+/**
+ * An extension that supplies view spaces to a track area.
+ * Must exist as an immediate child of the root model to be effective.
+ */
+class SEQUENCERCORE_API ITrackAreaViewSpaceProviderExtension
+{
+public:
+
+	UE_SEQUENCER_DECLARE_VIEW_MODEL_TYPE_ID(ITrackAreaViewSpaceProviderExtension)
+
+	virtual ~ITrackAreaViewSpaceProviderExtension(){}
+
+	virtual void UpdateViewSpaces(FTrackAreaViewModel& TrackAreaViewModel) = 0;
+};
+
+} // namespace UE::Sequencer
\ No newline at end of file
diff --git a/Engine/Source/Editor/SequencerCore/Public/MVVM/Extensions/ITrackLaneExtension.h b/Engine/Source/Editor/SequencerCore/Public/MVVM/Extensions/ITrackLaneExtension.h
--- a/Engine/Source/Editor/SequencerCore/Public/MVVM/Extensions/ITrackLaneExtension.h
+++ b/Engine/Source/Editor/SequencerCore/Public/MVVM/Extensions/ITrackLaneExtension.h
@@ -24,8 +24,10 @@ namespace Sequencer
 
 class FViewModel;
 class STrackLane;
+class STrackAreaView;
 class FEditorViewModel;
 struct FTrackLaneScreenAlignment;
+struct INonLinearTimeTransform;
 
 struct SEQUENCERCORE_API FTrackLaneVerticalArrangement
 {
@@ -52,19 +54,20 @@ struct SEQUENCERCORE_API FTrackLaneVirtualAlignment
 {
 	TRange<FFrameNumber> Range;
 	FTrackLaneVerticalAlignment VerticalAlignment;
+	FGuid ViewSpaceID;
 
 	bool IsVisible() const
 	{
 		return !Range.IsEmpty();
 	}
 
-	static FTrackLaneVirtualAlignment Fixed(const TRange<FFrameNumber>& InRange, float InFixedHeight, EVerticalAlignment InVAlign = VAlign_Center)
+	static FTrackLaneVirtualAlignment Fixed(const TRange<FFrameNumber>& InRange, float InFixedHeight, EVerticalAlignment InVAlign = VAlign_Center, const FGuid& InViewSpaceID = FGuid())
 	{
-		return FTrackLaneVirtualAlignment { InRange, { InFixedHeight, 0.f, InVAlign, FTrackLaneVerticalAlignment::ESizeMode::Fixed } };
+		return FTrackLaneVirtualAlignment { InRange, { InFixedHeight, 0.f, InVAlign, FTrackLaneVerticalAlignment::ESizeMode::Fixed }, InViewSpaceID };
 	}
-	static FTrackLaneVirtualAlignment Proportional(const TRange<FFrameNumber>& InRange, float InStretchFactor, EVerticalAlignment InVAlign = VAlign_Center)
+	static FTrackLaneVirtualAlignment Proportional(const TRange<FFrameNumber>& InRange, float InStretchFactor, EVerticalAlignment InVAlign = VAlign_Center, const FGuid& InViewSpaceID = FGuid())
 	{
-		return FTrackLaneVirtualAlignment { InRange, { InStretchFactor, 0.f, InVAlign, FTrackLaneVerticalAlignment::ESizeMode::Proportional } };
+		return FTrackLaneVirtualAlignment { InRange, { InStretchFactor, 0.f, InVAlign, FTrackLaneVerticalAlignment::ESizeMode::Proportional }, InViewSpaceID };
 	}
 
 	TOptional<FFrameNumber> GetFiniteLength() const;
@@ -74,8 +77,20 @@ struct SEQUENCERCORE_API FTrackLaneVirtualAlignment
 
 struct SEQUENCERCORE_API FTrackLaneScreenAlignment
 {
-	float LeftPosPx = 0.f;
-	float WidthPx = 0.f;
+	FTrackLaneScreenAlignment()
+		: LeftPosPx(0.f)
+		, WidthPx(0.f)
+	{
+	}
+
+	FTrackLaneScreenAlignment(float InLeftPosPx, float InWidthPx, FTrackLaneVerticalAlignment InVerticalAlignment)
+		: LeftPosPx(InLeftPosPx), WidthPx(InWidthPx), VerticalAlignment(InVerticalAlignment)
+	{
+	}
+
+	TSharedPtr<INonLinearTimeTransform> NonLinearTransform;
+	float LeftPosPx;
+	float WidthPx;
 	FTrackLaneVerticalAlignment VerticalAlignment;
 
 	bool IsVisible() const
@@ -92,6 +107,13 @@ struct SEQUENCERCORE_API FArrangedVirtualEntity
 	float VirtualTop, VirtualBottom;
 };
 
+/** Interface used for laying out track lane widgets in screen space */
+class SEQUENCERCORE_API ITrackLaneWidgetSpace
+{
+public:
+	virtual FTimeToPixel GetScreenSpace(const FGuid& ViewSpaceID = FGuid()) const = 0;
+};
+
 /** Base interface for track-area lanes */
 class SEQUENCERCORE_API ITrackLaneWidget
 {
@@ -111,7 +133,12 @@ public:
 	/**
 	 * Arrange this widget within its parent slot
 	 */
-	virtual FTrackLaneScreenAlignment GetAlignment(const FTimeToPixel& TimeToPixel, const FGeometry& InParentGeometry) const = 0;
+	UE_DEPRECATED(5.6, "Please use the ITrackLaneWidgetSpace overload")
+	virtual FTrackLaneScreenAlignment GetAlignment(const FTimeToPixel& TimeToPixel, const FGeometry& InParentGeometry) const
+	{
+		return FTrackLaneScreenAlignment();
+	}
+	virtual FTrackLaneScreenAlignment GetAlignment(const ITrackLaneWidgetSpace& ScreenSpace, const FGeometry& InParentGeometry) const;
 
 	/**
 	 * Gets this widget's overlap priority
diff --git a/Engine/Source/Editor/SequencerCore/Public/MVVM/Extensions/IViewSpaceClientExtension.h b/Engine/Source/Editor/SequencerCore/Public/MVVM/Extensions/IViewSpaceClientExtension.h
new file mode 100644
--- /dev/null
+++ b/Engine/Source/Editor/SequencerCore/Public/MVVM/Extensions/IViewSpaceClientExtension.h
@@ -0,0 +1,66 @@
+// Copyright Epic Games, Inc. All Rights Reserved.
+
+#pragma once
+
+#include "CoreTypes.h"
+#include "Misc/Guid.h"
+#include "MVVM/ViewModelTypeID.h"
+#include "MVVM/Extensions/DynamicExtensionContainer.h"
+
+struct FGuid;
+
+namespace UE::Sequencer
+{
+
+/**
+ * An extension that is used for view models that represent UObject data
+ */
+class SEQUENCERCORE_API IViewSpaceClientExtension
+{
+public:
+
+	UE_SEQUENCER_DECLARE_VIEW_MODEL_TYPE_ID(IViewSpaceClientExtension)
+
+	virtual ~IViewSpaceClientExtension(){}
+
+	virtual FGuid GetViewSpaceID() const = 0;
+	virtual void SetViewSpaceID(const FGuid&) = 0;
+};
+
+
+/**
+ * An extension that is used for view models that represent UObject data
+ */
+class SEQUENCERCORE_API FViewSpaceClientExtensionShim
+	: public IDynamicExtension
+	, public IViewSpaceClientExtension
+{
+public:
+
+	UE_SEQUENCER_DECLARE_VIEW_MODEL_TYPE_ID(FViewSpaceClientExtensionShim)
+
+	using Implements = TImplements<IDynamicExtension, IViewSpaceClientExtension>;
+
+	FViewSpaceClientExtensionShim(const FGuid& InViewSpaceID)
+		: ViewSpaceID(InViewSpaceID)
+	{
+	}
+
+	virtual FGuid GetViewSpaceID() const override
+	{
+		return ViewSpaceID;
+	}
+
+	virtual void SetViewSpaceID(const FGuid& InViewSpaceID) override
+	{
+		ViewSpaceID = InViewSpaceID;
+	}
+
+protected:
+
+	FGuid ViewSpaceID;
+};
+
+
+
+} // namespace UE::Sequencer
\ No newline at end of file
diff --git a/Engine/Source/Editor/SequencerCore/Public/MVVM/ViewModels/TrackAreaViewModel.h b/Engine/Source/Editor/SequencerCore/Public/MVVM/ViewModels/TrackAreaViewModel.h
--- a/Engine/Source/Editor/SequencerCore/Public/MVVM/ViewModels/TrackAreaViewModel.h
+++ b/Engine/Source/Editor/SequencerCore/Public/MVVM/ViewModels/TrackAreaViewModel.h
@@ -22,6 +22,7 @@ namespace Sequencer
 {
 
 class FEditorViewModel;
+class FTrackAreaViewSpace;
 class ISequencerEditTool;
 struct ITrackAreaHotspot;
 
@@ -39,6 +40,11 @@ public:
 
 public:
 
+	DECLARE_MULTICAST_DELEGATE_TwoParams(FViewSpaceEvent, const FGuid&, TSharedPtr<FTrackAreaViewSpace>)
+
+	FViewSpaceEvent OnViewSpaceAdded;
+	FViewSpaceEvent OnViewSpaceRemoved;
+
 	/** Get the editor view-model provided by the creation context */
 	TSharedPtr<FEditorViewModel> GetEditor() const;
 
@@ -50,6 +56,8 @@ public:
 	/** Gets the current view range of this track area */
 	virtual TRange<double> GetViewRange() const;
 
+	virtual void OnConstruct() override;
+
 	/** Get the current active hotspot */
 	TSharedPtr<ITrackAreaHotspot> GetHotspot() const;
 
@@ -78,8 +86,19 @@ public:
 	void LockEditTool();
 	void UnlockEditTool();
 
+	void AddViewSpace(const FGuid& ID, TSharedPtr<FTrackAreaViewSpace> ViewSpace);
+	void RemoveViewSpace(const FGuid& ID);
+
+	TSharedPtr<FTrackAreaViewSpace> FindViewSpace(const FGuid& InID) const;
+
 protected:
 
+	void UpdateViewSpaces();
+
+protected:
+
+	TMap<FGuid, TSharedPtr<FTrackAreaViewSpace>> ViewSpaces;
+
 	/** The current hotspot that can be set from anywhere to initiate drags */
 	TArray<TSharedPtr<ITrackAreaHotspot>> HotspotStack;
 
diff --git a/Engine/Source/Editor/SequencerCore/Public/MVVM/ViewModels/TrackAreaViewSpace.h b/Engine/Source/Editor/SequencerCore/Public/MVVM/ViewModels/TrackAreaViewSpace.h
new file mode 100644
--- /dev/null
+++ b/Engine/Source/Editor/SequencerCore/Public/MVVM/ViewModels/TrackAreaViewSpace.h
@@ -0,0 +1,28 @@
+// Copyright Epic Games, Inc. All Rights Reserved.
+
+#pragma once
+
+#include "CoreTypes.h"
+#include "MVVM/ViewModels/ViewModel.h"
+#include "TimeToPixel.h"
+
+namespace UE::Sequencer
+{
+
+class SEQUENCERCORE_API FTrackAreaViewSpace
+	: public FViewModel
+	, public INonLinearTimeTransform
+{
+public:
+
+	UE_SEQUENCER_DECLARE_CASTABLE(FTrackAreaViewSpace, FViewModel);
+
+	FTrackAreaViewSpace();
+	virtual ~FTrackAreaViewSpace();
+
+	virtual double SourceToView(double Seconds) const override;
+	virtual double ViewToSource(double Source) const override;
+};
+
+} // namespace UE::Sequencer
+
diff --git a/Engine/Source/Editor/SequencerCore/Public/MVVM/Views/SCompoundTrackLaneView.h b/Engine/Source/Editor/SequencerCore/Public/MVVM/Views/SCompoundTrackLaneView.h
--- a/Engine/Source/Editor/SequencerCore/Public/MVVM/Views/SCompoundTrackLaneView.h
+++ b/Engine/Source/Editor/SequencerCore/Public/MVVM/Views/SCompoundTrackLaneView.h
@@ -8,7 +8,6 @@
 #include "Widgets/SPanel.h"
 #include "Layout/Children.h"
 
-struct FTimeToPixel;
 struct FGeometry;
 
 class FPaintArgs;
@@ -18,9 +17,9 @@ namespace UE::Sequencer
 {
 
 class STrackLane;
+class STrackAreaView;
 class ITrackLaneWidget;
-
-DECLARE_DELEGATE_RetVal_OneParam(FTimeToPixel, FGetTimeToPixel, const FGeometry&);
+class ITrackLaneWidgetSpace;
 
 /**
  * 
@@ -31,15 +30,18 @@ class SEQUENCERCORE_API SCompoundTrackLaneView
 public:
 
 	SLATE_BEGIN_ARGS(SCompoundTrackLaneView){}
-
-		SLATE_ARGUMENT(FGetTimeToPixel, TimeToPixel)
-
 	SLATE_END_ARGS()
 
 	SCompoundTrackLaneView();
 	~SCompoundTrackLaneView();
 
-	void Construct(const FArguments& InArgs);
+	UE_DEPRECATED(5.6, "Please use the TSharedPtr<ITrackLaneWidgetSpace> overload")
+	void Construct(const FArguments& InArgs)
+	{
+		Construct(InArgs, nullptr);
+	}
+
+	void Construct(const FArguments& InArgs, TSharedPtr<ITrackLaneWidgetSpace> InTrackLaneWidgetSpace);
 
 	void AddWeakWidget(TSharedPtr<ITrackLaneWidget> InWidget, TWeakPtr<STrackLane> InOwningLane);
 	void AddStrongWidget(TSharedPtr<ITrackLaneWidget> InWidget, TWeakPtr<STrackLane> InOwningLane);
@@ -72,7 +74,7 @@ private:
 	/** All the widgets in the panel */
 	TPanelChildren<FSlot> Children;
 
-	FGetTimeToPixel TimeToPixelDelegate;
+	TWeakPtr<ITrackLaneWidgetSpace> WeakTrackLaneWidgetSpace;
 };
 
 } // namespace UE::Sequencer
diff --git a/Engine/Source/Editor/SequencerCore/Public/MVVM/Views/STrackAreaLaneView.h b/Engine/Source/Editor/SequencerCore/Public/MVVM/Views/STrackAreaLaneView.h
--- a/Engine/Source/Editor/SequencerCore/Public/MVVM/Views/STrackAreaLaneView.h
+++ b/Engine/Source/Editor/SequencerCore/Public/MVVM/Views/STrackAreaLaneView.h
@@ -26,7 +26,7 @@ public:
 
 	/*~ ITrackLaneWidget */
 	virtual TSharedRef<const SWidget> AsWidget() const;
-	virtual FTrackLaneScreenAlignment GetAlignment(const FTimeToPixel& InTimeToPixel, const FGeometry& InParentGeometry) const;
+	virtual FTrackLaneScreenAlignment GetAlignment(const ITrackLaneWidgetSpace& InScreenSpace, const FGeometry& InParentGeometry) const;
 
 	FTimeToPixel GetRelativeTimeToPixel() const;
 
diff --git a/Engine/Source/Editor/SequencerCore/Public/MVVM/Views/STrackAreaView.h b/Engine/Source/Editor/SequencerCore/Public/MVVM/Views/STrackAreaView.h
--- a/Engine/Source/Editor/SequencerCore/Public/MVVM/Views/STrackAreaView.h
+++ b/Engine/Source/Editor/SequencerCore/Public/MVVM/Views/STrackAreaView.h
@@ -13,6 +13,7 @@
 #include "MVVM/ViewModelPtr.h"
 #include "MVVM/ViewModels/ViewModel.h"
 #include "MVVM/Views/SequencerInputHandlerStack.h"
+#include "MVVM/Extensions/ITrackLaneExtension.h"
 #include "Math/Color.h"
 #include "Math/Vector2D.h"
 #include "Misc/Optional.h"
@@ -42,6 +43,7 @@ namespace Sequencer
 {
 
 class FTrackAreaViewModel;
+class FTrackAreaViewSpace;
 class FViewModel;
 class IOutlinerExtension;
 class SOutlinerView;
@@ -74,6 +76,7 @@ struct FTrackAreaViewLayers
  */
 class SEQUENCERCORE_API STrackAreaView
 	: public SPanel
+	, public ITrackLaneWidgetSpace
 {
 public:
 
@@ -93,7 +96,8 @@ public:
 	
 	void SetVirtualPosition(float InVirtualTop);
 
-	TSharedPtr<FTimeToPixel> GetTimeToPixel() const;
+	TSharedPtr<FTimeToPixel> GetFallbackTimeToPixel() const;
+	TSharedPtr<FTimeToPixel> GetTimeToPixel(const FGuid& InViewSpaceID = FGuid()) const;
 
 public:
 
@@ -145,6 +149,10 @@ public:
 	virtual FReply OnKeyDown(const FGeometry& MyGeometry, const FKeyEvent& InKeyEvent) override;
 	virtual FReply OnKeyUp(const FGeometry& MyGeometry, const FKeyEvent& InKeyEvent) override;
 
+	/*~ Begin ITrackLaneWidgetSpace Interface */
+	virtual FTimeToPixel GetScreenSpace(const FGuid& InViewSpaceID = FGuid()) const override;
+	/*~ End ITrackLaneWidgetSpace Interface */
+
 protected:
 
 	virtual void OnResized(const FVector2D& OldSize, const FVector2D& NewSize){}
@@ -152,8 +160,10 @@ protected:
 	/** Update any hover state required for the track area */
 	virtual void UpdateHoverStates( const FGeometry& MyGeometry, const FPointerEvent& MouseEvent );
 
-protected:
+private:
 
+	void HandleViewSpaceAdded(const FGuid& InID, TSharedPtr<FTrackAreaViewSpace> ViewSpace);
+	void HandleViewSpaceRemoved(const FGuid& InID, TSharedPtr<FTrackAreaViewSpace> ViewSpace);
 
 private:
 
@@ -181,6 +191,7 @@ protected:
 	TWeakViewModelPtr<IOutlinerExtension> WeakDroppedItem;
 
 	mutable TSharedPtr<FTimeToPixel> TimeToPixel;
+	mutable TMap<FGuid, TSharedPtr<FTimeToPixel>> ViewSpaceTimeToPixel;
 
 	float VirtualTop;
 
diff --git a/Engine/Source/Editor/SequencerCore/Public/TimeToPixel.h b/Engine/Source/Editor/SequencerCore/Public/TimeToPixel.h
--- a/Engine/Source/Editor/SequencerCore/Public/TimeToPixel.h
+++ b/Engine/Source/Editor/SequencerCore/Public/TimeToPixel.h
@@ -6,29 +6,33 @@
 #include "Layout/Geometry.h"
 #include "Misc/FrameRate.h"
 
-//TODO: Move to UE::Sequencer namespace
+namespace UE::Sequencer
+{
+
+struct INonLinearTimeTransform
+{
+	virtual ~INonLinearTimeTransform()
+	{
+	}
+
+	virtual double SourceToView(double Seconds) const = 0;
+	virtual double ViewToSource(double Source) const = 0;
+};
+
+
 
 /**
  * Utility for converting time units to slate pixel units and vice versa
  */
-struct FTimeToPixel
+struct FTimeToPixelSpace
 {
 public:
 
-	FTimeToPixel( const FGeometry& AllottedGeometry, const TRange<double>& InLocalViewRange, const FFrameRate& InTickResolution )
-		: FTimeToPixel(AllottedGeometry.GetLocalSize().X, InLocalViewRange, InTickResolution)
-	{}
+	SEQUENCERCORE_API FTimeToPixelSpace( const FGeometry& AllottedGeometry, const TRange<double>& InLocalViewRange, const FFrameRate& InTickResolution );
 
-	FTimeToPixel( float WidthPx, const TRange<double>& InLocalViewRange, const FFrameRate& InTickResolution )
-		: ViewRangeStartSeconds( InLocalViewRange.GetLowerBoundValue() )
-		, TickResolution( InTickResolution )
-	{
-		ensureMsgf(WidthPx >= 0, TEXT("Got negative pixel width!"));
-		const double VisibleWidth = InLocalViewRange.Size<double>();
+	SEQUENCERCORE_API FTimeToPixelSpace( float WidthPx, const TRange<double>& InLocalViewRange, const FFrameRate& InTickResolution );
 
-		const float MaxPixelsPerSecond = 1000.f;
-		PixelsPerSecond = VisibleWidth > 0 ? static_cast<float>(WidthPx / VisibleWidth) : MaxPixelsPerSecond;
-	}
+	SEQUENCERCORE_API FTimeToPixelSpace CreateNonLinear(TSharedPtr<INonLinearTimeTransform> InNonLinearTransform) const;
 
 	/**
 	 * Converts a time to a pixel point relative to the geometry of a widget (passed into the constructor)
@@ -36,10 +40,7 @@ public:
 	 * @param Time	The time to convert
 	 * @return The pixel equivalent of the time
 	 */
-	float SecondsToPixel( double Time ) const
-	{
-		return static_cast<float>((Time - ViewRangeStartSeconds) * PixelsPerSecond);
-	}
+	SEQUENCERCORE_API float SecondsToPixel( double Time ) const;
 
 	/**
 	 * Converts a time delta to a pixel delta
@@ -47,10 +48,7 @@ public:
 	 * @param TimeDelta  The time delta to convert
 	 * @return           The pixel equivalent of the delta time
 	 */
-	float SecondsDeltaToPixel( double TimeDelta ) const
-	{
-		return static_cast<float>(TimeDelta * PixelsPerSecond);
-	}
+	SEQUENCERCORE_API float SecondsDeltaToPixel( double TimeDelta ) const;
 
 	/**
 	 * Converts a pixel value to time 
@@ -58,10 +56,7 @@ public:
 	 * @param PixelX The x value of a pixel coordinate relative to the geometry that was passed into the constructor.
 	 * @return The time where the pixel is located
 	 */
-	double PixelToSeconds( float PixelX ) const
-	{
-		return (PixelX/PixelsPerSecond) + ViewRangeStartSeconds;
-	}
+	SEQUENCERCORE_API double PixelToSeconds( float PixelX ) const;
 
 	/**
 	 * Converts a frame time to a pixel point relative to the geometry of a widget (passed into the constructor)
@@ -69,10 +64,7 @@ public:
 	 * @param Time The time to convert
 	 * @return The pixel equivalent of the frame time
 	 */
-	float FrameToPixel( const FFrameTime& Time ) const
-	{
-		return static_cast<float>((Time / TickResolution - ViewRangeStartSeconds) * PixelsPerSecond);
-	}
+	SEQUENCERCORE_API float FrameToPixel( const FFrameTime& Time ) const;
 
 	/**
 	 * Converts a frame delta value to pixel delta
@@ -80,10 +72,7 @@ public:
 	 * @param TimeDelta   The time delta to convert
 	 * @return            The pixel equivalent of the delta time
 	 */
-	float FrameDeltaToPixel( const FFrameTime& TimeDelta) const
-	{
-		return static_cast<float>((TimeDelta / TickResolution) * PixelsPerSecond);
-	}
+	SEQUENCERCORE_API float FrameDeltaToPixel( const FFrameTime& TimeDelta) const;
 
 	/**
 	 * Converts a pixel value to frame time 
@@ -91,10 +80,7 @@ public:
 	 * @param PixelX The x value of a pixel coordinate relative to the geometry that was passed into the constructor.
 	 * @return The frame time where the pixel is located
 	 */
-	FFrameTime PixelToFrame( float PixelX ) const
-	{
-		return ( PixelX / PixelsPerSecond + ViewRangeStartSeconds ) * TickResolution;
-	}
+	SEQUENCERCORE_API FFrameTime PixelToFrame( float PixelX ) const;
 
 	/**
 	 * Converts a pixel delta value to delta frame time 
@@ -102,10 +88,7 @@ public:
 	 * @param PixelDelta The delta value in pixel space
 	 * @return The equivalent delta frame time
 	 */
-	FFrameTime PixelDeltaToFrame( float PixelDelta ) const
-	{
-		return ( PixelDelta / PixelsPerSecond ) * TickResolution;
-	}
+	SEQUENCERCORE_API FFrameTime PixelDeltaToFrame( float PixelDelta ) const;
 
 	/**
 	 * Converts a pixel delta value to delta seconds time 
@@ -113,37 +96,75 @@ public:
 	 * @param PixelDelta The delta value in pixel space
 	 * @return The equivalent delta time in seconds
 	 */
-	double PixelDeltaToSeconds( float PixelDelta ) const
-	{
-		return ( PixelDelta / PixelsPerSecond );
-	}
+	SEQUENCERCORE_API double PixelDeltaToSeconds( float PixelDelta ) const;
 
 	/**
 	 * Retrieve the tick resolution of the current sequence
 	 */
-	FFrameRate GetTickResolution() const
-	{
-		return TickResolution;
-	}
+	SEQUENCERCORE_API FFrameRate GetTickResolution() const;
 
 	/**
 	 * Make this converter relative to the specified time (ie, such that pixel 0 == FrameAmount)
 	 */
-	FTimeToPixel RelativeTo(const FFrameTime& FrameZero) const
-	{
-		FTimeToPixel Copy = *this;
-		Copy.ViewRangeStartSeconds = FrameZero / TickResolution;
-		return Copy;
-	}
+	SEQUENCERCORE_API FTimeToPixelSpace RelativeTo(const FFrameTime& FrameZero) const;
 
-private:
-	FTimeToPixel(){}
+	TSharedPtr<INonLinearTimeTransform> NonLinearTransform;
+
+protected:
+	FTimeToPixelSpace(){}
 
-	/** time range of the sequencer in seconds */
-	double ViewRangeStartSeconds;
 	/** The tick resolution of the current timeline */
 	FFrameRate TickResolution;
 	/** The number of pixels in the view range */
 	float PixelsPerSecond;
+	float PixelOffset;
+};
+
+
+
+} // namespace UE::Sequencer
+
+
+/**
+ * Utility for converting time units to slate pixel units and vice versa
+ * This class will eventually be deprecated in favor of UE::Sequencer::FTimeToPixelSpace
+ */
+struct FTimeToPixel : UE::Sequencer::FTimeToPixelSpace
+{
+public:
 
+	FTimeToPixel() = delete;
+
+	FTimeToPixel( const FGeometry& AllottedGeometry, const TRange<double>& InLocalViewRange, const FFrameRate& InTickResolution )
+		: UE::Sequencer::FTimeToPixelSpace(AllottedGeometry.GetLocalSize().X, InLocalViewRange, InTickResolution)
+	{
+	}
+
+	FTimeToPixel( float WidthPx, const TRange<double>& InLocalViewRange, const FFrameRate& InTickResolution )
+		: UE::Sequencer::FTimeToPixelSpace(WidthPx, InLocalViewRange, InTickResolution)
+	{
+	}
+
+	FTimeToPixel(const UE::Sequencer::FTimeToPixelSpace& InOther)
+		: UE::Sequencer::FTimeToPixelSpace(InOther)
+	{
+	}
+
+	FTimeToPixel CreateNonLinear(TSharedPtr<UE::Sequencer::INonLinearTimeTransform> InNonLinearTransform) const
+	{
+		FTimeToPixel Copy = *this;
+		Copy.NonLinearTransform = InNonLinearTransform;
+		return Copy;
+	}
+
+	FTimeToPixel& operator=(const UE::Sequencer::FTimeToPixelSpace& InOther)
+	{
+		static_cast<UE::Sequencer::FTimeToPixelSpace&>(*this) = InOther;
+		return *this;
+	}
+
+	/**
+	 * Make this converter relative to the specified time (ie, such that pixel 0 == FrameAmount)
+	 */
+	SEQUENCERCORE_API FTimeToPixel RelativeTo(const FFrameTime& FrameZero) const;
 };
diff --git a/Engine/Source/Runtime/MovieScene/Private/Compilation/MovieSceneCompiledDataManager.cpp b/Engine/Source/Runtime/MovieScene/Private/Compilation/MovieSceneCompiledDataManager.cpp
--- a/Engine/Source/Runtime/MovieScene/Private/Compilation/MovieSceneCompiledDataManager.cpp
+++ b/Engine/Source/Runtime/MovieScene/Private/Compilation/MovieSceneCompiledDataManager.cpp
@@ -11,9 +11,8 @@
 #include "MovieSceneSequence.h"
 #include "Sections/MovieSceneSubSection.h"
 #include "Tracks/MovieSceneSubTrack.h"
-#include "Tracks/MovieSceneTimeWarpTrack.h"
-#include "Channels/MovieSceneTimeWarpChannel.h"
 #include "Decorations/IMovieSceneDecoration.h"
+#include "Decorations/MovieSceneTimeWarpDecoration.h"
 #include "IMovieSceneModule.h"
 #include "MovieSceneTimeHelpers.h"
 #include "MovieSceneTransformTypes.h"
@@ -1643,31 +1642,16 @@ bool UMovieSceneCompiledDataManager::CompileHierarchy(UMovieSceneSequence* Seque
 
 	if (Params.SequenceID == MovieSceneSequenceID::Root)
 	{
-		UMovieScene* MovieScene = Sequence->GetMovieScene();
-
-		// Generate time warp transforms before any sub sequences
-		for (UMovieSceneTrack* Track : MovieScene->GetTracks())
+		UMovieSceneTimeWarpDecoration* TimeWarp = Sequence->GetMovieScene()->FindDecoration<UMovieSceneTimeWarpDecoration>();
+		if (TimeWarp)
 		{
-			UMovieSceneTimeWarpTrack* TimeWarpTrack = Cast<UMovieSceneTimeWarpTrack>(Track);
-			if (!TimeWarpTrack)
-			{
-				continue;
-			}
+			FMovieSceneSequenceTransform TimeWarpTransform = TimeWarp->GenerateTransform();
 
-			TimeWarpTrack->bIsActiveTimeWarp = false;
-
-			if (!bContainsTimeWarp && !TimeWarpTrack->IsEvalDisabled())
+			// Don't do anything for identity transforms
+			if (!TimeWarpTransform.IsIdentity())
 			{
-				FMovieSceneNestedSequenceTransform TimeWarpTransform = TimeWarpTrack->GenerateTransform();
-
-				// Don't do anything for identity transforms
-				if (!TimeWarpTransform.IsIdentity())
-				{
-					InOutHierarchy->SetRootTransform(FMovieSceneSequenceTransform(MoveTemp(TimeWarpTransform)));
-					bContainsTimeWarp = true;
-
-					TimeWarpTrack->bIsActiveTimeWarp = true;
-				}
+				InOutHierarchy->SetRootTransform(FMovieSceneSequenceTransform(MoveTemp(TimeWarpTransform)));
+				bContainsTimeWarp = true;
 			}
 		}
 	}
diff --git a/Engine/Source/Runtime/MovieScene/Private/Decorations/MovieSceneDecorationContainer.cpp b/Engine/Source/Runtime/MovieScene/Private/Decorations/MovieSceneDecorationContainer.cpp
--- a/Engine/Source/Runtime/MovieScene/Private/Decorations/MovieSceneDecorationContainer.cpp
+++ b/Engine/Source/Runtime/MovieScene/Private/Decorations/MovieSceneDecorationContainer.cpp
@@ -46,7 +46,7 @@ UObject* FMovieSceneDecorationContainer::GetOrCreateDecoration(const TSubclassOf
 	UObject* Found = FindDecoration(InClass);
 	if (!Found)
 	{
-		Found = NewObject<UObject>(Outer, InClass);
+		Found = NewObject<UObject>(Outer, InClass, NAME_None, RF_Transactional);
 		Decorations.Add(Found);
 		Event(Found);
 	}
@@ -73,6 +73,14 @@ TArrayView<const TObjectPtr<UObject>> FMovieSceneDecorationContainer::GetDecorat
 	return Decorations;
 }
 
+void UMovieSceneDecorationContainerObject::Serialize(FArchive& Ar)
+{
+	Super::Serialize(Ar);
+
+	// Remove null decorations for safety
+	Decorations.RemoveNulls();
+}
+
 void UMovieSceneDecorationContainerObject::AddDecoration(UObject* InDecoration)
 {
 	Decorations.AddDecoration(InDecoration, this, [this](UObject* Decoration){
diff --git a/Engine/Source/Runtime/MovieScene/Private/Decorations/MovieSceneScalingAnchors.cpp b/Engine/Source/Runtime/MovieScene/Private/Decorations/MovieSceneScalingAnchors.cpp
new file mode 100644
--- /dev/null
+++ b/Engine/Source/Runtime/MovieScene/Private/Decorations/MovieSceneScalingAnchors.cpp
@@ -0,0 +1,430 @@
+// Copyright Epic Games, Inc. All Rights Reserved.
+
+#include "Decorations/MovieSceneScalingAnchors.h"
+#include "Decorations/MovieSceneTimeWarpDecoration.h"
+#include "Variants/MovieSceneTimeWarpVariant.h"
+#include "Channels/MovieSceneChannelProxy.h"
+#include "Channels/MovieSceneInterpolation.h"
+#include "MovieSceneTransformTypes.h"
+#include "MovieSceneTimeHelpers.h"
+#include "MovieScene.h"
+#include "Evaluation/MovieSceneSequenceTransform.h"
+#include "EntitySystem/BuiltInComponentTypes.h"
+#include "EntitySystem/MovieSceneEntitySystemLinker.h"
+
+#if 0
+uint32 UMovieSceneScalingAnchors::DefaultTimeWarpEntityID = MAX_uint32;
+#endif
+
+UMovieSceneScalingAnchors::UMovieSceneScalingAnchors()
+{
+}
+
+void UMovieSceneScalingAnchors::Serialize(FArchive& Ar)
+{
+	Super::Serialize(Ar);
+
+	if (Ar.IsLoading())
+	{
+		// Remove null drivers for safety
+		ScalingDrivers.Remove(nullptr);
+	}
+}
+
+void UMovieSceneScalingAnchors::OnDecorationAdded(UMovieScene* MovieScene)
+{
+	UMovieSceneTimeWarpDecoration* TimeWarp = MovieScene->GetOrCreateDecoration<UMovieSceneTimeWarpDecoration>();
+	TimeWarp->AddTimeWarpSource(this);
+
+#if 0
+	// Anchors provide a default timewarp to all entities in the local sequence
+	UMovieSceneTimeWarpRegistryDecoration* TimewarpRegistry = MovieScene->GetOrCreateDecoration<UMovieSceneTimeWarpRegistryDecoration>();
+	TimewarpRegistry->SetDefaultTimeWarp(this);
+#endif
+}
+
+void UMovieSceneScalingAnchors::OnDecorationRemoved()
+{
+	if (UMovieSceneTimeWarpDecoration* TimeWarp = GetTypedOuter<UMovieScene>()->FindDecoration<UMovieSceneTimeWarpDecoration>())
+	{
+		TimeWarp->RemoveTimeWarpSource(this);
+	}
+#if 0
+	if (UMovieSceneTimeWarpRegistryDecoration* TimeWarp = GetTypedOuter<UMovieScene>()->FindDecoration<UMovieSceneTimeWarpRegistryDecoration>())
+	{
+		TimeWarp->SetDefaultTimeWarp(nullptr);
+	}
+#endif
+}
+
+void UMovieSceneScalingAnchors::ResetScaling()
+{
+	PlayRate.Reset();
+	bUpToDate = false;
+	bPlayRateCurveIsUpToDate = false;
+}
+
+FMovieSceneNestedSequenceTransform UMovieSceneScalingAnchors::GenerateTimeWarpTransform()
+{
+	return FMovieSceneNestedSequenceTransform(0, FMovieSceneTimeWarpVariant(this));
+}
+
+bool UMovieSceneScalingAnchors::IsTimeWarpActive() const
+{
+	return true;
+}
+
+void UMovieSceneScalingAnchors::SetIsTimeWarpActive(bool bInActive)
+{
+	// Unimplemented - always comes first
+}
+
+int32 UMovieSceneScalingAnchors::GetTimeWarpSortOrder() const
+{
+	return MIN_int32;
+}
+
+void UMovieSceneScalingAnchors::AddScalingDriver(TScriptInterface<IMovieSceneScalingDriver> InDriver)
+{
+	ScalingDrivers.AddUnique(InDriver);
+	ResetScaling();
+}
+
+void UMovieSceneScalingAnchors::RemoveScalingDriver(TScriptInterface<IMovieSceneScalingDriver> InDriver)
+{
+	ScalingDrivers.Remove(InDriver);
+	ResetScaling();
+}
+
+FMovieSceneAnchorsScalingGroup& UMovieSceneScalingAnchors::GetOrCreateScalingGroup(const FGuid& Guid)
+{
+	return ScalingGroups.FindOrAdd(Guid);
+}
+
+FMovieSceneAnchorsScalingGroup* UMovieSceneScalingAnchors::FindScalingGroup(const FGuid& Guid)
+{
+	return ScalingGroups.Find(Guid);
+}
+
+const TMap<FGuid, FMovieSceneScalingAnchor>& UMovieSceneScalingAnchors::GetInitialAnchors() const
+{
+	return InitialAnchors;
+}
+
+const TMap<FGuid, FMovieSceneScalingAnchor>& UMovieSceneScalingAnchors::GetCurrentAnchors() const
+{
+	return CurrentAnchors;
+}
+
+void UMovieSceneScalingAnchors::RemoveScalingGroup(const FGuid& Guid)
+{
+	ScalingGroups.Remove(Guid);
+}
+
+const TMap<FGuid, FMovieSceneAnchorsScalingGroup>& UMovieSceneScalingAnchors::GetScalingGroups() const
+{
+	return ScalingGroups;
+}
+
+void UMovieSceneScalingAnchors::OnPreDecorationCompiled()
+{
+	InitialAnchors.Reset();
+
+	for (TScriptInterface<IMovieSceneScalingDriver> Source : ScalingDrivers)
+	{
+		Source->PopulateInitialAnchors(InitialAnchors);
+	}
+
+	ResetScaling();
+}
+
+EMovieSceneChannelProxyType UMovieSceneScalingAnchors::PopulateChannelProxy(FMovieSceneChannelProxyData& OutProxyData, EAllowTopLevelChannels AllowTopLevel)
+{
+#if WITH_EDITOR
+	UMovieScene* MovieScene = GetTypedOuter<UMovieScene>();
+
+	FMovieSceneChannelMetaData ChannelMetaData;
+	ChannelMetaData.Name = "Anchors";
+	ChannelMetaData.bCanCollapseToTrack = (AllowTopLevel == EAllowTopLevelChannels::Yes);
+	ChannelMetaData.DisplayText = NSLOCTEXT("MovieSceneScalingAnchors", "Anchors_Label", "Anchors");
+	ChannelMetaData.WeakOwningObject = this;
+	ChannelMetaData.bRelativeToSection = true;
+
+	OutProxyData.Add(PlayRate, ChannelMetaData);
+
+#else
+	OutProxyData.Add(PlayRate);
+#endif
+
+	return EMovieSceneChannelProxyType::Static;
+}
+
+bool UMovieSceneScalingAnchors::DeleteChannel(FMovieSceneTimeWarpVariant& OutVariant, FName ChannelName)
+{
+	if (ChannelName == "PlayRate")
+	{
+		OutVariant.Set(1.0);
+		return true;
+	}
+	return false;
+}
+
+void UMovieSceneScalingAnchors::UpdateCurve(UMovieScenePlayRateCurve* Curve) const
+{
+	if (bPlayRateCurveIsUpToDate)
+	{
+		return;
+	}
+
+	bPlayRateCurveIsUpToDate = true;
+
+	Curve->bUpToDate = false;
+	Curve->PlayRate.Reset();
+
+	FFrameNumber PlaybackStart = UE::MovieScene::DiscreteInclusiveLower(GetTypedOuter<UMovieScene>()->GetPlaybackRange());
+
+	Curve->bManualPlaybackStart = true;
+	Curve->PlaybackStartFrame = 0;//PlaybackStart;
+
+	CurrentAnchors.Reset();
+	// Remove null drivers for safety
+	for (TScriptInterface<IMovieSceneScalingDriver> Source : ScalingDrivers)
+	{
+		Source->PopulateAnchors(CurrentAnchors);
+	}
+
+	if (CurrentAnchors.Num() == 0)
+	{
+		Curve->PlayRate.SetDefault(1.0);
+		return;
+	}
+
+	struct FSort
+	{
+		bool operator()(const TPair<FGuid, FMovieSceneScalingAnchor>& A, const TPair<FGuid, FMovieSceneScalingAnchor>& B) const
+		{
+			return A.Get<1>().Position < B.Get<1>().Position;
+		}
+	};
+
+	// Sort all the keys
+	struct FScaledAnchor
+	{
+		FGuid ID;
+		FMovieSceneScalingAnchor InitialAnchor;
+		FMovieSceneScalingAnchor CurrentAnchor;
+		FMovieSceneScalingAnchor ScaledAnchor;
+	};
+	TArray<FScaledAnchor> SortedAnchors;
+	for (TPair<FGuid, FMovieSceneScalingAnchor> Pair : InitialAnchors)
+	{
+		if (FMovieSceneScalingAnchor* NewAnchor = CurrentAnchors.Find(Pair.Key))
+		{
+			// Insert, allowing for duplicates at the same time
+			int32 Index = Algo::LowerBoundBy(SortedAnchors, Pair.Value.Position, [](FScaledAnchor In) { return In.InitialAnchor.Position; });
+			SortedAnchors.Insert(FScaledAnchor{ Pair.Key, Pair.Value, *NewAnchor, Pair.Value }, Index);
+		}
+	}
+
+	// The time to add the next play ate key at
+	FFrameNumber KeyTime = 0;
+	// A linear offset to apply to any non-overlapping anchors to accomodate for earlier scaled anchors
+	FFrameNumber CumulativeScaleOffset = 0;
+	// The scaled position of the last scale point encountered
+	TOptional<FFrameNumber> LastAnchorPosition;
+	// The initial position of the last scale point encountered
+	TOptional<FFrameNumber> LastAnchorInitialPosition;
+
+	// Place a key just before the first anchor to initialize the play rate to one
+	Curve->PlayRate.AddConstantKey(SortedAnchors[0].ScaledAnchor.Position - 100, 1.0);
+
+
+	TArray<int32> OverlappingIndex;
+
+	// Go through all scaling anchors by their initial order and compute the scaling,
+	//      taking account of overlapping anchors
+	for (int32 Index = 0; Index < SortedAnchors.Num(); ++Index)
+	{
+		FScaledAnchor& Entry = SortedAnchors[Index];
+
+		// Any entry we encounter in this outer loop is guaranteed to not be overlapping any other
+		//    since those will have been handled by the inner loop
+		Entry.ScaledAnchor.Position = CumulativeScaleOffset + Entry.CurrentAnchor.Position;
+		Entry.ScaledAnchor.Duration = Entry.CurrentAnchor.Duration;
+
+		// Move on the scale offset by how much this anchor moved (ie, the empty pace between the last point and this has scaled)
+		CumulativeScaleOffset += (Entry.CurrentAnchor.Position - Entry.InitialAnchor.Position);
+
+		// If we have a previous anchor position, add play rate for the empty spae up until this anchor
+		if (LastAnchorInitialPosition.IsSet())
+		{
+			const double PlayRateToThisAnchor = double(Entry.ScaledAnchor.Position.Value - LastAnchorPosition->Value) / double(Entry.InitialAnchor.Position.Value - LastAnchorInitialPosition->Value);
+			Curve->PlayRate.AddConstantKey(KeyTime, PlayRateToThisAnchor);
+		}
+
+		// Keep track of this anchor for the next iteration
+		KeyTime                   = Entry.ScaledAnchor.Position;
+		LastAnchorPosition        = Entry.ScaledAnchor.Position;
+		LastAnchorInitialPosition = Entry.InitialAnchor.Position;
+
+		// Store this anchor's final position for future external reference
+		CurrentAnchors.Add(Entry.ID, Entry.ScaledAnchor);
+
+		if (Entry.InitialAnchor.Duration <= 0.0 || Entry.CurrentAnchor.Duration <= 0.0)
+		{
+			// No duration on this anchor - just move on
+			continue;
+		}
+
+
+		// This anchor has duration: we need to scale anything that overlaps it proportionally
+		FFrameNumber OverlapBoundary        = Entry.ScaledAnchor.Position + Entry.ScaledAnchor.Duration;
+		FFrameNumber InitialOverlapBoundary = Entry.InitialAnchor.Position + Entry.InitialAnchor.Duration;
+
+		FFrameNumber MaxTime = OverlapBoundary;
+		FFrameNumber InitialMaxTime = InitialOverlapBoundary;
+
+
+		// Recursively process overlapping anchors until we find empty space
+		OverlappingIndex.Reset(1);
+		OverlappingIndex.Add(Index);
+
+		while (Index < SortedAnchors.Num()-1 && OverlappingIndex.Num() != 0)
+		{
+			const FScaledAnchor& LastOverlap = SortedAnchors[OverlappingIndex.Last()];
+
+			// Position everything proportionally along this anchor's duration that starts within the range
+			//     making sure not to stretch durations
+			for (int32 NextIndex = Index+1; NextIndex < SortedAnchors.Num(); ++NextIndex, ++Index)
+			{
+				FScaledAnchor& NextEntry = SortedAnchors[NextIndex];
+
+				if (NextEntry.InitialAnchor.Position >= LastOverlap.InitialAnchor.Position + LastOverlap.InitialAnchor.Duration)
+				{
+					// This anchor does not fall within the duration of the one we're processing
+					//     so pop the last overlap and start again from the next
+					OverlappingIndex.Pop();
+					break;
+				}
+
+				const double DurationScale = double(LastOverlap.ScaledAnchor.Duration) / double(LastOverlap.InitialAnchor.Duration);
+
+				// Reposition this anchor to be scaled proportionally within the most recent anchor's range
+				double NextAnchorPosition = LastOverlap.ScaledAnchor.Position.Value + (NextEntry.CurrentAnchor.Position.Value - LastOverlap.CurrentAnchor.Position.Value) * DurationScale;
+				NextEntry.ScaledAnchor.Position = FMath::RoundToInt32(NextAnchorPosition);
+				NextEntry.ScaledAnchor.Duration = NextEntry.CurrentAnchor.Duration;
+
+				// Store this anchor's final position for future external reference
+				CurrentAnchors.Add(NextEntry.ID, NextEntry.ScaledAnchor);
+
+				// If this anchor has a range that overflows the current maximum, we need to add a scale between
+				//    up until this boundary, and move on to start from the current boundary
+				if (NextEntry.ScaledAnchor.Position + NextEntry.ScaledAnchor.Duration > OverlapBoundary)
+				{
+					// Add a play rate key for the current boundary and move the boundary forward
+					const double OverlapScale = double(InitialOverlapBoundary.Value - LastAnchorInitialPosition->Value) / double(OverlapBoundary.Value - LastAnchorPosition->Value);
+					Curve->PlayRate.AddConstantKey(KeyTime, OverlapScale);
+
+					// The next play rate range should start from the current boundary
+					KeyTime = OverlapBoundary;
+					LastAnchorPosition = OverlapBoundary;
+					LastAnchorInitialPosition = InitialOverlapBoundary;
+
+					// The new boundary point will be the end of this anchor
+					OverlapBoundary = NextEntry.ScaledAnchor.Position + NextEntry.ScaledAnchor.Duration;
+					InitialOverlapBoundary = NextEntry.InitialAnchor.Position + NextEntry.InitialAnchor.Duration;
+
+					// Keep track of the maximums for the CumulativeScaleOffset
+					MaxTime = OverlapBoundary;
+					if (InitialOverlapBoundary > InitialMaxTime)
+					{
+						InitialMaxTime = InitialOverlapBoundary;
+					}
+				}
+
+				// If this anchor has duration we re-run this whole loop with the new anchor
+				if (NextEntry.CurrentAnchor.Duration > 0)
+				{
+					OverlappingIndex.Add(NextIndex);
+					++Index;
+					break;
+				}
+			}
+		}
+
+		// Accumulate the scale offset for this whole group
+		CumulativeScaleOffset += (MaxTime - Entry.ScaledAnchor.Position) - (InitialMaxTime - Entry.InitialAnchor.Position);
+
+		// Process the remaining space that is left from these overlaps (or the entire range if there are no overlaps)
+		const FFrameNumber InitialRemainingRange = InitialOverlapBoundary - LastAnchorInitialPosition.GetValue();
+		const FFrameNumber RemainingRange        = OverlapBoundary - LastAnchorPosition.GetValue();
+		if (InitialRemainingRange != 0 && RemainingRange != 0)
+		{
+			// Add a play rate key for the current boundary and move the boundary forward
+			const double OverlapScale = double(InitialRemainingRange.Value) / double(RemainingRange.Value);
+			Curve->PlayRate.AddConstantKey(KeyTime, OverlapScale);
+
+			// Move on from this point
+			KeyTime = OverlapBoundary;
+			LastAnchorPosition = OverlapBoundary;
+			LastAnchorInitialPosition = InitialOverlapBoundary;
+		}
+	}
+
+	// Add final play rate to cause subsequent regions to all play-back at 1.0x
+	Curve->PlayRate.AddConstantKey(*LastAnchorPosition, 1.0);
+}
+
+UMovieScenePlayRateCurve* UMovieSceneScalingAnchors::Initialize(TSharedPtr<UE::MovieScene::FSharedPlaybackState> SharedPlaybackState) const
+{
+	UMovieScenePlayRateCurve* ContextPlayRate = NewObject<UMovieScenePlayRateCurve>(SharedPlaybackState->GetPlaybackContext());
+	UpdateCurve(ContextPlayRate);
+
+#if 0
+	// // Update our desired duration
+	UMovieScene* MovieScene = GetTypedOuter<UMovieScene>();
+	TRange<FFrameNumber> PlaybackRange = MovieScene->GetPlaybackRange();
+
+	FFrameNumber StartTime = PlaybackRange.GetLowerBoundValue();
+	FFrameNumber EndTime   = PlaybackRange.GetLowerBoundValue();
+
+	TOptional<FFrameTime> NewEndTime = Curve->GetTimeWarpCurve().InverseEvaluate(EndTime.Value, EndTime, UE::MovieScene::EInverseEvaluateFlags::AnyDirection);
+
+	NewDuration = NewEndTime.Get(FFrameTime(EndTime)) - StartTime;
+#endif
+	return ContextPlayRate;
+}
+
+void UMovieSceneScalingAnchors::UpdateFromSource() const
+{
+	UpdateCurve(const_cast<UMovieSceneScalingAnchors*>(this));
+}
+
+TRange<FFrameTime> UMovieSceneScalingAnchors::ComputeTraversedHull(const TRange<FFrameTime>& Range) const
+{
+	UpdateFromSource();
+	return Super::ComputeTraversedHull(Range);
+}
+
+FFrameTime UMovieSceneScalingAnchors::RemapTime(FFrameTime In) const
+{
+	UpdateFromSource();
+	return Super::RemapTime(In);
+}
+
+TOptional<FFrameTime> UMovieSceneScalingAnchors::InverseRemapTimeCycled(FFrameTime InValue, FFrameTime InTimeHint, const UE::MovieScene::FInverseTransformTimeParams& Params) const
+{
+	UpdateFromSource();
+	return Super::InverseRemapTimeCycled(InValue, InTimeHint, Params);
+}
+
+bool UMovieSceneScalingAnchors::InverseRemapTimeWithinRange(FFrameTime InTime, FFrameTime RangeStart, FFrameTime RangeEnd, const TFunctionRef<bool(FFrameTime)>& VisitorCallback) const
+{
+	UpdateFromSource();
+	return Super::InverseRemapTimeWithinRange(InTime, RangeStart, RangeEnd, VisitorCallback);
+}
+
+void UMovieSceneScalingAnchors::ScaleBy(double UnwarpedScaleFactor)
+{
+	// Cannot Scale Anchors
+}
\ No newline at end of file
diff --git a/Engine/Source/Runtime/MovieScene/Private/Decorations/MovieSceneSectionAnchorsDecoration.cpp b/Engine/Source/Runtime/MovieScene/Private/Decorations/MovieSceneSectionAnchorsDecoration.cpp
new file mode 100644
--- /dev/null
+++ b/Engine/Source/Runtime/MovieScene/Private/Decorations/MovieSceneSectionAnchorsDecoration.cpp
@@ -0,0 +1,51 @@
+// Copyright Epic Games, Inc. All Rights Reserved.
+
+#include "Decorations/MovieSceneSectionAnchorsDecoration.h"
+#include "MovieScene.h"
+#include "Decorations/MovieSceneScalingAnchors.h"
+
+
+void UMovieSceneSectionAnchorsDecoration::OnReconstruct(UMovieScene* MovieScene)
+{
+	MovieScene->Modify();
+
+	UMovieSceneScalingAnchors* Anchors = MovieScene->GetOrCreateDecoration<UMovieSceneScalingAnchors>();
+
+	Anchors->Modify();
+	Anchors->AddScalingDriver(this);
+
+	StartAnchor = FGuid::NewGuid();
+}
+
+void UMovieSceneSectionAnchorsDecoration::OnDestroy(UMovieScene* MovieScene)
+{
+	if (UMovieSceneScalingAnchors* Anchors = MovieScene->FindDecoration<UMovieSceneScalingAnchors>())
+	{
+		Anchors->Modify();
+		Anchors->RemoveScalingDriver(this);
+	}
+}
+
+void UMovieSceneSectionAnchorsDecoration::PopulateInitialAnchors(TMap<FGuid, FMovieSceneScalingAnchor>& OutAnchors)
+{
+	IMovieSceneScalingDriver* SectionImpl = Cast<IMovieSceneScalingDriver>(GetTypedOuter<UMovieSceneSection>());
+	if (SectionImpl)
+	{
+		SectionImpl->PopulateInitialAnchors(OutAnchors);
+	}
+}
+
+void UMovieSceneSectionAnchorsDecoration::PopulateAnchors(TMap<FGuid, FMovieSceneScalingAnchor>& OutAnchors)
+{
+	IMovieSceneScalingDriver* SectionImpl = Cast<IMovieSceneScalingDriver>(GetTypedOuter<UMovieSceneSection>());
+	if (SectionImpl)
+	{
+		SectionImpl->PopulateAnchors(OutAnchors);
+	}
+}
+
+void UMovieSceneSectionAnchorsDecoration::PostEditImport()
+{
+	// Do not allow duplicate anchor GUIDs
+	StartAnchor = FGuid::NewGuid();
+}
\ No newline at end of file
diff --git a/Engine/Source/Runtime/MovieScene/Private/Decorations/MovieSceneTimeWarpDecoration.cpp b/Engine/Source/Runtime/MovieScene/Private/Decorations/MovieSceneTimeWarpDecoration.cpp
new file mode 100644
--- /dev/null
+++ b/Engine/Source/Runtime/MovieScene/Private/Decorations/MovieSceneTimeWarpDecoration.cpp
@@ -0,0 +1,59 @@
+// Copyright Epic Games, Inc. All Rights Reserved.
+
+#include "Decorations/MovieSceneTimeWarpDecoration.h"
+#include "Evaluation/MovieSceneSequenceTransform.h"
+
+
+FMovieSceneSequenceTransform UMovieSceneTimeWarpDecoration::GenerateTransform() const
+{
+	FMovieSceneSequenceTransform CombinedTransform;
+
+	Algo::SortBy(Sources, &IMovieSceneTimeWarpSource::GetTimeWarpSortOrder);
+
+	for (TScriptInterface<IMovieSceneTimeWarpSource> Source : Sources)
+	{
+		FMovieSceneNestedSequenceTransform TimeWarpTransform = Source->GenerateTimeWarpTransform();
+
+		// Don't do anything for identity transforms
+		if (!TimeWarpTransform.IsIdentity())
+		{
+			CombinedTransform.Add(TimeWarpTransform);
+		}
+	}
+
+
+	return CombinedTransform;
+}
+
+void UMovieSceneTimeWarpDecoration::OnCompiled()
+{
+
+}
+
+void UMovieSceneTimeWarpDecoration::AddTimeWarpSource(TScriptInterface<IMovieSceneTimeWarpSource> InSource)
+{
+	Sources.AddUnique(InSource);
+}
+
+void UMovieSceneTimeWarpDecoration::RemoveTimeWarpSource(TScriptInterface<IMovieSceneTimeWarpSource> InSource)
+{
+	Sources.Remove(InSource);
+}
+
+void UMovieSceneTimeWarpDecoration::ImportEntityImpl(UMovieSceneEntitySystemLinker* EntityLinker, const FEntityImportParams& Params, FImportedEntity* OutImportedEntity)
+{
+
+}
+bool UMovieSceneTimeWarpDecoration::PopulateEvaluationFieldImpl(const TRange<FFrameNumber>& EffectiveRange, const FMovieSceneEvaluationFieldEntityMetaData& InMetaData, FMovieSceneEntityComponentFieldBuilder* OutFieldBuilder)
+{
+	return true;
+}
+
+void UMovieSceneTimeWarpDecoration::Serialize(FArchive& Ar)
+{
+	Super::Serialize(Ar);
+	if (Ar.IsLoading())
+	{
+		Sources.Remove(nullptr);
+	}
+}
\ No newline at end of file
diff --git a/Engine/Source/Runtime/MovieScene/Private/Evaluation/MovieSceneSequenceTransform.cpp b/Engine/Source/Runtime/MovieScene/Private/Evaluation/MovieSceneSequenceTransform.cpp
--- a/Engine/Source/Runtime/MovieScene/Private/Evaluation/MovieSceneSequenceTransform.cpp
+++ b/Engine/Source/Runtime/MovieScene/Private/Evaluation/MovieSceneSequenceTransform.cpp
@@ -863,6 +863,23 @@ void FMovieSceneSequenceTransform::Add(FMovieSceneTimeTransform InTransform)
 	}
 }
 
+void FMovieSceneSequenceTransform::Add(FMovieSceneNestedSequenceTransform InTransform)
+{
+	if (InTransform.IsIdentity())
+	{
+		return;
+	}
+
+	if (IsLinear() && InTransform.IsLinear())
+	{
+		LinearTransform = InTransform.AsLinear() * LinearTransform;
+	}
+	else
+	{
+		NestedTransforms.Emplace(MoveTemp(InTransform));
+	}
+}
+
 void FMovieSceneSequenceTransform::Add(FFrameTime InOffset, FMovieSceneTimeWarpVariant&& InTimeWarp)
 {
 	if (InTimeWarp.GetType() == EMovieSceneTimeWarpType::FixedPlayRate)
diff --git a/Engine/Source/Runtime/MovieScene/Private/MovieSceneSection.cpp b/Engine/Source/Runtime/MovieScene/Private/MovieSceneSection.cpp
--- a/Engine/Source/Runtime/MovieScene/Private/MovieSceneSection.cpp
+++ b/Engine/Source/Runtime/MovieScene/Private/MovieSceneSection.cpp
@@ -7,6 +7,7 @@
 #include "MovieSceneCommonHelpers.h"
 #include "MovieSceneTimeHelpers.h"
 #include "Decorations/IMovieSceneSectionDecoration.h"
+#include "Decorations/IMovieSceneLifetimeDecoration.h"
 #include "Evaluation/MovieSceneEvalTemplate.h"
 #include "EntitySystem/MovieSceneEntityManager.h"
 #include "Generators/MovieSceneEasingCurves.h"
@@ -127,8 +128,47 @@ void UMovieSceneSection::Serialize(FArchive& Ar)
 	}
 }
 
+void UMovieSceneSection::OnAddedToTrack(UMovieSceneTrack* Track)
+{
+	OnAddedToTrackImpl(Track);
+
+	if (UMovieScene* MovieScene = Track->GetTypedOuter<UMovieScene>())
+	{
+		for (UObject* Decoration : GetDecorations())
+		{
+			if (IMovieSceneLifetimeDecoration* Lifetime = Cast<IMovieSceneLifetimeDecoration>(Decoration))
+			{
+				Lifetime->OnReconstruct(MovieScene);
+			}
+		}
+	}
+}
+
+void UMovieSceneSection::OnRemovedFromTrack()
+{
+	OnRemovedFromTrackImpl();
+
+	if (UMovieScene* MovieScene = GetTypedOuter<UMovieScene>())
+	{
+		for (UObject* Decoration : GetDecorations())
+		{
+			if (IMovieSceneLifetimeDecoration* Lifetime = Cast<IMovieSceneLifetimeDecoration>(Decoration))
+			{
+				Lifetime->OnDestroy(MovieScene);
+			}
+		}
+	}
+}
+
 void UMovieSceneSection::OnDecorationAdded(UObject* NewDecoration)
 {
+	UMovieScene* MovieScene = GetTypedOuter<UMovieScene>();
+	IMovieSceneLifetimeDecoration* Lifetime = Cast<IMovieSceneLifetimeDecoration>(NewDecoration);
+
+	if (MovieScene && Lifetime)
+	{
+		Lifetime->OnReconstruct(MovieScene);
+	}
 	if (IMovieSceneSectionDecoration* DecorationInterface = Cast<IMovieSceneSectionDecoration>(NewDecoration))
 	{
 		DecorationInterface->OnDecorationAdded(this);
@@ -143,6 +183,14 @@ void UMovieSceneSection::OnDecorationRemoved(UObject* Decoration)
 	{
 		DecorationInterface->OnDecorationRemoved();
 	}
+
+	UMovieScene* MovieScene = GetTypedOuter<UMovieScene>();
+	IMovieSceneLifetimeDecoration* Lifetime = Cast<IMovieSceneLifetimeDecoration>(Decoration);
+
+	if (MovieScene && Lifetime)
+	{
+		Lifetime->OnDestroy(MovieScene);
+	}
 	EventHandlers.Trigger(&UE::MovieScene::ISectionEventHandler::OnDecorationRemoved, Decoration);
 }
 
@@ -172,6 +220,19 @@ void UMovieSceneSection::PostRename(UObject* OldOuter, const FName OldName)
 		ChannelProxy = nullptr;
 	}
 
+	UMovieScene* NewMovieScene = GetTypedOuter<UMovieScene>();
+	if (NewMovieScene && NewMovieScene != OldOuter->GetTypedOuter<UMovieScene>())
+	{
+		for (UObject* Decoration : GetDecorations())
+		{
+			// If this is being renamed into a movie scene structure, call reconstruct on its decorations
+			if (IMovieSceneLifetimeDecoration* DecorationInterface = Cast<IMovieSceneLifetimeDecoration>(Decoration))
+			{
+				DecorationInterface->OnReconstruct(NewMovieScene);
+			}
+		}
+	}
+
 	Super::PostRename(OldOuter, OldName);
 }
 
@@ -1037,6 +1098,11 @@ void UMovieSceneSection::PostPaste()
 	{
 		DefaultEaseOut->ClearFlags(RF_Transient);
 	}
+
+	for (UObject* Decoration : GetDecorations())
+	{
+		Decoration->ClearFlags(RF_Transient);
+	}
 }
 
 ECookOptimizationFlags UMovieSceneSection::GetCookOptimizationFlags() const
diff --git a/Engine/Source/Runtime/MovieScene/Private/MovieSceneTrack.cpp b/Engine/Source/Runtime/MovieScene/Private/MovieSceneTrack.cpp
--- a/Engine/Source/Runtime/MovieScene/Private/MovieSceneTrack.cpp
+++ b/Engine/Source/Runtime/MovieScene/Private/MovieSceneTrack.cpp
@@ -13,6 +13,8 @@
 #include "Compilation/IMovieSceneTemplateGenerator.h"
 
 #include "Decorations/IMovieSceneTrackDecoration.h"
+#include "Decorations/IMovieSceneSectionDecoration.h"
+#include "Decorations/IMovieSceneLifetimeDecoration.h"
 
 #include "Evaluation/MovieSceneEvaluationTrack.h"
 #include "Evaluation/MovieSceneEvaluationTemplate.h"
@@ -93,14 +95,80 @@ void UMovieSceneTrack::PostLoad()
 	}
 }
 
+void UMovieSceneTrack::AddSection(FSectionParameter Section)
+{
+	CallAddSection(*Section.Section);
+	Section.Section->OnAddedToTrack(this);
+}
+
+void UMovieSceneTrack::RemoveSection(FSectionParameter Section)
+{
+	CallRemoveSection(*Section.Section);
+	Section.Section->OnRemovedFromTrack();
+}
+
+void UMovieSceneTrack::RemoveSectionAt(FSectionIndexParameter SectionIndex)
+{
+	const TArray<UMovieSceneSection*>& Sections = GetAllSections();
+	if (Sections.IsValidIndex(SectionIndex.SectionIndex))
+	{
+		UMovieSceneSection* Section = Sections[SectionIndex.SectionIndex];
+		CallRemoveSectionAt(SectionIndex.SectionIndex);
+		Section->OnRemovedFromTrack();
+	}
+}
+
 void UMovieSceneTrack::OnAddedToMovieScene(UMovieScene* MovieScene)
 {
 	OnAddedToMovieSceneImpl(MovieScene);
+
+	// If this is being renamed into a movie scene structure, simulate re-addition of the
+	//    decorations to ensure any external references within the MovieScene are updated.
+	for (UObject* Decoration : GetDecorations())
+	{
+		if (IMovieSceneLifetimeDecoration* Lifetime = Cast<IMovieSceneLifetimeDecoration>(Decoration))
+		{
+			Lifetime->OnReconstruct(MovieScene);
+		}
+	}
+
+	for (UMovieSceneSection* Section : GetAllSections())
+	{
+		for (UObject* Decoration : Section->GetDecorations())
+		{
+			if (IMovieSceneLifetimeDecoration* Lifetime = Cast<IMovieSceneLifetimeDecoration>(Decoration))
+			{
+				Lifetime->OnReconstruct(MovieScene);
+			}
+		}
+	}
 }
 
 void UMovieSceneTrack::OnRemovedFromMovieScene()
 {
 	OnRemovedFromMovieSceneImpl();
+
+	if (UMovieScene* MovieScene = GetTypedOuter<UMovieScene>())
+	{
+		for (UObject* Decoration : GetDecorations())
+		{
+			if (IMovieSceneLifetimeDecoration* Lifetime = Cast<IMovieSceneLifetimeDecoration>(Decoration))
+			{
+				Lifetime->OnDestroy(MovieScene);
+			}
+		}
+
+		for (UMovieSceneSection* Section : GetAllSections())
+		{
+			for (UObject* Decoration : Section->GetDecorations())
+			{
+				if (IMovieSceneLifetimeDecoration* Lifetime = Cast<IMovieSceneLifetimeDecoration>(Decoration))
+				{
+					Lifetime->OnDestroy(MovieScene);
+				}
+			}
+		}
+	}
 }
 
 void UMovieSceneTrack::OnDecorationAdded(UObject* NewDecoration)
diff --git a/Engine/Source/Runtime/MovieScene/Private/Tracks/MovieSceneTimeWarpTrack.cpp b/Engine/Source/Runtime/MovieScene/Private/Tracks/MovieSceneTimeWarpTrack.cpp
--- a/Engine/Source/Runtime/MovieScene/Private/Tracks/MovieSceneTimeWarpTrack.cpp
+++ b/Engine/Source/Runtime/MovieScene/Private/Tracks/MovieSceneTimeWarpTrack.cpp
@@ -1,9 +1,11 @@
 // Copyright Epic Games, Inc. All Rights Reserved.
 
 #include "Tracks/MovieSceneTimeWarpTrack.h"
+#include "Evaluation/MovieSceneSequenceTransform.h"
+#include "MovieScene.h"
 #include "Sections/MovieSceneTimeWarpSection.h"
 #include "Variants/MovieSceneTimeWarpGetter.h"
-#include "Evaluation/MovieSceneSequenceTransform.h"
+#include "Decorations/IMovieSceneDecoration.h"
 
 
 UMovieSceneTimeWarpTrack::UMovieSceneTimeWarpTrack(const FObjectInitializer& ObjectInitializer)
@@ -27,11 +29,49 @@ UMovieSceneSection* UMovieSceneTimeWarpTrack::CreateNewSection()
 	return NewObject<UMovieSceneTimeWarpSection>(this, NAME_None, RF_Transactional);
 }
 
+void UMovieSceneTimeWarpTrack::OnAddedToMovieSceneImpl(UMovieScene* InMovieScene)
+{
+	UMovieSceneTimeWarpDecoration* TimeWarp = InMovieScene->GetOrCreateDecoration<UMovieSceneTimeWarpDecoration>();
+	TimeWarp->AddTimeWarpSource(this);
+}
+
+void UMovieSceneTimeWarpTrack::OnRemovedFromMovieSceneImpl()
+{
+	if (UMovieSceneTimeWarpDecoration* TimeWarp = GetTypedOuter<UMovieScene>()->FindDecoration<UMovieSceneTimeWarpDecoration>())
+	{
+		TimeWarp->RemoveTimeWarpSource(this);
+	}
+}
+
+FMovieSceneNestedSequenceTransform UMovieSceneTimeWarpTrack::GenerateTimeWarpTransform()
+{
+	return GenerateTransform();
+}
+
+bool UMovieSceneTimeWarpTrack::IsTimeWarpActive() const
+{
+	return bIsActiveTimeWarp && !IsEvalDisabled();
+}
+
+void UMovieSceneTimeWarpTrack::SetIsTimeWarpActive(bool bInActive)
+{
+	bIsActiveTimeWarp = bInActive;
+}
+
+int32 UMovieSceneTimeWarpTrack::GetTimeWarpSortOrder() const
+{
+#if WITH_EDITORONLY_DATA
+	return GetSortingOrder();
+#else
+	return 0;
+#endif
+}
+
 FMovieSceneNestedSequenceTransform UMovieSceneTimeWarpTrack::GenerateTransform() const
 {
 	for (UMovieSceneSection* Section : GetAllSections())
 	{
-		UMovieSceneTimeWarpSection* TimeWarpSection = CastChecked<UMovieSceneTimeWarpSection>(Section);
+		UMovieSceneTimeWarpSection* TimeWarpSection = Cast<UMovieSceneTimeWarpSection>(Section);
 		if (TimeWarpSection && TimeWarpSection->IsActive() && !this->IsRowEvalDisabled(TimeWarpSection->GetRowIndex()))
 		{
 			return TimeWarpSection->GenerateTransform();
@@ -53,7 +93,6 @@ bool UMovieSceneTimeWarpTrack::HasSection(const UMovieSceneSection& Section) con
 
 void UMovieSceneTimeWarpTrack::AddSection(UMovieSceneSection& Section)
 {
-	check(Section.IsA<UMovieSceneTimeWarpSection>());
 	Sections.AddUnique(&Section);
 }
 
@@ -82,8 +121,8 @@ FText UMovieSceneTimeWarpTrack::GetDisplayName() const
 {
 	for (UMovieSceneSection* Section : Sections)
 	{
-		UMovieSceneTimeWarpSection* TimeWarpSection = CastChecked<UMovieSceneTimeWarpSection>(Section);
-		UMovieSceneTimeWarpGetter* Custom = TimeWarpSection->TimeWarp.GetType() == EMovieSceneTimeWarpType::Custom ? TimeWarpSection->TimeWarp.AsCustom() : nullptr;
+		UMovieSceneTimeWarpSection* TimeWarpSection = Cast<UMovieSceneTimeWarpSection>(Section);
+		UMovieSceneTimeWarpGetter* Custom = (TimeWarpSection && TimeWarpSection->TimeWarp.GetType() == EMovieSceneTimeWarpType::Custom) ? TimeWarpSection->TimeWarp.AsCustom() : nullptr;
 
 		if (Custom)
 		{
diff --git a/Engine/Source/Runtime/MovieScene/Private/Variants/MovieScenePlayRateCurve.cpp b/Engine/Source/Runtime/MovieScene/Private/Variants/MovieScenePlayRateCurve.cpp
--- a/Engine/Source/Runtime/MovieScene/Private/Variants/MovieScenePlayRateCurve.cpp
+++ b/Engine/Source/Runtime/MovieScene/Private/Variants/MovieScenePlayRateCurve.cpp
@@ -49,8 +49,8 @@ const UE::MovieScene::FPiecewiseCurve& UMovieScenePlayRateCurve::GetTimeWarpCurv
 
 		IntegratedTimeWarp = PlayRate.AsPiecewiseCurve().Integral();
 
-		FFrameTime IntegralStartTime;
-		if (!Section)
+		FFrameTime IntegralStartTime = PlaybackStartFrame;
+		if (!bManualPlaybackStart && MovieScene && !Section)
 		{
 			IntegralStartTime = MovieScene->GetPlaybackRange().GetLowerBoundValue();
 		}
diff --git a/Engine/Source/Runtime/MovieScene/Private/Variants/MovieSceneTimeWarpVariant.cpp b/Engine/Source/Runtime/MovieScene/Private/Variants/MovieSceneTimeWarpVariant.cpp
--- a/Engine/Source/Runtime/MovieScene/Private/Variants/MovieSceneTimeWarpVariant.cpp
+++ b/Engine/Source/Runtime/MovieScene/Private/Variants/MovieSceneTimeWarpVariant.cpp
@@ -168,6 +168,43 @@ void FMovieSceneTimeWarpVariant::ScaleBy(double ScaleFactor)
 	}
 }
 
+
+FFrameTime FMovieSceneTimeWarpVariant::RemapTime(FFrameTime InTime) const
+{
+	switch(GetType())
+	{
+	case EMovieSceneTimeWarpType::FixedPlayRate:
+		return InTime * AsFixedPlayRate();
+
+	case EMovieSceneTimeWarpType::Custom:
+		if (UMovieSceneTimeWarpGetter* Custom = AsCustom())
+		{
+			return Custom->RemapTime(InTime);
+		}
+		return InTime;
+
+	case EMovieSceneTimeWarpType::FixedTime:
+		return AsFixedTime().FrameNumber;
+
+	case EMovieSceneTimeWarpType::FrameRate:
+		return ConvertFrameTime(InTime, FFrameRate(1, 1), AsFrameRate().GetFrameRate());
+
+	case EMovieSceneTimeWarpType::Loop:
+		return AsLoop().LoopTime(InTime);
+
+	case EMovieSceneTimeWarpType::Clamp:
+		return AsClamp().Clamp(InTime);
+
+	case EMovieSceneTimeWarpType::LoopFloat:
+		return AsLoopFloat().LoopTime(InTime);
+
+	case EMovieSceneTimeWarpType::ClampFloat:
+		return AsClampFloat().Clamp(InTime);
+	}
+
+	return InTime;
+}
+
 bool FMovieSceneTimeWarpVariant::Serialize(FArchive& Ar)
 {
 	Ar.UsingCustomVersion(FFortniteMainBranchObjectVersion::GUID);
diff --git a/Engine/Source/Runtime/MovieScene/Public/Decorations/IMovieSceneLifetimeDecoration.h b/Engine/Source/Runtime/MovieScene/Public/Decorations/IMovieSceneLifetimeDecoration.h
new file mode 100644
--- /dev/null
+++ b/Engine/Source/Runtime/MovieScene/Public/Decorations/IMovieSceneLifetimeDecoration.h
@@ -0,0 +1,45 @@
+// Copyright Epic Games, Inc. All Rights Reserved.
+
+#pragma once
+
+#include "CoreTypes.h"
+#include "UObject/Interface.h"
+#include "IMovieSceneLifetimeDecoration.generated.h"
+
+class UMovieScene;
+
+UINTERFACE(MinimalAPI)
+class UMovieSceneLifetimeDecoration : public UInterface
+{
+public:
+	GENERATED_BODY()
+};
+
+
+/** 
+ * Optional interface that can be added to any decoration to provide 'construct' / 'destroy' semantics for when a decoration
+ *        is added to or removed from a well-formed UMovieScene hierarchy.
+ */
+class IMovieSceneLifetimeDecoration
+{
+public:
+	GENERATED_BODY()
+
+
+	/**
+	 * Called to reconstruct this decoration for the specified MovieScene.
+	 * @note: this may get called multiple times on the same instance
+	 */
+	virtual void OnReconstruct(UMovieScene* MovieScene)
+	{
+	}
+
+
+	/**
+	 * Called when this decoration is removed from a UMovieScene hierarchy
+	 */
+	virtual void OnDestroy(UMovieScene* MovieScene)
+	{
+	}
+};
+
diff --git a/Engine/Source/Runtime/MovieScene/Public/Decorations/MovieSceneDecorationContainer.h b/Engine/Source/Runtime/MovieScene/Public/Decorations/MovieSceneDecorationContainer.h
--- a/Engine/Source/Runtime/MovieScene/Public/Decorations/MovieSceneDecorationContainer.h
+++ b/Engine/Source/Runtime/MovieScene/Public/Decorations/MovieSceneDecorationContainer.h
@@ -89,6 +89,14 @@ public:
 		RemoveDecoration(DecorationType::StaticClass());
 	}
 
+	/**
+	 * Remove any null decoration ptrs
+	 */
+	void RemoveNulls()
+	{
+		Decorations.Remove(nullptr);
+	}
+
 protected:
 
 	UPROPERTY()
@@ -166,6 +174,14 @@ public:
 		return RemoveDecoration(DecorationType::StaticClass());
 	}
 
+	/**
+	 * Retrieve the list of compatible decorations for this object
+	 */
+	void GetCompatibleUserDecorations(TSet<UClass*>& OutClasses) const
+	{
+		GetCompatibleUserDecorationsImpl(OutClasses);
+	}
+
 protected:
 
 	virtual void OnDecorationAdded(UObject* Decoration)
@@ -174,6 +190,13 @@ protected:
 	virtual void OnDecorationRemoved(UObject* Decoration)
 	{}
 
+	virtual void GetCompatibleUserDecorationsImpl(TSet<UClass*>& OutClasses) const
+	{}
+
+protected:
+
+	MOVIESCENE_API virtual void Serialize(FArchive& Ar) override;
+
 protected:
 
 	/** Array of decorations for this movie scene */
diff --git a/Engine/Source/Runtime/MovieScene/Public/Decorations/MovieSceneLanguagePreviewDecoration.h b/Engine/Source/Runtime/MovieScene/Public/Decorations/MovieSceneLanguagePreviewDecoration.h
new file mode 100644
--- /dev/null
+++ b/Engine/Source/Runtime/MovieScene/Public/Decorations/MovieSceneLanguagePreviewDecoration.h
@@ -0,0 +1,20 @@
+// Copyright Epic Games, Inc. All Rights Reserved.
+
+#pragma once
+
+#include "CoreTypes.h"
+#include "UObject/Interface.h"
+#include "MovieSceneLanguagePreviewDecoration.generated.h"
+
+
+UCLASS(MinimalAPI)
+class UMovieSceneLanguagePreviewDecoration
+	: public UObject
+{
+public:
+
+	GENERATED_BODY()
+
+	UPROPERTY()
+	FString PreviewLanguage;
+};
\ No newline at end of file
diff --git a/Engine/Source/Runtime/MovieScene/Public/Decorations/MovieSceneScalingAnchors.h b/Engine/Source/Runtime/MovieScene/Public/Decorations/MovieSceneScalingAnchors.h
new file mode 100644
--- /dev/null
+++ b/Engine/Source/Runtime/MovieScene/Public/Decorations/MovieSceneScalingAnchors.h
@@ -0,0 +1,209 @@
+// Copyright Epic Games, Inc. All Rights Reserved.
+
+#pragma once
+
+#include "Channels/MovieSceneTimeWarpChannel.h"
+#include "Channels/MovieScenePiecewiseCurve.h"
+#include "Decorations/IMovieSceneDecoration.h"
+#include "Decorations/MovieSceneTimeWarpDecoration.h"
+#include "Variants/MovieScenePlayRateCurve.h"
+#include "MovieSceneScalingAnchors.generated.h"
+
+
+
+/**
+ * Structure that defines a single anchor with an optional duration
+ * Anchor durations do not get scaled by other anchors, but positions will be offset based on
+ * previous anchor stretching or movement
+ */
+USTRUCT()
+struct FMovieSceneScalingAnchor
+{
+	GENERATED_BODY()
+
+	/** Pefines the position of this anchor in ticks */
+	UPROPERTY()
+	FFrameNumber Position;
+
+	/** Defines the duration of this anchor in ticks */
+	UPROPERTY()
+	int32 Duration = 0;
+};
+
+
+
+
+UINTERFACE(MinimalAPI)
+class UMovieSceneScalingDriver : public UInterface
+{
+public:
+	GENERATED_BODY()
+};
+
+
+/**
+ * Interface that can be added to any object within a MovieScene in order to supply scaling anchors to the
+ *      UMovieSceneScalingAnchors decoration that acts as a registry for all anchors
+ */
+class IMovieSceneScalingDriver
+{
+public:
+	GENERATED_BODY()
+
+	/**
+	 * Populate the 'unscaled' map of anchors to use as a basis for scaling the seqence.
+	 * Anchors are represented by a GUID that uniquely identifies them within the sequence.
+	 */
+	virtual void PopulateInitialAnchors(TMap<FGuid, FMovieSceneScalingAnchor>& OutAnchors)
+	{
+		PopulateAnchors(OutAnchors);
+	}
+
+	/**
+	 * Populate the scaled map of anchors from which scaling factors will be computed
+	 */
+	virtual void PopulateAnchors(TMap<FGuid, FMovieSceneScalingAnchor>& OutAnchors) = 0;
+};
+
+
+
+/**
+ * Structure that defines a grouping of sections to be scaled as one
+ */
+USTRUCT()
+struct FMovieSceneAnchorsScalingGroup
+{
+	GENERATED_BODY()
+
+	/** Set of all the sections that are contained within this group */
+	UPROPERTY()
+	TSet<TObjectPtr<UMovieSceneSection>> Sections;
+};
+
+
+/**
+ * Decoration that is added to a UMovieScene in order to define dynamic scaling anchors to a the sequence.
+ * Anchors are defined by scaling 'drivers' which control both the initial (unscaled) position of their anchors,
+ * and their scaled position.
+ */
+UCLASS(MinimalAPI, DisplayName="Anchors", meta=(Hidden))
+class UMovieSceneScalingAnchors
+	: public UMovieScenePlayRateCurve
+	, public IMovieSceneDecoration
+	, public IMovieSceneTimeWarpSource
+{
+public:
+
+	GENERATED_BODY()
+
+	UMovieSceneScalingAnchors();
+
+#if 0
+	static uint32 DefaultTimeWarpEntityID;
+#endif
+
+public:
+
+	/**
+	 * Add a scaling driver to this anchor registry
+	 * 
+	 * @param InScalingDriver     The scaling driver to add to this registry. Must not be null.
+	 */
+	MOVIESCENE_API void AddScalingDriver(TScriptInterface<IMovieSceneScalingDriver> InScalingDriver);
+
+
+	/**
+	 * Remove a  scaling driver to this anchor registry
+	 */
+	MOVIESCENE_API void RemoveScalingDriver(TScriptInterface<IMovieSceneScalingDriver> InScalingDriver);
+
+public:
+
+	/**
+	 * Retrieve a scaling group of the specified identifier, creating it if necessary
+	 */
+	MOVIESCENE_API FMovieSceneAnchorsScalingGroup& GetOrCreateScalingGroup(const FGuid& Guid);
+
+	/**
+	 * Retrieve a scaling group by its ID
+	 */
+	MOVIESCENE_API FMovieSceneAnchorsScalingGroup* FindScalingGroup(const FGuid& Guid);
+
+	/**
+	 * Destroy a scaling group by its ID
+	 */
+	MOVIESCENE_API void RemoveScalingGroup(const FGuid& Guid);
+
+	/**
+	 * Retrieve all scaling groups
+	 */
+	MOVIESCENE_API const TMap<FGuid, FMovieSceneAnchorsScalingGroup>& GetScalingGroups() const;
+
+public:
+
+	/**
+	 * Retrieve the initial anchors stored in this registry that act as the unscaled basis for the scaling
+	 */
+	MOVIESCENE_API const TMap<FGuid, FMovieSceneScalingAnchor>& GetInitialAnchors() const;
+
+
+	/**
+	 * Retrieve the current, scaled anchors stored in this registry
+	 */
+	MOVIESCENE_API const TMap<FGuid, FMovieSceneScalingAnchor>& GetCurrentAnchors() const;
+
+public:
+
+	/*~ Begin UObject Implementation */
+	virtual void Serialize(FArchive& Ar) override;
+	/*~ End UObject Implementation */
+
+	/*~ Begin UMovieSceneTimeWarpGetter Implementation */
+	virtual FFrameTime RemapTime(FFrameTime In) const override;
+	virtual TOptional<FFrameTime> InverseRemapTimeCycled(FFrameTime InValue, FFrameTime InTimeHint, const UE::MovieScene::FInverseTransformTimeParams& Params) const override;
+	virtual TRange<FFrameTime> ComputeTraversedHull(const TRange<FFrameTime>& Range) const override;
+	virtual bool InverseRemapTimeWithinRange(FFrameTime InTime, FFrameTime RangeStart, FFrameTime RangeEnd, const TFunctionRef<bool(FFrameTime)>& VisitorCallback) const override;
+	virtual EMovieSceneChannelProxyType PopulateChannelProxy(FMovieSceneChannelProxyData& OutProxyData, EAllowTopLevelChannels AllowTopLevel) override;
+	virtual bool DeleteChannel(FMovieSceneTimeWarpVariant& OutVariant, FName ChannelName) override;
+	virtual void ScaleBy(double UnwarpedScaleFactor) override;
+	/* End UMovieSceneTimeWarpGetter Implementation */
+
+	/*~ Begin IMovieSceneDecoration Implementation */
+	virtual void OnDecorationAdded(UMovieScene* MovieScene) override;
+	virtual void OnDecorationRemoved() override;
+	virtual void OnPreDecorationCompiled() override;
+	/*~ End IMovieSceneDecoration Implementation */
+
+	/*~ Begin IMovieSceneTimeWarpSource Implementation */
+	virtual FMovieSceneNestedSequenceTransform GenerateTimeWarpTransform() override;
+	virtual bool IsTimeWarpActive() const override;
+	virtual void SetIsTimeWarpActive(bool bInActive) override;
+	virtual int32 GetTimeWarpSortOrder() const override;
+	/*~ End IMovieSceneTimeWarpSource Implementation */
+
+private:
+
+	UMovieScenePlayRateCurve* Initialize(TSharedPtr<UE::MovieScene::FSharedPlaybackState> SharedPlaybackState) const;
+
+	void UpdateCurve(UMovieScenePlayRateCurve* Curve) const;
+
+	void UpdateFromSource() const;
+
+	void ResetScaling();
+
+private:
+
+	UPROPERTY()
+	TArray<TScriptInterface<IMovieSceneScalingDriver>> ScalingDrivers;
+
+	UPROPERTY()
+	TMap<FGuid, FMovieSceneScalingAnchor> InitialAnchors;
+
+	UPROPERTY()
+	TMap<FGuid, FMovieSceneAnchorsScalingGroup> ScalingGroups;
+
+	mutable TMap<FGuid, FMovieSceneScalingAnchor> CurrentAnchors;
+
+	mutable FFrameTime NewDuration;
+	mutable bool bPlayRateCurveIsUpToDate = false;
+};
\ No newline at end of file
diff --git a/Engine/Source/Runtime/MovieScene/Public/Decorations/MovieSceneSectionAnchorsDecoration.h b/Engine/Source/Runtime/MovieScene/Public/Decorations/MovieSceneSectionAnchorsDecoration.h
new file mode 100644
--- /dev/null
+++ b/Engine/Source/Runtime/MovieScene/Public/Decorations/MovieSceneSectionAnchorsDecoration.h
@@ -0,0 +1,36 @@
+// Copyright Epic Games, Inc. All Rights Reserved.
+
+#pragma once
+
+#include "CoreTypes.h"
+#include "Decorations/MovieSceneScalingAnchors.h"
+#include "Decorations/IMovieSceneSectionDecoration.h"
+#include "Decorations/IMovieSceneLifetimeDecoration.h"
+#include "UObject/Interface.h"
+#include "MovieSceneSectionAnchorsDecoration.generated.h"
+
+
+UCLASS(MinimalAPI)
+class UMovieSceneSectionAnchorsDecoration
+	: public UObject
+	, public IMovieSceneSectionDecoration
+	, public IMovieSceneLifetimeDecoration
+	, public IMovieSceneScalingDriver
+{
+public:
+
+	GENERATED_BODY()
+
+	UPROPERTY()
+	FGuid StartAnchor;
+
+private:
+
+	virtual void OnReconstruct(UMovieScene* MovieScene) override;
+	virtual void OnDestroy(UMovieScene* MovieScene) override;
+
+	virtual void PopulateInitialAnchors(TMap<FGuid, FMovieSceneScalingAnchor>& OutAnchors) override;
+	virtual void PopulateAnchors(TMap<FGuid, FMovieSceneScalingAnchor>& OutAnchors) override;
+
+	virtual void PostEditImport() override;
+};
\ No newline at end of file
diff --git a/Engine/Source/Runtime/MovieScene/Public/Decorations/MovieSceneTimeWarpDecoration.h b/Engine/Source/Runtime/MovieScene/Public/Decorations/MovieSceneTimeWarpDecoration.h
new file mode 100644
--- /dev/null
+++ b/Engine/Source/Runtime/MovieScene/Public/Decorations/MovieSceneTimeWarpDecoration.h
@@ -0,0 +1,65 @@
+// Copyright Epic Games, Inc. All Rights Reserved.
+
+#pragma once
+
+#include "CoreTypes.h"
+#include "UObject/Interface.h"
+#include "EntitySystem/IMovieSceneEntityProvider.h"
+#include "MovieSceneTimeWarpDecoration.generated.h"
+
+
+struct FMovieSceneSequenceTransform;
+struct FMovieSceneNestedSequenceTransform;
+
+
+UINTERFACE(MinimalAPI)
+class UMovieSceneTimeWarpSource : public UInterface
+{
+public:
+	GENERATED_BODY()
+};
+
+class IMovieSceneTimeWarpSource
+{
+public:
+	GENERATED_BODY()
+
+	virtual FMovieSceneNestedSequenceTransform GenerateTimeWarpTransform() = 0;
+
+	virtual bool IsTimeWarpActive() const = 0;
+
+	virtual void SetIsTimeWarpActive(bool bInActive) = 0;
+
+	virtual int32 GetTimeWarpSortOrder() const = 0;
+};
+
+
+UCLASS(MinimalAPI)
+class UMovieSceneTimeWarpDecoration
+	: public UObject
+	, public IMovieSceneEntityProvider
+{
+public:
+
+	GENERATED_BODY()
+
+	MOVIESCENE_API FMovieSceneSequenceTransform GenerateTransform() const;
+
+	MOVIESCENE_API void OnCompiled();
+
+	MOVIESCENE_API void AddTimeWarpSource(TScriptInterface<IMovieSceneTimeWarpSource> InSource);
+
+	MOVIESCENE_API void RemoveTimeWarpSource(TScriptInterface<IMovieSceneTimeWarpSource> InSource);
+
+private:
+
+	virtual void ImportEntityImpl(UMovieSceneEntitySystemLinker* EntityLinker, const FEntityImportParams& Params, FImportedEntity* OutImportedEntity) override;
+	virtual bool PopulateEvaluationFieldImpl(const TRange<FFrameNumber>& EffectiveRange, const FMovieSceneEvaluationFieldEntityMetaData& InMetaData, FMovieSceneEntityComponentFieldBuilder* OutFieldBuilder) override;
+
+	virtual void Serialize(FArchive& AR) override;
+private:
+
+
+	UPROPERTY()
+	TArray<TScriptInterface<IMovieSceneTimeWarpSource>> Sources;
+};
\ No newline at end of file
diff --git a/Engine/Source/Runtime/MovieScene/Public/Evaluation/MovieSceneSequenceTransform.h b/Engine/Source/Runtime/MovieScene/Public/Evaluation/MovieSceneSequenceTransform.h
--- a/Engine/Source/Runtime/MovieScene/Public/Evaluation/MovieSceneSequenceTransform.h
+++ b/Engine/Source/Runtime/MovieScene/Public/Evaluation/MovieSceneSequenceTransform.h
@@ -762,6 +762,13 @@ public:
 	MOVIESCENE_API void Add(FMovieSceneTimeTransform InTransform);
 
 
+
+	/**
+	 * Add the specified nested transform to the end of this transform stack (ie, applying it last). Does nothing if the supplied transform is an identity.
+	 */
+	MOVIESCENE_API void Add(FMovieSceneNestedSequenceTransform InTransform);
+
+
 	/**
 	 * Add the specified warping transform to the end of this transform stack with an offset (ie, applying it last)
 	 */
diff --git a/Engine/Source/Runtime/MovieScene/Public/IMovieSceneTrackVirtualAPI.h b/Engine/Source/Runtime/MovieScene/Public/IMovieSceneTrackVirtualAPI.h
new file mode 100644
--- /dev/null
+++ b/Engine/Source/Runtime/MovieScene/Public/IMovieSceneTrackVirtualAPI.h
@@ -0,0 +1,51 @@
+// Copyright Epic Games, Inc. All Rights Reserved.
+
+#pragma once
+
+#include "CoreTypes.h"
+#include "Misc/AssertionMacros.h"
+
+class UMovieSceneTrack;
+class UMovieSceneSection;
+
+class IMovieSceneTrackVirtualAPI
+{
+protected:
+
+	/**
+	 * Add a section to this track.
+	 *
+	 * @param Section The section to add.
+	 */
+	virtual void AddSection(UMovieSceneSection& Section) PURE_VIRTUAL(IMovieSceneTrackVirtualAPI::AddSection,);
+
+	/**
+	 * Removes a section from this track.
+	 *
+	 * @param Section The section to remove.
+	 */
+	virtual void RemoveSection(UMovieSceneSection& Section) PURE_VIRTUAL(IMovieSceneTrackVirtualAPI::RemoveSection, );
+
+	/**
+	 * Removes a section from this track at a particular index
+	 *
+	 * @param SectionIndex The section index to remove.
+	 */
+	virtual void RemoveSectionAt(int32 SectionIndex) PURE_VIRTUAL(IMovieSceneTrackVirtualAPI::RemoveSectionAt, );
+
+private:
+	friend UMovieSceneTrack;
+
+	void CallAddSection(UMovieSceneSection& Section)
+	{
+		AddSection(Section);
+	}
+	void CallRemoveSection(UMovieSceneSection& Section)
+	{
+		RemoveSection(Section);
+	}
+	void CallRemoveSectionAt(int32 SectionIndex)
+	{
+		RemoveSectionAt(SectionIndex);
+	}
+};
\ No newline at end of file
diff --git a/Engine/Source/Runtime/MovieScene/Public/MovieSceneSection.h b/Engine/Source/Runtime/MovieScene/Public/MovieSceneSection.h
--- a/Engine/Source/Runtime/MovieScene/Public/MovieSceneSection.h
+++ b/Engine/Source/Runtime/MovieScene/Public/MovieSceneSection.h
@@ -42,6 +42,7 @@ class IMovieSceneEasingFunction;
 class IMovieScenePlayer;
 class UMovieSceneEntitySystemLinker;
 class UObject;
+class UMovieSceneTrack;
 enum class ECookOptimizationFlags;
 struct FEasingComponentData;
 struct FFrame;
@@ -681,6 +682,9 @@ public:
 	/** Get the object associated with this section */
 	MOVIESCENE_API virtual UObject* GetSourceObject() const { return nullptr; }
 
+	void OnAddedToTrack(UMovieSceneTrack* Track);
+	void OnRemovedFromTrack();
+
 #if WITH_EDITOR
 	MOVIESCENE_API virtual void PostEditChangeProperty(FPropertyChangedEvent& PropertyChangedEvent) override;
 
@@ -723,6 +727,13 @@ protected:
 	MOVIESCENE_API virtual void OnDecorationAdded(UObject* NewDecoration) override;
 	MOVIESCENE_API virtual void OnDecorationRemoved(UObject* NewDecoration) override;
 
+	virtual void OnAddedToTrackImpl(UMovieSceneTrack* Track)
+	{
+	}
+	virtual void OnRemovedFromTrackImpl()
+	{
+	}
+
 	virtual void OnMoved(int32 DeltaTime) {}
 	virtual void OnDilated(float DilationFactor, FFrameNumber Origin) {}
 
diff --git a/Engine/Source/Runtime/MovieScene/Public/MovieSceneTrack.h b/Engine/Source/Runtime/MovieScene/Public/MovieSceneTrack.h
--- a/Engine/Source/Runtime/MovieScene/Public/MovieSceneTrack.h
+++ b/Engine/Source/Runtime/MovieScene/Public/MovieSceneTrack.h
@@ -3,6 +3,7 @@
 #pragma once
 
 #include "Compilation/MovieSceneSegmentCompiler.h"
+#include "IMovieSceneTrackVirtualAPI.h"
 #include "Containers/Array.h"
 #include "Containers/ArrayView.h"
 #include "CoreMinimal.h"
@@ -181,12 +182,27 @@ struct MOVIESCENE_API FMovieSceneTrackRowMetadata
 	FMovieSceneConditionContainer ConditionContainer;
 };
 
+
+// Disable some VS warnings that are emitted by Add/RemoveSection functions hiding the
+//    private virtual methods in IMovieSceneTrackVirtualAPI
+#if PLATFORM_COMPILER_CLANG
+	#pragma clang diagnostic push
+	#pragma clang diagnostic ignored "-Woverloaded-virtual"
+#elif PLATFORM_MICROSOFT
+	#pragma warning(push)
+	#pragma warning(disable:4263)
+	#pragma warning(disable:4264)
+#else
+	static_assert(false, "Unexpected compiler detected");
+#endif
+
 /**
  * Base class for a track in a Movie Scene
  */
 UCLASS(abstract, DefaultToInstanced, MinimalAPI, BlueprintType)
 class UMovieSceneTrack
 	: public UMovieSceneDecorationContainerObject
+	, public IMovieSceneTrackVirtualAPI
 {
 	GENERATED_BODY()
 
@@ -436,13 +452,40 @@ public:
 	 */
 	virtual bool SupportsType(TSubclassOf<UMovieSceneSection> SectionClass) const PURE_VIRTUAL(UMovieSceneTrack::SupportsType, return false;);
 
+	// Parameter structs that only exist to differentiate overloads from IMovieSceneTrackVirtualAPI
+	struct FSectionParameter
+	{
+		FSectionParameter(UMovieSceneSection& InSection) : Section(&InSection) {};
+		UMovieSceneSection* Section;
+	};
+	struct FSectionIndexParameter
+	{
+		FSectionIndexParameter(int32 InSectionIndex) : SectionIndex(InSectionIndex) {};
+		int32 SectionIndex;
+	};
+
 	/**
 	 * Add a section to this track.
 	 *
 	 * @param Section The section to add.
 	 */
-	virtual void AddSection(UMovieSceneSection& Section) PURE_VIRTUAL(UMovieSceneTrack::AddSection,);
+	MOVIESCENE_API void AddSection(FSectionParameter Section);
+
+	/**
+	 * Removes a section from this track.
+	 *
+	 * @param Section The section to remove.
+	 */
+	MOVIESCENE_API void RemoveSection(FSectionParameter Section);
+
+	/**
+	 * Removes a section from this track at a particular index
+	 *
+	 * @param SectionIndex The section index to remove.
+	 */
+	MOVIESCENE_API void RemoveSectionAt(FSectionIndexParameter SectionIndex);
 
+public:
 	/**
 	 * Generates a new section suitable for use with this track.
 	 *
@@ -465,20 +508,6 @@ public:
 	 */
 	virtual bool HasSection(const UMovieSceneSection& Section) const PURE_VIRTUAL(UMovieSceneTrack::HasSection, return false;);
 
-	/**
-	 * Removes a section from this track.
-	 *
-	 * @param Section The section to remove.
-	 */
-	virtual void RemoveSection(UMovieSceneSection& Section) PURE_VIRTUAL(UMovieSceneTrack::RemoveSection, );
-
-	/**
-	 * Removes a section from this track at a particular index
-	 *
-	 * @param SectionIndex The section index to remove.
-	 */
-	virtual void RemoveSectionAt(int32 SectionIndex) PURE_VIRTUAL(UMovieSceneTrack::RemoveSectionAt, );
-
 #if WITH_EDITOR
 
 	/**
@@ -623,3 +652,9 @@ public:
 
 	#endif
 };
+
+#if PLATFORM_COMPILER_CLANG
+	#pragma clang diagnostic pop
+#elif PLATFORM_MICROSOFT
+	#pragma warning(pop)
+#endif
\ No newline at end of file
diff --git a/Engine/Source/Runtime/MovieScene/Public/Tracks/MovieSceneTimeWarpTrack.h b/Engine/Source/Runtime/MovieScene/Public/Tracks/MovieSceneTimeWarpTrack.h
--- a/Engine/Source/Runtime/MovieScene/Public/Tracks/MovieSceneTimeWarpTrack.h
+++ b/Engine/Source/Runtime/MovieScene/Public/Tracks/MovieSceneTimeWarpTrack.h
@@ -3,6 +3,7 @@
 #pragma once
 
 #include "CoreTypes.h"
+#include "Decorations/MovieSceneTimeWarpDecoration.h"
 #include "MovieSceneNameableTrack.h"
 #include "MovieSceneTimeWarpTrack.generated.h"
 
@@ -14,6 +15,7 @@ struct FMovieSceneNestedSequenceTransform;
 UCLASS(MinimalAPI)
 class UMovieSceneTimeWarpTrack
 	: public UMovieSceneTrack
+	, public IMovieSceneTimeWarpSource
 {
 	GENERATED_BODY()
 
@@ -30,7 +32,7 @@ public:
 	 * Generate the a sequence transform that perform's this track's time-warp.
 	 * @note: The resulting transform may be the identity transform
 	 */
-	MOVIESCENE_API FMovieSceneNestedSequenceTransform GenerateTransform() const;
+	MOVIESCENE_API virtual FMovieSceneNestedSequenceTransform GenerateTransform() const;
 
 	/** UMovieSceneTrack interface */
 	virtual bool SupportsType(TSubclassOf<UMovieSceneSection> SectionClass) const override;
@@ -39,10 +41,6 @@ public:
 
 public:
 
-	/** True if this track is the active timewarp in its sequence. Assigned as part of compilation */
-	UPROPERTY()
-	bool bIsActiveTimeWarp = true;
-
 	// UMovieSceneTrack interface
 
 	virtual void RemoveAllAnimationData() override;
@@ -52,6 +50,13 @@ public:
 	virtual void RemoveSectionAt(int32 SectionIndex) override;
 	virtual bool IsEmpty() const override;
 	virtual const TArray<UMovieSceneSection*>& GetAllSections() const override;
+	virtual void OnAddedToMovieSceneImpl(UMovieScene* InMovieScene) override;
+	virtual void OnRemovedFromMovieSceneImpl() override;
+
+	virtual FMovieSceneNestedSequenceTransform GenerateTimeWarpTransform() override;
+	virtual bool IsTimeWarpActive() const override;
+	virtual void SetIsTimeWarpActive(bool bInUse) override;
+	virtual int32 GetTimeWarpSortOrder() const override;
 
 #if WITH_EDITORONLY_DATA
 	virtual FText GetDisplayName() const override;
@@ -64,4 +69,8 @@ private:
 	/** Array of sections contained within this track - should only ever be one. */
 	UPROPERTY()
 	TArray<TObjectPtr<UMovieSceneSection>> Sections;
+
+	/** True if this track is the active timewarp in its sequence. Assigned as part of compilation */
+	UPROPERTY()
+	bool bIsActiveTimeWarp = true;
 };
diff --git a/Engine/Source/Runtime/MovieScene/Public/Variants/MovieScenePlayRateCurve.h b/Engine/Source/Runtime/MovieScene/Public/Variants/MovieScenePlayRateCurve.h
--- a/Engine/Source/Runtime/MovieScene/Public/Variants/MovieScenePlayRateCurve.h
+++ b/Engine/Source/Runtime/MovieScene/Public/Variants/MovieScenePlayRateCurve.h
@@ -39,19 +39,29 @@ public:
 
 	MOVIESCENE_API const UE::MovieScene::FPiecewiseCurve& GetTimeWarpCurve() const;
 
-private:
+protected:
 
 	void InvalidateTimeWarp();
 
-private:
+public:
 
 	/** Curve defined in play-rate space. Does not support cycling. */
 	UPROPERTY(EditAnywhere, Category="TimeWarp")
 	FMovieSceneTimeWarpChannel PlayRate;
 
+	UPROPERTY(EditAnywhere, Category="TimeWarp")
+	FFrameNumber PlaybackStartFrame;
+
+	UPROPERTY(EditAnywhere, Category="TimeWarp")
+	bool bManualPlaybackStart = false;
+
+private:
+
 	/** Curve defined in time-warp space as an integral of the PlayRate curve */
 	mutable UE::MovieScene::FPiecewiseCurve IntegratedTimeWarp;
 
+public:
+
 	/** false when IntegratedTimeWarp needs to be udpated */
 	mutable bool bUpToDate = false;
 };
\ No newline at end of file
diff --git a/Engine/Source/Runtime/MovieScene/Public/Variants/MovieSceneTimeWarpVariant.h b/Engine/Source/Runtime/MovieScene/Public/Variants/MovieSceneTimeWarpVariant.h
--- a/Engine/Source/Runtime/MovieScene/Public/Variants/MovieSceneTimeWarpVariant.h
+++ b/Engine/Source/Runtime/MovieScene/Public/Variants/MovieSceneTimeWarpVariant.h
@@ -205,45 +205,6 @@ public:
 	MOVIESCENE_API FFrameTime RemapTime(FFrameTime InTime) const;
 
 
-	/**
-	 * Given a continuous input range, compute the hull of warped times that are contained.
-	 * eg: for looping time-warps and an input range with size > loop duration, this function will return the full loop range
-	 */
-	MOVIESCENE_API TRange<FFrameTime> ComputeTraversedHull(const TRange<FFrameTime>& Range) const;
-
-
-	/**
-	 * Attempt to transform a time-warped time into its corresponding non-time-warped time.
-	 * @note: this operation can fail for warped times that have no representation in the non-time-warped space
-	 * 
-	 * @param WarpedTime         The time-warped time to remap
-	 * @param UnwarpedTimeHint   A time 'hint' in the unwarped space to serve as a guide where there may be more than one solution.
-	 *                           For example; consider a looping or cycling warp with duration 10, calling this function with InValue=3.
-	 *                           There are an infinite number of unwarped times that can result in frame 3 (ie, every frame where %3==0). If we want to get the time that resulted in this
-	 *                           from the 10th loop, we could pass the middle frame of that loop as the hint (ie, frame 95). The result would be the closest frame to frame 95 that warps
-	 *                           to frame 3.
-	 * @param Params             Parameter structure to control the inverse algorithm, eg limiting it to only searching one direction, ignoring clamping and cycles etc.
-	 * @return The resulting unwarped frame time, or nothing if the time does not map to an unwarped time.
-	 */
-	MOVIESCENE_API TOptional<FFrameTime> InverseRemapTimeCycled(FFrameTime WarpedTime, FFrameTime UnwarpedTimeHint, const UE::MovieScene::FInverseTransformTimeParams& Params) const;
-
-
-	/**
-	 * Attempt to transform a time-warped time into its corresponding non-time-warped time within a specified unwarped range.
-	 * @note: this operation can fail for warped times that have no representation in the non-time-warped space
-	 * 
-	 * @param InValue      The time-warped time to remap
-	 * @param InTimeHint   A time 'hint' in the unwarped space to serve as a guide where there may be more than one solution.
-	 *                     For example; consider a looping or cycling warp with duration 10, calling this function with InValue=3.
-	 *                     There are an infinite number of unwarped times that can result in frame 3 (ie, every frame where %3==0). If we want to get the time that resulted in this
-	 *                     from the 10th loop, we could pass the middle frame of that loop as the hint (ie, frame 95). The result would be the closest frame to frame 95 that warps
-	 *                     to frame 3.
-	 * @param Params       Parameter structure to control the inverse algorithm, eg limiting it to only searching one direction, ignoring clamping and cycles etc.
-	 * @return The resulting unwarped frame time, or nothing if the time does not map to an unwarped time.
-	 */
-	MOVIESCENE_API bool InverseRemapTimeWithinRange(FFrameTime InTime, FFrameTime UnwarpedRangeStart, FFrameTime UnwarpedRangeEnd, const TFunctionRef<bool(FFrameTime)>& VisitorCallback) const;
-
-
 public:
 
 
diff --git a/Engine/Source/Runtime/MovieSceneTracks/Private/Sections/MovieSceneAudioSection.cpp b/Engine/Source/Runtime/MovieSceneTracks/Private/Sections/MovieSceneAudioSection.cpp
--- a/Engine/Source/Runtime/MovieSceneTracks/Private/Sections/MovieSceneAudioSection.cpp
+++ b/Engine/Source/Runtime/MovieSceneTracks/Private/Sections/MovieSceneAudioSection.cpp
@@ -1,6 +1,8 @@
 // Copyright Epic Games, Inc. All Rights Reserved.
 
 #include "Sections/MovieSceneAudioSection.h"
+#include "Decorations/MovieSceneLanguagePreviewDecoration.h"
+#include "Decorations/MovieSceneSectionAnchorsDecoration.h"
 #include "Tracks/MovieSceneAudioTrack.h"
 #include "Sound/SoundBase.h"
 #include "UObject/SequencerObjectVersion.h"
@@ -10,6 +12,7 @@
 #include "MovieSceneCommonHelpers.h"
 #include "Misc/FrameRate.h"
 #include "Misc/GeneratedTypeName.h"
+#include "Misc/PackageName.h"
 #include "EntitySystem/BuiltInComponentTypes.h"
 #include "MovieSceneTracksComponentTypes.h"
 
@@ -149,6 +152,71 @@ EMovieSceneChannelProxyType  UMovieSceneAudioSection::CacheChannelProxy()
 	return EMovieSceneChannelProxyType::Dynamic;
 }
 
+USoundBase* UMovieSceneAudioSection::GetPlaybackSound() const
+{
+	if (!Sound)
+	{
+		return nullptr;
+	}
+
+#if WITH_EDITORONLY_DATA
+
+	// Handle custom language preview if necessary
+	UMovieSceneLanguagePreviewDecoration* PreviewLanguage = GetTypedOuter<UMovieScene>()->FindDecoration<UMovieSceneLanguagePreviewDecoration>();
+	if (PreviewLanguage && PreviewLanguage->PreviewLanguage.Len() != 0)
+	{
+		FString LocalizedPackageName = FPackageName::GetLocalizedPackagePath(Sound->GetOutermost()->GetPathName(), PreviewLanguage->PreviewLanguage);
+		if (LocalizedPackageName.Len() != 0)
+		{
+			// Attempt to load the localized asset
+			USoundBase* LocalizedSound = LoadObject<USoundBase>(nullptr, *FString(LocalizedPackageName + TEXT(".") + Sound->GetName()));
+			if (LocalizedSound)
+			{
+				return LocalizedSound;
+			}
+		}
+	}
+
+#endif
+
+	// Sound will have already been localized if we're in game by package-level localization
+	return Sound;
+}
+
+void UMovieSceneAudioSection::PopulateInitialAnchors(TMap<FGuid, FMovieSceneScalingAnchor>& OutAnchors)
+{
+	UMovieSceneSectionAnchorsDecoration* Decoration = FindDecoration<UMovieSceneSectionAnchorsDecoration>();
+	if (Decoration && HasStartFrame() && HasEndFrame())
+	{
+		FFrameNumber StartTime = GetInclusiveStartFrame();
+		OutAnchors.Add(Decoration->StartAnchor, FMovieSceneScalingAnchor{ StartTime, (GetExclusiveEndFrame() - StartTime).Value });
+	}
+}
+
+void UMovieSceneAudioSection::PopulateAnchors(TMap<FGuid, FMovieSceneScalingAnchor>& OutAnchors)
+{
+	UMovieSceneSectionAnchorsDecoration* Decoration = FindDecoration<UMovieSceneSectionAnchorsDecoration>();
+	if (Decoration && HasStartFrame() && HasEndFrame())
+	{
+		USoundBase* ReferenceSound = GetSound();
+		USoundBase* PlaybackSound  = GetPlaybackSound();
+
+		if (PlaybackSound && ReferenceSound && ReferenceSound != PlaybackSound)
+		{
+			const UMovieScene* MovieScene = GetTypedOuter<UMovieScene>();
+			const float ReferenceDuration = MovieSceneHelpers::GetSoundDuration(ReferenceSound);
+			const float PlaybackDuration  = MovieSceneHelpers::GetSoundDuration(PlaybackSound);
+
+			FFrameNumber StartTime = GetInclusiveStartFrame();
+			FFrameNumber EndTime   = GetExclusiveEndFrame();
+
+			int32 SectionLength = FMath::CeilToInt32((EndTime.Value - StartTime.Value) * double(PlaybackDuration / ReferenceDuration));
+
+			OutAnchors.Add(Decoration->StartAnchor, FMovieSceneScalingAnchor{ StartTime, SectionLength });
+		}
+	}
+}
+
 UObject* UMovieSceneAudioSection::GetSourceObject() const
 {
 	return GetSound();
diff --git a/Engine/Source/Runtime/MovieSceneTracks/Private/Systems/MovieSceneAudioSystem.cpp b/Engine/Source/Runtime/MovieSceneTracks/Private/Systems/MovieSceneAudioSystem.cpp
--- a/Engine/Source/Runtime/MovieSceneTracks/Private/Systems/MovieSceneAudioSystem.cpp
+++ b/Engine/Source/Runtime/MovieSceneTracks/Private/Systems/MovieSceneAudioSystem.cpp
@@ -13,6 +13,8 @@
 #include "GameFramework/WorldSettings.h"
 #include "IMovieScenePlayer.h"
 #include "MovieScene.h"
+#include "Decorations/MovieSceneScalingAnchors.h"
+#include "Decorations/MovieSceneSectionAnchorsDecoration.h"
 #include "MovieSceneTracksComponentTypes.h"
 #include "Sections/MovieSceneAudioSection.h"
 #include "Sound/SoundAttenuation.h"
@@ -515,7 +517,26 @@ private:
 		}
 
 		FFrameNumber SectionStartFrame = (AudioSection.HasStartFrame() ? AudioSection.GetInclusiveStartFrame() : 0);
-		float SectionStartTimeSeconds = SectionStartFrame / AudioSection.GetTypedOuter<UMovieScene>()->GetTickResolution();
+		
+		// If this audio seciton is a scaling driver (ie, it has the section anchors decoration),
+		//     we need to 'undo' the scaling from the evaluation time and use the scaled section start time
+		FFrameTime EvalTime = Context.GetTime();
+		UMovieSceneSectionAnchorsDecoration* AnchorsDecoration = AudioSection.FindDecoration<UMovieSceneSectionAnchorsDecoration>();
+		if (AnchorsDecoration)
+		{
+			UMovieSceneScalingAnchors* ScalingAnchors = AudioSection.GetTypedOuter<UMovieScene>()->FindDecoration<UMovieSceneScalingAnchors>();
+			if (const FMovieSceneScalingAnchor* AnchoredStart = ScalingAnchors->GetCurrentAnchors().Find(AnchorsDecoration->StartAnchor))
+			{
+				SectionStartFrame = AnchoredStart->Position;
+				TOptional<FFrameTime> UnwarpedTime = ScalingAnchors->InverseRemapTimeCycled(EvalTime, EvalTime, FInverseTransformTimeParams());
+				if (UnwarpedTime.IsSet())
+				{
+					EvalTime = UnwarpedTime.GetValue();
+				}
+			}
+		}
+
+		float SectionStartTimeSeconds = SectionStartFrame / Context.GetFrameRate();
 
 		float InverseTimeDilation = 1.0f;
 		const bool bUseTimeDilationToAdjustment = (bUseTimeDilationToAdjustPlayDurationCVar && WorldSettings);
@@ -533,9 +554,9 @@ private:
 		}
 
 		const FFrameNumber AudioStartOffset = AudioSection.GetStartOffset();
-		USoundBase* Sound = AudioSection.GetSound();
+		USoundBase* Sound = AudioSection.GetPlaybackSound();
 
-		float AudioTime = (Context.GetTime() / Context.GetFrameRate()) * InverseTimeDilation
+		float AudioTime = (EvalTime / Context.GetFrameRate()) * InverseTimeDilation
 			- SectionStartTimeSeconds 
 			+ (float)Context.GetFrameRate().AsSeconds(AudioStartOffset);
 
@@ -548,13 +569,13 @@ private:
 				{
 					// Store current audio time and current context time
 					EvaluationData.LastAudioTime = AudioTime;
-					EvaluationData.LastContextTime = Context.GetTime() / Context.GetFrameRate();
+					EvaluationData.LastContextTime = EvalTime / Context.GetFrameRate();
 				}
 				else
 				{
 					// Get previous AudioTime
 					AudioTime = (EvaluationData.LastAudioTime.GetValue() < AudioTime) ? EvaluationData.LastAudioTime.GetValue() : AudioTime;
-					float CurrContextTime = (Context.GetTime() / Context.GetFrameRate());
+					float CurrContextTime = (EvalTime / Context.GetFrameRate());
 					float PrevContextTime = EvaluationData.LastContextTime.GetValue();
 					PrevContextTime = (PrevContextTime < CurrContextTime) ? PrevContextTime : CurrContextTime;
 
@@ -680,7 +701,7 @@ private:
 
 			if (AudioTime >= 0.f)
 			{
-				UE_LOG(LogMovieScene, Verbose, TEXT("Audio Component Play at Local Time: %6.2f CurrentTime: %6.2f(s) SectionStart: %6.2f(s), SoundDur: %6.2f OffsetIntoClip: %6.2f Component: %s Sound: %s"), AudioTime, (Context.GetTime() / Context.GetFrameRate()), SectionStartTimeSeconds, AudioComponent.Sound ? AudioComponent.Sound->GetDuration() : 0.0f, (float)Context.GetFrameRate().AsSeconds(AudioStartOffset), *AudioComponent.GetName(), *GetNameSafe(AudioComponent.Sound));
+				UE_LOG(LogMovieScene, Verbose, TEXT("Audio Component Play at Local Time: %6.2f CurrentTime: %6.2f(s) SectionStart: %6.2f(s), SoundDur: %6.2f OffsetIntoClip: %6.2f Component: %s Sound: %s"), AudioTime, (EvalTime / Context.GetFrameRate()), SectionStartTimeSeconds, AudioComponent.Sound ? AudioComponent.Sound->GetDuration() : 0.0f, (float)Context.GetFrameRate().AsSeconds(AudioStartOffset), *AudioComponent.GetName(), *GetNameSafe(AudioComponent.Sound));
 				AudioComponent.Play(AudioTime);
 
 				// Keep track of when we asked this audio clip to play (in game time) so that we can figure 
diff --git a/Engine/Source/Runtime/MovieSceneTracks/Private/Tracks/MovieSceneAudioTrack.cpp b/Engine/Source/Runtime/MovieSceneTracks/Private/Tracks/MovieSceneAudioTrack.cpp
--- a/Engine/Source/Runtime/MovieSceneTracks/Private/Tracks/MovieSceneAudioTrack.cpp
+++ b/Engine/Source/Runtime/MovieSceneTracks/Private/Tracks/MovieSceneAudioTrack.cpp
@@ -1,6 +1,7 @@
 // Copyright Epic Games, Inc. All Rights Reserved.
 
 #include "Tracks/MovieSceneAudioTrack.h"
+#include "Decorations/MovieSceneSectionAnchorsDecoration.h"
 #include "Audio.h"
 #include "Sound/SoundWave.h"
 #include "MovieScene.h"
@@ -78,6 +79,12 @@ bool UMovieSceneAudioTrack::IsEmpty() const
 }
 
 
+void UMovieSceneAudioTrack::GetCompatibleUserDecorationsImpl(TSet<UClass*>& OutClasses) const
+{
+	OutClasses.Add(UMovieSceneSectionAnchorsDecoration::StaticClass());
+}
+
+
 UMovieSceneSection* UMovieSceneAudioTrack::AddNewSoundOnRow(USoundBase* Sound, FFrameNumber Time, int32 RowIndex)
 {
 	check(Sound);
diff --git a/Engine/Source/Runtime/MovieSceneTracks/Public/Sections/MovieSceneAudioSection.h b/Engine/Source/Runtime/MovieSceneTracks/Public/Sections/MovieSceneAudioSection.h
--- a/Engine/Source/Runtime/MovieSceneTracks/Public/Sections/MovieSceneAudioSection.h
+++ b/Engine/Source/Runtime/MovieSceneTracks/Public/Sections/MovieSceneAudioSection.h
@@ -8,6 +8,7 @@
 #include "Channels/MovieSceneIntegerChannel.h"
 #include "Channels/MovieSceneStringChannel.h"
 #include "Components/AudioComponent.h"
+#include "Decorations/MovieSceneScalingAnchors.h"
 #include "CoreMinimal.h"
 #include "EntitySystem/IMovieSceneEntityProvider.h"
 #include "MovieSceneSection.h"
@@ -26,6 +27,7 @@ UCLASS(MinimalAPI)
 class UMovieSceneAudioSection
 	: public UMovieSceneSection
 	, public IMovieSceneEntityProvider
+	, public IMovieSceneScalingDriver
 {
 	GENERATED_UCLASS_BODY()
 
@@ -39,6 +41,10 @@ public:
 	UFUNCTION(BlueprintPure, Category = "Sequencer|Section")
 	class USoundBase* GetSound() const {return Sound;}
 
+	/** Gets the sound that this section should use for playback, taking into account localization concerns */
+	UFUNCTION(BlueprintPure, Category = "Sequencer|Section")
+	MOVIESCENETRACKS_API USoundBase* GetPlaybackSound() const;
+
 	/** Set the offset into the beginning of the audio clip */
 	UFUNCTION(BlueprintCallable, Category = "Sequencer|Section")
 	void SetStartOffset(FFrameNumber InStartOffset) {StartFrameOffset = InStartOffset;}
@@ -196,6 +202,10 @@ public:
 	MOVIESCENETRACKS_API virtual bool PopulateEvaluationFieldImpl(const TRange<FFrameNumber>& EffectiveRange, const FMovieSceneEvaluationFieldEntityMetaData& InMetaData, FMovieSceneEntityComponentFieldBuilder* OutFieldBuilder) override;
 	MOVIESCENETRACKS_API virtual void ImportEntityImpl(UMovieSceneEntitySystemLinker* EntityLinker, const FEntityImportParams& Params, FImportedEntity* OutImportedEntity) override;
 
+	//~ IMovieSceneScalingDriver interface
+	MOVIESCENETRACKS_API virtual void PopulateInitialAnchors(TMap<FGuid, FMovieSceneScalingAnchor>& OutAnchors) override;
+	MOVIESCENETRACKS_API virtual void PopulateAnchors(TMap<FGuid, FMovieSceneScalingAnchor>& OutAnchors) override;
+
 private:
 
 	template<typename ChannelType, typename ForEachFunction>
diff --git a/Engine/Source/Runtime/MovieSceneTracks/Public/Tracks/MovieSceneAudioTrack.h b/Engine/Source/Runtime/MovieSceneTracks/Public/Tracks/MovieSceneAudioTrack.h
--- a/Engine/Source/Runtime/MovieSceneTracks/Public/Tracks/MovieSceneAudioTrack.h
+++ b/Engine/Source/Runtime/MovieSceneTracks/Public/Tracks/MovieSceneAudioTrack.h
@@ -53,6 +53,8 @@ public:
 	MOVIESCENETRACKS_API virtual bool SupportsMultipleRows() const override;
 	MOVIESCENETRACKS_API virtual UMovieSceneSection* CreateNewSection() override;
 
+	MOVIESCENETRACKS_API virtual void GetCompatibleUserDecorationsImpl(TSet<UClass*>& OutClasses) const override;
+
 	// ~UObject interface
 	MOVIESCENETRACKS_API virtual void PostRename(UObject* OldOuter, const FName OldName) override;
 
