Sequencer: Fix exit race condition in the task scheduler

There was a race condition between GameThreadSignal and NumRemainingTasks where NumRemainingTasks could be seen on the gamethread to be zero, which causes it to finalize the execution and clean up the event, all while the background thread might still be running OnAllTasksFinished.

To fix this, in a threading context, we assign special meaning to NumTasksRemaining == -1 to mean that OnAllTasksFinished has completely finished executing. This allows the gamethread to reason about the state of the final task's thread, and only break out of the loop when it is completely finished.

diff --git a/Engine/Source/Runtime/MovieScene/Private/EntitySystem/MovieSceneTaskScheduler.cpp b/Engine/Source/Runtime/MovieScene/Private/EntitySystem/MovieSceneTaskScheduler.cpp
--- a/Engine/Source/Runtime/MovieScene/Private/EntitySystem/MovieSceneTaskScheduler.cpp
+++ b/Engine/Source/Runtime/MovieScene/Private/EntitySystem/MovieSceneTaskScheduler.cpp
@@ -591,17 +591,39 @@ void FEntitySystemScheduler::ExecuteTasks()
 			Task->Run(this, FTaskExecutionFlags());
 		}
 
-		if (NumTasksRemaining.Load(ThreadingModel) == 0)
+		// Process exit conditions for this loop, where we have one of two situations:
+		//
+		//    1. The final task was just processed on this thread.
+		//       In this case OnAllTasksFinished will have decremented
+		//       NumTasksRemaining to -1 and we can break immediately.
+		//
+		//    2. The final task was completed on a thread.
+		//       In this case, this thread may see NumTasksRemaining as either 0 or -1
+		//       depending on how far through OnAllTasksFinished the bg thread is.
+		//       In this case, -1 means OnAllTasksFinished has finished, 0 means it is
+		//       part way through, so we just spin on the atomic until it hits -1
+		const int32 RemainingCount = NumTasksRemaining.Load(ThreadingModel);
+		if (RemainingCount == -1)
+		{
+			break;
+		}
+		else if (RemainingCount == 0)
 		{
+			// Handle the race condition where this code runs in-between
+			//    NumTasksRemaining reaching 0 and the GameThreadSignal being triggered.
+			// If the Trigger is called before the Wait below, some platforms may deadlock
+			while (NumTasksRemaining.Load(ThreadingModel) != -1);
 			break;
 		}
 
 		GameThreadSignal->Wait();
 	}
 
-	check(NumTasksRemaining.Load(ThreadingModel) == 0);
+	check(NumTasksRemaining.Load(ThreadingModel) == -1);
 	EntityManager->IncrementSystemSerial(SystemSerialIncrement);
 
+	NumTasksRemaining.Exchange(ThreadingModel, 0);
+
 	// Lastly return the game thread signal event to the pool as we are
 	// done executing.
 	check(GameThreadSignal);
@@ -714,6 +736,7 @@ void FEntitySystemScheduler::OnAllTasksFinished() const
 	if (ThreadingModel != EEntityThreadingModel::NoThreading)
 	{
 		GameThreadSignal->Trigger();
+		NumTasksRemaining.Sub(ThreadingModel, 1);
 	}
 }
 
