Sequencer: Fixed object replacement resulting in a mismatch in entity group bookkeeping

diff --git a/Engine/Source/Runtime/MovieScene/Public/EntitySystem/MovieSceneEntityGroupingSystem.h b/Engine/Source/Runtime/MovieScene/Public/EntitySystem/MovieSceneEntityGroupingSystem.h
--- a/Engine/Source/Runtime/MovieScene/Public/EntitySystem/MovieSceneEntityGroupingSystem.h
+++ b/Engine/Source/Runtime/MovieScene/Public/EntitySystem/MovieSceneEntityGroupingSystem.h
@@ -230,8 +230,8 @@ struct TEntityGroupingHandlerImpl<GroupingPolicy, TIntegerSequence<int, Componen
 	virtual void OnGroupIndexFreed(int32 GroupIndex) override
 	{
 		GroupKeyType Key = this->IndexToGroupKey.FindChecked(GroupIndex);
-		this->IndexToGroupKey.Remove(GroupIndex);
-		this->GroupKeyToIndex.Remove(Key);
+		ensure(this->IndexToGroupKey.Remove(GroupIndex) == 1);
+		ensure(this->GroupKeyToIndex.Remove(Key) == 1);
 	}
 
 #if WITH_EDITOR
@@ -255,6 +255,7 @@ struct TEntityGroupingHandlerImpl<GroupingPolicy, TIntegerSequence<int, Componen
 			if (ensure(bRemoved))
 			{
 				this->GroupKeyToIndex.Add(Pair.Value, GroupIndex);
+				this->IndexToGroupKey.Add(GroupIndex, Pair.Value);
 			}
 		}
 	}
