Sequencer: When moving tracks to a different owner with Rename(), use the existing track's names.

diff --git a/Engine/Source/Runtime/MovieScene/Private/MovieScene.cpp b/Engine/Source/Runtime/MovieScene/Private/MovieScene.cpp
--- a/Engine/Source/Runtime/MovieScene/Private/MovieScene.cpp
+++ b/Engine/Source/Runtime/MovieScene/Private/MovieScene.cpp
@@ -316,7 +316,11 @@ void UMovieScene::AddSpawnable(const FMovieSceneSpawnable& InNewSpawnable, const
 	FMovieSceneBinding NewBinding = InNewBinding;
 	for (UMovieSceneTrack* Track : NewBinding.GetTracks())
 	{
-		Track->Rename(nullptr, this, REN_DontCreateRedirectors);
+		if (!Track->IsIn(this))
+		{
+			FName NewName = MakeUniqueObjectName(this, Track->GetClass(), *Track->GetName());
+			Track->Rename(*NewName.ToString(), this, REN_DontCreateRedirectors);
+		}
 	}
 	int32 NewBindingIndex = InsertSorted(ObjectBindings, MoveTemp(NewBinding));
 
@@ -399,9 +403,13 @@ void UMovieScene::AddPossessable(const FMovieScenePossessable& InNewPossessable,
 	InsertSorted(Possessables, MoveTemp(NewPossessable));
 
 	FMovieSceneBinding NewBinding = InNewBinding;
-	for (auto Track : NewBinding.GetTracks())
+	for (UMovieSceneTrack* Track : NewBinding.GetTracks())
 	{
-		Track->Rename(nullptr, this, REN_DontCreateRedirectors);
+		if (!Track->IsIn(this))
+		{
+			FName NewName = MakeUniqueObjectName(this, Track->GetClass(), *Track->GetName());
+			Track->Rename(*NewName.ToString(), this, REN_DontCreateRedirectors);
+		}
 	}
 	int32 NewBindingIndex = InsertSorted(ObjectBindings, MoveTemp(NewBinding));
 	check(ObjectBindings.IsValidIndex(NewBindingIndex));
@@ -1639,9 +1647,13 @@ void UMovieScene::ReplaceBinding(const FGuid& BindingToReplaceGuid, const FMovie
 	{
 		*Binding = NewBinding;
 		// We also need to change the track's owners to be the MovieScene.
-		for (auto Track : Binding->GetTracks())
+		for (UMovieSceneTrack* Track : Binding->GetTracks())
 		{
-			Track->Rename(nullptr, this, REN_DontCreateRedirectors);
+			if (!Track->IsIn(this))
+			{
+				FName NewName = MakeUniqueObjectName(this, Track->GetClass(), *Track->GetName());
+				Track->Rename(*NewName.ToString(), this, REN_DontCreateRedirectors);
+			}
 		}
 
 		EventHandlers.Trigger(&UE::MovieScene::ISequenceDataEventHandler::OnBindingRemoved, BindingToReplaceGuid);
