Sequencer: add support for custom property handling

This makes it possible for Sequencer to animate properties on bound objects that are not UProperties. An example of this is animating data-driven properties exposed through a data asset, such as dynamic audio parameters on a MetaSound asset or dynamic camera parameters on a CameraRig asset.

The core of the change is to make the FTrackInstancePropertyBindings support custom handlers for finding properties and getting their value's address. This is done by registering factory methods that can create those handlers for a given type of object.

At the moment, these handlers are only queried on the top level object of a property path. There is also a flag that can disable querying and allocating such custom handlers, because some existing call sites in the editor are doing queries every tick for resolving property bindings, so it isn't desirable to potentially create custom handlers continuously like this. Another change is that property bindings don't always resolve to an FProperty so some call sites have been changed to handle the case of custom properties.

diff --git a/Engine/Source/Runtime/MovieScene/Private/EntitySystem/MovieScenePropertyRegistry.cpp b/Engine/Source/Runtime/MovieScene/Private/EntitySystem/MovieScenePropertyRegistry.cpp
--- a/Engine/Source/Runtime/MovieScene/Private/EntitySystem/MovieScenePropertyRegistry.cpp
+++ b/Engine/Source/Runtime/MovieScene/Private/EntitySystem/MovieScenePropertyRegistry.cpp
@@ -199,7 +199,7 @@ TOptional<FResolvedProperty> FPropertyRegistry::ResolveProperty(UObject* Object,
 
 	// None of the above optimized paths can apply to this property (probably because it has a setter function or because it is within a compound property), so we must use the slow property bindings
 	TSharedPtr<FTrackInstancePropertyBindings> SlowBindings = MakeShared<FTrackInstancePropertyBindings>(PropertyBinding.PropertyName, PropertyBinding.PropertyPath.ToString());
-	if (SlowBindings->GetProperty(*Object) == nullptr)
+	if (!SlowBindings->HasValidBinding(*Object))
 	{
 		UE_LOG(LogMovieSceneECS, Warning, TEXT("Unable to resolve property '%s' from '%s' instance '%s'"), *PropertyBinding.PropertyPath.ToString(), *Object->GetClass()->GetName(), *Object->GetName());
 		return TOptional<FResolvedProperty>();
@@ -224,4 +224,4 @@ void FPropertyDefinition::SetupInitialValueProcessor() const
 }
 
 } // namespace MovieScene
-} // namespace UE
\ No newline at end of file
+} // namespace UE
diff --git a/Engine/Source/Runtime/MovieScene/Private/MovieSceneCommonHelpers.cpp b/Engine/Source/Runtime/MovieScene/Private/MovieSceneCommonHelpers.cpp
--- a/Engine/Source/Runtime/MovieScene/Private/MovieSceneCommonHelpers.cpp
+++ b/Engine/Source/Runtime/MovieScene/Private/MovieSceneCommonHelpers.cpp
@@ -1185,9 +1185,25 @@ MovieSceneHelpers::FMovieSceneScopedPackageDirtyGuard::~FMovieSceneScopedPackage
 	}
 }
 
-FTrackInstancePropertyBindings::FTrackInstancePropertyBindings( FName InPropertyName, const FString& InPropertyPath )
-	: PropertyPath( InPropertyPath )
-	, PropertyName( InPropertyName )
+TMap<UClass*, FOnCreateTrackInstanceCustomPropertyBindingHandler> GCustomHandlers;
+
+FDelegateHandle FTrackInstancePropertyBindings::RegisterCustomHandler(UClass* InClass, FOnCreateTrackInstanceCustomPropertyBindingHandler InDelegate)
+{
+	using namespace UE::MovieScene;
+	GCustomHandlers.Add(InClass, InDelegate);
+	return GCustomHandlers[InClass].GetHandle();
+}
+
+void FTrackInstancePropertyBindings::UnregisterCustomHandler(UClass* InClass)
+{
+	using namespace UE::MovieScene;
+	GCustomHandlers.Remove(InClass);
+}
+
+FTrackInstancePropertyBindings::FTrackInstancePropertyBindings(FName InPropertyName, const FString& InPropertyPath, bool bInAllowCustomHandlers)
+	: PropertyPath(InPropertyPath)
+	, PropertyName(InPropertyName)
+	, bAllowCustomHandlers(bInAllowCustomHandlers)
 {
 	static const FString Set(TEXT("Set"));
 	const FString FunctionString = Set + PropertyName.ToString();
@@ -1195,27 +1211,61 @@ FTrackInstancePropertyBindings::FTrackInstancePropertyBindings( FName InProperty
 	FunctionName = FName(*FunctionString);
 }
 
-struct FPropertyAndIndex
+void FTrackInstancePropertyBindings::FindProperty(FTrackInstancePropertyBindings& Bindings, void* BasePointer, UStruct* InStruct, const FString& InPropertyName, FResolvedProperty& OutResolvedProperty, bool bAllowCustomHandlers)
 {
-	FPropertyAndIndex() : Property(nullptr), ArrayIndex(INDEX_NONE) {}
+	// Find a property via UObject reflection first.
+	FProperty* FoundProperty = FindFProperty<FProperty>(InStruct, *InPropertyName);
 
-	FProperty* Property;
-	int32 ArrayIndex;
-};
+	if (FoundProperty)
+	{
+		OutResolvedProperty.Property = FoundProperty;
+	}
+	else if (bAllowCustomHandlers)
+	{
+		// Find an existing handler for this class, or try to create one.
+		// For now, custom handlers must match the exact object class. Handlers for a parent class can't
+		// necessarily resolve properties on children classes.
+		UClass* ObjectClass = Cast<UClass>(InStruct);
+		TSharedPtr<ITrackInstanceCustomPropertyBindingHandler> Handler = Bindings.ObjectClassToCustomHandlerMap.FindRef(ObjectClass);
+		if (!Handler)
+		{
+			FOnCreateTrackInstanceCustomPropertyBindingHandler* Factory = GCustomHandlers.Find(ObjectClass);
+			if (Factory)
+			{
+				Handler = Factory->Execute();
+				if (ensure(Handler))
+				{
+					Bindings.ObjectClassToCustomHandlerMap.Add(ObjectClass, Handler);
+				}
+			}
+		}
+		if (Handler)
+		{
+			FTrackInstanceCustomProperty CustomProperty;
+			if (Handler->FindProperty(BasePointer, InPropertyName, CustomProperty))
+			{
+				OutResolvedProperty.CustomProperty.Handler = Handler;
+				OutResolvedProperty.CustomProperty.Property = CustomProperty;
+			}
+		}
+	}
+}
 
-FPropertyAndIndex FindPropertyAndArrayIndex(UStruct* InStruct, const FString& PropertyName)
+FTrackInstancePropertyBindings::FResolvedProperty FTrackInstancePropertyBindings::FindPropertyAndArrayIndex(FTrackInstancePropertyBindings& Bindings, void* BasePointer, UStruct* InStruct, const FString& PropertyName, bool bAllowCustomHandlers)
 {
-	FPropertyAndIndex PropertyAndIndex;
+	FResolvedProperty PropertyAndIndex;
 
-	// Calculate the array index if possible
+	// Calculate the array index if possible.
 	int32 ArrayIndex = -1;
 	if (PropertyName.Len() > 0 && PropertyName.GetCharArray()[PropertyName.Len() - 1] == ']')
 	{
 		int32 OpenIndex = 0;
 		if (PropertyName.FindLastChar('[', OpenIndex))
 		{
+			// We have a property name of the form "Foo[123]". Resolve the property itself ("Foo") and then
+			// parse the array element index (123).
 			FString TruncatedPropertyName = FString::ConstructFromPtrSize(*PropertyName, OpenIndex);
-			PropertyAndIndex.Property = FindFProperty<FProperty>(InStruct, *TruncatedPropertyName);
+			FindProperty(Bindings, BasePointer, InStruct, *TruncatedPropertyName, PropertyAndIndex, bAllowCustomHandlers);
 
 			const int32 NumberLength = PropertyName.Len() - OpenIndex - 2;
 			if (NumberLength > 0 && NumberLength <= 10)
@@ -1229,146 +1279,162 @@ FPropertyAndIndex FindPropertyAndArrayIndex(UStruct* InStruct, const FString& Pr
 		}
 	}
 
-	PropertyAndIndex.Property = FindFProperty<FProperty>(InStruct, *PropertyName);
-	return PropertyAndIndex;
-}
-
-
-FProperty* FTrackInstancePropertyBindings::FindPropertyRecursive(UStruct* InStruct, TArray<FString>& InPropertyNames, uint32 Index)
-{
-	FProperty* Property = FindPropertyAndArrayIndex(InStruct, *InPropertyNames[Index]).Property;
-	if (!Property)
-	{
-		return nullptr;
-	}
-
-	/* Recursive property types */
-	if (Property->IsA(FArrayProperty::StaticClass()))
-	{
-		FArrayProperty* ArrayProp = CastFieldChecked<FArrayProperty>(Property);
-
-		FStructProperty* InnerStructProp = CastField<FStructProperty>(ArrayProp->Inner);
-		if (InnerStructProp && InPropertyNames.IsValidIndex(Index + 1))
-		{
-			return FindPropertyRecursive(InnerStructProp->Struct, InPropertyNames, Index + 1);
-		}
-		else
-		{
-			Property = ArrayProp->Inner;
-		}
-	}
-	else if (FStructProperty* StructProp = CastField<FStructProperty>(Property))
-	{
-		if( InPropertyNames.IsValidIndex(Index+1) )
-		{
-			return FindPropertyRecursive(StructProp->Struct, InPropertyNames, Index+1);
-		}
-		else
-		{
-			check( StructProp->GetName() == InPropertyNames[Index] );
-		}
-	}
-
-	return Property;
-}
-
-FProperty* FTrackInstancePropertyBindings::FindProperty(const UObject* Object, const FString& InPropertyPath)
-{
-	check(Object);
-
-	TArray<FString> PropertyNames;
-	InPropertyPath.ParseIntoArray(PropertyNames, TEXT("."), true);
-
-	if (PropertyNames.Num() > 0)
-	{
-		return FindPropertyRecursive(Object->GetClass(), PropertyNames, 0);
-	}
+	// Not index found in this property name, just find the property normally.
+	FindProperty(Bindings, BasePointer, InStruct, *PropertyName, PropertyAndIndex, bAllowCustomHandlers);
 
-	return nullptr;
+	return PropertyAndIndex;
 }
 
-FTrackInstancePropertyBindings::FPropertyAddress FTrackInstancePropertyBindings::FindPropertyAddressRecursive( void* BasePointer, UStruct* InStruct, TArray<FString>& InPropertyNames, uint32 Index )
+FTrackInstancePropertyBindings::FResolvedProperty FTrackInstancePropertyBindings::ResolvePropertyRecursive(FTrackInstancePropertyBindings& Bindings, void* BasePointer, UStruct* InStruct, TArray<FString>& InPropertyNames, uint32 Index)
 {
-	FPropertyAndIndex PropertyAndIndex = FindPropertyAndArrayIndex(InStruct, *InPropertyNames[Index]);
+	const bool bAllowCustomHandlers = Bindings.bAllowCustomHandlers && (Index == 0);
+	FResolvedProperty ResolvedProperty = FindPropertyAndArrayIndex(Bindings, BasePointer, InStruct, *InPropertyNames[Index], bAllowCustomHandlers);
 	
-	FTrackInstancePropertyBindings::FPropertyAddress NewAddress;
+	FResolvedProperty NewProperty;
 
-	if (PropertyAndIndex.ArrayIndex != INDEX_NONE)
+	FProperty* ResolvedPropertyProperty = ResolvedProperty.Property.Get();
+	FCustomProperty ResolvedCustomProperty = ResolvedProperty.CustomProperty;
+
+	if (ResolvedProperty.ArrayIndex != INDEX_NONE)
 	{
-		if (PropertyAndIndex.Property)
+		// We found that this segment of the property path reaches an element inside an array.
+		if (ResolvedPropertyProperty)
 		{
-			if (PropertyAndIndex.Property->IsA(FArrayProperty::StaticClass()))
+			if (FArrayProperty* ArrayProp = CastField<FArrayProperty>(ResolvedPropertyProperty))
 			{
-				FArrayProperty* ArrayProp = CastFieldChecked<FArrayProperty>(PropertyAndIndex.Property);
-
 				FScriptArrayHelper ArrayHelper(ArrayProp, ArrayProp->ContainerPtrToValuePtr<void>(BasePointer));
-				if (ArrayHelper.IsValidIndex(PropertyAndIndex.ArrayIndex))
+				if (ArrayHelper.IsValidIndex(ResolvedProperty.ArrayIndex))
 				{
 					FStructProperty* InnerStructProp = CastField<FStructProperty>(ArrayProp->Inner);
 					if (InnerStructProp && InPropertyNames.IsValidIndex(Index + 1))
 					{
-						return FindPropertyAddressRecursive(ArrayHelper.GetRawPtr(PropertyAndIndex.ArrayIndex), InnerStructProp->Struct, InPropertyNames, Index + 1);
+						// Move the BasePointer to the array element and keep resolving the property path on it.
+						void* ArrayElement = ArrayHelper.GetRawPtr(ResolvedProperty.ArrayIndex);
+						return ResolvePropertyRecursive(Bindings, ArrayElement, InnerStructProp->Struct, InPropertyNames, Index + 1);
 					}
 					else
 					{
-						NewAddress.Property = ArrayProp->Inner;
-						NewAddress.Address = ArrayHelper.GetRawPtr(PropertyAndIndex.ArrayIndex);
+						// The property path ends here (e.g. "Foo.Bar[1]").
+						NewProperty.Property = ArrayProp;
+						NewProperty.ContainerAddress = BasePointer;
+						NewProperty.ArrayIndex = ResolvedProperty.ArrayIndex;
 					}
 				}
 			}
 			else
 			{
-				UE_LOG(LogMovieScene, Error, TEXT("Mismatch in property evaluation. %s is not of type: %s"), *PropertyAndIndex.Property->GetName(), *FArrayProperty::StaticClass()->GetName());
+				UE_LOG(LogMovieScene, Error, TEXT("Mismatch in property evaluation. %s is not of type: %s"), *ResolvedPropertyProperty->GetName(), *FArrayProperty::StaticClass()->GetName());
+			}
+		}
+		else if (ResolvedCustomProperty.IsValid())
+		{
+			if (ResolvedCustomProperty.Property.ContainerType == ETrackInstancePropertyContainerType::Array)
+			{
+				if (InPropertyNames.IsValidIndex(Index + 1))
+				{
+					// Move the BasePointer to the array element and keep resolving the property path on it.
+					void* ArrayElement = ResolvedCustomProperty.Handler->ContainerPtrToValuePtr(BasePointer, ResolvedCustomProperty.Property, ResolvedProperty.ArrayIndex);
+					return ResolvePropertyRecursive(Bindings, ArrayElement, ResolvedCustomProperty.Property.PropertyStruct, InPropertyNames, Index + 1);
+				}
+				else
+				{
+					// The property path ends here (e.g. "Foo.Bar[1]").
+					NewProperty.CustomProperty = ResolvedCustomProperty;
+					NewProperty.ContainerAddress = BasePointer;
+					NewProperty.ArrayIndex = ResolvedProperty.ArrayIndex;
+				}
+			}
+			else
+			{
+				UE_LOG(LogMovieScene, Error, TEXT("Mismatch in property evaluation. Custom property %d is not an array property"), ResolvedCustomProperty.Property.PropertyID);
 			}
 		}
 	}
-	else if (FStructProperty* StructProp = CastField<FStructProperty>(PropertyAndIndex.Property))
+	else if (FStructProperty* StructProp = CastField<FStructProperty>(ResolvedPropertyProperty))
 	{
-		NewAddress.Property = StructProp;
-		NewAddress.Address = BasePointer;
+		// This segment of the property path reaches a struct FProperty.
+		NewProperty.Property = StructProp;
+		NewProperty.ContainerAddress = BasePointer;
 
-		if( InPropertyNames.IsValidIndex(Index+1) )
+		if (InPropertyNames.IsValidIndex(Index + 1))
 		{
 			void* StructContainer = StructProp->ContainerPtrToValuePtr<void>(BasePointer);
-			return FindPropertyAddressRecursive( StructContainer, StructProp->Struct, InPropertyNames, Index+1 );
+			return ResolvePropertyRecursive(Bindings, StructContainer, StructProp->Struct, InPropertyNames, Index + 1);
 		}
 		else
 		{
-			check( StructProp->GetName() == InPropertyNames[Index] );
+			check(StructProp->GetName() == InPropertyNames[Index]);
 		}
 	}
-	else if(PropertyAndIndex.Property)
+	else if (ResolvedPropertyProperty)
 	{
-		NewAddress.Property = PropertyAndIndex.Property;
-		NewAddress.Address = BasePointer;
+		NewProperty.Property = ResolvedProperty.Property;
+		NewProperty.ContainerAddress = BasePointer;
 	}
+	else if (ResolvedCustomProperty.IsValid())
+	{
+		NewProperty.CustomProperty = ResolvedCustomProperty;
+		NewProperty.ContainerAddress = BasePointer;
 
-	return NewAddress;
+		const FTrackInstanceCustomProperty& CustomProperty = NewProperty.CustomProperty.Property;
+		if (CustomProperty.PropertyStruct && 
+				CustomProperty.PropertyStruct->IsA<UScriptStruct>() && 
+				InPropertyNames.IsValidIndex(Index + 1))
+		{
+			void* StructContainer = NewProperty.CustomProperty.Handler->ContainerPtrToValuePtr(BasePointer, CustomProperty);
+			return ResolvePropertyRecursive(Bindings, StructContainer, CustomProperty.PropertyStruct, InPropertyNames, Index + 1);
+		}
+	}
 
+	return NewProperty;
 }
 
-
-FTrackInstancePropertyBindings::FPropertyAddress FTrackInstancePropertyBindings::FindPropertyAddress( const UObject& InObject, const FString& InPropertyPath )
+FTrackInstancePropertyBindings::FResolvedProperty FTrackInstancePropertyBindings::ResolveProperty(FTrackInstancePropertyBindings& Bindings, const UObject& InObject)
 {
 	TArray<FString> PropertyNames;
 
-	InPropertyPath.ParseIntoArray(PropertyNames, TEXT("."), true);
+	const FString& PropertyPath = Bindings.PropertyPath;
+	PropertyPath.ParseIntoArray(PropertyNames, TEXT("."), true);
 
-	if(IsValid(&InObject) && PropertyNames.Num() > 0)
+	if (IsValid(&InObject) && PropertyNames.Num() > 0)
 	{
-		return FindPropertyAddressRecursive( (void*)&InObject, InObject.GetClass(), PropertyNames, 0 );
+		return ResolvePropertyRecursive(Bindings, (void*)&InObject, InObject.GetClass(), PropertyNames, 0);
 	}
 	else
 	{
-		return FTrackInstancePropertyBindings::FPropertyAddress();
+		return FTrackInstancePropertyBindings::FResolvedProperty();
+	}
+}
+
+FProperty* FTrackInstancePropertyBindings::FindProperty(const UObject* Object, const FString& InPropertyPath)
+{
+	FTrackInstancePropertyBindings Temp(NAME_None, InPropertyPath, false);
+	FResolvedProperty ResolvedProperty = ResolveProperty(Temp, *Object);
+	return ResolvedProperty.Property.Get();
+}
+
+FTrackInstancePropertyBindings::FResolvedPropertyAndFunction FTrackInstancePropertyBindings::FindOrAdd(const UObject& InObject)
+{
+	FObjectKey ObjectKey(&InObject);
+
+	const FResolvedPropertyAndFunction* PropAndFunction = RuntimeObjectToFunctionMap.Find(ObjectKey);
+	if (PropAndFunction && (
+				PropAndFunction->SetterFunction.IsValid() ||
+				PropAndFunction->ResolvedProperty.GetValidProperty() ||
+				PropAndFunction->ResolvedProperty.CustomProperty.IsValid()))
+	{
+		return *PropAndFunction;
 	}
+
+	CacheBinding(InObject);
+	return RuntimeObjectToFunctionMap.FindRef(ObjectKey);
 }
 
-void FTrackInstancePropertyBindings::CallFunctionForEnum( UObject& InRuntimeObject, int64 PropertyValue )
+void FTrackInstancePropertyBindings::CallFunctionForEnum(UObject& InRuntimeObject, int64 PropertyValue)
 {
-	FPropertyAndFunction PropAndFunction = FindOrAdd(InRuntimeObject);
+	FResolvedPropertyAndFunction PropAndFunction = FindOrAdd(InRuntimeObject);
 
-	FProperty* Property = GetProperty(InRuntimeObject);
+	FProperty* Property = PropAndFunction.ResolvedProperty.GetValidProperty();
 	if (Property && Property->HasSetter())
 	{
 		Property->CallSetter(&InRuntimeObject, &PropertyValue);
@@ -1377,18 +1443,25 @@ void FTrackInstancePropertyBindings::CallFunctionForEnum( UObject& InRuntimeObje
 	{
 		InvokeSetterFunction(&InRuntimeObject, SetterFunction, PropertyValue);
 	}
-	else if (Property && Property->IsA(FEnumProperty::StaticClass()))
+	else if (Property)
 	{
-		if (FEnumProperty* EnumProperty = CastFieldChecked<FEnumProperty>(Property))
+		if (FEnumProperty* EnumProperty = CastField<FEnumProperty>(Property))
 		{
 			FNumericProperty* UnderlyingProperty = EnumProperty->GetUnderlyingProperty();
-			void* ValueAddr = EnumProperty->ContainerPtrToValuePtr<void>(PropAndFunction.PropertyAddress.Address);
+			void* ValueAddr = EnumProperty->ContainerPtrToValuePtr<void>(PropAndFunction.ResolvedProperty.ContainerAddress);
 			UnderlyingProperty->SetIntPropertyValue(ValueAddr, PropertyValue);
 		}
+		else
+		{
+			UE_LOG(LogMovieScene, Error, TEXT("Mismatch in property evaluation. %s is not of type: %s"), *Property->GetName(), *FEnumProperty::StaticClass()->GetName());
+		}
 	}
-	else if (Property)
+	else if (PropAndFunction.ResolvedProperty.HasValidCustomProperty())
 	{
-		UE_LOG(LogMovieScene, Error, TEXT("Mismatch in property evaluation. %s is not of type: %s"), *Property->GetName(), *FEnumProperty::StaticClass()->GetName());
+		// NOTE: custom handlers must always use uint8 values for enums.
+		const FCustomProperty& CustomProperty = PropAndFunction.ResolvedProperty.CustomProperty;
+		void* ValueAddr = CustomProperty.Handler->ContainerPtrToValuePtr(PropAndFunction.ResolvedProperty.ContainerAddress, CustomProperty.Property);
+		*(uint8*)ValueAddr = (uint8)PropertyValue;
 	}
 
 	if (UFunction* NotifyFunction = PropAndFunction.NotifyFunction.Get())
@@ -1399,16 +1472,19 @@ void FTrackInstancePropertyBindings::CallFunctionForEnum( UObject& InRuntimeObje
 
 void FTrackInstancePropertyBindings::CacheBinding(const UObject& Object)
 {
-	FPropertyAndFunction PropAndFunction;
+	FResolvedPropertyAndFunction PropAndFunction;
 	{
-		PropAndFunction.PropertyAddress = FindPropertyAddress(Object, PropertyPath);
+		PropAndFunction.ResolvedProperty = ResolveProperty(*this, Object);
 
-		UFunction* SetterFunction = Object.FindFunction(FunctionName);
-		if (SetterFunction && SetterFunction->NumParms >= 1)
+		if (!PropAndFunction.ResolvedProperty.CustomProperty.IsValid())
 		{
-			PropAndFunction.SetterFunction = SetterFunction;
+			UFunction* SetterFunction = Object.FindFunction(FunctionName);
+			if (SetterFunction && SetterFunction->NumParms >= 1)
+			{
+				PropAndFunction.SetterFunction = SetterFunction;
+			}
 		}
-		
+
 		UFunction* NotifyFunction = NotifyFunctionName != NAME_None ? Object.FindFunction(NotifyFunctionName) : nullptr;
 		if (NotifyFunction && NotifyFunction->NumParms == 0 && NotifyFunction->ReturnValueOffset == MAX_uint16)
 		{
@@ -1419,47 +1495,75 @@ void FTrackInstancePropertyBindings::CacheBinding(const UObject& Object)
 	RuntimeObjectToFunctionMap.Add(FObjectKey(&Object), PropAndFunction);
 }
 
-FProperty* FTrackInstancePropertyBindings::GetProperty(const UObject& Object) const
+FProperty* FTrackInstancePropertyBindings::GetProperty(const UObject& Object)
+{
+	FResolvedPropertyAndFunction PropAndFunction = FindOrAdd(Object);
+	return PropAndFunction.ResolvedProperty.GetValidProperty();
+}
+
+bool FTrackInstancePropertyBindings::HasValidBinding(const UObject& Object)
+{
+	FResolvedPropertyAndFunction PropAndFunction = FindOrAdd(Object);
+	return PropAndFunction.ResolvedProperty.GetValidProperty() || PropAndFunction.ResolvedProperty.HasValidCustomProperty();
+}
+
+const UStruct* FTrackInstancePropertyBindings::GetPropertyStruct(const UObject& Object)
 {
-	FPropertyAndFunction PropAndFunction = RuntimeObjectToFunctionMap.FindRef(&Object);
-	if (FProperty* Property = PropAndFunction.PropertyAddress.GetProperty())
+	FResolvedPropertyAndFunction PropAndFunction = FindOrAdd(Object);
+
+	if (FProperty* Property = PropAndFunction.ResolvedProperty.GetValidProperty())
+	{
+		if (FStructProperty* StructProperty = CastField<FStructProperty>(Property))
+		{
+			return StructProperty->Struct;
+		}
+		return nullptr;
+	}
+	else if (PropAndFunction.ResolvedProperty.HasValidCustomProperty())
 	{
-		return Property;
+		return PropAndFunction.ResolvedProperty.CustomProperty.Property.PropertyStruct;
 	}
 
-	return FindPropertyAddress(Object, PropertyPath).GetProperty();
+	return nullptr;
 }
 
 int64 FTrackInstancePropertyBindings::GetCurrentValueForEnum(const UObject& Object)
 {
-	FPropertyAndFunction PropAndFunction = FindOrAdd(Object);
+	FResolvedPropertyAndFunction PropAndFunction = FindOrAdd(Object);
+
+	FProperty* Property = PropAndFunction.ResolvedProperty.GetValidProperty();
 
-	if (FProperty* Property = PropAndFunction.PropertyAddress.GetProperty())
+	if (Property)
 	{
-		if (Property->IsA(FEnumProperty::StaticClass()))
+		if (FEnumProperty* EnumProperty = CastField<FEnumProperty>(Property))
 		{
-			if (FEnumProperty* EnumProperty = CastFieldChecked<FEnumProperty>(Property))
-			{
-				FNumericProperty* UnderlyingProperty = EnumProperty->GetUnderlyingProperty();
-				void* ValueAddr = EnumProperty->ContainerPtrToValuePtr<void>(PropAndFunction.PropertyAddress.Address);
-				int64 Result = UnderlyingProperty->GetSignedIntPropertyValue(ValueAddr);
-				return Result;
-			}
+			FNumericProperty* UnderlyingProperty = EnumProperty->GetUnderlyingProperty();
+			void* ValueAddr = EnumProperty->ContainerPtrToValuePtr<void>(PropAndFunction.ResolvedProperty.ContainerAddress);
+			int64 Result = UnderlyingProperty->GetSignedIntPropertyValue(ValueAddr);
+			return Result;
 		}
 		else
 		{
 			UE_LOG(LogMovieScene, Error, TEXT("Mismatch in property evaluation. %s is not of type: %s"), *Property->GetName(), *FEnumProperty::StaticClass()->GetName());
 		}
 	}
+	else if (PropAndFunction.ResolvedProperty.HasValidCustomProperty())
+	{
+		// NOTE: custom handlers must always use uint8 values for enums.
+		const FCustomProperty& CustomProperty = PropAndFunction.ResolvedProperty.CustomProperty;
+		void* ValueAddr = CustomProperty.Handler->ContainerPtrToValuePtr(PropAndFunction.ResolvedProperty.ContainerAddress, CustomProperty.Property);
+		return (int64)(*(uint8*)ValueAddr);
+	}
 
 	return 0;
 }
 
 template<> void FTrackInstancePropertyBindings::CallFunction<bool>(UObject& InRuntimeObject, TCallTraits<bool>::ParamType PropertyValue)
 {
-	FPropertyAndFunction PropAndFunction = FindOrAdd(InRuntimeObject);
+	FResolvedPropertyAndFunction PropAndFunction = FindOrAdd(InRuntimeObject);
+
+	FProperty* Property = PropAndFunction.ResolvedProperty.GetValidProperty();
 
-	FProperty* Property = GetProperty(InRuntimeObject);
 	if (Property && Property->HasSetter())
 	{
 		Property->CallSetter(&InRuntimeObject, &PropertyValue);
@@ -1468,17 +1572,23 @@ template<> void FTrackInstancePropertyBindings::CallFunction<bool>(UObject& InRu
 	{
 		InvokeSetterFunction(&InRuntimeObject, SetterFunction, PropertyValue);
 	}
-	else if (Property && Property->IsA(FBoolProperty::StaticClass()))
+	else if (Property)
 	{
-		if (FBoolProperty* BoolProperty = CastFieldChecked<FBoolProperty>(Property))
+		if (FBoolProperty* BoolProperty = CastField<FBoolProperty>(Property))
 		{
-			uint8* ValuePtr = BoolProperty->ContainerPtrToValuePtr<uint8>(PropAndFunction.PropertyAddress.Address);
+			uint8* ValuePtr = BoolProperty->ContainerPtrToValuePtr<uint8>(PropAndFunction.ResolvedProperty.ContainerAddress);
 			BoolProperty->SetPropertyValue(ValuePtr, PropertyValue);
 		}
+		else
+		{
+			UE_LOG(LogMovieScene, Error, TEXT("Mismatch in property evaluation. %s is not of type: %s"), *Property->GetName(), *FBoolProperty::StaticClass()->GetName());
+		}
 	}
-	else if (Property)
+	else if (PropAndFunction.ResolvedProperty.HasValidCustomProperty())
 	{
-		UE_LOG(LogMovieScene, Error, TEXT("Mismatch in property evaluation. %s is not of type: %s"), *Property->GetName(), *FBoolProperty::StaticClass()->GetName());
+		// NOTE: custom boolean properties should never be bitfields.
+		bool* ValuePtr = PropAndFunction.ResolvedProperty.GetPropertyAddress<bool>();
+		*ValuePtr = PropertyValue;
 	}
 
 	if (UFunction* NotifyFunction = PropAndFunction.NotifyFunction.Get())
@@ -1487,39 +1597,51 @@ template<> void FTrackInstancePropertyBindings::CallFunction<bool>(UObject& InRu
 	}
 }
 
-template<> bool FTrackInstancePropertyBindings::ResolvePropertyValue<bool>(const FPropertyAddress& Address, bool& OutValue)
+template<> bool FTrackInstancePropertyBindings::TryGetPropertyValue<bool>(const FResolvedProperty& ResolvedProperty, bool& OutValue)
 {
-	if (FProperty* Property = Address.GetProperty())
+	if (FProperty* Property = ResolvedProperty.GetValidProperty())
 	{
-		if (Property->IsA(FBoolProperty::StaticClass()))
+		if (FBoolProperty* BoolProperty = CastField<FBoolProperty>(Property))
 		{
-			if (FBoolProperty* BoolProperty = CastFieldChecked<FBoolProperty>(Property))
-			{
-				const uint8* ValuePtr = BoolProperty->ContainerPtrToValuePtr<uint8>(Address.Address);
-				OutValue = BoolProperty->GetPropertyValue(ValuePtr);
-				return true;
-			}
+			const uint8* ValuePtr = BoolProperty->ContainerPtrToValuePtr<uint8>(ResolvedProperty.ContainerAddress);
+			OutValue = BoolProperty->GetPropertyValue(ValuePtr);
+			return true;
 		}
 		else
 		{
 			UE_LOG(LogMovieScene, Error, TEXT("Mismatch in property evaluation. %s is not of type: %s"), *Property->GetName(), *FBoolProperty::StaticClass()->GetName());
 		}
 	}
+	else if (ResolvedProperty.HasValidCustomProperty())
+	{
+		// NOTE: custom boolean properties should never be bitfields.
+		const FCustomProperty& CustomProperty = ResolvedProperty.CustomProperty;
+		const void* ValuePtr = CustomProperty.Handler->ContainerPtrToValuePtr(ResolvedProperty.ContainerAddress, CustomProperty.Property);
+		OutValue = *(bool*)ValuePtr;
+	}
 
 	return false;
 }
 
 template<> void FTrackInstancePropertyBindings::SetCurrentValue<bool>(UObject& Object, TCallTraits<bool>::ParamType InValue)
 {
-	FPropertyAndFunction PropAndFunction = FindOrAdd(Object);
-	if (FProperty* Property = PropAndFunction.PropertyAddress.GetProperty())
+	FResolvedPropertyAndFunction PropAndFunction = FindOrAdd(Object);
+
+	if (FProperty* Property = PropAndFunction.ResolvedProperty.GetValidProperty())
 	{
 		if (FBoolProperty* BoolProperty = CastField<FBoolProperty>(Property))
 		{
-			uint8* ValuePtr = BoolProperty->ContainerPtrToValuePtr<uint8>(PropAndFunction.PropertyAddress.Address);
+			uint8* ValuePtr = BoolProperty->ContainerPtrToValuePtr<uint8>(PropAndFunction.ResolvedProperty.ContainerAddress);
 			BoolProperty->SetPropertyValue(ValuePtr, InValue);
 		}
 	}
+	else if (PropAndFunction.ResolvedProperty.HasValidCustomProperty())
+	{
+		// NOTE: custom boolean properties should never be bitfields.
+		const FCustomProperty& CustomProperty = PropAndFunction.ResolvedProperty.CustomProperty;
+		void* ValuePtr = CustomProperty.Handler->ContainerPtrToValuePtr(PropAndFunction.ResolvedProperty.ContainerAddress, CustomProperty.Property);
+		*(bool*)ValuePtr = InValue;
+	}
 
 	if (UFunction* NotifyFunction = PropAndFunction.NotifyFunction.Get())
 	{
@@ -1527,12 +1649,12 @@ template<> void FTrackInstancePropertyBindings::SetCurrentValue<bool>(UObject& O
 	}
 }
 
-
 template<> void FTrackInstancePropertyBindings::CallFunction<UObject*>(UObject& InRuntimeObject, UObject* PropertyValue)
 {
-	FPropertyAndFunction PropAndFunction = FindOrAdd(InRuntimeObject);
+	FResolvedPropertyAndFunction PropAndFunction = FindOrAdd(InRuntimeObject);
+
+	FProperty* Property = PropAndFunction.ResolvedProperty.GetValidProperty();
 
-	FProperty* Property = GetProperty(InRuntimeObject);
 	if (Property && Property->HasSetter())
 	{
 		Property->CallSetter(&InRuntimeObject, &PropertyValue);
@@ -1541,17 +1663,24 @@ template<> void FTrackInstancePropertyBindings::CallFunction<UObject*>(UObject&
 	{
 		InvokeSetterFunction(&InRuntimeObject, SetterFunction, PropertyValue);
 	}
-	else if (Property && Property->IsA(FObjectPropertyBase::StaticClass()))
+	else if (Property)
 	{
-		if (FObjectPropertyBase* ObjectProperty = CastFieldChecked<FObjectPropertyBase>(Property))
+		if (FObjectPropertyBase* ObjectProperty = CastField<FObjectPropertyBase>(Property))
 		{
-			uint8* ValuePtr = ObjectProperty->ContainerPtrToValuePtr<uint8>(PropAndFunction.PropertyAddress.Address);
+			uint8* ValuePtr = ObjectProperty->ContainerPtrToValuePtr<uint8>(PropAndFunction.ResolvedProperty.ContainerAddress);
 			ObjectProperty->SetObjectPropertyValue(ValuePtr, PropertyValue);
 		}
+		else
+		{
+			UE_LOG(LogMovieScene, Error, TEXT("Mismatch in property evaluation. %s is not of type: %s"), *Property->GetName(), *FObjectPropertyBase::StaticClass()->GetName());
+		}
 	}
-	else if (Property)
+	else if (PropAndFunction.ResolvedProperty.HasValidCustomProperty())
 	{
-		UE_LOG(LogMovieScene, Error, TEXT("Mismatch in property evaluation. %s is not of type: %s"), *Property->GetName(), *FObjectPropertyBase::StaticClass()->GetName());
+		// NOTE: custom handlers must treat object pointers as TObjectPtr<T>. Any type fits, since the templated
+		// smart object pointers are covariant.
+		TObjectPtr<UObject>* ValuePtr = PropAndFunction.ResolvedProperty.GetPropertyAddress<TObjectPtr<UObject>>();
+		(*ValuePtr) = PropertyValue;
 	}
 
 	if (UFunction* NotifyFunction = PropAndFunction.NotifyFunction.Get())
@@ -1560,29 +1689,57 @@ template<> void FTrackInstancePropertyBindings::CallFunction<UObject*>(UObject&
 	}
 }
 
-template<> bool FTrackInstancePropertyBindings::ResolvePropertyValue<UObject*>(const FPropertyAddress& Address, UObject*& OutValue)
+template<> bool FTrackInstancePropertyBindings::TryGetPropertyValue<UObject*>(const FResolvedProperty& ResolvedProperty, UObject*& OutValue)
 {
-	if (FObjectPropertyBase* ObjectProperty = CastField<FObjectPropertyBase>(Address.GetProperty()))
+	if (FProperty* Property = ResolvedProperty.GetValidProperty())
 	{
-		const uint8* ValuePtr = ObjectProperty->ContainerPtrToValuePtr<uint8>(Address.Address);
-		OutValue = ObjectProperty->GetObjectPropertyValue(ValuePtr);
-		return true;
+		if (FObjectPropertyBase* ObjectProperty = CastField<FObjectPropertyBase>(ResolvedProperty.GetValidProperty()))
+		{
+			const uint8* ValuePtr = ObjectProperty->ContainerPtrToValuePtr<uint8>(ResolvedProperty.ContainerAddress);
+			OutValue = ObjectProperty->GetObjectPropertyValue(ValuePtr);
+			return true;
+		}
+		else
+		{
+			UE_LOG(LogMovieScene, Error, TEXT("Mismatch in property evaluation. %s is not of type: %s"), *Property->GetName(), *FObjectPropertyBase::StaticClass()->GetName());
+		}
+	}
+	else if (ResolvedProperty.HasValidCustomProperty())
+	{
+		// NOTE: custom handlers must treat object pointers as TObjectPtr<T>. Any type fits, since the templated
+		// smart object pointers are covariant.
+		const FCustomProperty& CustomProperty = ResolvedProperty.CustomProperty;
+		const void* ValuePtr = CustomProperty.Handler->ContainerPtrToValuePtr(ResolvedProperty.ContainerAddress, CustomProperty.Property);
+		OutValue = ((TObjectPtr<UObject>*)ValuePtr)->Get();
 	}
 
 	return false;
 }
 
-template<> void FTrackInstancePropertyBindings::SetCurrentValue<UObject*>(UObject& InRuntimeObject, UObject* InValue)
+template<> void FTrackInstancePropertyBindings::SetCurrentValue<UObject*>(UObject& Object, UObject* InValue)
 {
-	FPropertyAndFunction PropAndFunction = FindOrAdd(InRuntimeObject);
-	if (FObjectPropertyBase* ObjectProperty = CastField<FObjectPropertyBase>(PropAndFunction.PropertyAddress.GetProperty()))
+	FResolvedPropertyAndFunction PropAndFunction = FindOrAdd(Object);
+
+	if (FProperty* Property = PropAndFunction.ResolvedProperty.GetValidProperty())
+	{
+		if (FObjectPropertyBase* ObjectProperty = CastField<FObjectPropertyBase>(Property))
+		{
+			uint8* ValuePtr = ObjectProperty->ContainerPtrToValuePtr<uint8>(PropAndFunction.ResolvedProperty.ContainerAddress);
+			ObjectProperty->SetObjectPropertyValue(ValuePtr, InValue);
+		}
+	}
+	else if (PropAndFunction.ResolvedProperty.HasValidCustomProperty())
 	{
-		uint8* ValuePtr = ObjectProperty->ContainerPtrToValuePtr<uint8>(PropAndFunction.PropertyAddress.Address);
-		ObjectProperty->SetObjectPropertyValue(ValuePtr, InValue);
+		// NOTE: custom handlers must treat object pointers as TObjectPtr<T>. Any type fits, since the templated
+		// smart object pointers are covariant.
+		const FCustomProperty& CustomProperty = PropAndFunction.ResolvedProperty.CustomProperty;
+		void* ValuePtr = CustomProperty.Handler->ContainerPtrToValuePtr(PropAndFunction.ResolvedProperty.ContainerAddress, CustomProperty.Property);
+		*(TObjectPtr<UObject>*)ValuePtr = InValue;
 	}
 
 	if (UFunction* NotifyFunction = PropAndFunction.NotifyFunction.Get())
 	{
-		InRuntimeObject.ProcessEvent(NotifyFunction, nullptr);
+		Object.ProcessEvent(NotifyFunction, nullptr);
 	}
 }
+
diff --git a/Engine/Source/Runtime/MovieScene/Public/MovieSceneCommonHelpers.h b/Engine/Source/Runtime/MovieScene/Public/MovieSceneCommonHelpers.h
--- a/Engine/Source/Runtime/MovieScene/Public/MovieSceneCommonHelpers.h
+++ b/Engine/Source/Runtime/MovieScene/Public/MovieSceneCommonHelpers.h
@@ -292,6 +292,46 @@ public:
 	static MOVIESCENE_API bool EvaluateSequenceCondition(const FGuid& BindingID, const FMovieSceneSequenceID& SequenceID, const UMovieSceneCondition* Condition, UObject* ConditionOwnerObject, TSharedRef<const UE::MovieScene::FSharedPlaybackState> SharedPlaybackState);
 };
 
+/** 
+ * Container type for a resolved property binding.
+ */
+enum class ETrackInstancePropertyContainerType
+{
+	None,
+	Array,
+};
+
+/**
+ * Information about a resolved property binding from a custom handler.
+ */
+struct FTrackInstanceCustomProperty
+{
+	/** The container type of the property. */
+	ETrackInstancePropertyContainerType ContainerType = ETrackInstancePropertyContainerType::None;
+	/** Structure or object class if the property is of that type. */
+	UStruct* PropertyStruct = nullptr;
+	/** Custom ID for the handler to identify a property. */
+	int32 PropertyID = 0;
+};
+
+/**
+ * Interface for a custom property binding resolver.
+ */
+class ITrackInstanceCustomPropertyBindingHandler : public TSharedFromThis<ITrackInstanceCustomPropertyBindingHandler>
+{
+public:
+
+	virtual ~ITrackInstanceCustomPropertyBindingHandler() {}
+
+	/** Called when a property binding must be resolved on the given object. */
+	virtual bool FindProperty(void* BasePointer, const FString& PropertyName, FTrackInstanceCustomProperty& OutProperty) = 0;
+	/** Get a raw pointer to the resolved property's value for the given raw object address. */
+	virtual void* ContainerPtrToValuePtr(void* BasePointer, const FTrackInstanceCustomProperty& CustomProperty, int32 ArrayIndex = INDEX_NONE) = 0;
+};
+
+/** A delegate for creating a custom property binding handler. */
+DECLARE_DELEGATE_RetVal(TSharedRef<ITrackInstanceCustomPropertyBindingHandler>, FOnCreateTrackInstanceCustomPropertyBindingHandler);
+
 /**
  * Manages bindings to keyed properties for a track instance. 
  * Calls UFunctions to set the value on runtime objects
@@ -299,7 +339,15 @@ public:
 class FTrackInstancePropertyBindings
 {
 public:
-	MOVIESCENE_API FTrackInstancePropertyBindings( FName InPropertyName, const FString& InPropertyPath);
+
+	/** Register a custom property binding handler. */
+	static MOVIESCENE_API FDelegateHandle RegisterCustomHandler(UClass* InClass, FOnCreateTrackInstanceCustomPropertyBindingHandler InDelegate);
+	/** Unregister a custom property binding handler. */
+	static MOVIESCENE_API void UnregisterCustomHandler(UClass* InClass);
+
+public:
+
+	MOVIESCENE_API FTrackInstancePropertyBindings(FName InPropertyName, const FString& InPropertyPath, bool bInAllowCustomHandlers = true);
 
 	/**
 	 * Calls the setter function for a specific runtime object or if the setter function does not exist, the property is set directly
@@ -308,11 +356,11 @@ public:
 	 * @param PropertyValue The new value to assign to the property
 	 */
 	template <typename ValueType>
-	void CallFunction( UObject& InRuntimeObject, typename TCallTraits<ValueType>::ParamType PropertyValue )
+	void CallFunction(UObject& InRuntimeObject, typename TCallTraits<ValueType>::ParamType PropertyValue)
 	{
-		FPropertyAndFunction PropAndFunction = FindOrAdd(InRuntimeObject);
+		FResolvedPropertyAndFunction PropAndFunction = FindOrAdd(InRuntimeObject);
 
-		FProperty* Property = GetProperty(InRuntimeObject);
+		FProperty* Property = PropAndFunction.ResolvedProperty.GetValidProperty();
 		if (Property && Property->HasSetter())
 		{
 			Property->CallSetter(&InRuntimeObject, &PropertyValue);
@@ -345,15 +393,33 @@ public:
 	 *
 	 * @param InRuntimeObject	The object to cache mappings for
 	 */
-	MOVIESCENE_API void CacheBinding( const UObject& InRuntimeObject );
+	MOVIESCENE_API void CacheBinding(const UObject& InRuntimeObject);
 
 	/**
-	 * Gets the FProperty that is bound to the track instance
+	 * Gets the FProperty that is bound to the track instance. This may be null if the property
+	 * is resolved by a custom handler, or if the property simply doesn't exist on this object.
 	 *
 	 * @param Object	The Object that owns the property
 	 * @return			The property on the object if it exists
 	 */
-	MOVIESCENE_API FProperty* GetProperty(const UObject& Object) const;
+	MOVIESCENE_API FProperty* GetProperty(const UObject& Object);
+
+	/**
+	 * Returns whether this binding is valid for the given object. This could be because it
+	 * resolves to an FProperty, or to a custom one.
+	 *
+	 * @param Object	The Object that owns the property
+	 * @return          Whether the property binding is valid
+	 */
+	MOVIESCENE_API bool HasValidBinding(const UObject& Object);
+
+	/**
+	 * Gets the structure type of the bound property if it is a property of that type.
+	 *
+	 * @param Object	The Object that owns the property
+	 * @return          The structure type
+	 */
+	MOVIESCENE_API const UStruct* GetPropertyStruct(const UObject& Object);
 
 	/**
 	 * Gets the current value of a property on an object
@@ -366,8 +432,8 @@ public:
 	{
 		ValueType Value{};
 
-		FPropertyAndFunction PropAndFunction = FindOrAdd(Object);
-		ResolvePropertyValue<ValueType>(PropAndFunction.PropertyAddress, Value);
+		FResolvedPropertyAndFunction PropAndFunction = FindOrAdd(Object);
+		TryGetPropertyValue<ValueType>(PropAndFunction.ResolvedProperty, Value);
 
 		return Value;
 	}
@@ -383,8 +449,8 @@ public:
 	{
 		ValueType Value{};
 
-		FPropertyAndFunction PropAndFunction = FindOrAdd(Object);
-		if (ResolvePropertyValue<ValueType>(PropAndFunction.PropertyAddress, Value))
+		FResolvedPropertyAndFunction PropAndFunction = FindOrAdd(Object);
+		if (TryGetPropertyValue<ValueType>(PropAndFunction.ResolvedProperty, Value))
 		{
 			return Value;
 		}
@@ -404,10 +470,11 @@ public:
 	{
 		checkf(Object, TEXT("No object specified"));
 
-		FPropertyAddress Address = FindPropertyAddress(*Object, InPropertyPath);
+		FTrackInstancePropertyBindings Temp(NAME_None, InPropertyPath);
+		FResolvedProperty ResolvedProperty = ResolveProperty(Temp, *Object);
 
 		ValueType Value;
-		if (ResolvePropertyValue<ValueType>(Address, Value))
+		if (TryGetPropertyValue<ValueType>(ResolvedProperty, Value))
 		{
 			return Value;
 		}
@@ -432,9 +499,9 @@ public:
 	template <typename ValueType>
 	void SetCurrentValue(UObject& Object, typename TCallTraits<ValueType>::ParamType InValue)
 	{
-		FPropertyAndFunction PropAndFunction = FindOrAdd(Object);
+		FResolvedPropertyAndFunction PropAndFunction = FindOrAdd(Object);
 
-		if(ValueType* Val = PropAndFunction.GetPropertyAddress<ValueType>())
+		if (ValueType* Val = PropAndFunction.GetPropertyAddress<ValueType>())
 		{
 			*Val = InValue;
 
@@ -457,6 +524,7 @@ public:
 		return PropertyName;
 	}
 
+	UE_DEPRECATED(5.6, "Create an instance and call GetProperty")
 	static MOVIESCENE_API FProperty* FindProperty(const UObject* Object, const FString& InPropertyPath);
 
 private:
@@ -469,57 +537,92 @@ private:
 	template<typename T>
 	static void InvokeSetterFunction(UObject* InRuntimeObject, UFunction* Setter, T&& InPropertyValue);
 
-	struct FPropertyAddress
+	struct FCustomProperty
+	{
+		TSharedPtr<ITrackInstanceCustomPropertyBindingHandler> Handler;
+		FTrackInstanceCustomProperty Property;
+
+		bool IsValid() const { return Handler.IsValid(); }
+	};
+
+	struct FResolvedProperty
 	{
 		TWeakFieldPtr<FProperty> Property;
-		void* Address;
 
-		FProperty* GetProperty() const
+		FCustomProperty CustomProperty;
+
+		void* ContainerAddress;
+
+		int32 ArrayIndex = INDEX_NONE;
+
+		FProperty* GetValidProperty() const
 		{
 			FProperty* PropertyPtr = Property.Get();
-			if (PropertyPtr && Address && !PropertyPtr->HasAnyFlags(RF_BeginDestroyed | RF_FinishDestroyed))
+			if (PropertyPtr && ContainerAddress && !PropertyPtr->HasAnyFlags(RF_BeginDestroyed | RF_FinishDestroyed))
 			{
 				return PropertyPtr;
 			}
 			return nullptr;
 		}
 
+		bool HasValidCustomProperty() const
+		{
+			return ContainerAddress && CustomProperty.IsValid();
+		}
+
 		template<typename ValueType>
 		ValueType* GetPropertyAddress() const
 		{
-			FProperty* PropertyPtr = GetProperty();
-			return PropertyPtr ? PropertyPtr->ContainerPtrToValuePtr<ValueType>(Address) : nullptr;
+			if (FProperty* PropertyPtr = GetValidProperty())
+			{
+				if (ArrayIndex == INDEX_NONE)
+				{
+					return PropertyPtr->ContainerPtrToValuePtr<ValueType>(ContainerAddress);
+				}
+				else if (FArrayProperty* ArrayProp = CastFieldChecked<FArrayProperty>(PropertyPtr))
+				{
+					FScriptArrayHelper ArrayHelper(ArrayProp, ArrayProp->ContainerPtrToValuePtr<void>(ContainerAddress));
+					check(ArrayHelper.IsValidIndex(ArrayIndex));
+					return (ValueType*)ArrayHelper.GetRawPtr(ArrayIndex);
+				}
+			}
+			if (CustomProperty.IsValid())
+			{
+				void* RawAddress = CustomProperty.Handler->ContainerPtrToValuePtr(ContainerAddress, CustomProperty.Property, ArrayIndex);
+				return reinterpret_cast<ValueType*>(RawAddress);
+			}
+			return nullptr;
 		}
 
-		FPropertyAddress()
+		FResolvedProperty()
 			: Property(nullptr)
-			, Address(nullptr)
+			, ContainerAddress(nullptr)
 		{}
 	};
 
-	struct FPropertyAndFunction
+	struct FResolvedPropertyAndFunction
 	{
-		FPropertyAddress PropertyAddress;
+		FResolvedProperty ResolvedProperty;
 		TWeakObjectPtr<UFunction> SetterFunction;
 		TWeakObjectPtr<UFunction> NotifyFunction;
 
 		template<typename ValueType>
 		ValueType* GetPropertyAddress() const
 		{
-			return PropertyAddress.GetPropertyAddress<ValueType>();
+			return ResolvedProperty.GetPropertyAddress<ValueType>();
 		}
 
-		FPropertyAndFunction()
-			: PropertyAddress()
+		FResolvedPropertyAndFunction()
+			: ResolvedProperty()
 			, SetterFunction( nullptr )
 			, NotifyFunction( nullptr )
 		{}
 	};
 
 	template <typename ValueType>
-	static bool ResolvePropertyValue(const FPropertyAddress& Address, ValueType& OutValue)
+	static bool TryGetPropertyValue(const FResolvedProperty& ResolvedProperty, ValueType& OutValue)
 	{
-		if (const ValueType* Value = Address.GetPropertyAddress<ValueType>())
+		if (const ValueType* Value = ResolvedProperty.GetPropertyAddress<ValueType>())
 		{
 			OutValue = *Value;
 			return true;
@@ -527,29 +630,22 @@ private:
 		return false;
 	}
 
-	static MOVIESCENE_API FPropertyAddress FindPropertyAddressRecursive(void* BasePointer, UStruct* InStruct, TArray<FString>& InPropertyNames, uint32 Index);
-	static MOVIESCENE_API FPropertyAddress FindPropertyAddress(const UObject& Object, const FString& InPropertyPath);
+	static MOVIESCENE_API FResolvedProperty ResolvePropertyRecursive(FTrackInstancePropertyBindings& Bindings, void* BasePointer, UStruct* InStruct, TArray<FString>& InPropertyNames, uint32 Index);
+	static MOVIESCENE_API FResolvedProperty ResolveProperty(FTrackInstancePropertyBindings& Bindings, const UObject& Object);
 
-	static MOVIESCENE_API FProperty* FindPropertyRecursive(UStruct* InStruct, TArray<FString>& InPropertyNames, uint32 Index);
+	static MOVIESCENE_API void FindProperty(FTrackInstancePropertyBindings& Bindings, void* BasePointer, UStruct* InStruct, const FString& InPropertyName, FResolvedProperty& OutResolvedProperty, bool bAllowCustomHandlers);
+	static MOVIESCENE_API FResolvedProperty FindPropertyAndArrayIndex(FTrackInstancePropertyBindings& Bindings, void* BasePointer, UStruct* InStruct, const FString& PropertyName, bool bAllowCustomHandlers);
 
-	/** Find or add the FPropertyAndFunction for the specified object */
-	FPropertyAndFunction FindOrAdd(const UObject& InObject)
-	{
-		FObjectKey ObjectKey(&InObject);
-
-		const FPropertyAndFunction* PropAndFunction = RuntimeObjectToFunctionMap.Find(ObjectKey);
-		if (PropAndFunction && (PropAndFunction->SetterFunction.IsValid() || PropAndFunction->PropertyAddress.Property.Get()))
-		{
-			return *PropAndFunction;
-		}
-
-		CacheBinding(InObject);
-		return RuntimeObjectToFunctionMap.FindRef(ObjectKey);
-	}
+	/** Find or add the FResolvedPropertyAndFunction for the specified object */
+	MOVIESCENE_API FResolvedPropertyAndFunction FindOrAdd(const UObject& InObject);
 
 private:
+
 	/** Mapping of objects to bound functions that will be called to update data on the track */
-	TMap< FObjectKey, FPropertyAndFunction > RuntimeObjectToFunctionMap;
+	TMap< FObjectKey, FResolvedPropertyAndFunction > RuntimeObjectToFunctionMap;
+
+	/** Mapping of object classes to custom handlers. */
+	TMap< const UClass*, TSharedPtr<ITrackInstanceCustomPropertyBindingHandler> > ObjectClassToCustomHandlerMap;
 
 	/** Path to the property we are bound to */
 	FString PropertyPath;
@@ -563,17 +659,19 @@ private:
 	/** Actual name of the property we are bound to */
 	FName PropertyName;
 
+	/** Whether custom handlers are allowed to be created and used. */
+	bool bAllowCustomHandlers = true;
 };
 
 /** Explicit specializations for bools */
 template<> MOVIESCENE_API void FTrackInstancePropertyBindings::CallFunction<bool>(UObject& InRuntimeObject, TCallTraits<bool>::ParamType PropertyValue);
-template<> MOVIESCENE_API bool FTrackInstancePropertyBindings::ResolvePropertyValue<bool>(const FPropertyAddress& Address, bool& OutValue);
+template<> MOVIESCENE_API bool FTrackInstancePropertyBindings::TryGetPropertyValue<bool>(const FResolvedProperty& ResolvedProperty, bool& OutValue);
 template<> MOVIESCENE_API void FTrackInstancePropertyBindings::SetCurrentValue<bool>(UObject& Object, TCallTraits<bool>::ParamType InValue);
 
+/** Explicit specializations for object pointers */
 template<> MOVIESCENE_API void FTrackInstancePropertyBindings::CallFunction<UObject*>(UObject& InRuntimeObject, UObject* PropertyValue);
-template<> MOVIESCENE_API bool FTrackInstancePropertyBindings::ResolvePropertyValue<UObject*>(const FPropertyAddress& Address, UObject*& OutValue);
-template<> MOVIESCENE_API void FTrackInstancePropertyBindings::SetCurrentValue<UObject*>(UObject& InRuntimeObject, UObject* InValue);
-
+template<> MOVIESCENE_API bool FTrackInstancePropertyBindings::TryGetPropertyValue<UObject*>(const FResolvedProperty& ResolvedProperty, UObject*& OutValue);
+template<> MOVIESCENE_API void FTrackInstancePropertyBindings::SetCurrentValue<UObject*>(UObject& Object, UObject* InValue);
 
 template<typename T>
 void FTrackInstancePropertyBindings::InvokeSetterFunction(UObject* InRuntimeObject, UFunction* Setter, T&& InPropertyValue)
@@ -623,3 +721,4 @@ void FTrackInstancePropertyBindings::InvokeSetterFunction(UObject* InRuntimeObje
 	// Now we have the parameters set up correctly, call the function
 	InRuntimeObject->ProcessEvent(Setter, Params);
 }
+
diff --git a/Engine/Source/Runtime/MovieSceneTracks/Private/MovieSceneTracksComponentTypes.cpp b/Engine/Source/Runtime/MovieSceneTracks/Private/MovieSceneTracksComponentTypes.cpp
--- a/Engine/Source/Runtime/MovieSceneTracks/Private/MovieSceneTracksComponentTypes.cpp
+++ b/Engine/Source/Runtime/MovieSceneTracks/Private/MovieSceneTracksComponentTypes.cpp
@@ -341,23 +341,24 @@ struct FColorHandler : TPropertyComponentHandler<FColorPropertyTraits, double, d
 		.FilterAll({ BuiltInComponents->Tags.NeedsLink })
 		.Iterate_PerEntity(&Linker->EntityManager, [](UObject* Object, const FMovieScenePropertyBinding& Binding, EColorPropertyType& OutType)
 		{
-			FStructProperty* BoundProperty = CastField<FStructProperty>(FTrackInstancePropertyBindings::FindProperty(Object, Binding.PropertyPath.ToString()));
-			if (ensure(BoundProperty && BoundProperty->Struct))
+			FTrackInstancePropertyBindings Bindings(Binding.PropertyName, Binding.PropertyPath.ToString());
+			const UStruct* PropertyStruct = Bindings.GetPropertyStruct(*Object);
+			if (ensure(PropertyStruct))
 			{
-				if (BoundProperty->Struct == TBaseStructure<FColor>::Get())
+				if (PropertyStruct == TBaseStructure<FColor>::Get())
 				{
 					// We assume the color we get back is in sRGB, assigning it to a linear color will implicitly
 					// convert it to a linear color instead of using ReinterpretAsLinear which will just change the
 					// bytes into floats using divide by 255.
 					OutType = EColorPropertyType::Color;
 				}
-				else if (BoundProperty->Struct == TBaseStructure<FSlateColor>::Get())
+				else if (PropertyStruct == TBaseStructure<FSlateColor>::Get())
 				{
 					OutType = EColorPropertyType::Slate;
 				}
 				else
 				{
-					ensure(BoundProperty->Struct == TBaseStructure<FLinearColor>::Get());
+					ensure(PropertyStruct == TBaseStructure<FLinearColor>::Get());
 					OutType = EColorPropertyType::Linear;
 				}
 			}
@@ -400,20 +401,21 @@ struct FFloatVectorHandler : TPropertyComponentHandler<FFloatVectorPropertyTrait
 		.FilterAll({ BuiltInComponents->Tags.NeedsLink })
 		.Iterate_PerEntity(&Linker->EntityManager, [](UObject* Object, const FMovieScenePropertyBinding& Binding, FVectorPropertyMetaData& OutMetaData)
 		{
-			FStructProperty* BoundProperty = CastField<FStructProperty>(FTrackInstancePropertyBindings::FindProperty(Object, Binding.PropertyPath.ToString()));
-			if (ensure(BoundProperty && BoundProperty->Struct))
+			FTrackInstancePropertyBindings Bindings(Binding.PropertyName, Binding.PropertyPath.ToString());
+			const UStruct* PropertyStruct = Bindings.GetPropertyStruct(*Object);
+			if (ensure(PropertyStruct))
 			{
-				if (BoundProperty->Struct == TVariantStructure<FVector2f>::Get())
+				if (PropertyStruct == TVariantStructure<FVector2f>::Get())
 				{
 					OutMetaData.NumChannels = 2;
 				}
-				else if (BoundProperty->Struct == TVariantStructure<FVector3f>::Get())
+				else if (PropertyStruct == TVariantStructure<FVector3f>::Get())
 				{
 					OutMetaData.NumChannels = 3;
 				}
 				else
 				{
-					ensure(BoundProperty->Struct == TVariantStructure<FVector4f>::Get());
+					ensure(PropertyStruct == TVariantStructure<FVector4f>::Get());
 					OutMetaData.NumChannels = 4;
 				}
 			}
@@ -440,20 +442,21 @@ struct FDoubleVectorHandler : TPropertyComponentHandler<FDoubleVectorPropertyTra
 		.FilterAll({ BuiltInComponents->Tags.NeedsLink })
 		.Iterate_PerEntity(&Linker->EntityManager, [](UObject* Object, const FMovieScenePropertyBinding& Binding, FVectorPropertyMetaData& OutMetaData)
 		{
-			FStructProperty* BoundProperty = CastField<FStructProperty>(FTrackInstancePropertyBindings::FindProperty(Object, Binding.PropertyPath.ToString()));
-			if (ensure(BoundProperty && BoundProperty->Struct))
+			FTrackInstancePropertyBindings Bindings(Binding.PropertyName, Binding.PropertyPath.ToString());
+			const UStruct* PropertyStruct = Bindings.GetPropertyStruct(*Object);
+			if (ensure(PropertyStruct))
 			{
-				if (BoundProperty->Struct == TBaseStructure<FVector2D>::Get())
+				if (PropertyStruct == TBaseStructure<FVector2D>::Get())
 				{
 					OutMetaData.NumChannels = 2;
 				}
-				else if (BoundProperty->Struct == TBaseStructure<FVector>::Get() || BoundProperty->Struct == TVariantStructure<FVector3d>::Get())
+				else if (PropertyStruct == TBaseStructure<FVector>::Get() || PropertyStruct == TVariantStructure<FVector3d>::Get())
 				{
 					OutMetaData.NumChannels = 3;
 				}
 				else
 				{
-					ensure(BoundProperty->Struct == TBaseStructure<FVector4>::Get() || BoundProperty->Struct == TVariantStructure<FVector4d>::Get());
+					ensure(PropertyStruct == TBaseStructure<FVector4>::Get() || PropertyStruct == TVariantStructure<FVector4d>::Get());
 					OutMetaData.NumChannels = 4;
 				}
 			}
@@ -479,22 +482,19 @@ struct FBoolHandler : TPropertyComponentHandler<FBoolPropertyTraits, bool>
 		.FilterAll({ BuiltInComponents->Tags.NeedsLink })
 		.Iterate_PerEntity(&Linker->EntityManager, [](UObject* Object, const FMovieScenePropertyBinding& Binding, FBoolPropertyTraits::FBoolMetaData& OutMetaData)
 		{
-			FBoolProperty* BoundProperty = CastField<FBoolProperty>(FTrackInstancePropertyBindings::FindProperty(Object, Binding.PropertyPath.ToString()));
-			if (ensure(BoundProperty))
+			OutMetaData.BitFieldSize = 0;
+			OutMetaData.BitIndex     = 0;
+
+			FTrackInstancePropertyBindings Bindings(Binding.PropertyName, Binding.PropertyPath.ToString());
+			ensure(Bindings.HasValidBinding(*Object));
+			FBoolProperty* BoundProperty = CastField<FBoolProperty>(Bindings.GetProperty(*Object));
+			if (BoundProperty && !BoundProperty->IsNativeBool())
 			{
-				if (BoundProperty->IsNativeBool())
-				{
-					OutMetaData.BitFieldSize = 0;
-					OutMetaData.BitIndex     = 0;
-				}
-				else
-				{
-					auto FieldMask = BoundProperty->GetFieldMask();
-					static_assert(std::is_same_v<decltype(FieldMask), uint8>, "Unexpected size of field mask returned from FBoolProperty::FieldMask");
+				auto FieldMask = BoundProperty->GetFieldMask();
+				static_assert(std::is_same_v<decltype(FieldMask), uint8>, "Unexpected size of field mask returned from FBoolProperty::FieldMask");
 
-					OutMetaData.BitFieldSize = static_cast<uint8>(BoundProperty->GetElementSize());
-					OutMetaData.BitIndex     = static_cast<uint8>(FMath::CountTrailingZeros(FieldMask));
-				}
+				OutMetaData.BitFieldSize = static_cast<uint8>(BoundProperty->GetElementSize());
+				OutMetaData.BitIndex     = static_cast<uint8>(FMath::CountTrailingZeros(FieldMask));
 			}
 		});
 	}
@@ -522,24 +522,26 @@ struct FObjectHandler : TPropertyComponentHandler<FObjectPropertyTraits, FObject
 			.Write(TrackComponents->Object.MetaDataComponents.GetType<0>())
 			.FilterAll({ BuiltInComponents->Tags.NeedsLink })
 			.Iterate_PerEntity(&Linker->EntityManager, [TrackComponents](UObject* Object, const FMovieScenePropertyBinding& Binding, const FCustomPropertyIndex* OptionalCustomPropertyIndex, FObjectPropertyTraits::FObjectMetadata& OutMetaData)
+			{
+				if (OptionalCustomPropertyIndex)
 				{
-					if (OptionalCustomPropertyIndex)
+					if (const auto* CustomObjectMetaData = TrackComponents->Accessors.Object.MetaData.Find(OptionalCustomPropertyIndex->Value))
 					{
-						if (const auto* CustomObjectMetaData = TrackComponents->Accessors.Object.MetaData.Find(OptionalCustomPropertyIndex->Value))
-						{
-							OutMetaData.ObjectClass = CustomObjectMetaData->AllowedClass.Get();
-							OutMetaData.bAllowsClear = CustomObjectMetaData->bAllowsClear;
-							return;
-						}
+						OutMetaData.ObjectClass = CustomObjectMetaData->AllowedClass.Get();
+						OutMetaData.bAllowsClear = CustomObjectMetaData->bAllowsClear;
+						return;
 					}
+				}
 
-					FObjectPropertyBase* BoundProperty = CastField<FObjectPropertyBase>(FTrackInstancePropertyBindings::FindProperty(Object, Binding.PropertyPath.ToString()));
-					if (ensure(BoundProperty))
-					{
-						OutMetaData.ObjectClass = BoundProperty->PropertyClass;
-						OutMetaData.bAllowsClear = !BoundProperty->HasAnyPropertyFlags(CPF_NoClear);
-					}
-				});
+				FTrackInstancePropertyBindings Bindings(Binding.PropertyName, Binding.PropertyPath.ToString());
+				ensure(Bindings.HasValidBinding(*Object));
+				FObjectPropertyBase* BoundProperty = CastField<FObjectPropertyBase>(Bindings.GetProperty(*Object));
+				if (BoundProperty)
+				{
+					OutMetaData.ObjectClass = BoundProperty->PropertyClass;
+					OutMetaData.bAllowsClear = !BoundProperty->HasAnyPropertyFlags(CPF_NoClear);
+				}
+			});
 	}
 };
 
diff --git a/Engine/Source/Runtime/MovieSceneTracks/Private/Sections/MovieScene3DTransformSection.cpp b/Engine/Source/Runtime/MovieSceneTracks/Private/Sections/MovieScene3DTransformSection.cpp
--- a/Engine/Source/Runtime/MovieSceneTracks/Private/Sections/MovieScene3DTransformSection.cpp
+++ b/Engine/Source/Runtime/MovieSceneTracks/Private/Sections/MovieScene3DTransformSection.cpp
@@ -166,18 +166,18 @@ struct F3DTransformChannelEditorData
 
 	static TOptional<FVector> GetTranslation(UObject& InObject, FTrackInstancePropertyBindings* Bindings)
 	{
-		const FStructProperty* TransformProperty = Bindings ? CastField<FStructProperty>(Bindings->GetProperty(InObject)) : nullptr;
+		const UStruct* TransformStruct = Bindings ? Bindings->GetPropertyStruct(InObject) : nullptr;
 
-		if (TransformProperty)
+		if (TransformStruct)
 		{
-			if (TransformProperty->Struct == TBaseStructure<FTransform>::Get())
+			if (TransformStruct == TBaseStructure<FTransform>::Get())
 			{
 				if (TOptional<FTransform> Transform = Bindings->GetOptionalValue<FTransform>(InObject))
 				{
 					return Transform->GetTranslation();
 				}
 			}
-			else if (TransformProperty->Struct == TBaseStructure<FEulerTransform>::Get())
+			else if (TransformStruct == TBaseStructure<FEulerTransform>::Get())
 			{
 				if (TOptional<FEulerTransform> EulerTransform = Bindings->GetOptionalValue<FEulerTransform>(InObject))
 				{
@@ -202,18 +202,18 @@ struct F3DTransformChannelEditorData
 
 	static TOptional<FRotator> GetRotator(UObject& InObject, FTrackInstancePropertyBindings* Bindings)
 	{
-		const FStructProperty* TransformProperty = Bindings ? CastField<FStructProperty>(Bindings->GetProperty(InObject)) : nullptr;
+		const UStruct* TransformStruct = Bindings ? Bindings->GetPropertyStruct(InObject) : nullptr;
 
-		if (TransformProperty)
+		if (TransformStruct)
 		{
-			if (TransformProperty->Struct == TBaseStructure<FTransform>::Get())
+			if (TransformStruct == TBaseStructure<FTransform>::Get())
 			{
 				if (TOptional<FTransform> Transform = Bindings->GetOptionalValue<FTransform>(InObject))
 				{
 					return Transform->GetRotation().Rotator();
 				}
 			}
-			else if (TransformProperty->Struct == TBaseStructure<FEulerTransform>::Get())
+			else if (TransformStruct == TBaseStructure<FEulerTransform>::Get())
 			{
 				if (TOptional<FEulerTransform> EulerTransform = Bindings->GetOptionalValue<FEulerTransform>(InObject))
 				{
@@ -238,18 +238,18 @@ struct F3DTransformChannelEditorData
 
 	static TOptional<FVector> GetScale(UObject& InObject, FTrackInstancePropertyBindings* Bindings)
 	{
-		const FStructProperty* TransformProperty = Bindings ? CastField<FStructProperty>(Bindings->GetProperty(InObject)) : nullptr;
+		const UStruct* TransformStruct = Bindings ? Bindings->GetPropertyStruct(InObject) : nullptr;
 
-		if (TransformProperty)
+		if (TransformStruct)
 		{
-			if (TransformProperty->Struct == TBaseStructure<FTransform>::Get())
+			if (TransformStruct == TBaseStructure<FTransform>::Get())
 			{
 				if (TOptional<FTransform> Transform = Bindings->GetOptionalValue<FTransform>(InObject))
 				{
 					return Transform->GetScale3D();
 				}
 			}
-			else if (TransformProperty->Struct == TBaseStructure<FEulerTransform>::Get())
+			else if (TransformStruct == TBaseStructure<FEulerTransform>::Get())
 			{
 				if (TOptional<FEulerTransform> EulerTransform = Bindings->GetOptionalValue<FEulerTransform>(InObject))
 				{
diff --git a/Engine/Source/Runtime/MovieSceneTracks/Private/Sections/MovieSceneColorSection.cpp b/Engine/Source/Runtime/MovieSceneTracks/Private/Sections/MovieSceneColorSection.cpp
--- a/Engine/Source/Runtime/MovieSceneTracks/Private/Sections/MovieSceneColorSection.cpp
+++ b/Engine/Source/Runtime/MovieSceneTracks/Private/Sections/MovieSceneColorSection.cpp
@@ -49,20 +49,20 @@ struct FColorSectionEditorData
 	{
 		const FName SlateColorName("SlateColor");
 	
-		FStructProperty* ColorStructProperty = CastField<FStructProperty>(Bindings.GetProperty(InObject));
-		if (ColorStructProperty != nullptr)
+		const UStruct* ColorStruct = Bindings.GetPropertyStruct(InObject);
+		if (ColorStruct != nullptr)
 		{
-			if (ColorStructProperty->Struct->GetFName() == SlateColorName)
+			if (ColorStruct->GetFName() == SlateColorName)
 			{
 				return Bindings.GetCurrentValue<FSlateColor>(InObject).GetSpecifiedColor();
 			}
 	
-			if (ColorStructProperty->Struct->GetFName() == NAME_LinearColor)
+			if (ColorStruct->GetFName() == NAME_LinearColor)
 			{
 				return Bindings.GetCurrentValue<FLinearColor>(InObject);
 			}
 	
-			if (ColorStructProperty->Struct->GetFName() == NAME_Color)
+			if (ColorStruct->GetFName() == NAME_Color)
 			{
 				return Bindings.GetCurrentValue<FColor>(InObject);
 			}
diff --git a/Engine/Source/Runtime/MovieSceneTracks/Private/Sections/MovieSceneRotatorSection.cpp b/Engine/Source/Runtime/MovieSceneTracks/Private/Sections/MovieSceneRotatorSection.cpp
--- a/Engine/Source/Runtime/MovieSceneTracks/Private/Sections/MovieSceneRotatorSection.cpp
+++ b/Engine/Source/Runtime/MovieSceneTracks/Private/Sections/MovieSceneRotatorSection.cpp
@@ -35,9 +35,9 @@ struct FRotatorSectionEditorData
 
 	static FRotator GetPropertyValue(const UObject& InObject, FTrackInstancePropertyBindings& InBindings)
 	{
-		if (const FStructProperty* RotatorStructProperty = CastField<FStructProperty>(InBindings.GetProperty(InObject)))
+		if (const UStruct* RotatorStruct = InBindings.GetPropertyStruct(InObject))
 		{
-			if (RotatorStructProperty->Struct->GetFName() == NAME_Rotator)
+			if (RotatorStruct->GetFName() == NAME_Rotator)
 			{
 				return InBindings.GetCurrentValue<FRotator>(InObject);
 			}
diff --git a/Engine/Source/Runtime/MovieSceneTracks/Private/Systems/MovieScenePropertyInstantiator.cpp b/Engine/Source/Runtime/MovieSceneTracks/Private/Systems/MovieScenePropertyInstantiator.cpp
--- a/Engine/Source/Runtime/MovieSceneTracks/Private/Systems/MovieScenePropertyInstantiator.cpp
+++ b/Engine/Source/Runtime/MovieSceneTracks/Private/Systems/MovieScenePropertyInstantiator.cpp
@@ -424,7 +424,8 @@ void UMovieScenePropertyInstantiatorSystem::UpgradeFloatToDoubleProperties(const
 				continue;
 			}
 
-			FProperty* BoundProperty = FTrackInstancePropertyBindings::FindProperty(PropertyInfo.BoundObject, PropertyInfo.PropertyBinding.PropertyPath.ToString());
+			FTrackInstancePropertyBindings Bindings(PropertyInfo.PropertyBinding.PropertyName, PropertyInfo.PropertyBinding.PropertyPath.ToString());
+			FProperty* BoundProperty = Bindings.GetProperty(*PropertyInfo.BoundObject);
 			if (!BoundProperty)
 			{
 				continue;
diff --git a/Engine/Source/Runtime/MovieSceneTracks/Private/Tracks/MovieScenePropertyTrack.cpp b/Engine/Source/Runtime/MovieSceneTracks/Private/Tracks/MovieScenePropertyTrack.cpp
--- a/Engine/Source/Runtime/MovieSceneTracks/Private/Tracks/MovieScenePropertyTrack.cpp
+++ b/Engine/Source/Runtime/MovieSceneTracks/Private/Tracks/MovieScenePropertyTrack.cpp
@@ -74,7 +74,7 @@ FText UMovieScenePropertyTrack::GetDisplayNameToolTipText(const FMovieSceneLabel
 	const TArrayView<TWeakObjectPtr<>> FoundBoundObjects = LabelParams.Player->FindBoundObjects(LabelParams.BindingID, LabelParams.SequenceID);
 	for (const TWeakObjectPtr<> BoundObject : FoundBoundObjects)
 	{
-		FTrackInstancePropertyBindings InstancePropertyBinding(GetPropertyName(), GetPropertyPath().ToString());
+		FTrackInstancePropertyBindings InstancePropertyBinding(GetPropertyName(), GetPropertyPath().ToString(), false);
 		if (FProperty* BoundProperty = InstancePropertyBinding.GetProperty(*BoundObject))
 		{
 			FString PropertyName = BoundProperty->GetMetaData(TEXT("DisplayName"));
@@ -106,11 +106,12 @@ FSlateColor UMovieScenePropertyTrack::GetLabelColor(const FMovieSceneLabelParams
 
 	// Return a normal colour if we have at least one bound object for which the property binding resolves
 	// correctly. Otherwise, return a red colour indicating a binding issue.
+	// TODO: if the property is custom, we are going to allocate and throw away a custom handler every frame here :(
+	FTrackInstancePropertyBindings InstancePropertyBinding(GetPropertyName(), GetPropertyPath().ToString());
 	const TArrayView<TWeakObjectPtr<>> FoundBoundObjects = LabelParams.Player->FindBoundObjects(LabelParams.BindingID, LabelParams.SequenceID);
 	for (const TWeakObjectPtr<> BoundObject : FoundBoundObjects)
 	{
-		FTrackInstancePropertyBindings InstancePropertyBinding(GetPropertyName(), GetPropertyPath().ToString());
-		if (InstancePropertyBinding.GetProperty(*BoundObject))
+		if (InstancePropertyBinding.HasValidBinding(*BoundObject))
 		{
 			return LabelParams.bIsDimmed ? FSlateColor::UseSubduedForeground() : FSlateColor::UseForeground();
 		}
