Sequencer: Numeric variants stored in FMovieSceneSequenceTransforms are now stored as weak references

This change is required in order to ensure that no strong back references exist from the compiled data to the source data when Time Warp exists in a sequence. In order to achieve this, we have to hide the UMovieSceneNumericVariantGetter reference from GC when it is stored within the sequence transforms. I've achieved this by forcing the getters to have an alignment of 8 bytes (which seems to be standard for all UObjects anyway), which allows me to (ab)use the 3 LSBs for flags in the variant class. When the ptr is marked as weak, we use the LSB to denote that it should be hidden from the reference graph.

To provide safety guarantees in non-shipping builds, I have also included an actual WeakObjectPtr to the object which can be used for verification purposes.

diff --git a/Engine/Source/Runtime/MovieScene/Private/Evaluation/MovieSceneSequenceTransform.cpp b/Engine/Source/Runtime/MovieScene/Private/Evaluation/MovieSceneSequenceTransform.cpp
--- a/Engine/Source/Runtime/MovieScene/Private/Evaluation/MovieSceneSequenceTransform.cpp
+++ b/Engine/Source/Runtime/MovieScene/Private/Evaluation/MovieSceneSequenceTransform.cpp
@@ -148,6 +148,14 @@ FMovieSceneWarpCounter& FMovieSceneWarpCounter::operator=(FMovieSceneWarpCounter
 FMovieSceneWarpCounter::~FMovieSceneWarpCounter() = default;
 PRAGMA_ENABLE_DEPRECATION_WARNINGS
 
+void FMovieSceneNestedSequenceTransform::PostSerialize(const FArchive& Ar)
+{
+	if (Ar.IsLoading())
+	{
+		TimeScale.MakeWeakUnsafe();
+	}
+}
+
 void FMovieSceneNestedSequenceTransform::ToString(TStringBuilderBase<TCHAR>& OutBuilder) const
 {
 	switch (TimeScale.GetType())
diff --git a/Engine/Source/Runtime/MovieScene/Private/Variants/MovieSceneNumericVariant.cpp b/Engine/Source/Runtime/MovieScene/Private/Variants/MovieSceneNumericVariant.cpp
--- a/Engine/Source/Runtime/MovieScene/Private/Variants/MovieSceneNumericVariant.cpp
+++ b/Engine/Source/Runtime/MovieScene/Private/Variants/MovieSceneNumericVariant.cpp
@@ -81,6 +81,10 @@ FMovieSceneNumericVariant FMovieSceneNumericVariant::ShallowCopy() const
 	// Bitwise copy
 	FMovieSceneNumericVariant New(NoInit);
 	FMemory::Memcpy(New.Data, this->Data, sizeof(Data));
+
+#if UE_MOVIESCENE_WEAKNUMERICVARIANT_CHECKS
+	New.WeakCustomGetter = this->WeakCustomGetter;
+#endif
 	return New;
 }
 
@@ -93,7 +97,15 @@ FMovieSceneNumericVariant FMovieSceneNumericVariant::DeepCopy(UObject* NewOuter)
 		{
 			Getter = DuplicateObject(Getter, NewOuter);
 		}
-		return FMovieSceneNumericVariant(Getter);
+
+		FMovieSceneNumericVariant New(Getter);
+
+		if (this->HasCustomWeakPtrFlag())
+		{
+			New.MakeWeakUnsafe();
+		}
+
+		return New;
 	}
 
 	return ShallowCopy();
@@ -102,6 +114,10 @@ FMovieSceneNumericVariant FMovieSceneNumericVariant::DeepCopy(UObject* NewOuter)
 void FMovieSceneNumericVariant::Set(double InLiteralValue)
 {
 	*reinterpret_cast<double*>(Data) = InLiteralValue;
+
+#if UE_MOVIESCENE_WEAKNUMERICVARIANT_CHECKS
+	WeakCustomGetter = nullptr;
+#endif
 }
 
 void FMovieSceneNumericVariant::Set(UMovieSceneNumericVariantGetter* InDynamicValue)
@@ -116,6 +132,24 @@ void FMovieSceneNumericVariant::Set(UMovieSceneNumericVariantGetter* InDynamicVa
 	*reinterpret_cast<uint64*>(Data) = NewValue | TAGGED_Bits | TYPE_CustomPtr;
 }
 
+void FMovieSceneNumericVariant::SetWeakUnsafe(UMovieSceneNumericVariantGetter* InDynamicValue)
+{
+	Set(InDynamicValue);
+	MakeWeakUnsafe();
+}
+
+
+void FMovieSceneNumericVariant::MakeWeakUnsafe()
+{
+	if (IsCustomPtr())
+	{
+#if UE_MOVIESCENE_WEAKNUMERICVARIANT_CHECKS
+		WeakCustomGetter = GetCustomPtr();
+#endif
+		*reinterpret_cast<uint64*>(Data) |= CUSTOMPTR_Weak;
+	}
+}
+
 double FMovieSceneNumericVariant::Get() const
 {
 	if (IsLiteral())
@@ -141,6 +175,10 @@ void FMovieSceneNumericVariant::SetTypeBits(uint8 InType)
 
 	const uint64 TypeMask = uint64(InType & 0x7) << 48;
 	*reinterpret_cast<uint64*>(Data) |= (TypeMask & TYPE_Bits);
+
+#if UE_MOVIESCENE_WEAKNUMERICVARIANT_CHECKS
+	WeakCustomGetter = nullptr;
+#endif
 }
 
 uint8 FMovieSceneNumericVariant::GetTypeBits() const
@@ -151,12 +189,27 @@ uint8 FMovieSceneNumericVariant::GetTypeBits() const
 	return uint8(Type >> 48);
 }
 
+bool FMovieSceneNumericVariant::HasCustomWeakPtrFlag() const
+{
+	check(!IsLiteral() && (GetTypeBits() == TYPE_CustomPtr));
+	return (*reinterpret_cast<const uint64*>(Data) & CUSTOMPTR_Weak) != 0ull;
+}
 
 UMovieSceneNumericVariantGetter* FMovieSceneNumericVariant::GetCustomPtr() const
 {
 	check(IsCustomPtr());
-	const uint64 PtrValue = *reinterpret_cast<const uint64*>(Data) & PAYLOAD_Bits;
+
+	const uint64 PtrValue = (*reinterpret_cast<const uint64*>(Data) & PAYLOAD_Bits) & ~CUSTOMPTR_FlagBits;
+
 	UMovieSceneNumericVariantGetter* Ptr = reinterpret_cast<UMovieSceneNumericVariantGetter*>(PtrValue);
+
+#if UE_MOVIESCENE_WEAKNUMERICVARIANT_CHECKS
+	if (Ptr && HasCustomWeakPtrFlag() && !ensureMsgf(WeakCustomGetter.Get() != nullptr, TEXT("Weakly referenced getter stored in FMovieSceneNumericVariant has been destroyed!")))
+	{
+		return nullptr;
+	}
+#endif
+
 	return Ptr;
 }
 
@@ -342,9 +395,12 @@ void FMovieSceneNumericVariant::AddStructReferencedObjects(FReferenceCollector&
 	if (!IsLiteral() && GetTypeBits() == TYPE_CustomPtr)
 	{
 		UMovieSceneNumericVariantGetter* Getter = GetCustomPtr();
-		if (Getter)
+
+		const bool bReportReference = !HasCustomWeakPtrFlag();
+		if (bReportReference && Getter)
 		{
-			Collector.AddReferencedObject(Getter->ReferenceToSelf);
+			const FProperty* Property = Collector.GetSerializedProperty();
+			Collector.AddReferencedObject(Getter->ReferenceToSelf, nullptr, Property);
 		}
 	}
 }
diff --git a/Engine/Source/Runtime/MovieScene/Private/Variants/MovieSceneTimeWarpVariant.cpp b/Engine/Source/Runtime/MovieScene/Private/Variants/MovieSceneTimeWarpVariant.cpp
--- a/Engine/Source/Runtime/MovieScene/Private/Variants/MovieSceneTimeWarpVariant.cpp
+++ b/Engine/Source/Runtime/MovieScene/Private/Variants/MovieSceneTimeWarpVariant.cpp
@@ -102,7 +102,7 @@ FMovieSceneTimeWarpClampFloat FMovieSceneTimeWarpVariant::AsClampFloat() const
 UMovieSceneTimeWarpGetter* FMovieSceneTimeWarpVariant::AsCustom() const
 {
 	check(GetType() == EMovieSceneTimeWarpType::Custom);
-	return Variant.UnsafePayloadCast<UMovieSceneTimeWarpGetter*>();
+	return static_cast<UMovieSceneTimeWarpGetter*>(Variant.GetCustomPtr());
 }
 
 void FMovieSceneTimeWarpVariant::Set(UMovieSceneTimeWarpGetter* InDynamicValue)
diff --git a/Engine/Source/Runtime/MovieScene/Public/Evaluation/MovieSceneSequenceTransform.h b/Engine/Source/Runtime/MovieScene/Public/Evaluation/MovieSceneSequenceTransform.h
--- a/Engine/Source/Runtime/MovieScene/Public/Evaluation/MovieSceneSequenceTransform.h
+++ b/Engine/Source/Runtime/MovieScene/Public/Evaluation/MovieSceneSequenceTransform.h
@@ -390,7 +390,9 @@ struct FMovieSceneNestedSequenceTransform
 	 */
 	FMovieSceneNestedSequenceTransform(FMovieSceneTimeWarpVariant&& InVariant)
 		: TimeScale(MoveTemp(InVariant))
-	{}
+	{
+		TimeScale.MakeWeakUnsafe();
+	}
 
 
 	/**
@@ -414,7 +416,9 @@ struct FMovieSceneNestedSequenceTransform
 	 */
 	FMovieSceneNestedSequenceTransform(FFrameTime InOffset, FMovieSceneTimeWarpVariant&& InTimeScale)
 		: TimeScale(MoveTemp(InTimeScale)), Offset(InOffset)
-	{}
+	{
+		TimeScale.MakeWeakUnsafe();
+	}
 
 
 public:
@@ -546,6 +550,10 @@ public:
 		return A.TimeScale != B.TimeScale || A.Offset != B.Offset;
 	}
 
+	/*~ Begin StructOpsTypeTraits */
+	void PostSerialize(const FArchive& Ar);
+	/*~ End StructOpsTypeTraits */
+
 private:
 
 	/**
@@ -596,6 +604,14 @@ public:
 	FMovieSceneTimeWarping Warping;
 	/*~~~~~~~~~~~~~~~~~~~~~~~~~~~~~~~~~~~~~  End API deprecation  ~~~~~~~~~~~~~~~~~~~~~~~~~~~~~~~~~~~~~*/
 };
+template<>
+struct TStructOpsTypeTraits<FMovieSceneNestedSequenceTransform> : public TStructOpsTypeTraitsBase2<FMovieSceneNestedSequenceTransform>
+{
+	enum
+	{
+		WithPostSerialize = true,
+	};
+};
 
 
 
diff --git a/Engine/Source/Runtime/MovieScene/Public/Variants/MovieSceneNumericVariant.h b/Engine/Source/Runtime/MovieScene/Public/Variants/MovieSceneNumericVariant.h
--- a/Engine/Source/Runtime/MovieScene/Public/Variants/MovieSceneNumericVariant.h
+++ b/Engine/Source/Runtime/MovieScene/Public/Variants/MovieSceneNumericVariant.h
@@ -7,6 +7,10 @@
 
 class UMovieSceneNumericVariantGetter;
 
+#ifndef UE_MOVIESCENE_WEAKNUMERICVARIANT_CHECKS
+	#define UE_MOVIESCENE_WEAKNUMERICVARIANT_CHECKS DO_CHECK
+#endif
+
 
 /**
  * A variant type that masquerades as a numeric (double) value.
@@ -96,6 +100,22 @@ public:
 	MOVIESCENE_API void Set(UMovieSceneNumericVariantGetter* InDynamicValue);
 
 
+	/**
+	 * Assign a new dynamic value to this variant as an unsafe weak ptr.
+	 * @warning: This will result in a dangling pointer if it is not referenced strongly elsewhere. Use with caution.
+	 * 
+	 * @param InDynamicValue      The new dynamic value this variant should represent.
+	 */
+	MOVIESCENE_API void SetWeakUnsafe(UMovieSceneNumericVariantGetter* InDynamicValue);
+
+
+	/**
+	 * If this variant wraps a UMovieSceneTimeWarpGetter, turn it into an unsafe weak reference.
+	 * @warning: This will result in a dangling pointer if it is not referenced strongly elsewhere. Use with caution.
+	 */
+	MOVIESCENE_API void MakeWeakUnsafe();
+
+
 	/**
 	 * Retrieve this variant's numeric value
 	 * 
@@ -217,6 +237,9 @@ private:
 	static constexpr uint64 TYPE_CustomPtr = 0x0000000000000000u; // INTENTIONALLY ZERO: Special value for TYPE_Bits when our data points to a UMovieSceneNumericVariantGetter*. 
 	static constexpr uint64 PAYLOAD_Bits   = 0x0000FFFFFFFFFFFFu; // Bitmask specifying valid bits that can be used for custom payloads when any of TAGGED_Bits is set
 
+	static constexpr uint64 CUSTOMPTR_FlagBits = 0x0000000000000003u; // Low bitmask that (ab)uses the alignment of UMovieSceneNumericVariantGetter to encode additional flags
+	static constexpr uint64 CUSTOMPTR_Weak     = 0x0000000000000001u; // Low bit that signifies the wrapped custom pointer should not be reported to the reference graph.
+
 	static constexpr uint64 TAGGED_Bits   = SIGN_Bit | EXP_Bits | QUIET_Bit;
 
 
@@ -226,6 +249,13 @@ private:
 		return *reinterpret_cast<double*>(Data);
 	}
 
+
+	/**
+	 * Checks whether this variant has the weak ptr flag.
+	 * @note: Only valid to be called if IsCustomPtr() is true
+	 */
+	bool HasCustomWeakPtrFlag() const;
+
 private:
 
 	/** Functionally non-copyable except to the reflection layer - use Deep/Shallow copy instead */
@@ -236,6 +266,10 @@ private:
 	friend UScriptStruct::TCppStructOps<FMovieSceneNumericVariant>;
 
 	alignas(8) uint8 Data[8];
+
+#if UE_MOVIESCENE_WEAKNUMERICVARIANT_CHECKS
+	TWeakObjectPtr<UMovieSceneNumericVariantGetter> WeakCustomGetter;
+#endif
 };
 
 template<>
diff --git a/Engine/Source/Runtime/MovieScene/Public/Variants/MovieSceneNumericVariantGetter.h b/Engine/Source/Runtime/MovieScene/Public/Variants/MovieSceneNumericVariantGetter.h
--- a/Engine/Source/Runtime/MovieScene/Public/Variants/MovieSceneNumericVariantGetter.h
+++ b/Engine/Source/Runtime/MovieScene/Public/Variants/MovieSceneNumericVariantGetter.h
@@ -9,9 +9,11 @@
 
 /**
  * Base class for all dynamic getter types supplied to an FMovieSceneNumericVariant
+ * 
+ * Must be at least aligned to 8 bits since the low bits are used for referencing flags in FMovieSceneNumericVariant
  */
 UCLASS(Abstract, MinimalAPI)
-class UMovieSceneNumericVariantGetter : public UMovieSceneSignedObject
+class alignas(8) UMovieSceneNumericVariantGetter : public UMovieSceneSignedObject
 {
 public:
 
diff --git a/Engine/Source/Runtime/MovieScene/Public/Variants/MovieSceneTimeWarpVariant.h b/Engine/Source/Runtime/MovieScene/Public/Variants/MovieSceneTimeWarpVariant.h
--- a/Engine/Source/Runtime/MovieScene/Public/Variants/MovieSceneTimeWarpVariant.h
+++ b/Engine/Source/Runtime/MovieScene/Public/Variants/MovieSceneTimeWarpVariant.h
@@ -193,6 +193,15 @@ struct FMovieSceneTimeWarpVariant
 	}
 
 
+	/**
+	 * If this variant wraps a UMovieSceneTimeWarpGetter, turn it into an unsafe weak reference.
+	 * @warning: This will result in a dangling pointer if it is not referenced strongly elsewhere. Use with caution.
+	 */
+	void MakeWeakUnsafe()
+	{
+		Variant.MakeWeakUnsafe();
+	}
+
 public:
 
 
