Sequencer: query decorations for extending imported entities in all import situations

diff --git a/Engine/Source/Runtime/MovieScene/Private/EntitySystem/MovieSceneEntityLedger.cpp b/Engine/Source/Runtime/MovieScene/Private/EntitySystem/MovieSceneEntityLedger.cpp
--- a/Engine/Source/Runtime/MovieScene/Private/EntitySystem/MovieSceneEntityLedger.cpp
+++ b/Engine/Source/Runtime/MovieScene/Private/EntitySystem/MovieSceneEntityLedger.cpp
@@ -124,6 +124,17 @@ void FEntityLedger::UpdateOneShotEntities(UMovieSceneEntitySystemLinker* Linker,
 
 		if (!ImportedEntity.IsEmpty())
 		{
+			if (UMovieSceneDecorationContainerObject* DecorationContainer = Cast<UMovieSceneDecorationContainerObject>(EntityOwner))
+			{
+				for (UObject* Decoration : DecorationContainer->GetDecorations())
+				{
+					if (IMovieSceneEntityDecorator* EntityDecorator = Cast<IMovieSceneEntityDecorator>(Decoration))
+					{
+						EntityDecorator->ExtendEntity(Linker, Params, &ImportedEntity);
+					}
+				}
+			}
+
 			if (UMovieSceneSection* Section = Cast<UMovieSceneSection>(EntityOwner))
 			{
 				Section->BuildDefaultComponents(Linker, Params, &ImportedEntity);
@@ -342,6 +353,17 @@ void FEntityLedger::ImportEntity(UMovieSceneEntitySystemLinker* Linker, const FE
 
 	if (!ImportedEntity.IsEmpty())
 	{
+		if (UMovieSceneDecorationContainerObject* DecorationContainer = Cast<UMovieSceneDecorationContainerObject>(EntityOwner))
+		{
+			for (UObject* Decoration : DecorationContainer->GetDecorations())
+			{
+				if (IMovieSceneEntityDecorator* EntityDecorator = Cast<IMovieSceneEntityDecorator>(Decoration))
+				{
+					EntityDecorator->ExtendEntity(Linker, Params, &ImportedEntity);
+				}
+			}
+		}
+
 		if (UMovieSceneSection* Section = Cast<UMovieSceneSection>(EntityOwner))
 		{
 			Section->BuildDefaultComponents(Linker, Params, &ImportedEntity);
