Sequencer: add global delegates for when sections are added to/removed from a track.

This allows client code to add and manage decorators on those sections.

diff --git a/Engine/Source/Runtime/MovieScene/Private/MovieSceneSection.cpp b/Engine/Source/Runtime/MovieScene/Private/MovieSceneSection.cpp
--- a/Engine/Source/Runtime/MovieScene/Private/MovieSceneSection.cpp
+++ b/Engine/Source/Runtime/MovieScene/Private/MovieSceneSection.cpp
@@ -25,6 +25,11 @@
 
 #include UE_INLINE_GENERATED_CPP_BY_NAME(MovieSceneSection)
 
+#if WITH_EDITORONLY_DATA
+FOnMovieSceneSectionAddedToTrack UMovieSceneSection::OnAddedToTrackEvent;
+FOnMovieSceneSectionChanged UMovieSceneSection::OnRemovedFromTrackEvent;
+#endif
+
 UMovieSceneSection::UMovieSceneSection(const FObjectInitializer& ObjectInitializer)
 	: Super( ObjectInitializer )
 	, PreRollFrames(0)
@@ -128,6 +133,20 @@ void UMovieSceneSection::Serialize(FArchive& Ar)
 	}
 }
 
+#if WITH_EDITORONLY_DATA
+
+FOnMovieSceneSectionAddedToTrack& UMovieSceneSection::GetOnSectionAddedToTrack()
+{
+	return OnAddedToTrackEvent;
+}
+
+FOnMovieSceneSectionChanged& UMovieSceneSection::GetOnSectionRemovedFromTrack()
+{
+	return OnRemovedFromTrackEvent;
+}
+
+#endif  // WITH_EDITORONLY_DATA
+
 void UMovieSceneSection::OnAddedToTrack(UMovieSceneTrack* Track)
 {
 	OnAddedToTrackImpl(Track);
@@ -142,6 +161,10 @@ void UMovieSceneSection::OnAddedToTrack(UMovieSceneTrack* Track)
 			}
 		}
 	}
+
+#if WITH_EDITORONLY_DATA
+	OnAddedToTrackEvent.Broadcast(Track, this);
+#endif
 }
 
 void UMovieSceneSection::OnRemovedFromTrack()
@@ -158,6 +181,10 @@ void UMovieSceneSection::OnRemovedFromTrack()
 			}
 		}
 	}
+
+#if WITH_EDITORONLY_DATA
+	OnRemovedFromTrackEvent.Broadcast(this);
+#endif
 }
 
 void UMovieSceneSection::OnDecorationAdded(UObject* NewDecoration)
diff --git a/Engine/Source/Runtime/MovieScene/Public/MovieSceneSection.h b/Engine/Source/Runtime/MovieScene/Public/MovieSceneSection.h
--- a/Engine/Source/Runtime/MovieScene/Public/MovieSceneSection.h
+++ b/Engine/Source/Runtime/MovieScene/Public/MovieSceneSection.h
@@ -78,6 +78,10 @@ enum class EMovieSceneChannelProxyType : uint8
 	Dynamic
 };
 
+/** Delegate for when a section is added to a track. */
+DECLARE_MULTICAST_DELEGATE_TwoParams(FOnMovieSceneSectionAddedToTrack, UMovieSceneTrack*, UMovieSceneSection*);
+/** Delegate for when a section is removed from a track. */
+DECLARE_MULTICAST_DELEGATE_OneParam(FOnMovieSceneSectionChanged, UMovieSceneSection*);
 
 USTRUCT()
 struct FMovieSceneSectionEvalOptions
@@ -714,6 +718,22 @@ public:
 	MOVIESCENE_API static void DeclareConstructClasses(TArray<FTopLevelAssetPath>& OutConstructClasses, const UClass* SpecificSubclass);
 #endif
 
+#if WITH_EDITORONLY_DATA
+
+public:
+
+	/** A delegate for when a section is added to a track. */
+	MOVIESCENE_API static FOnMovieSceneSectionAddedToTrack& GetOnSectionAddedToTrack();
+	/** A delegate for when a section is removed from a track. */
+	MOVIESCENE_API static FOnMovieSceneSectionChanged& GetOnSectionRemovedFromTrack();
+
+private:
+
+	static FOnMovieSceneSectionAddedToTrack OnAddedToTrackEvent;
+	static FOnMovieSceneSectionChanged OnRemovedFromTrackEvent;
+
+#endif
+
 protected:
 
 	//~ UObject interface
