Sequencer: Fixed object replacement not updating bound object keys

diff --git a/Engine/Source/Runtime/MovieScene/Private/EntitySystem/MovieSceneEntitySystemLinker.cpp b/Engine/Source/Runtime/MovieScene/Private/EntitySystem/MovieSceneEntitySystemLinker.cpp
--- a/Engine/Source/Runtime/MovieScene/Private/EntitySystem/MovieSceneEntitySystemLinker.cpp
+++ b/Engine/Source/Runtime/MovieScene/Private/EntitySystem/MovieSceneEntitySystemLinker.cpp
@@ -478,12 +478,14 @@ void UMovieSceneEntitySystemLinker::OnObjectsReplaced(const TMap<UObject*, UObje
 	FBuiltInComponentTypes* BuiltInComponents = FBuiltInComponentTypes::Get();
 
 	FEntityTaskBuilder()
+	.Write(BuiltInComponents->BoundObjectKey)
 	.Write(BuiltInComponents->BoundObject)
-	.Iterate_PerEntity(&EntityManager, [&ReplacementMap](UObject*& Object)
+	.Iterate_PerEntity(&EntityManager, [&ReplacementMap](FObjectKey& ObjectKey, UObject*& Object)
 	{
 		if (UObject* Replacement = ReplacementMap.FindRef(Object))
 		{
 			Object = Replacement;
+			ObjectKey = Replacement;
 		}
 	});
 #endif
