Sequencer: remove custom property binding handler API

diff --git a/Engine/Source/Runtime/MovieScene/Private/MovieSceneCommonHelpers.cpp b/Engine/Source/Runtime/MovieScene/Private/MovieSceneCommonHelpers.cpp
--- a/Engine/Source/Runtime/MovieScene/Private/MovieSceneCommonHelpers.cpp
+++ b/Engine/Source/Runtime/MovieScene/Private/MovieSceneCommonHelpers.cpp
@@ -1186,25 +1186,9 @@ MovieSceneHelpers::FMovieSceneScopedPackageDirtyGuard::~FMovieSceneScopedPackage
 	}
 }
 
-TMap<UClass*, FOnCreateTrackInstanceCustomPropertyBindingHandler> GCustomHandlers;
-
-FDelegateHandle FTrackInstancePropertyBindings::RegisterCustomHandler(UClass* InClass, FOnCreateTrackInstanceCustomPropertyBindingHandler InDelegate)
-{
-	using namespace UE::MovieScene;
-	GCustomHandlers.Add(InClass, InDelegate);
-	return GCustomHandlers[InClass].GetHandle();
-}
-
-void FTrackInstancePropertyBindings::UnregisterCustomHandler(UClass* InClass)
-{
-	using namespace UE::MovieScene;
-	GCustomHandlers.Remove(InClass);
-}
-
-FTrackInstancePropertyBindings::FTrackInstancePropertyBindings(FName InPropertyName, const FString& InPropertyPath, bool bInAllowCustomHandlers)
+FTrackInstancePropertyBindings::FTrackInstancePropertyBindings(FName InPropertyName, const FString& InPropertyPath)
 	: PropertyPath(InPropertyPath)
 	, PropertyName(InPropertyName)
-	, bAllowCustomHandlers(bInAllowCustomHandlers)
 {
 	static const FString Set(TEXT("Set"));
 	const FString FunctionString = Set + PropertyName.ToString();
@@ -1212,7 +1196,7 @@ FTrackInstancePropertyBindings::FTrackInstancePropertyBindings(FName InPropertyN
 	FunctionName = FName(*FunctionString);
 }
 
-void FTrackInstancePropertyBindings::FindProperty(FTrackInstancePropertyBindings& Bindings, void* BasePointer, UStruct* InStruct, const FString& InPropertyName, FResolvedProperty& OutResolvedProperty, bool bAllowCustomHandlers)
+void FTrackInstancePropertyBindings::FindProperty(FTrackInstancePropertyBindings& Bindings, void* BasePointer, UStruct* InStruct, const FString& InPropertyName, FResolvedProperty& OutResolvedProperty)
 {
 	// Find a property via UObject reflection first.
 	FProperty* FoundProperty = FindFProperty<FProperty>(InStruct, *InPropertyName);
@@ -1221,38 +1205,9 @@ void FTrackInstancePropertyBindings::FindProperty(FTrackInstancePropertyBindings
 	{
 		OutResolvedProperty.Property = FoundProperty;
 	}
-	else if (bAllowCustomHandlers)
-	{
-		// Find an existing handler for this class, or try to create one.
-		// For now, custom handlers must match the exact object class. Handlers for a parent class can't
-		// necessarily resolve properties on children classes.
-		UClass* ObjectClass = Cast<UClass>(InStruct);
-		TSharedPtr<ITrackInstanceCustomPropertyBindingHandler> Handler = Bindings.ObjectClassToCustomHandlerMap.FindRef(ObjectClass);
-		if (!Handler)
-		{
-			FOnCreateTrackInstanceCustomPropertyBindingHandler* Factory = GCustomHandlers.Find(ObjectClass);
-			if (Factory)
-			{
-				Handler = Factory->Execute();
-				if (ensure(Handler))
-				{
-					Bindings.ObjectClassToCustomHandlerMap.Add(ObjectClass, Handler);
-				}
-			}
-		}
-		if (Handler)
-		{
-			FTrackInstanceCustomProperty CustomProperty;
-			if (Handler->FindProperty(BasePointer, InPropertyName, CustomProperty))
-			{
-				OutResolvedProperty.CustomProperty.Handler = Handler;
-				OutResolvedProperty.CustomProperty.Property = CustomProperty;
-			}
-		}
-	}
 }
 
-FTrackInstancePropertyBindings::FResolvedProperty FTrackInstancePropertyBindings::FindPropertyAndArrayIndex(FTrackInstancePropertyBindings& Bindings, void* BasePointer, UStruct* InStruct, const FString& PropertyName, bool bAllowCustomHandlers)
+FTrackInstancePropertyBindings::FResolvedProperty FTrackInstancePropertyBindings::FindPropertyAndArrayIndex(FTrackInstancePropertyBindings& Bindings, void* BasePointer, UStruct* InStruct, const FString& PropertyName)
 {
 	FResolvedProperty PropertyAndIndex;
 
@@ -1266,7 +1221,7 @@ FTrackInstancePropertyBindings::FResolvedProperty FTrackInstancePropertyBindings
 			// We have a property name of the form "Foo[123]". Resolve the property itself ("Foo") and then
 			// parse the array element index (123).
 			FString TruncatedPropertyName = FString::ConstructFromPtrSize(*PropertyName, OpenIndex);
-			FindProperty(Bindings, BasePointer, InStruct, *TruncatedPropertyName, PropertyAndIndex, bAllowCustomHandlers);
+			FindProperty(Bindings, BasePointer, InStruct, *TruncatedPropertyName, PropertyAndIndex);
 
 			const int32 NumberLength = PropertyName.Len() - OpenIndex - 2;
 			if (NumberLength > 0 && NumberLength <= 10)
@@ -1281,20 +1236,18 @@ FTrackInstancePropertyBindings::FResolvedProperty FTrackInstancePropertyBindings
 	}
 
 	// Not index found in this property name, just find the property normally.
-	FindProperty(Bindings, BasePointer, InStruct, *PropertyName, PropertyAndIndex, bAllowCustomHandlers);
+	FindProperty(Bindings, BasePointer, InStruct, *PropertyName, PropertyAndIndex);
 
 	return PropertyAndIndex;
 }
 
 FTrackInstancePropertyBindings::FResolvedProperty FTrackInstancePropertyBindings::ResolvePropertyRecursive(FTrackInstancePropertyBindings& Bindings, void* BasePointer, UStruct* InStruct, TArray<FString>& InPropertyNames, uint32 Index)
 {
-	const bool bAllowCustomHandlers = Bindings.bAllowCustomHandlers && (Index == 0);
-	FResolvedProperty ResolvedProperty = FindPropertyAndArrayIndex(Bindings, BasePointer, InStruct, *InPropertyNames[Index], bAllowCustomHandlers);
+	FResolvedProperty ResolvedProperty = FindPropertyAndArrayIndex(Bindings, BasePointer, InStruct, *InPropertyNames[Index]);
 	
 	FResolvedProperty NewProperty;
 
 	FProperty* ResolvedPropertyProperty = ResolvedProperty.Property.Get();
-	FCustomProperty ResolvedCustomProperty = ResolvedProperty.CustomProperty;
 
 	if (ResolvedProperty.ArrayIndex != INDEX_NONE)
 	{
@@ -1327,29 +1280,6 @@ FTrackInstancePropertyBindings::FResolvedProperty FTrackInstancePropertyBindings
 				UE_LOG(LogMovieScene, Error, TEXT("Mismatch in property evaluation. %s is not of type: %s"), *ResolvedPropertyProperty->GetName(), *FArrayProperty::StaticClass()->GetName());
 			}
 		}
-		else if (ResolvedCustomProperty.IsValid())
-		{
-			if (ResolvedCustomProperty.Property.ContainerType == ETrackInstancePropertyContainerType::Array)
-			{
-				if (InPropertyNames.IsValidIndex(Index + 1))
-				{
-					// Move the BasePointer to the array element and keep resolving the property path on it.
-					void* ArrayElement = ResolvedCustomProperty.Handler->ContainerPtrToValuePtr(BasePointer, ResolvedCustomProperty.Property, ResolvedProperty.ArrayIndex);
-					return ResolvePropertyRecursive(Bindings, ArrayElement, ResolvedCustomProperty.Property.PropertyStruct, InPropertyNames, Index + 1);
-				}
-				else
-				{
-					// The property path ends here (e.g. "Foo.Bar[1]").
-					NewProperty.CustomProperty = ResolvedCustomProperty;
-					NewProperty.ContainerAddress = BasePointer;
-					NewProperty.ArrayIndex = ResolvedProperty.ArrayIndex;
-				}
-			}
-			else
-			{
-				UE_LOG(LogMovieScene, Error, TEXT("Mismatch in property evaluation. Custom property %d is not an array property"), ResolvedCustomProperty.Property.PropertyID);
-			}
-		}
 	}
 	else if (FStructProperty* StructProp = CastField<FStructProperty>(ResolvedPropertyProperty))
 	{
@@ -1385,20 +1315,6 @@ FTrackInstancePropertyBindings::FResolvedProperty FTrackInstancePropertyBindings
 		NewProperty.Property = ResolvedProperty.Property;
 		NewProperty.ContainerAddress = BasePointer;
 	}
-	else if (ResolvedCustomProperty.IsValid())
-	{
-		NewProperty.CustomProperty = ResolvedCustomProperty;
-		NewProperty.ContainerAddress = BasePointer;
-
-		const FTrackInstanceCustomProperty& CustomProperty = NewProperty.CustomProperty.Property;
-		if (CustomProperty.PropertyStruct && 
-				CustomProperty.PropertyStruct->IsA<UScriptStruct>() && 
-				InPropertyNames.IsValidIndex(Index + 1))
-		{
-			void* StructContainer = NewProperty.CustomProperty.Handler->ContainerPtrToValuePtr(BasePointer, CustomProperty);
-			return ResolvePropertyRecursive(Bindings, StructContainer, CustomProperty.PropertyStruct, InPropertyNames, Index + 1);
-		}
-	}
 
 	return NewProperty;
 }
@@ -1422,7 +1338,7 @@ FTrackInstancePropertyBindings::FResolvedProperty FTrackInstancePropertyBindings
 
 FProperty* FTrackInstancePropertyBindings::FindProperty(const UObject* Object, const FString& InPropertyPath)
 {
-	FTrackInstancePropertyBindings Temp(NAME_None, InPropertyPath, false);
+	FTrackInstancePropertyBindings Temp(NAME_None, InPropertyPath);
 	FResolvedProperty ResolvedProperty = ResolveProperty(Temp, *Object);
 	return ResolvedProperty.Property.Get();
 }
@@ -1434,8 +1350,7 @@ FTrackInstancePropertyBindings::FResolvedPropertyAndFunction FTrackInstancePrope
 	const FResolvedPropertyAndFunction* PropAndFunction = RuntimeObjectToFunctionMap.Find(ObjectKey);
 	if (PropAndFunction && (
 				PropAndFunction->SetterFunction.IsValid() ||
-				PropAndFunction->ResolvedProperty.GetValidProperty() ||
-				PropAndFunction->ResolvedProperty.CustomProperty.IsValid()))
+				PropAndFunction->ResolvedProperty.GetValidProperty()))
 	{
 		return *PropAndFunction;
 	}
@@ -1470,13 +1385,6 @@ void FTrackInstancePropertyBindings::CallFunctionForEnum(UObject& InRuntimeObjec
 			UE_LOG(LogMovieScene, Error, TEXT("Mismatch in property evaluation. %s is not of type: %s"), *Property->GetName(), *FEnumProperty::StaticClass()->GetName());
 		}
 	}
-	else if (PropAndFunction.ResolvedProperty.HasValidCustomProperty())
-	{
-		// NOTE: custom handlers must always use uint8 values for enums.
-		const FCustomProperty& CustomProperty = PropAndFunction.ResolvedProperty.CustomProperty;
-		void* ValueAddr = CustomProperty.Handler->ContainerPtrToValuePtr(PropAndFunction.ResolvedProperty.ContainerAddress, CustomProperty.Property);
-		*(uint8*)ValueAddr = (uint8)PropertyValue;
-	}
 
 	if (UFunction* NotifyFunction = PropAndFunction.NotifyFunction.Get())
 	{
@@ -1490,13 +1398,10 @@ void FTrackInstancePropertyBindings::CacheBinding(const UObject& Object)
 	{
 		PropAndFunction.ResolvedProperty = ResolveProperty(*this, Object);
 
-		if (!PropAndFunction.ResolvedProperty.CustomProperty.IsValid())
+		UFunction* SetterFunction = Object.FindFunction(FunctionName);
+		if (SetterFunction && SetterFunction->NumParms >= 1)
 		{
-			UFunction* SetterFunction = Object.FindFunction(FunctionName);
-			if (SetterFunction && SetterFunction->NumParms >= 1)
-			{
-				PropAndFunction.SetterFunction = SetterFunction;
-			}
+			PropAndFunction.SetterFunction = SetterFunction;
 		}
 
 		UFunction* NotifyFunction = NotifyFunctionName != NAME_None ? Object.FindFunction(NotifyFunctionName) : nullptr;
@@ -1518,7 +1423,7 @@ FProperty* FTrackInstancePropertyBindings::GetProperty(const UObject& Object)
 bool FTrackInstancePropertyBindings::HasValidBinding(const UObject& Object)
 {
 	FResolvedPropertyAndFunction PropAndFunction = FindOrAdd(Object);
-	return PropAndFunction.ResolvedProperty.GetValidProperty() || PropAndFunction.ResolvedProperty.HasValidCustomProperty();
+	return PropAndFunction.ResolvedProperty.GetValidProperty() != nullptr;
 }
 
 const UStruct* FTrackInstancePropertyBindings::GetPropertyStruct(const UObject& Object)
@@ -1533,10 +1438,6 @@ const UStruct* FTrackInstancePropertyBindings::GetPropertyStruct(const UObject&
 		}
 		return nullptr;
 	}
-	else if (PropAndFunction.ResolvedProperty.HasValidCustomProperty())
-	{
-		return PropAndFunction.ResolvedProperty.CustomProperty.Property.PropertyStruct;
-	}
 
 	return nullptr;
 }
@@ -1561,13 +1462,6 @@ int64 FTrackInstancePropertyBindings::GetCurrentValueForEnum(const UObject& Obje
 			UE_LOG(LogMovieScene, Error, TEXT("Mismatch in property evaluation. %s is not of type: %s"), *Property->GetName(), *FEnumProperty::StaticClass()->GetName());
 		}
 	}
-	else if (PropAndFunction.ResolvedProperty.HasValidCustomProperty())
-	{
-		// NOTE: custom handlers must always use uint8 values for enums.
-		const FCustomProperty& CustomProperty = PropAndFunction.ResolvedProperty.CustomProperty;
-		void* ValueAddr = CustomProperty.Handler->ContainerPtrToValuePtr(PropAndFunction.ResolvedProperty.ContainerAddress, CustomProperty.Property);
-		return (int64)(*(uint8*)ValueAddr);
-	}
 
 	return 0;
 }
@@ -1598,12 +1492,6 @@ template<> void FTrackInstancePropertyBindings::CallFunction<bool>(UObject& InRu
 			UE_LOG(LogMovieScene, Error, TEXT("Mismatch in property evaluation. %s is not of type: %s"), *Property->GetName(), *FBoolProperty::StaticClass()->GetName());
 		}
 	}
-	else if (PropAndFunction.ResolvedProperty.HasValidCustomProperty())
-	{
-		// NOTE: custom boolean properties should never be bitfields.
-		bool* ValuePtr = PropAndFunction.ResolvedProperty.GetPropertyAddress<bool>();
-		*ValuePtr = PropertyValue;
-	}
 
 	if (UFunction* NotifyFunction = PropAndFunction.NotifyFunction.Get())
 	{
@@ -1626,13 +1514,6 @@ template<> bool FTrackInstancePropertyBindings::TryGetPropertyValue<bool>(const
 			UE_LOG(LogMovieScene, Error, TEXT("Mismatch in property evaluation. %s is not of type: %s"), *Property->GetName(), *FBoolProperty::StaticClass()->GetName());
 		}
 	}
-	else if (ResolvedProperty.HasValidCustomProperty())
-	{
-		// NOTE: custom boolean properties should never be bitfields.
-		const FCustomProperty& CustomProperty = ResolvedProperty.CustomProperty;
-		const void* ValuePtr = CustomProperty.Handler->ContainerPtrToValuePtr(ResolvedProperty.ContainerAddress, CustomProperty.Property);
-		OutValue = *(bool*)ValuePtr;
-	}
 
 	return false;
 }
@@ -1649,13 +1530,6 @@ template<> void FTrackInstancePropertyBindings::SetCurrentValue<bool>(UObject& O
 			BoolProperty->SetPropertyValue(ValuePtr, InValue);
 		}
 	}
-	else if (PropAndFunction.ResolvedProperty.HasValidCustomProperty())
-	{
-		// NOTE: custom boolean properties should never be bitfields.
-		const FCustomProperty& CustomProperty = PropAndFunction.ResolvedProperty.CustomProperty;
-		void* ValuePtr = CustomProperty.Handler->ContainerPtrToValuePtr(PropAndFunction.ResolvedProperty.ContainerAddress, CustomProperty.Property);
-		*(bool*)ValuePtr = InValue;
-	}
 
 	if (UFunction* NotifyFunction = PropAndFunction.NotifyFunction.Get())
 	{
@@ -1689,13 +1563,6 @@ template<> void FTrackInstancePropertyBindings::CallFunction<UObject*>(UObject&
 			UE_LOG(LogMovieScene, Error, TEXT("Mismatch in property evaluation. %s is not of type: %s"), *Property->GetName(), *FObjectPropertyBase::StaticClass()->GetName());
 		}
 	}
-	else if (PropAndFunction.ResolvedProperty.HasValidCustomProperty())
-	{
-		// NOTE: custom handlers must treat object pointers as TObjectPtr<T>. Any type fits, since the templated
-		// smart object pointers are covariant.
-		TObjectPtr<UObject>* ValuePtr = PropAndFunction.ResolvedProperty.GetPropertyAddress<TObjectPtr<UObject>>();
-		(*ValuePtr) = PropertyValue;
-	}
 
 	if (UFunction* NotifyFunction = PropAndFunction.NotifyFunction.Get())
 	{
@@ -1718,14 +1585,6 @@ template<> bool FTrackInstancePropertyBindings::TryGetPropertyValue<UObject*>(co
 			UE_LOG(LogMovieScene, Error, TEXT("Mismatch in property evaluation. %s is not of type: %s"), *Property->GetName(), *FObjectPropertyBase::StaticClass()->GetName());
 		}
 	}
-	else if (ResolvedProperty.HasValidCustomProperty())
-	{
-		// NOTE: custom handlers must treat object pointers as TObjectPtr<T>. Any type fits, since the templated
-		// smart object pointers are covariant.
-		const FCustomProperty& CustomProperty = ResolvedProperty.CustomProperty;
-		const void* ValuePtr = CustomProperty.Handler->ContainerPtrToValuePtr(ResolvedProperty.ContainerAddress, CustomProperty.Property);
-		OutValue = ((TObjectPtr<UObject>*)ValuePtr)->Get();
-	}
 
 	return false;
 }
@@ -1742,14 +1601,6 @@ template<> void FTrackInstancePropertyBindings::SetCurrentValue<UObject*>(UObjec
 			ObjectProperty->SetObjectPropertyValue(ValuePtr, InValue);
 		}
 	}
-	else if (PropAndFunction.ResolvedProperty.HasValidCustomProperty())
-	{
-		// NOTE: custom handlers must treat object pointers as TObjectPtr<T>. Any type fits, since the templated
-		// smart object pointers are covariant.
-		const FCustomProperty& CustomProperty = PropAndFunction.ResolvedProperty.CustomProperty;
-		void* ValuePtr = CustomProperty.Handler->ContainerPtrToValuePtr(PropAndFunction.ResolvedProperty.ContainerAddress, CustomProperty.Property);
-		*(TObjectPtr<UObject>*)ValuePtr = InValue;
-	}
 
 	if (UFunction* NotifyFunction = PropAndFunction.NotifyFunction.Get())
 	{
diff --git a/Engine/Source/Runtime/MovieScene/Public/MovieSceneCommonHelpers.h b/Engine/Source/Runtime/MovieScene/Public/MovieSceneCommonHelpers.h
--- a/Engine/Source/Runtime/MovieScene/Public/MovieSceneCommonHelpers.h
+++ b/Engine/Source/Runtime/MovieScene/Public/MovieSceneCommonHelpers.h
@@ -292,46 +292,6 @@ public:
 	static MOVIESCENE_API bool EvaluateSequenceCondition(const FGuid& BindingID, const FMovieSceneSequenceID& SequenceID, const UMovieSceneCondition* Condition, UObject* ConditionOwnerObject, TSharedRef<const UE::MovieScene::FSharedPlaybackState> SharedPlaybackState);
 };
 
-/** 
- * Container type for a resolved property binding.
- */
-enum class ETrackInstancePropertyContainerType
-{
-	None,
-	Array,
-};
-
-/**
- * Information about a resolved property binding from a custom handler.
- */
-struct FTrackInstanceCustomProperty
-{
-	/** The container type of the property. */
-	ETrackInstancePropertyContainerType ContainerType = ETrackInstancePropertyContainerType::None;
-	/** Structure or object class if the property is of that type. */
-	UStruct* PropertyStruct = nullptr;
-	/** Custom ID for the handler to identify a property. */
-	int32 PropertyID = 0;
-};
-
-/**
- * Interface for a custom property binding resolver.
- */
-class ITrackInstanceCustomPropertyBindingHandler : public TSharedFromThis<ITrackInstanceCustomPropertyBindingHandler>
-{
-public:
-
-	virtual ~ITrackInstanceCustomPropertyBindingHandler() {}
-
-	/** Called when a property binding must be resolved on the given object. */
-	virtual bool FindProperty(void* BasePointer, const FString& PropertyName, FTrackInstanceCustomProperty& OutProperty) = 0;
-	/** Get a raw pointer to the resolved property's value for the given raw object address. */
-	virtual void* ContainerPtrToValuePtr(void* BasePointer, const FTrackInstanceCustomProperty& CustomProperty, int32 ArrayIndex = INDEX_NONE) = 0;
-};
-
-/** A delegate for creating a custom property binding handler. */
-DECLARE_DELEGATE_RetVal(TSharedRef<ITrackInstanceCustomPropertyBindingHandler>, FOnCreateTrackInstanceCustomPropertyBindingHandler);
-
 /**
  * Manages bindings to keyed properties for a track instance. 
  * Calls UFunctions to set the value on runtime objects
@@ -340,14 +300,7 @@ class FTrackInstancePropertyBindings
 {
 public:
 
-	/** Register a custom property binding handler. */
-	static MOVIESCENE_API FDelegateHandle RegisterCustomHandler(UClass* InClass, FOnCreateTrackInstanceCustomPropertyBindingHandler InDelegate);
-	/** Unregister a custom property binding handler. */
-	static MOVIESCENE_API void UnregisterCustomHandler(UClass* InClass);
-
-public:
-
-	MOVIESCENE_API FTrackInstancePropertyBindings(FName InPropertyName, const FString& InPropertyPath, bool bInAllowCustomHandlers = true);
+	MOVIESCENE_API FTrackInstancePropertyBindings(FName InPropertyName, const FString& InPropertyPath);
 
 	/**
 	 * Calls the setter function for a specific runtime object or if the setter function does not exist, the property is set directly
@@ -396,8 +349,7 @@ public:
 	MOVIESCENE_API void CacheBinding(const UObject& InRuntimeObject);
 
 	/**
-	 * Gets the FProperty that is bound to the track instance. This may be null if the property
-	 * is resolved by a custom handler, or if the property simply doesn't exist on this object.
+	 * Gets the FProperty that is bound to the track instance.
 	 *
 	 * @param Object	The Object that owns the property
 	 * @return			The property on the object if it exists
@@ -405,8 +357,7 @@ public:
 	MOVIESCENE_API FProperty* GetProperty(const UObject& Object);
 
 	/**
-	 * Returns whether this binding is valid for the given object. This could be because it
-	 * resolves to an FProperty, or to a custom one.
+	 * Returns whether this binding is valid for the given object.
 	 *
 	 * @param Object	The Object that owns the property
 	 * @return          Whether the property binding is valid
@@ -524,7 +475,6 @@ public:
 		return PropertyName;
 	}
 
-	UE_DEPRECATED(5.6, "Create an instance and call GetProperty")
 	static MOVIESCENE_API FProperty* FindProperty(const UObject* Object, const FString& InPropertyPath);
 
 private:
@@ -537,20 +487,10 @@ private:
 	template<typename T>
 	static void InvokeSetterFunction(UObject* InRuntimeObject, UFunction* Setter, T&& InPropertyValue);
 
-	struct FCustomProperty
-	{
-		TSharedPtr<ITrackInstanceCustomPropertyBindingHandler> Handler;
-		FTrackInstanceCustomProperty Property;
-
-		bool IsValid() const { return Handler.IsValid(); }
-	};
-
 	struct FResolvedProperty
 	{
 		TWeakFieldPtr<FProperty> Property;
 
-		FCustomProperty CustomProperty;
-
 		void* ContainerAddress;
 
 		int32 ArrayIndex = INDEX_NONE;
@@ -565,11 +505,6 @@ private:
 			return nullptr;
 		}
 
-		bool HasValidCustomProperty() const
-		{
-			return ContainerAddress && CustomProperty.IsValid();
-		}
-
 		template<typename ValueType>
 		ValueType* GetPropertyAddress() const
 		{
@@ -586,11 +521,6 @@ private:
 					return (ValueType*)ArrayHelper.GetRawPtr(ArrayIndex);
 				}
 			}
-			if (CustomProperty.IsValid())
-			{
-				void* RawAddress = CustomProperty.Handler->ContainerPtrToValuePtr(ContainerAddress, CustomProperty.Property, ArrayIndex);
-				return reinterpret_cast<ValueType*>(RawAddress);
-			}
 			return nullptr;
 		}
 
@@ -633,8 +563,8 @@ private:
 	static MOVIESCENE_API FResolvedProperty ResolvePropertyRecursive(FTrackInstancePropertyBindings& Bindings, void* BasePointer, UStruct* InStruct, TArray<FString>& InPropertyNames, uint32 Index);
 	static MOVIESCENE_API FResolvedProperty ResolveProperty(FTrackInstancePropertyBindings& Bindings, const UObject& Object);
 
-	static MOVIESCENE_API void FindProperty(FTrackInstancePropertyBindings& Bindings, void* BasePointer, UStruct* InStruct, const FString& InPropertyName, FResolvedProperty& OutResolvedProperty, bool bAllowCustomHandlers);
-	static MOVIESCENE_API FResolvedProperty FindPropertyAndArrayIndex(FTrackInstancePropertyBindings& Bindings, void* BasePointer, UStruct* InStruct, const FString& PropertyName, bool bAllowCustomHandlers);
+	static MOVIESCENE_API void FindProperty(FTrackInstancePropertyBindings& Bindings, void* BasePointer, UStruct* InStruct, const FString& InPropertyName, FResolvedProperty& OutResolvedProperty);
+	static MOVIESCENE_API FResolvedProperty FindPropertyAndArrayIndex(FTrackInstancePropertyBindings& Bindings, void* BasePointer, UStruct* InStruct, const FString& PropertyName);
 
 	/** Find or add the FResolvedPropertyAndFunction for the specified object */
 	MOVIESCENE_API FResolvedPropertyAndFunction FindOrAdd(const UObject& InObject);
@@ -644,9 +574,6 @@ private:
 	/** Mapping of objects to bound functions that will be called to update data on the track */
 	TMap< FObjectKey, FResolvedPropertyAndFunction > RuntimeObjectToFunctionMap;
 
-	/** Mapping of object classes to custom handlers. */
-	TMap< const UClass*, TSharedPtr<ITrackInstanceCustomPropertyBindingHandler> > ObjectClassToCustomHandlerMap;
-
 	/** Path to the property we are bound to */
 	FString PropertyPath;
 
@@ -658,9 +585,6 @@ private:
 
 	/** Actual name of the property we are bound to */
 	FName PropertyName;
-
-	/** Whether custom handlers are allowed to be created and used. */
-	bool bAllowCustomHandlers = true;
 };
 
 /** Explicit specializations for bools */
diff --git a/Engine/Source/Runtime/MovieSceneTracks/Private/Tracks/MovieScenePropertyTrack.cpp b/Engine/Source/Runtime/MovieSceneTracks/Private/Tracks/MovieScenePropertyTrack.cpp
--- a/Engine/Source/Runtime/MovieSceneTracks/Private/Tracks/MovieScenePropertyTrack.cpp
+++ b/Engine/Source/Runtime/MovieSceneTracks/Private/Tracks/MovieScenePropertyTrack.cpp
@@ -74,7 +74,7 @@ FText UMovieScenePropertyTrack::GetDisplayNameToolTipText(const FMovieSceneLabel
 	const TArrayView<TWeakObjectPtr<>> FoundBoundObjects = LabelParams.Player->FindBoundObjects(LabelParams.BindingID, LabelParams.SequenceID);
 	for (const TWeakObjectPtr<> BoundObject : FoundBoundObjects)
 	{
-		FTrackInstancePropertyBindings InstancePropertyBinding(GetPropertyName(), GetPropertyPath().ToString(), false);
+		FTrackInstancePropertyBindings InstancePropertyBinding(GetPropertyName(), GetPropertyPath().ToString());
 		if (FProperty* BoundProperty = InstancePropertyBinding.GetProperty(*BoundObject))
 		{
 			FString PropertyName = BoundProperty->GetMetaData(TEXT("DisplayName"));
