Fixed usage of TEXT macro in UE_DEPRECATED

diff --git a/Engine/Source/Runtime/MovieScene/Public/SkeletalMeshRestoreState.h b/Engine/Source/Runtime/MovieScene/Public/SkeletalMeshRestoreState.h
--- a/Engine/Source/Runtime/MovieScene/Public/SkeletalMeshRestoreState.h
+++ b/Engine/Source/Runtime/MovieScene/Public/SkeletalMeshRestoreState.h
@@ -77,12 +77,12 @@ struct FSkeletalMeshRestoreState
 		}
 	}
 
-	UE_DEPRECATED(5.6, TEXT("InComponent is not required"))
+	UE_DEPRECATED(5.6, "InComponent is not required")
 	void RestoreState(USkeletalMeshComponent* InComponent)
 	{
 		RestoreState();
 	}
-	UE_DEPRECATED(5.6, TEXT("InComponent is not required"))
+	UE_DEPRECATED(5.6, "InComponent is not required")
 	void RestoreLOD(USkeletalMeshComponent* InComponent)
 	{
 		RestoreLOD();
