[Sequencer] Navigation Tool initial submit

diff --git a/Engine/Source/Editor/Sequencer/Public/Tools/NavigationTool/NavigationToolScopedSelection.h b/Engine/Source/Editor/Sequencer/Public/Tools/NavigationTool/NavigationToolScopedSelection.h
new file mode 100644
--- /dev/null
+++ b/Engine/Source/Editor/Sequencer/Public/Tools/NavigationTool/NavigationToolScopedSelection.h
@@ -0,0 +1,59 @@
+// Copyright Epic Games, Inc. All Rights Reserved.
+
+#pragma once
+
+#include "Containers/Array.h"
+#include "Containers/Set.h"
+
+class ISequencer;
+class UMovieSceneSection;
+class UMovieSceneTrack;
+class UObject;
+
+namespace UE::Sequencer::NavigationTool
+{
+
+enum class ENavigationToolScopedSelectionPurpose
+{
+	/** At the end of the Scope, it will set whatever has been added to the Selected List to be the new Selection */
+	Sync,
+	/** Used only to check for whether an Object is Selected or not. Cannot execute "Select" */
+	Read,
+};
+
+/** Handler to Sync Selection from Navigation Tool to the Sequencer */
+class FNavigationToolScopedSelection
+{
+public:
+	explicit FNavigationToolScopedSelection(ISequencer& InSequencer
+		, ENavigationToolScopedSelectionPurpose InPurpose = ENavigationToolScopedSelectionPurpose::Read);
+
+	~FNavigationToolScopedSelection();
+
+	SEQUENCER_API void Select(const FGuid& InObjectGuid);
+	SEQUENCER_API void Select(UMovieSceneSection* const InSection);
+	SEQUENCER_API void Select(UMovieSceneTrack* const InTrack);
+
+	SEQUENCER_API bool IsSelected(const UObject* const InObject) const;
+	SEQUENCER_API bool IsSelected(const FGuid& InObjectGuid) const;
+	SEQUENCER_API bool IsSelected(UMovieSceneSection* const InSection) const;
+	SEQUENCER_API bool IsSelected(UMovieSceneTrack* const InTrack) const;
+
+	SEQUENCER_API ISequencer& GetSequencer() const;
+
+private:
+	void SyncSelections();
+
+	ISequencer& Sequencer;
+
+	/** All Objects Selected (Sections, Tracks, Objects) */
+	TSet<const UObject*> ObjectsSet;
+
+	TArray<FGuid> SelectedObjectGuids;
+	TArray<UMovieSceneSection*> SelectedSections;
+	TArray<UMovieSceneTrack*> SelectedTracks;
+
+	ENavigationToolScopedSelectionPurpose Purpose;
+};
+
+}
diff --git a/Engine/Source/Editor/Sequencer/Public/Tools/NavigationTool/NavigationToolViewModelTypeID.h b/Engine/Source/Editor/Sequencer/Public/Tools/NavigationTool/NavigationToolViewModelTypeID.h
new file mode 100644
--- /dev/null
+++ b/Engine/Source/Editor/Sequencer/Public/Tools/NavigationTool/NavigationToolViewModelTypeID.h
@@ -0,0 +1,109 @@
+﻿// Copyright Epic Games, Inc. All Rights Reserved.
+
+#pragma once
+
+#include "CoreTypes.h"
+#include "NavigationToolCastableTypeTable.h"
+
+namespace UE::Sequencer::NavigationTool
+{
+
+#define UE_NAVIGATIONTOOL_DECLARE_VIEW_MODEL_TYPE_ID(Type)							\
+	static ::UE::Sequencer::TNavigationToolAutoRegisterViewModelTypeID<Type> ID;	\
+	static void RegisterTypeID();
+
+#define UE_NAVIGATIONTOOL_DECLARE_VIEW_MODEL_TYPE_ID_API(MODULE_API, Type)					\
+	MODULE_API static ::UE::Sequencer::TNavigationToolAutoRegisterViewModelTypeID<Type> ID;	\
+	MODULE_API static void RegisterTypeID();
+
+#define UE_NAVIGATIONTOOL_DEFINE_VIEW_MODEL_TYPE_ID(Type)																			\
+	::UE::Sequencer::TNavigationToolAutoRegisterViewModelTypeID<Type> Type::ID;														\
+	void Type::RegisterTypeID()																										\
+	{																																\
+		Type::ID.ID        = ::UE::Sequencer::FNavigationToolViewModelTypeID::RegisterNewID();										\
+		Type::ID.TypeTable = ::UE::Sequencer::FNavigationToolCastableTypeTable::MakeTypeTable<Type>((Type*)0, Type::ID.ID, #Type);	\
+	}
+
+struct FNavigationToolCastableTypeTable;
+
+struct FNavigationToolViewModelTypeID
+{
+	friend uint32 GetTypeHash(FNavigationToolViewModelTypeID In)
+	{
+		return In.ID;
+	}
+
+	friend bool operator<(FNavigationToolViewModelTypeID A, FNavigationToolViewModelTypeID B)
+	{
+		return A.ID < B.ID;
+	}
+
+	friend bool operator==(FNavigationToolViewModelTypeID A, FNavigationToolViewModelTypeID B)
+	{
+		return A.ID == B.ID;
+	}
+
+	uint32 GetTypeID() const
+	{
+		return ID;
+	}
+
+	SEQUENCERCORE_API static uint32 RegisterNewID();
+
+	FNavigationToolViewModelTypeID(FNavigationToolCastableTypeTable* InTypeTable, uint32 InID)
+		: TypeTable(InTypeTable)
+		, ID(InID)
+	{}
+
+private:
+	FNavigationToolCastableTypeTable* TypeTable = nullptr;
+	uint32 ID;
+};
+
+template<typename T>
+struct TNavigationToolViewModelTypeID
+{
+	operator FNavigationToolViewModelTypeID() const
+	{
+		Register();
+		return FNavigationToolViewModelTypeID{ TypeTable, ID };
+	}
+
+	uint32 GetTypeID() const
+	{
+		Register();
+		return ID;
+	}
+
+	FNavigationToolCastableTypeTable* GetTypeTable() const
+	{
+		Register();
+		return TypeTable;
+	}
+
+	void Register() const
+	{
+		if (!IsRegistered())
+		{
+			T::RegisterTypeID();
+		}
+	}
+
+protected:
+	friend T;
+
+	bool IsRegistered() const
+	{
+		return ID != ~0u;
+	}
+
+	FNavigationToolCastableTypeTable* TypeTable = nullptr;
+	uint32 ID = ~0u;
+};
+
+template<typename T>
+struct TNavigationToolAutoRegisterViewModelTypeID : TNavigationToolViewModelTypeID<T>
+{
+};
+
+} // namespace UE::Sequencer::NavigationTool
diff --git a/Engine/Source/Editor/Sequencer/Public/Tools/NavigationTool/Providers/NavigationToolProvider.h b/Engine/Source/Editor/Sequencer/Public/Tools/NavigationTool/Providers/NavigationToolProvider.h
new file mode 100644
--- /dev/null
+++ b/Engine/Source/Editor/Sequencer/Public/Tools/NavigationTool/Providers/NavigationToolProvider.h
@@ -0,0 +1,44 @@
+﻿// Copyright Epic Games, Inc. All Rights Reserved.
+
+#pragma once
+
+#include "Tools/NavigationTool/Providers/INavigationToolProvider.h"
+
+struct FNavigationToolViewSaveState;
+
+namespace UE::Sequencer::NavigationTool
+{
+
+class INavigationTool;
+
+/** Base Navigation Tool Provider to extend from */
+class FNavigationToolProvider : public INavigationToolProvider
+{
+public:
+	//~ Begin INavigationToolProvider
+
+	SEQUENCER_API virtual void OnExtendColumnViews(TSet<FNavigationToolColumnView>& OutColumnViews) override;
+
+	SEQUENCER_API virtual bool ShouldLockTool() const override { return false; }
+	SEQUENCER_API virtual bool ShouldHideItem(const FNavigationToolItemPtr& InItem) const override { return false; }
+
+	SEQUENCER_API virtual void UpdateItemIdContexts(const INavigationTool& InTool) final override;
+
+	SEQUENCER_API virtual TOptional<EItemDropZone> OnToolItemCanAcceptDrop(const FDragDropEvent& InDragDropEvent, EItemDropZone InDropZone, const FNavigationToolItemPtr& InTargetItem) const override;
+	SEQUENCER_API virtual FReply OnToolItemAcceptDrop(const FDragDropEvent& InDragDropEvent, EItemDropZone InDropZone, const FNavigationToolItemPtr& InTargetItem) override;
+
+	//~ End INavigationToolProvider
+
+	/** @return Editor only saved data for a specific view */
+	FNavigationToolViewSaveState* GetViewSaveState(const INavigationTool& InTool, const int32 InToolViewId) const;
+
+	/** Ensure saved editor only data contains an entry for a specific view, creating it if necessary */
+	void EnsureToolViewCount(const INavigationTool& InTool, const int32 InToolViewId);
+
+	const TArray<FName>& GetExtendedColumnNames() const;
+
+protected:
+	TArray<FName> ExtendedColumnNames;
+};
+
+}
diff --git a/Engine/Source/Editor/Sequencer/Public/Tools/NavigationTool/Widgets/SNavigationToolTreeRow.h b/Engine/Source/Editor/Sequencer/Public/Tools/NavigationTool/Widgets/SNavigationToolTreeRow.h
new file mode 100644
--- /dev/null
+++ b/Engine/Source/Editor/Sequencer/Public/Tools/NavigationTool/Widgets/SNavigationToolTreeRow.h
@@ -0,0 +1,53 @@
+// Copyright Epic Games, Inc. All Rights Reserved.
+
+#pragma once
+
+#include "Tools/NavigationTool/NavigationToolDefines.h"
+#include "Widgets/Views/STreeView.h"
+#include "Widgets/Views/STableRow.h"
+
+namespace UE::Sequencer::NavigationTool
+{
+
+class FNavigationToolItem;
+class FNavigationToolView;
+class SNavigationToolView;
+class SNavigationToolTreeView;
+
+class SNavigationToolTreeRow : public SMultiColumnTableRow<FNavigationToolItemPtr>
+{
+public:
+	SLATE_BEGIN_ARGS(SNavigationToolTreeRow) {}
+		SLATE_ATTRIBUTE(FText, HighlightText)
+	SLATE_END_ARGS()
+
+	void Construct(const FArguments& InArgs
+		, const TSharedRef<FNavigationToolView>& InToolView
+		, const TSharedPtr<SNavigationToolTreeView>& InTreeView
+		, const FNavigationToolItemPtr& InItem);
+
+	virtual TSharedRef<SWidget> GenerateWidgetForColumn(const FName& InColumnName) override;
+
+	TAttribute<FText> GetHighlightText() const { return HighlightText; }
+
+	virtual FReply OnMouseButtonUp(const FGeometry& MyGeometry, const FPointerEvent& MouseEvent) override;
+	virtual FReply OnMouseButtonDoubleClick(const FGeometry& InGeometry, const FPointerEvent& InMouseEvent) override;
+
+	TSharedPtr<FNavigationToolView> GetToolView() const;
+
+	const FTableRowStyle* GetStyle() const { return Style; }
+
+	/** The default reply if a row did not handle AcceptDrop */
+	FReply OnDefaultDrop(const FDragDropEvent& InDragDropEvent) const;
+
+private:
+	FNavigationToolItemPtr Item;
+	
+	TWeakPtr<SNavigationToolTreeView> WeakTreeView;
+	
+	TWeakPtr<FNavigationToolView> WeakToolView;
+	
+	TAttribute<FText> HighlightText;
+};
+
+}
diff --git a/Engine/Source/Editor/Sequencer/Sequencer.Build.cs b/Engine/Source/Editor/Sequencer/Sequencer.Build.cs
--- a/Engine/Source/Editor/Sequencer/Sequencer.Build.cs
+++ b/Engine/Source/Editor/Sequencer/Sequencer.Build.cs
@@ -22,6 +22,7 @@ public class Sequencer : ModuleRules
 				"Core", 
 				"CoreUObject", 
 				"CurveEditor",
+				"DeveloperSettings",
 				"InputCore",
 				"Engine", 
 				"Slate", 
@@ -44,6 +45,7 @@ public class Sequencer : ModuleRules
 				"PropertyEditor",
 				"ViewportInteraction",
 				"SerializedRecorderInterface",
+				"SourceControl",
 				"SubobjectDataInterface",
 				"ToolMenus",
 				"ToolWidgets",
@@ -67,6 +69,7 @@ public class Sequencer : ModuleRules
 			new string[] {
 				"LevelEditor",
 				"DesktopPlatform",
+				"WorkspaceMenuStructure"
 				}
 			);
 
diff --git a/Engine/Source/Runtime/MovieScene/Public/EditorSequenceNavigationDefs.h b/Engine/Source/Runtime/MovieScene/Public/EditorSequenceNavigationDefs.h
new file mode 100644
--- /dev/null
+++ b/Engine/Source/Runtime/MovieScene/Public/EditorSequenceNavigationDefs.h
@@ -0,0 +1,187 @@
+﻿// Copyright Epic Games, Inc. All Rights Reserved.
+
+#pragma once
+
+#include "EditorSequenceNavigationDefs.generated.h"
+
+UENUM()
+enum class ENavigationToolItemFlags : uint8
+{
+	None = 0,
+	/** Whether item should get the underlying UObject, ignoring if it's pending kill */
+	IgnorePendingKill = 1 << 0,
+	/** Item pending removal from the Navigation Tool */
+	PendingRemoval = 1 << 1,
+	/** Whether the item is in expanded state to show its child items */
+	Expanded = 1 << 2,
+};
+ENUM_CLASS_FLAGS(ENavigationToolItemFlags);
+
+USTRUCT()
+struct MOVIESCENE_API FNavigationToolSerializedItem
+{
+	GENERATED_BODY()
+
+	FNavigationToolSerializedItem() = default;
+
+	FNavigationToolSerializedItem(const FString& InId)
+	{
+		Id = InId;
+	}
+
+	bool IsValid() const
+	{
+		return !Id.IsEmpty();
+	}
+
+	friend uint32 GetTypeHash(const FNavigationToolSerializedItem& InItem)
+	{
+		return GetTypeHash(InItem.Id);
+	}
+
+	bool operator==(const FNavigationToolSerializedItem& InOther) const
+	{
+		return Id == InOther.Id;
+	}
+
+private:
+	UPROPERTY()
+	FString Id;
+};
+
+USTRUCT()
+struct MOVIESCENE_API FNavigationToolSerializedTreeNode
+{
+	GENERATED_BODY()
+
+	friend FNavigationToolSerializedTree;
+
+	int32 GetLocalIndex() const { return LocalIndex; }
+	int32 GetGlobalIndex() const { return GlobalIndex; }
+	int32 GetParentIndex() const { return ParentIndex; }
+
+	TConstArrayView<int32> GetChildrenIndices() const { return ChildrenIndices; }
+
+	const FNavigationToolSerializedTreeNode* GetParentTreeNode() const;
+
+	int32 CalculateHeight() const;
+
+	TArray<const FNavigationToolSerializedTreeNode*> FindPath(const TArray<const FNavigationToolSerializedTreeNode*>& InItems) const;
+
+	void Reset();
+
+private:
+	/** Index of this tree node relative to the parent node children items. Can be used as means of ordering. */
+	UPROPERTY()
+	int32 LocalIndex = INDEX_NONE;
+
+	/** Index of this tree node in the owning tree */
+	UPROPERTY()
+	int32 GlobalIndex = INDEX_NONE;
+
+	/** Absolute Index of the Parent Node in the owning tree. If INDEX_NONE, it means Parent is Root. */
+	UPROPERTY()
+	int32 ParentIndex = INDEX_NONE;
+
+	/** Absolute indices of the children in the owning tree */
+	UPROPERTY()
+	TArray<int32> ChildrenIndices;
+
+	/** Pointer to the tree that owns this node */
+	FNavigationToolSerializedTree* OwningTree = nullptr;
+};
+
+USTRUCT()
+struct MOVIESCENE_API FNavigationToolSerializedTree
+{
+	GENERATED_BODY()
+
+	FNavigationToolSerializedTree();
+
+	void PostSerialize(const FArchive& Ar);
+
+	FNavigationToolSerializedTreeNode& GetRootNode() { return RootNode; }
+	const FNavigationToolSerializedTreeNode& GetRootNode() const { return RootNode; }
+
+	FNavigationToolSerializedTreeNode* FindTreeNode(const FNavigationToolSerializedItem& InItem);
+	const FNavigationToolSerializedTreeNode* FindTreeNode(const FNavigationToolSerializedItem& InItem) const;
+
+	const FNavigationToolSerializedItem* GetItemAtIndex(int32 InIndex) const;
+
+	FNavigationToolSerializedTreeNode& GetOrAddTreeNode(const FNavigationToolSerializedItem& InItem, const FNavigationToolSerializedItem& InParentItem);
+
+	static const FNavigationToolSerializedTreeNode* FindLowestCommonAncestor(const TArray<const FNavigationToolSerializedTreeNode*>& InItems);
+
+	static bool CompareTreeItemOrder(const FNavigationToolSerializedTreeNode* InA, const FNavigationToolSerializedTreeNode* InB);
+
+	void Reset();
+
+private:
+	void UpdateTreeNodes();
+
+	UPROPERTY()
+	FNavigationToolSerializedTreeNode RootNode;
+
+	UPROPERTY()
+	TArray<FNavigationToolSerializedItem> SceneItems;
+
+	UPROPERTY()
+	TMap<FNavigationToolSerializedItem, FNavigationToolSerializedTreeNode> ItemTreeMap;
+};
+
+USTRUCT()
+struct MOVIESCENE_API FNavigationToolViewColumnSaveState
+{
+	GENERATED_BODY()
+
+	UPROPERTY()
+	bool bVisible = false;
+
+	UPROPERTY()
+	float Size = 1.f;
+};
+
+USTRUCT()
+struct MOVIESCENE_API FNavigationToolViewSaveState
+{
+	GENERATED_BODY()
+
+	/** Items specific to this Navigation Tool Instance, rather than being shared across Navigation Tools (e.g. Expansion flags) */
+	UPROPERTY()
+	TMap<FString, ENavigationToolItemFlags> ViewItemFlags;
+
+	/** Map of the column Ids to their overriden (i.e. saved) visibility */
+	UPROPERTY()
+	TMap<FName, FNavigationToolViewColumnSaveState> ColumnsState;
+
+	/** Active list of item filters */
+	UPROPERTY()
+	TSet<FName> ActiveItemFilters;
+};
+
+USTRUCT()
+struct MOVIESCENE_API FNavigationToolSaveState
+{
+	GENERATED_BODY()
+
+	UPROPERTY(Transient)
+	FNavigationToolSerializedTree SerializedTree;
+
+	UPROPERTY()
+	TMap<FString, FColor> ItemColorMap;
+
+	UPROPERTY()
+	TArray<FNavigationToolViewSaveState> ToolViewSaveStates;
+
+	UPROPERTY()
+	FString ContextPath;
+};
+
+template<>
+struct TStructOpsTypeTraits<FNavigationToolSerializedTree> : TStructOpsTypeTraitsBase2<FNavigationToolSerializedTree>
+{
+	enum 
+	{
+		WithPostSerialize = true,
+	};
+};
diff --git a/Engine/Source/Runtime/MovieScene/Public/MovieScene.h b/Engine/Source/Runtime/MovieScene/Public/MovieScene.h
--- a/Engine/Source/Runtime/MovieScene/Public/MovieScene.h
+++ b/Engine/Source/Runtime/MovieScene/Public/MovieScene.h
@@ -10,6 +10,7 @@
 #include "CoreMinimal.h"
 #include "CoreTypes.h"
 #include "Delegates/Delegate.h"
+#include "EditorSequenceNavigationDefs.h"
 #include "EventHandlers/MovieSceneDataEventContainer.h"
 #include "Decorations/MovieSceneDecorationContainer.h"
 #include "HAL/PlatformCrt.h"
@@ -129,6 +130,10 @@ struct FMovieSceneEditorData
 	FFloatRange WorkingRange_DEPRECATED;
 	UPROPERTY()
 	FFloatRange ViewRange_DEPRECATED;
+
+	/** Saved Navigation Tool state to restore */
+	UPROPERTY()
+	FNavigationToolSaveState NavigationToolState;
 };
 
 
