Fix for a crash where Conditions are not weakly referenced by MovieSceneCompiledDataManager, and when a level sequence actor is destroyed that uses them, GC will crash on breaking Disregard for GC assumptions.

diff --git a/Engine/Source/Runtime/MovieScene/Private/EntitySystem/MovieSceneSequenceUpdaters.cpp b/Engine/Source/Runtime/MovieScene/Private/EntitySystem/MovieSceneSequenceUpdaters.cpp
--- a/Engine/Source/Runtime/MovieScene/Private/EntitySystem/MovieSceneSequenceUpdaters.cpp
+++ b/Engine/Source/Runtime/MovieScene/Private/EntitySystem/MovieSceneSequenceUpdaters.cpp
@@ -841,25 +841,26 @@ void FSequenceUpdater_Hierarchical::Update(TSharedRef<const FSharedPlaybackState
 			checkf(SubData, TEXT("Sub data does not exist for a SequenceID that exists in the hierarchical tree - this indicates a corrupt compilation product."));
 
 			bool bSubSequenceConditionFailed = false;
-			if (SubData->Condition)
+			const UMovieSceneCondition* Condition = SubData->WeakCondition.Get();
+			if (Condition)
 			{
 				// If we're able to cache the condition result, then it should be cached above when its entity got processed- retrieve that value.
 				// Otherwise, test it again.
-				if (SubData->Condition->CanCacheResult(SharedPlaybackState))
+				if (Condition->CanCacheResult(SharedPlaybackState))
 				{
-					if (bool* ConditionResult = CachedConditionResults.Find(SubData->Condition->ComputeCacheKey(FGuid(), RootOverrideSequenceID, SharedPlaybackState, FindObject<UMovieSceneSubSection>(CompiledDataManager->GetEntryRef(RootCompiledDataID).GetSequence(), *SubData->SectionPath.ToString()))))
+					if (bool* ConditionResult = CachedConditionResults.Find(Condition->ComputeCacheKey(FGuid(), RootOverrideSequenceID, SharedPlaybackState, FindObject<UMovieSceneSubSection>(CompiledDataManager->GetEntryRef(RootCompiledDataID).GetSequence(), *SubData->SectionPath.ToString()))))
 					{
 						if (*ConditionResult == false)
 						{
 							bSubSequenceConditionFailed = true;
 						}
 					}
-					else if (!SubData->Condition->EvaluateCondition(FGuid(), RootOverrideSequenceID, SharedPlaybackState))
+					else if (!Condition->EvaluateCondition(FGuid(), RootOverrideSequenceID, SharedPlaybackState))
 					{
 						bSubSequenceConditionFailed = true;
 					}
 				}
-				else if (!SubData->Condition->EvaluateCondition(FGuid(), RootOverrideSequenceID, SharedPlaybackState))
+				else if (!Condition->EvaluateCondition(FGuid(), RootOverrideSequenceID, SharedPlaybackState))
 				{
 					bSubSequenceConditionFailed = true;
 				}
diff --git a/Engine/Source/Runtime/MovieScene/Private/Evaluation/MovieSceneSequenceHierarchy.cpp b/Engine/Source/Runtime/MovieScene/Private/Evaluation/MovieSceneSequenceHierarchy.cpp
--- a/Engine/Source/Runtime/MovieScene/Private/Evaluation/MovieSceneSequenceHierarchy.cpp
+++ b/Engine/Source/Runtime/MovieScene/Private/Evaluation/MovieSceneSequenceHierarchy.cpp
@@ -101,7 +101,7 @@ FMovieSceneSubSequenceData::FMovieSceneSubSequenceData(const UMovieSceneSubSecti
 		PostRollRange = TRange<FFrameNumber>(PostRollStart, PostRollEnd);
 	}
 
-	Condition = MovieSceneHelpers::GetSequenceCondition(InSubSection.GetTypedOuter<UMovieSceneTrack>(), &InSubSection, true);
+	WeakCondition = MovieSceneHelpers::GetSequenceCondition(InSubSection.GetTypedOuter<UMovieSceneTrack>(), &InSubSection, true);
 }
 
 UMovieSceneSequence* FMovieSceneSubSequenceData::GetSequence() const
diff --git a/Engine/Source/Runtime/MovieScene/Public/Evaluation/MovieSceneSequenceHierarchy.h b/Engine/Source/Runtime/MovieScene/Public/Evaluation/MovieSceneSequenceHierarchy.h
--- a/Engine/Source/Runtime/MovieScene/Public/Evaluation/MovieSceneSequenceHierarchy.h
+++ b/Engine/Source/Runtime/MovieScene/Public/Evaluation/MovieSceneSequenceHierarchy.h
@@ -178,7 +178,7 @@ struct FMovieSceneSubSequenceData
 	
 	/* Optional condition pulled off of the UMovieSceneSubSection or UMovieSceneSubTrack */
 	UPROPERTY()
-	TObjectPtr<const UMovieSceneCondition> Condition;
+	TWeakObjectPtr<const UMovieSceneCondition> WeakCondition;
 
 
 private:
