Fix static analysis warning

diff --git a/Engine/Source/Runtime/MovieScene/Private/MovieSceneCommonHelpers.cpp b/Engine/Source/Runtime/MovieScene/Private/MovieSceneCommonHelpers.cpp
--- a/Engine/Source/Runtime/MovieScene/Private/MovieSceneCommonHelpers.cpp
+++ b/Engine/Source/Runtime/MovieScene/Private/MovieSceneCommonHelpers.cpp
@@ -543,11 +543,18 @@ UCameraComponent* MovieSceneHelpers::CameraComponentFromRuntimeObject(UObject* R
 
 float MovieSceneHelpers::GetSoundDuration(USoundBase* Sound)
 {
-	if (Sound->IsProcedurallyGenerated() || !Sound->IsOneShot())
+	if (Sound)
 	{
-		return INDEFINITELY_LOOPING_DURATION;
+		if (Sound->IsProcedurallyGenerated() || !Sound->IsOneShot())
+		{
+			return INDEFINITELY_LOOPING_DURATION;
+		}
+		else
+		{
+			return FMath::Max(0.0f, Sound->GetDuration());
+		}
 	}
-	return Sound ? FMath::Max(0.0f, Sound->GetDuration()) : 0.0f;
+	return 0.0f;
 }
 
 float MovieSceneHelpers::CalculateWeightForBlending(UMovieSceneSection* SectionToKey, FFrameNumber Time)
