Sequencer: add debug trap for leaky shared playback state

diff --git a/Engine/Source/Runtime/MovieScene/Private/EntitySystem/MovieSceneSequenceInstance.cpp b/Engine/Source/Runtime/MovieScene/Private/EntitySystem/MovieSceneSequenceInstance.cpp
--- a/Engine/Source/Runtime/MovieScene/Private/EntitySystem/MovieSceneSequenceInstance.cpp
+++ b/Engine/Source/Runtime/MovieScene/Private/EntitySystem/MovieSceneSequenceInstance.cpp
@@ -120,10 +120,15 @@ void FSequenceInstance::Initialize()
 FSequenceInstance::~FSequenceInstance()
 {
 	const int32 RefCount = SharedPlaybackState.GetSharedReferenceCount();
-	ensureAlwaysMsgf(
+	if (!ensureAlwaysMsgf(
 			RootInstanceHandle != InstanceHandle || RefCount == 1,
 			TEXT("References to SharedPlaybackState should not be held past the lifetime of its root sequence instance. SharedPlaybackState has %d references."),
-			RefCount);
+			RefCount))
+	{
+#if !UE_BUILD_SHIPPING
+		SharedPlaybackState->SetDebugBreakOnDestroy();
+#endif
+	}
 }
 
 FSequenceInstance::FSequenceInstance(FSequenceInstance&&) = default;
diff --git a/Engine/Source/Runtime/MovieScene/Private/EntitySystem/MovieSceneSharedPlaybackState.cpp b/Engine/Source/Runtime/MovieScene/Private/EntitySystem/MovieSceneSharedPlaybackState.cpp
--- a/Engine/Source/Runtime/MovieScene/Private/EntitySystem/MovieSceneSharedPlaybackState.cpp
+++ b/Engine/Source/Runtime/MovieScene/Private/EntitySystem/MovieSceneSharedPlaybackState.cpp
@@ -31,6 +31,16 @@ FSharedPlaybackState::FSharedPlaybackState(
 	}
 }
 
+FSharedPlaybackState::~FSharedPlaybackState()
+{
+#if !UE_BUILD_SHIPPING
+	if (bDebugBreakOnDestroy)
+	{
+		ensureAlwaysMsgf(false, TEXT("Debug break was requested upon destruction of this state."));
+	}
+#endif
+}
+
 TSharedPtr<FMovieSceneEntitySystemRunner> FSharedPlaybackState::GetRunner() const
 {
 	if (UMovieSceneEntitySystemLinker* Linker = WeakLinker.Get())
diff --git a/Engine/Source/Runtime/MovieScene/Public/EntitySystem/MovieSceneSharedPlaybackState.h b/Engine/Source/Runtime/MovieScene/Public/EntitySystem/MovieSceneSharedPlaybackState.h
--- a/Engine/Source/Runtime/MovieScene/Public/EntitySystem/MovieSceneSharedPlaybackState.h
+++ b/Engine/Source/Runtime/MovieScene/Public/EntitySystem/MovieSceneSharedPlaybackState.h
@@ -62,6 +62,8 @@ public:
 			UMovieSceneSequence& InRootSequence,
 			const FSharedPlaybackStateCreateParams& CreateParams);
 
+	~FSharedPlaybackState();
+
 public:
 	
 	/** Gets the playback context */
@@ -288,6 +290,10 @@ public:
 
 	void InvalidateCachedData();
 
+#if !UE_BUILD_SHIPPING
+	void SetDebugBreakOnDestroy() { bDebugBreakOnDestroy = true; }
+#endif
+
 private:
 
 	template<typename T>
@@ -325,6 +331,10 @@ private:
 
 	/** Pre-animated state utility for the sequence hierarchy */
 	FMovieSceneInstancePreAnimatedState PreAnimatedState;
+
+#if !UE_BUILD_SHIPPING
+	bool bDebugBreakOnDestroy = false;
+#endif
 };
 
 } // namespace UE::MovieScene
