[MovieScene]
* Moved dll storage to methods/variables instead of on types. This is needed to get merged modules working

diff --git a/Engine/Source/Runtime/MovieScene/Public/Bindings/MovieSceneCustomBinding.h b/Engine/Source/Runtime/MovieScene/Public/Bindings/MovieSceneCustomBinding.h
--- a/Engine/Source/Runtime/MovieScene/Public/Bindings/MovieSceneCustomBinding.h
+++ b/Engine/Source/Runtime/MovieScene/Public/Bindings/MovieSceneCustomBinding.h
@@ -24,7 +24,7 @@ struct FMovieSceneBindingReference;
 class UMovieSceneSpawnableBindingBase;
 
 USTRUCT(BlueprintType)
-struct MOVIESCENE_API FMovieSceneBindingResolveResult
+struct FMovieSceneBindingResolveResult
 {
 	GENERATED_BODY()
 
@@ -37,7 +37,7 @@ struct MOVIESCENE_API FMovieSceneBindingResolveResult
 * Blueprint-specific resolution context for custom bindings.
 */
 USTRUCT(BlueprintType)
-struct MOVIESCENE_API FMovieSceneBindingResolveContext
+struct FMovieSceneBindingResolveContext
 {
 	GENERATED_BODY()
 
@@ -53,16 +53,16 @@ struct MOVIESCENE_API FMovieSceneBindingResolveContext
 /**
  * A custom binding. Allows users to define their own binding resolution types, including dynamic 'Replaceable' bindings with previews in editor, as well as Spawnable types.
  */
-UCLASS(abstract, DefaultToInstanced, EditInlineNew)
-class MOVIESCENE_API UMovieSceneCustomBinding
+UCLASS(abstract, DefaultToInstanced, EditInlineNew, MinimalAPI)
+class UMovieSceneCustomBinding
 	: public UObject
 {
 public:
 
 	GENERATED_BODY()
 
-	static const int32 BaseEnginePriority;
-	static const int32 BaseCustomPriority;
+	static MOVIESCENE_API const int32 BaseEnginePriority;
+	static MOVIESCENE_API const int32 BaseCustomPriority;
 
 	UFUNCTION(BlueprintCallable, Category = "Sequencer|Binding")
 	static int32 GetBaseEnginePriority() { return BaseEnginePriority; }
diff --git a/Engine/Source/Runtime/MovieScene/Public/Conditions/MovieSceneCondition.h b/Engine/Source/Runtime/MovieScene/Public/Conditions/MovieSceneCondition.h
--- a/Engine/Source/Runtime/MovieScene/Public/Conditions/MovieSceneCondition.h
+++ b/Engine/Source/Runtime/MovieScene/Public/Conditions/MovieSceneCondition.h
@@ -48,7 +48,7 @@ enum class EMovieSceneConditionCheckFrequency : uint8
 * Blueprint-friendly struct containing any context needed to evaluate conditions. 
 */
 USTRUCT(BlueprintType)
-struct MOVIESCENE_API FMovieSceneConditionContext
+struct FMovieSceneConditionContext
 {
 GENERATED_BODY()
 public:
@@ -70,7 +70,7 @@ public:
 * Container struct for instanced UMovieSceneConditions, existing only to allow property type customization for choosing conditions.
 */
 USTRUCT(BlueprintType)
-struct MOVIESCENE_API FMovieSceneConditionContainer
+struct FMovieSceneConditionContainer
 {
 	GENERATED_BODY()
 
@@ -82,8 +82,8 @@ struct MOVIESCENE_API FMovieSceneConditionContainer
  * Abstract condition class. Conditions can be applied to sections, tracks, and track rows to determine whether or not they are evaluated at runtime.
  * This allows developers to create Sequences with dynamic behavior based on gameplay state, local player state, player hardware, etc.
  */
-UCLASS(abstract, Blueprintable, DefaultToInstanced, EditInlineNew, meta = (ShowWorldContextPin), CollapseCategories)
-class MOVIESCENE_API UMovieSceneCondition
+UCLASS(abstract, Blueprintable, DefaultToInstanced, EditInlineNew, meta = (ShowWorldContextPin), CollapseCategories, MinimalAPI)
+class UMovieSceneCondition
 	: public UMovieSceneSignedObject
 {
 	GENERATED_BODY()
@@ -91,17 +91,17 @@ class MOVIESCENE_API UMovieSceneCondition
 public:
 
 	/* Called by Sequencer code to evaluate this condition, passing relevant context. Note that BindingGuid will be invalid for conditions on global sections/tracks. */
-	bool EvaluateCondition(FGuid BindingGuid, FMovieSceneSequenceID SequenceID, TSharedRef<const UE::MovieScene::FSharedPlaybackState> SharedPlaybackState) const;
+	MOVIESCENE_API bool EvaluateCondition(FGuid BindingGuid, FMovieSceneSequenceID SequenceID, TSharedRef<const UE::MovieScene::FSharedPlaybackState> SharedPlaybackState) const;
 
 	/* Called by Sequencer to compute a cache key for this condition given the passed in context.
 	*  By default, this key will be computed based on the Binding Scope, and if relevant, the binding and entity owner.
 	* If a condition returns the same cache key given the same or different contexts, it will not be rechecked, and a cached value may be used.
 	*/
-	virtual uint32 ComputeCacheKey(FGuid BindingGuid, FMovieSceneSequenceID SequenceID, TSharedRef<const UE::MovieScene::FSharedPlaybackState> SharedPlaybackState, UObject* EntityOwner) const;
+	MOVIESCENE_API virtual uint32 ComputeCacheKey(FGuid BindingGuid, FMovieSceneSequenceID SequenceID, TSharedRef<const UE::MovieScene::FSharedPlaybackState> SharedPlaybackState, UObject* EntityOwner) const;
 
-	virtual bool CanCacheResult(TSharedRef<const UE::MovieScene::FSharedPlaybackState> SharedPlaybackState) const;
+	MOVIESCENE_API virtual bool CanCacheResult(TSharedRef<const UE::MovieScene::FSharedPlaybackState> SharedPlaybackState) const;
 
-	EMovieSceneConditionScope GetConditionScope() const;
+	MOVIESCENE_API EMovieSceneConditionScope GetConditionScope() const;
 
 #if WITH_EDITORONLY_DATA
 	/* If true, will skip evaluating the condition and always return true. Useful for authoring or debugging. */
@@ -115,24 +115,24 @@ protected:
 
 	/* Override to implement your condition.*/
 	UFUNCTION(BlueprintImplementableEvent, CallInEditor, meta=(DisplayName = "On Evaluate Condition"))
-	bool BP_EvaluateCondition(const FMovieSceneConditionContext& ConditionContext) const;
+	MOVIESCENE_API bool BP_EvaluateCondition(const FMovieSceneConditionContext& ConditionContext) const;
 
 	/* Override in native code to implement your condition. Note that BindingGuid will be invalid for conditions on global sections/tracks.*/
-	virtual bool EvaluateConditionInternal(FGuid BindingGuid, FMovieSceneSequenceID SequenceID, TSharedRef<const UE::MovieScene::FSharedPlaybackState> SharedPlaybackState) const;
+	MOVIESCENE_API virtual bool EvaluateConditionInternal(FGuid BindingGuid, FMovieSceneSequenceID SequenceID, TSharedRef<const UE::MovieScene::FSharedPlaybackState> SharedPlaybackState) const;
 	
 	/* Returns the scope of the condition, which determines whether the condition needs to be re-evaluated for different bindings or entities in the Sequence. */
 	UFUNCTION(BlueprintNativeEvent, CallInEditor, meta=(DisplayName = "Get Scope"))
-	EMovieSceneConditionScope BP_GetScope() const;
+	MOVIESCENE_API EMovieSceneConditionScope BP_GetScope() const;
 
 	/* Returns the scope of the condition, which determines whether the condition needs to be re-evaluated for different bindings or entities in the Sequence. */
-	virtual EMovieSceneConditionScope GetScopeInternal() const;
+	MOVIESCENE_API virtual EMovieSceneConditionScope GetScopeInternal() const;
 
 	/* Returns the check frequency of the condition, which determines whether the condition result can change during playback and needs to get re-evaluated. */
 	UFUNCTION(BlueprintNativeEvent, CallInEditor, meta=(DisplayName = "Get Check Frequency"))
-	EMovieSceneConditionCheckFrequency BP_GetCheckFrequency() const;
+	MOVIESCENE_API EMovieSceneConditionCheckFrequency BP_GetCheckFrequency() const;
 
 	/* Returns the check frequency of the condition, which determines whether the condition result can change during playback and needs to get re-evaluated. */
-	virtual EMovieSceneConditionCheckFrequency GetCheckFrequencyInternal() const;
+	MOVIESCENE_API virtual EMovieSceneConditionCheckFrequency GetCheckFrequencyInternal() const;
 
 protected:
 	/* If true, inverts the result of the condition check.*/
diff --git a/Engine/Source/Runtime/MovieScene/Public/EntitySystem/MovieSceneEntityGroupingSystem.h b/Engine/Source/Runtime/MovieScene/Public/EntitySystem/MovieSceneEntityGroupingSystem.h
--- a/Engine/Source/Runtime/MovieScene/Public/EntitySystem/MovieSceneEntityGroupingSystem.h
+++ b/Engine/Source/Runtime/MovieScene/Public/EntitySystem/MovieSceneEntityGroupingSystem.h
@@ -17,6 +17,8 @@
 
 #include "MovieSceneEntityGroupingSystem.generated.h"
 
+#define UE_API MOVIESCENE_API
+
 class UMovieSceneEntityGroupingSystem;
 
 namespace UE::MovieScene
@@ -35,22 +37,22 @@ struct CBatchGroupingPolicy
 /**
  * Utility class used by the grouping system's policies (see below) to manage groups.
  */
-struct MOVIESCENE_API FEntityGroupBuilder
+struct FEntityGroupBuilder
 {
-	FEntityGroupBuilder(UMovieSceneEntityGroupingSystem* InOwner, FEntityGroupingPolicyKey InPolicyKey);
+	UE_API FEntityGroupBuilder(UMovieSceneEntityGroupingSystem* InOwner, FEntityGroupingPolicyKey InPolicyKey);
 
 	/** Make a full group ID from an existing group index */
 	FEntityGroupID MakeGroupID(int32 GroupIndex) const { return FEntityGroupID(PolicyKey, GroupIndex); }
 	/** Make an invalid group ID that is associated with the policy key */
 	FEntityGroupID MakeInvalidGroupID() const { return FEntityGroupID(PolicyKey, INDEX_NONE); }
 	/** Add the entity to the given group. The entity must already have the group ID component. */
-	void AddEntityToGroup(const FMovieSceneEntityID& InEntity, const FEntityGroupID& InNewGroupID);
+	UE_API void AddEntityToGroup(const FMovieSceneEntityID& InEntity, const FEntityGroupID& InNewGroupID);
 	/** Remove the entity from the given group. The entity must already have the group ID component. */
-	void RemoveEntityFromGroup(const FMovieSceneEntityID& InEntity, const FEntityGroupID& InPreviousGroupID);
+	UE_API void RemoveEntityFromGroup(const FMovieSceneEntityID& InEntity, const FEntityGroupID& InPreviousGroupID);
 
-	int32 AllocateGroupIndex();
+	UE_API int32 AllocateGroupIndex();
 
-	void ReportUsedGroupIndex(int32 GroupIndex);
+	UE_API void ReportUsedGroupIndex(int32 GroupIndex);
 
 private:
 	UMovieSceneEntityGroupingSystem* Owner;
@@ -312,8 +314,8 @@ struct TTupleGroupingPolicy
 } // namespace UE::MovieScene
 
 
-UCLASS()
-class MOVIESCENE_API UMovieSceneEntityGroupingSystem : public UMovieSceneEntitySystem
+UCLASS(MinimalAPI)
+class UMovieSceneEntityGroupingSystem : public UMovieSceneEntitySystem
 {
 public:
 
@@ -324,7 +326,7 @@ public:
 
 	GENERATED_BODY()
 
-	UMovieSceneEntityGroupingSystem(const FObjectInitializer& ObjInit);
+	UE_API UMovieSceneEntityGroupingSystem(const FObjectInitializer& ObjInit);
 
 	/**
 	 * Add a new grouping policy that will use the given components to make up a group key.
@@ -392,13 +394,13 @@ public:
 	/**
 	 * Remove a previously added grouping policy.
 	 */
-	void RemoveGrouping(FEntityGroupingPolicyKey InPolicyKey);
+	UE_API void RemoveGrouping(FEntityGroupingPolicyKey InPolicyKey);
 
 	/**
 	 * Allocate a new group index used to uniquely identify a collection of entities that animate the same target.
 	 * Group indices are globally unique within this system, regardless of the 'type' of the target.
 	 */
-	int32 AllocateGroupIndex(FEntityGroupingPolicyKey InPolicy);
+	UE_API int32 AllocateGroupIndex(FEntityGroupingPolicyKey InPolicy);
 
 
 	/**
@@ -456,3 +458,5 @@ private:
 	friend struct UE::MovieScene::FEntityGroupBuilder;
 };
 
+
+#undef UE_API 
\ No newline at end of file
diff --git a/Engine/Source/Runtime/MovieScene/Public/Evaluation/EventContextsPlaybackCapability.h b/Engine/Source/Runtime/MovieScene/Public/Evaluation/EventContextsPlaybackCapability.h
--- a/Engine/Source/Runtime/MovieScene/Public/Evaluation/EventContextsPlaybackCapability.h
+++ b/Engine/Source/Runtime/MovieScene/Public/Evaluation/EventContextsPlaybackCapability.h
@@ -11,9 +11,9 @@ namespace UE::MovieScene
 {
 
 /** Playback capability for controlling how events are triggered */
-struct MOVIESCENE_API IEventContextsPlaybackCapability
+struct IEventContextsPlaybackCapability
 {
-	UE_DECLARE_MOVIESCENE_PLAYBACK_CAPABILITY(IEventContextsPlaybackCapability)
+	UE_DECLARE_MOVIESCENE_PLAYBACK_CAPABILITY_API(MOVIESCENE_API, IEventContextsPlaybackCapability)
 
 	/** Virtual destructor */
 	virtual ~IEventContextsPlaybackCapability() {}
diff --git a/Engine/Source/Runtime/MovieScene/Public/Evaluation/EventTriggerControlPlaybackCapability.h b/Engine/Source/Runtime/MovieScene/Public/Evaluation/EventTriggerControlPlaybackCapability.h
--- a/Engine/Source/Runtime/MovieScene/Public/Evaluation/EventTriggerControlPlaybackCapability.h
+++ b/Engine/Source/Runtime/MovieScene/Public/Evaluation/EventTriggerControlPlaybackCapability.h
@@ -10,16 +10,16 @@ namespace UE::MovieScene
 {
 
 /** Playback capability for controlling how events are triggered */
-struct MOVIESCENE_API FEventTriggerControlPlaybackCapability
+struct FEventTriggerControlPlaybackCapability
 {
-	UE_DECLARE_MOVIESCENE_PLAYBACK_CAPABILITY(FEventTriggerControlPlaybackCapability)
+	UE_DECLARE_MOVIESCENE_PLAYBACK_CAPABILITY_API(MOVIESCENE_API, FEventTriggerControlPlaybackCapability)
 
 	/**
 	 * Returns whether triggering events should be temporarily disabled.
 	 *
 	 * @param OutDisabledUntilTime  The time until which to disable triggering events
 	 */
-	bool IsDisablingEventTriggers(FFrameTime& OutDisabledUntilTime) const;
+	MOVIESCENE_API bool IsDisablingEventTriggers(FFrameTime& OutDisabledUntilTime) const;
 
 	TOptional<FFrameTime> DisableEventTriggersUntilTime;
 };
diff --git a/Engine/Source/Runtime/MovieScene/Public/Evaluation/IMovieScenePlaybackCapability.h b/Engine/Source/Runtime/MovieScene/Public/Evaluation/IMovieScenePlaybackCapability.h
--- a/Engine/Source/Runtime/MovieScene/Public/Evaluation/IMovieScenePlaybackCapability.h
+++ b/Engine/Source/Runtime/MovieScene/Public/Evaluation/IMovieScenePlaybackCapability.h
@@ -89,7 +89,7 @@ struct IPlaybackCapability
 /**
  * Structure providing basic information on a playback capability type.
  */
-struct MOVIESCENE_API FPlaybackCapabilityIDInfo
+struct FPlaybackCapabilityIDInfo
 {
 #if UE_MOVIESCENE_ENTITY_DEBUG
 	/** Display name for debugging. */
@@ -100,15 +100,15 @@ struct MOVIESCENE_API FPlaybackCapabilityIDInfo
 /**
  * A registry for all known playback capability types.
  */
-class MOVIESCENE_API FPlaybackCapabilityIDRegistry
+class FPlaybackCapabilityIDRegistry
 {
 public:
 
 	/** Gets the registry. */
-	static FPlaybackCapabilityIDRegistry* Get();
+	static MOVIESCENE_API FPlaybackCapabilityIDRegistry* Get();
 
 	/** Registers a new playback capability type with the given display name. */
-	FPlaybackCapabilityID RegisterNewID(const TCHAR* InDebugName);
+	MOVIESCENE_API FPlaybackCapabilityID RegisterNewID(const TCHAR* InDebugName);
 
 private:
 
@@ -123,7 +123,7 @@ extern MOVIESCENE_API FPlaybackCapabilityIDRegistry* GPlaybackCapabilityIDRegist
  * Interface to a playback capability pointer that has a valid v-table for the debugger
  * to be able to show what it is.
  */
-struct MOVIESCENE_API IPlaybackCapabilityDebuggingTypedPtr
+struct IPlaybackCapabilityDebuggingTypedPtr
 {
 	virtual ~IPlaybackCapabilityDebuggingTypedPtr() {}
 	void* Ptr = nullptr;
diff --git a/Engine/Source/Runtime/MovieScene/Public/Evaluation/MovieSceneEvaluationState.h b/Engine/Source/Runtime/MovieScene/Public/Evaluation/MovieSceneEvaluationState.h
--- a/Engine/Source/Runtime/MovieScene/Public/Evaluation/MovieSceneEvaluationState.h
+++ b/Engine/Source/Runtime/MovieScene/Public/Evaluation/MovieSceneEvaluationState.h
@@ -40,9 +40,9 @@ namespace UE::MovieScene
 	/**
 	 * Playback capability for being notified of object bindings changing.
 	 */
-	struct MOVIESCENE_API IObjectBindingNotifyPlaybackCapability
+	struct IObjectBindingNotifyPlaybackCapability
 	{
-		UE_DECLARE_MOVIESCENE_PLAYBACK_CAPABILITY(IObjectBindingNotifyPlaybackCapability)
+		UE_DECLARE_MOVIESCENE_PLAYBACK_CAPABILITY_API(MOVIESCENE_API, IObjectBindingNotifyPlaybackCapability)
 
 		/** Called when multiple object bindings have changed. */
 		virtual void NotifyBindingsChanged() {}
@@ -54,9 +54,9 @@ namespace UE::MovieScene
 	/**
 	 * Playback capability for storing static object binding overrides.
 	 */
-	struct MOVIESCENE_API IStaticBindingOverridesPlaybackCapability
+	struct IStaticBindingOverridesPlaybackCapability
 	{
-		UE_DECLARE_MOVIESCENE_PLAYBACK_CAPABILITY(IStaticBindingOverridesPlaybackCapability)
+		UE_DECLARE_MOVIESCENE_PLAYBACK_CAPABILITY_API(MOVIESCENE_API, IStaticBindingOverridesPlaybackCapability)
 
 		virtual ~IStaticBindingOverridesPlaybackCapability() {}
 
diff --git a/Engine/Source/Runtime/MovieScene/Public/Evaluation/SequenceDirectorPlaybackCapability.h b/Engine/Source/Runtime/MovieScene/Public/Evaluation/SequenceDirectorPlaybackCapability.h
--- a/Engine/Source/Runtime/MovieScene/Public/Evaluation/SequenceDirectorPlaybackCapability.h
+++ b/Engine/Source/Runtime/MovieScene/Public/Evaluation/SequenceDirectorPlaybackCapability.h
@@ -17,20 +17,20 @@ struct FSharedPlaybackState;
 /**
  * Playback capability for sequences that have a director blueprint.
  */
-struct MOVIESCENE_API FSequenceDirectorPlaybackCapability : public IPlaybackCapability
+struct FSequenceDirectorPlaybackCapability : public IPlaybackCapability
 {
-	UE_DECLARE_MOVIESCENE_PLAYBACK_CAPABILITY(FSequenceDirectorPlaybackCapability)
+	UE_DECLARE_MOVIESCENE_PLAYBACK_CAPABILITY_API(MOVIESCENE_API, FSequenceDirectorPlaybackCapability)
 
 	/** Remove all director blueprint instances */
-	void ResetDirectorInstances();
+	MOVIESCENE_API void ResetDirectorInstances();
 
 	/** Gets a new or existing director blueprint instance for the given root or sub sequence */
-	UObject* GetOrCreateDirectorInstance(TSharedRef<const FSharedPlaybackState> SharedPlaybackState, FMovieSceneSequenceIDRef SequenceID);
+	MOVIESCENE_API UObject* GetOrCreateDirectorInstance(TSharedRef<const FSharedPlaybackState> SharedPlaybackState, FMovieSceneSequenceIDRef SequenceID);
 
 public:
 
 	// IPlaybackCapability interface.
-	virtual void InvalidateCachedData(UMovieSceneEntitySystemLinker* Linker) override;
+	MOVIESCENE_API virtual void InvalidateCachedData(UMovieSceneEntitySystemLinker* Linker) override;
 
 private:
 
diff --git a/Engine/Source/Runtime/MovieScene/Public/IMovieScenePlayer.h b/Engine/Source/Runtime/MovieScene/Public/IMovieScenePlayer.h
--- a/Engine/Source/Runtime/MovieScene/Public/IMovieScenePlayer.h
+++ b/Engine/Source/Runtime/MovieScene/Public/IMovieScenePlayer.h
@@ -37,12 +37,12 @@ namespace UE::MovieScene
 	/**
 	 * Playback capability for storing an IMovieScenePlayer unique index.
 	 */
-	struct MOVIESCENE_API FPlayerIndexPlaybackCapability
+	struct FPlayerIndexPlaybackCapability
 	{
-		UE_DECLARE_MOVIESCENE_PLAYBACK_CAPABILITY(FPlayerIndexPlaybackCapability)
+		UE_DECLARE_MOVIESCENE_PLAYBACK_CAPABILITY_API(MOVIESCENE_API, FPlayerIndexPlaybackCapability)
 
-		static IMovieScenePlayer* GetPlayer(TSharedRef<const FSharedPlaybackState> Owner);
-		static uint16 GetPlayerIndex(TSharedRef<const FSharedPlaybackState> Owner);
+		static MOVIESCENE_API IMovieScenePlayer* GetPlayer(TSharedRef<const FSharedPlaybackState> Owner);
+		static MOVIESCENE_API uint16 GetPlayerIndex(TSharedRef<const FSharedPlaybackState> Owner);
 
 		FPlayerIndexPlaybackCapability(uint16 InPlayerIndex)
 			: PlayerIndex(InPlayerIndex)
diff --git a/Engine/Source/Runtime/MovieScene/Public/MovieSceneLegacyPlayer.h b/Engine/Source/Runtime/MovieScene/Public/MovieSceneLegacyPlayer.h
--- a/Engine/Source/Runtime/MovieScene/Public/MovieSceneLegacyPlayer.h
+++ b/Engine/Source/Runtime/MovieScene/Public/MovieSceneLegacyPlayer.h
@@ -11,9 +11,9 @@ namespace UE::MovieScene
 
 	struct FSharedPlaybackState;
 
-	struct MOVIESCENE_API ILegacyPlayerProviderPlaybackCapability
+	struct ILegacyPlayerProviderPlaybackCapability
 	{
-		UE_DECLARE_MOVIESCENE_PLAYBACK_CAPABILITY(ILegacyPlayerProviderPlaybackCapability)
+		UE_DECLARE_MOVIESCENE_PLAYBACK_CAPABILITY_API(MOVIESCENE_API, ILegacyPlayerProviderPlaybackCapability)
 
 		virtual ~ILegacyPlayerProviderPlaybackCapability() {}
 
diff --git a/Engine/Source/Runtime/MovieSceneTracks/Public/Evaluation/CameraCutPlaybackCapability.h b/Engine/Source/Runtime/MovieSceneTracks/Public/Evaluation/CameraCutPlaybackCapability.h
--- a/Engine/Source/Runtime/MovieSceneTracks/Public/Evaluation/CameraCutPlaybackCapability.h
+++ b/Engine/Source/Runtime/MovieSceneTracks/Public/Evaluation/CameraCutPlaybackCapability.h
@@ -23,9 +23,9 @@ namespace UE::MovieScene
 	/**
 	 * Playback capability for sequences that can run camera cuts.
 	 */
-	struct MOVIESCENETRACKS_API FCameraCutPlaybackCapability
+	struct FCameraCutPlaybackCapability
 	{
-		UE_DECLARE_MOVIESCENE_PLAYBACK_CAPABILITY(FCameraCutPlaybackCapability)
+		UE_DECLARE_MOVIESCENE_PLAYBACK_CAPABILITY_API(MOVIESCENETRACKS_API, FCameraCutPlaybackCapability)
 
 		/** Whether the associated sequence should execute camera cuts */
 		virtual bool ShouldUpdateCameraCut() { return true; }
diff --git a/Engine/Source/Runtime/MovieSceneTracks/Public/Evaluation/ViewportSettingsPlaybackCapability.h b/Engine/Source/Runtime/MovieSceneTracks/Public/Evaluation/ViewportSettingsPlaybackCapability.h
--- a/Engine/Source/Runtime/MovieSceneTracks/Public/Evaluation/ViewportSettingsPlaybackCapability.h
+++ b/Engine/Source/Runtime/MovieSceneTracks/Public/Evaluation/ViewportSettingsPlaybackCapability.h
@@ -40,9 +40,9 @@ namespace UE::MovieScene
 /**
  * Playback capability for controlling game and editor viewports.
  */
-struct MOVIESCENETRACKS_API FViewportSettingsPlaybackCapability
+struct FViewportSettingsPlaybackCapability
 {
-	UE_DECLARE_MOVIESCENE_PLAYBACK_CAPABILITY(FViewportSettingsPlaybackCapability)
+	UE_DECLARE_MOVIESCENE_PLAYBACK_CAPABILITY_API(MOVIESCENETRACKS_API, FViewportSettingsPlaybackCapability)
 
 	/*
 	 * Set the perspective viewport settings
diff --git a/Engine/Source/Runtime/MovieSceneTracks/Public/Sections/MovieSceneParameterSection.h b/Engine/Source/Runtime/MovieSceneTracks/Public/Sections/MovieSceneParameterSection.h
--- a/Engine/Source/Runtime/MovieSceneTracks/Public/Sections/MovieSceneParameterSection.h
+++ b/Engine/Source/Runtime/MovieSceneTracks/Public/Sections/MovieSceneParameterSection.h
@@ -167,7 +167,7 @@ struct FBaseParameterNameAndValue
 	{}
 
 	/** Creates a new FScalarParameterNameAndCurve for a specific scalar parameter. */
-	MOVIESCENETRACKS_API FBaseParameterNameAndValue(FName InParameterName)
+	FBaseParameterNameAndValue(FName InParameterName)
 	: ParameterName(InParameterName)
 	{}
 
@@ -187,7 +187,7 @@ struct FBoolParameterNameAndCurve : public FBaseParameterNameAndValue
 	FBoolParameterNameAndCurve()
 	{}
 
-	MOVIESCENETRACKS_API FBoolParameterNameAndCurve(FName InParameterName)
+	FBoolParameterNameAndCurve(FName InParameterName)
 	: FBaseParameterNameAndValue(InParameterName)
 	{}
 
@@ -207,7 +207,7 @@ struct FScalarParameterNameAndCurve : public FBaseParameterNameAndValue
 	FScalarParameterNameAndCurve()
 	{}
 
-	MOVIESCENETRACKS_API FScalarParameterNameAndCurve(FName InParameterName)
+	FScalarParameterNameAndCurve(FName InParameterName)
 	: FBaseParameterNameAndValue(InParameterName)
 	{}
 
@@ -227,7 +227,7 @@ struct FVector2DParameterNameAndCurves : public FBaseParameterNameAndValue
 	FVector2DParameterNameAndCurves()
 	{}
 
-	MOVIESCENETRACKS_API FVector2DParameterNameAndCurves(FName InParameterName)
+	FVector2DParameterNameAndCurves(FName InParameterName)
 	: FBaseParameterNameAndValue(InParameterName)
 	{}
 
@@ -253,7 +253,7 @@ struct FVectorParameterNameAndCurves : public FBaseParameterNameAndValue
 	FVectorParameterNameAndCurves()
 	{}
 
-	MOVIESCENETRACKS_API FVectorParameterNameAndCurves(FName InParameterName)
+	FVectorParameterNameAndCurves(FName InParameterName)
 	: FBaseParameterNameAndValue(InParameterName)
 	{}
 
@@ -282,7 +282,7 @@ struct FColorParameterNameAndCurves : public FBaseParameterNameAndValue
 	FColorParameterNameAndCurves()
 	{}
 
-	MOVIESCENETRACKS_API FColorParameterNameAndCurves(FName InParameterName)
+	FColorParameterNameAndCurves(FName InParameterName)
 	: FBaseParameterNameAndValue(InParameterName)
 	{}
 
@@ -314,7 +314,7 @@ struct FTransformParameterNameAndCurves : public FBaseParameterNameAndValue
 	FTransformParameterNameAndCurves()
 	{}
 
-	MOVIESCENETRACKS_API FTransformParameterNameAndCurves(FName InParameterName)
+	FTransformParameterNameAndCurves(FName InParameterName)
 	: FBaseParameterNameAndValue(InParameterName)
 	{}
 
