Sequencer: fix pre-animated state being incorrectly saved as "persistent" for root tracks

#virtualized

diff --git a/Engine/Source/Runtime/MovieScene/Private/Evaluation/PreAnimatedState/MovieScenePreAnimatedStateExtension.cpp b/Engine/Source/Runtime/MovieScene/Private/Evaluation/PreAnimatedState/MovieScenePreAnimatedStateExtension.cpp
--- a/Engine/Source/Runtime/MovieScene/Private/Evaluation/PreAnimatedState/MovieScenePreAnimatedStateExtension.cpp
+++ b/Engine/Source/Runtime/MovieScene/Private/Evaluation/PreAnimatedState/MovieScenePreAnimatedStateExtension.cpp
@@ -856,11 +856,8 @@ void FPreAnimatedStateExtension::SavePreAnimatedStateDirectly(UObject& InObject,
 
 	AddSourceMetaData(Entry);
 
-	EPreAnimatedStorageRequirement Requirement = HasActiveCaptureSource()
-		? EPreAnimatedStorageRequirement::Transient
-		: EPreAnimatedStorageRequirement::Persistent;
-
-	if (!ObjectStorage->IsStorageRequirementSatisfied(StorageIndex, Requirement))
+	EPreAnimatedStorageRequirement StorageRequirement = GetStorageRequirement(Entry);
+	if (!ObjectStorage->IsStorageRequirementSatisfied(StorageIndex, StorageRequirement))
 	{
 		IMovieScenePreAnimatedTokenPtr Token = Producer.CacheExistingState(InObject);
 		if (Token.IsValid())
@@ -871,7 +868,7 @@ void FPreAnimatedStateExtension::SavePreAnimatedStateDirectly(UObject& InObject,
 				Producer.InitializeObjectForAnimation(InObject);
 			}
 
-			ObjectStorage->AssignPreAnimatedValue(StorageIndex, Requirement, MoveTemp(Token));
+			ObjectStorage->AssignPreAnimatedValue(StorageIndex, StorageRequirement, MoveTemp(Token));
 		}
 	}
 }
@@ -900,11 +897,8 @@ void FPreAnimatedStateExtension::SavePreAnimatedStateDirectly(FMovieSceneAnimTyp
 
 	AddSourceMetaData(Entry);
 
-	EPreAnimatedStorageRequirement Requirement = HasActiveCaptureSource()
-		? EPreAnimatedStorageRequirement::Transient
-		: EPreAnimatedStorageRequirement::Persistent;
-
-	if (!RootStorage->IsStorageRequirementSatisfied(StorageIndex, Requirement))
+	EPreAnimatedStorageRequirement StorageRequirement = GetStorageRequirement(Entry);
+	if (!RootStorage->IsStorageRequirementSatisfied(StorageIndex, StorageRequirement))
 	{
 		IMovieScenePreAnimatedGlobalTokenPtr Token = Producer.CacheExistingState();
 		if (Token.IsValid())
@@ -915,7 +909,7 @@ void FPreAnimatedStateExtension::SavePreAnimatedStateDirectly(FMovieSceneAnimTyp
 				Producer.InitializeForAnimation();
 			}
 
-			RootStorage->AssignPreAnimatedValue(StorageIndex, Requirement, MoveTemp(Token));
+			RootStorage->AssignPreAnimatedValue(StorageIndex, StorageRequirement, MoveTemp(Token));
 		}
 	}
 }
diff --git a/Engine/Source/Runtime/MovieScene/Public/Evaluation/PreAnimatedState/MovieScenePreAnimatedStateStorage.h b/Engine/Source/Runtime/MovieScene/Public/Evaluation/PreAnimatedState/MovieScenePreAnimatedStateStorage.h
--- a/Engine/Source/Runtime/MovieScene/Public/Evaluation/PreAnimatedState/MovieScenePreAnimatedStateStorage.h
+++ b/Engine/Source/Runtime/MovieScene/Public/Evaluation/PreAnimatedState/MovieScenePreAnimatedStateStorage.h
@@ -707,14 +707,11 @@ public:
 
 			TrackCaptureSource(Entry, TrackingMode);
 
-			EPreAnimatedStorageRequirement Requirement = ParentExtension->HasActiveCaptureSource()
-				? EPreAnimatedStorageRequirement::Transient
-				: EPreAnimatedStorageRequirement::Persistent;
-
-			if (!IsStorageRequirementSatisfied(Entry.ValueHandle.StorageIndex, Requirement))
+			EPreAnimatedStorageRequirement StorageRequirement = this->ParentExtension->GetStorageRequirement(Entry);
+			if (!IsStorageRequirementSatisfied(Entry.ValueHandle.StorageIndex, StorageRequirement))
 			{
 				StorageType NewValue = this->Traits.CachePreAnimatedValue(this->Traits.ResolveComponent(Values)...);
-				this->AssignPreAnimatedValue(Entry.ValueHandle.StorageIndex, Requirement, MoveTemp(NewValue));
+				this->AssignPreAnimatedValue(Entry.ValueHandle.StorageIndex, StorageRequirement, MoveTemp(NewValue));
 			}
 		}
 	}
@@ -748,14 +745,11 @@ public:
 		FPreAnimatedStateEntry Entry{ GroupHandle, FPreAnimatedStateCachedValueHandle{ this->StorageID, StorageIndex } };
 		TrackCaptureSource(Entry, TrackingMode);
 
-		EPreAnimatedStorageRequirement Requirement = ParentExtension->HasActiveCaptureSource()
-			? EPreAnimatedStorageRequirement::Transient
-			: EPreAnimatedStorageRequirement::Persistent;
-
-		if (!IsStorageRequirementSatisfied(Entry.ValueHandle.StorageIndex, Requirement))
+		EPreAnimatedStorageRequirement StorageRequirement = this->ParentExtension->GetStorageRequirement(Entry);
+		if (!IsStorageRequirementSatisfied(Entry.ValueHandle.StorageIndex, StorageRequirement))
 		{
 			StorageType NewValue = CacheCallback(InKey);
-			AssignPreAnimatedValue(StorageIndex, Requirement, MoveTemp(NewValue));
+			AssignPreAnimatedValue(StorageIndex, StorageRequirement, MoveTemp(NewValue));
 		}
 	}
 
diff --git a/Engine/Source/Runtime/MovieSceneTracks/Private/Tests/MovieSceneFadeTrackTests.cpp b/Engine/Source/Runtime/MovieSceneTracks/Private/Tests/MovieSceneFadeTrackTests.cpp
new file mode 100644
--- /dev/null
+++ b/Engine/Source/Runtime/MovieSceneTracks/Private/Tests/MovieSceneFadeTrackTests.cpp
@@ -0,0 +1,11 @@
+// Copyright Epic Games, Inc. All Rights Reserved.
+
+#include "Tests/MovieSceneFadeTrackTests.h"
+
+#include "Camera/PlayerCameraManager.h"
+
+float UMovieSceneFadeTrackTestLibrary::GetManualFadeAmount(APlayerCameraManager* PlayerCameraManager)
+{
+	return PlayerCameraManager ? PlayerCameraManager->FadeAmount : 0.f;
+}
+
diff --git a/Engine/Source/Runtime/MovieSceneTracks/Public/Tests/MovieSceneFadeTrackTests.h b/Engine/Source/Runtime/MovieSceneTracks/Public/Tests/MovieSceneFadeTrackTests.h
new file mode 100644
--- /dev/null
+++ b/Engine/Source/Runtime/MovieSceneTracks/Public/Tests/MovieSceneFadeTrackTests.h
@@ -0,0 +1,22 @@
+// Copyright Epic Games, Inc. All Rights Reserved.
+
+#pragma once
+
+#include "CoreTypes.h"
+#include "Kismet/BlueprintFunctionLibrary.h"
+
+#include "MovieSceneFadeTrackTests.generated.h"
+
+class APlayerCameraManager;
+
+UCLASS()
+class UMovieSceneFadeTrackTestLibrary : public UBlueprintFunctionLibrary
+{
+	GENERATED_BODY()
+
+public:
+
+	UFUNCTION(BlueprintCallable, Category="Sequencer Tests")
+	static float GetManualFadeAmount(APlayerCameraManager* PlayerCameraManager);
+};
+
