Sequencer: fix broken support for ping-pong playback when we only cross the loop boundary once

diff --git a/Engine/Source/Runtime/MovieScene/Private/Evaluation/MovieScenePlaybackManager.cpp b/Engine/Source/Runtime/MovieScene/Private/Evaluation/MovieScenePlaybackManager.cpp
--- a/Engine/Source/Runtime/MovieScene/Private/Evaluation/MovieScenePlaybackManager.cpp
+++ b/Engine/Source/Runtime/MovieScene/Private/Evaluation/MovieScenePlaybackManager.cpp
@@ -206,6 +206,10 @@ void FMovieScenePlaybackManager::InternalUpdateToTick(const FFrameNumber NextTic
 				}
 			}
 		}
+		else if (bPingPongPlayback)
+		{
+			ReversePlayDirection();
+		}
 
 		// Complete the loops we said we completed.
 		NumLoopsCompleted += NumLoopsNewlyCompleted;
