Sequencer: Fix ensure hit when creating piecewise curves from curves with multiple keys on the same frame

diff --git a/Engine/Source/Runtime/MovieScene/Private/Channels/MovieSceneDoubleChannel.cpp b/Engine/Source/Runtime/MovieScene/Private/Channels/MovieSceneDoubleChannel.cpp
--- a/Engine/Source/Runtime/MovieScene/Private/Channels/MovieSceneDoubleChannel.cpp
+++ b/Engine/Source/Runtime/MovieScene/Private/Channels/MovieSceneDoubleChannel.cpp
@@ -617,14 +617,27 @@ UE::MovieScene::FPiecewiseCurve FMovieSceneDoubleChannel::AsPiecewiseCurve(bool
 		}
 	}
 
-	if (Times.Num() > 1)
+	Curve.Values.Reserve(Times.Num());
+	for (int32 Index = 0; Index < Times.Num() - 1; ++Index)
 	{
-		for (int32 Index = 0; Index < Times.Num()-1; ++Index)
+		using namespace UE::MovieScene::Interpolation;
+
+		// Add a constant interp if this is the index of the last key, or if the next key sits on the same frame
+		const bool bNoRange =
+			Index == Times.Num() - 1 ||
+			Times[Index] == Times[Index + 1];
+
+		if (bNoRange)
+		{
+			FCachedInterpolationRange Range = FCachedInterpolationRange::Only(Times[Index]);
+			Curve.Values.Emplace(Range, FConstantValue(Range.Start, Values[Index].Value));
+		}
+		else
 		{
 			Interpolation::FCachedInterpolation Interp = FMovieSceneDoubleChannelImpl::GetInterpolationForKey(this, Index);
 			if (ensure(Interp.IsValid()))
 			{
-				Curve.Values.Emplace(Interp);
+				Curve.Values.Add(MoveTemp(Interp));
 			}
 		}
 	}
diff --git a/Engine/Source/Runtime/MovieScene/Private/Channels/MovieSceneFloatChannel.cpp b/Engine/Source/Runtime/MovieScene/Private/Channels/MovieSceneFloatChannel.cpp
--- a/Engine/Source/Runtime/MovieScene/Private/Channels/MovieSceneFloatChannel.cpp
+++ b/Engine/Source/Runtime/MovieScene/Private/Channels/MovieSceneFloatChannel.cpp
@@ -109,14 +109,27 @@ UE::MovieScene::FPiecewiseCurve FMovieSceneFloatChannel::AsPiecewiseCurve(bool b
 		}
 	}
 
-	if (Times.Num() > 1)
+	Curve.Values.Reserve(Times.Num());
+	for (int32 Index = 0; Index < Times.Num() - 1; ++Index)
 	{
-		for (int32 Index = 0; Index < Times.Num() - 1; ++Index)
+		using namespace UE::MovieScene::Interpolation;
+
+		// Add a constant interp if this is the index of the last key, or if the next key sits on the same frame
+		const bool bNoRange =
+			Index == Times.Num() - 1 ||
+			Times[Index] == Times[Index + 1];
+
+		if (bNoRange)
+		{
+			FCachedInterpolationRange Range = FCachedInterpolationRange::Only(Times[Index]);
+			Curve.Values.Emplace(Range, FConstantValue(Range.Start, Values[Index].Value));
+		}
+		else
 		{
 			Interpolation::FCachedInterpolation Interp = FMovieSceneFloatChannelImpl::GetInterpolationForKey(this, Index);
 			if (ensure(Interp.IsValid()))
 			{
-				Curve.Values.Emplace(Interp);
+				Curve.Values.Add(MoveTemp(Interp));
 			}
 		}
 	}
