Support Additive Animations in Mixer

Add new tag for additive animations
Add tag to animation entities if the animation is a valid additive animation
In the anim mixer check that flag to determine if the animation is additive

diff --git a/Engine/Plugins/Experimental/MovieSceneAnimMixer/Source/MovieSceneAnimMixer/Private/Systems/MovieSceneAnimMixerSystem.cpp b/Engine/Plugins/Experimental/MovieSceneAnimMixer/Source/MovieSceneAnimMixer/Private/Systems/MovieSceneAnimMixerSystem.cpp
--- a/Engine/Plugins/Experimental/MovieSceneAnimMixer/Source/MovieSceneAnimMixer/Private/Systems/MovieSceneAnimMixerSystem.cpp
+++ b/Engine/Plugins/Experimental/MovieSceneAnimMixer/Source/MovieSceneAnimMixer/Private/Systems/MovieSceneAnimMixerSystem.cpp
@@ -708,7 +708,7 @@ void UMovieSceneAnimMixerSystem::OnRun(FSystemTaskPrerequisites& InPrerequisites
 				}
 
 				const FComponentMask& Type = Linker->EntityManager.GetEntityType(EntityID);
-				Entry->bAdditive = Type.ContainsAny({BuiltInComponents->Tags.AdditiveBlend, BuiltInComponents->Tags.AdditiveFromBaseBlend});
+				Entry->bAdditive = Type.Contains(BuiltInComponents->Tags.AdditiveAnimation);
 				Entry->bRequiresBlend = Type.Contains(AnimMixerComponents->Tags.RequiresBlending);
 
 				TSharedPtr<FMovieSceneAnimMixer> ExistingMixer = Entry->WeakMixer.Pin();
diff --git a/Engine/Source/Runtime/MovieScene/Private/EntitySystem/BuiltInComponentTypes.cpp b/Engine/Source/Runtime/MovieScene/Private/EntitySystem/BuiltInComponentTypes.cpp
--- a/Engine/Source/Runtime/MovieScene/Private/EntitySystem/BuiltInComponentTypes.cpp
+++ b/Engine/Source/Runtime/MovieScene/Private/EntitySystem/BuiltInComponentTypes.cpp
@@ -409,6 +409,8 @@ FBuiltInComponentTypes::FBuiltInComponentTypes()
 	Tags.AdditiveFromBaseBlend   = ComponentRegistry->NewTag(TEXT("Is Additive From Base Blend"));
 	Tags.OverrideBlend			 = ComponentRegistry->NewTag(TEXT("Is Override Blend"));
 
+	Tags.AdditiveAnimation       = ComponentRegistry->NewTag(TEXT("Is Additive Animation"));
+
 	Tags.NeedsLink               = ComponentRegistry->NewTag(TEXT("Needs Link"));
 	Tags.NeedsUnlink             = ComponentRegistry->NewTag(TEXT("Needs Unlink"));
 	Tags.HasUnresolvedBinding    = ComponentRegistry->NewTag(TEXT("Has Unresolved Binding"));
@@ -450,6 +452,7 @@ FBuiltInComponentTypes::FBuiltInComponentTypes()
 	ComponentRegistry->Factories.DefineChildComponent(Tags.AdditiveBlend, Tags.AdditiveBlend);
 	ComponentRegistry->Factories.DefineChildComponent(Tags.OverrideBlend, Tags.OverrideBlend);
 	ComponentRegistry->Factories.DefineChildComponent(Tags.AdditiveFromBaseBlend, Tags.AdditiveFromBaseBlend);
+	ComponentRegistry->Factories.DefineChildComponent(Tags.AdditiveAnimation, Tags.AdditiveAnimation);
 	ComponentRegistry->Factories.DefineChildComponent(Tags.FixedTime,     Tags.FixedTime);
 	ComponentRegistry->Factories.DefineChildComponent(Tags.PreRoll,       Tags.PreRoll);
 	ComponentRegistry->Factories.DefineChildComponent(Tags.SectionPreRoll,Tags.SectionPreRoll);
diff --git a/Engine/Source/Runtime/MovieScene/Public/EntitySystem/BuiltInComponentTypes.h b/Engine/Source/Runtime/MovieScene/Public/EntitySystem/BuiltInComponentTypes.h
--- a/Engine/Source/Runtime/MovieScene/Public/EntitySystem/BuiltInComponentTypes.h
+++ b/Engine/Source/Runtime/MovieScene/Public/EntitySystem/BuiltInComponentTypes.h
@@ -750,6 +750,8 @@ public:
 		FComponentTypeID AdditiveFromBaseBlend;
 		FComponentTypeID OverrideBlend;
 
+		FComponentTypeID AdditiveAnimation;
+
 		FComponentTypeID NeedsLink;
 		FComponentTypeID NeedsUnlink;
 
diff --git a/Engine/Source/Runtime/MovieSceneTracks/Private/Sections/MovieSceneSkeletalAnimationSection.cpp b/Engine/Source/Runtime/MovieSceneTracks/Private/Sections/MovieSceneSkeletalAnimationSection.cpp
--- a/Engine/Source/Runtime/MovieSceneTracks/Private/Sections/MovieSceneSkeletalAnimationSection.cpp
+++ b/Engine/Source/Runtime/MovieSceneTracks/Private/Sections/MovieSceneSkeletalAnimationSection.cpp
@@ -530,6 +530,7 @@ void UMovieSceneSkeletalAnimationSection::ImportEntityImpl(UMovieSceneEntitySyst
 	const FMovieSceneTracksComponentTypes* TrackComponents = FMovieSceneTracksComponentTypes::Get();
 
 	const FGuid ObjectBindingID = InParams.GetObjectBindingID();
+	const bool bIsAdditive = Params.Animation->IsValidAdditive();
 
 	if (!ObjectBindingID.IsValid())
 	{
@@ -544,6 +545,7 @@ void UMovieSceneSkeletalAnimationSection::ImportEntityImpl(UMovieSceneEntitySyst
 		.Add(BuiltInComponents->GenericObjectBinding, ObjectBindingID)
 		.Add(BuiltInComponents->BoundObjectResolver, UMovieSceneSkeletalAnimationSystem::ResolveSkeletalMeshComponentBinding)
 		.AddConditional(BuiltInComponents->WeightChannel, &Params.Weight, Params.Weight.HasAnyData())
+		.AddTagConditional(BuiltInComponents->Tags.AdditiveAnimation, bIsAdditive)
 	);
 }
 
