Make semetrical changes to FSequencer to match previous changes to UMovieSceneSequencePlayer. AND... change bool to an enum to help make the code easier to understand.

This is in support of previously added support for situations where a sequence player's time controller might not be immediately ready to start providing time.

See https://p4-swarm.epicgames.net/changes/38756133 for original changes to UMovieSceneSequencePlayer.

diff --git a/Engine/Source/Editor/Sequencer/Private/Sequencer.cpp b/Engine/Source/Editor/Sequencer/Private/Sequencer.cpp
--- a/Engine/Source/Editor/Sequencer/Private/Sequencer.cpp
+++ b/Engine/Source/Editor/Sequencer/Private/Sequencer.cpp
@@ -842,6 +842,7 @@ FSequencer::FSequencer()
 	, bUpdatingExternalSelection( false )
 	, bNeedsEvaluate(false)
 	, bNeedsInvalidateCachedData(false)
+	, TimeControllerState(ETimeControllerState::ReadyToPlay)
 {
 	// Exposes the sequencer and curve editor command lists to subscribers from other systems
 	FInputBindingManager::Get().RegisterCommandList(FSequencerCommands::Get().GetContextName(), SequencerCommandBindings);
@@ -1067,6 +1068,12 @@ void FSequencer::Tick(float InDeltaTime)
 	UWorld* World = PlaybackContext ? PlaybackContext->GetWorld() : nullptr;
 	float Dilation = World ? World->GetWorldSettings()->CinematicTimeDilation : 1.f;
 
+	if (TimeControllerState == ETimeControllerState::PreparingToPlay && TimeController->IsReadyToPlay())
+	{
+		SetPlaybackStatus(EMovieScenePlayerStatus::Playing);
+		TimeControllerState = ETimeControllerState::ReadyToPlay;
+	}
+
 	TimeController->Tick(InDeltaTime, PlaybackSpeed * Dilation);
 
 	FQualifiedFrameTime GlobalTime = GetGlobalTime();
@@ -4502,6 +4509,10 @@ void FSequencer::SetPlaybackStatus(EMovieScenePlayerStatus::Type InPlaybackStatu
 	}
 	else if (InPlaybackStatus == EMovieScenePlayerStatus::Stopped || InPlaybackStatus == EMovieScenePlayerStatus::Scrubbing || InPlaybackStatus == EMovieScenePlayerStatus::Stepping)
 	{
+		// We may have been waiting for our time controller to tell us it was ready to play,
+		// so TimeControllerState may be ETimeControllerState::PreparingToPlay. Here we reset that back to its
+		// initialized state of true.
+		TimeControllerState = ETimeControllerState::ReadyToPlay;
 		SequencerState = ESS_Paused;
 	}
 	
@@ -5042,7 +5053,16 @@ FReply FSequencer::OnPlay(bool bTogglePlay)
 		SpeedIndexBeforePlay = CurrentSpeedIndex;
 		PlaybackSpeedBeforePlay = PlaybackSpeed;
 
-		SetPlaybackStatus(EMovieScenePlayerStatus::Playing);
+		TimeController->PrepareToPlay(GetGlobalTime());
+		if (TimeController->IsReadyToPlay())
+		{
+			SetPlaybackStatus(EMovieScenePlayerStatus::Playing);
+			TimeControllerState = ETimeControllerState::ReadyToPlay;
+		}
+		else
+		{
+			TimeControllerState = ETimeControllerState::PreparingToPlay;
+		}
 
 		// Make sure Slate ticks during playback
 		SequencerWidget->RegisterActiveTimerForPlayback();
diff --git a/Engine/Source/Editor/Sequencer/Private/Sequencer.h b/Engine/Source/Editor/Sequencer/Private/Sequencer.h
--- a/Engine/Source/Editor/Sequencer/Private/Sequencer.h
+++ b/Engine/Source/Editor/Sequencer/Private/Sequencer.h
@@ -1645,4 +1645,11 @@ private:
 
 	/** Navigation Tool instance for this Sequencer */
 	TSharedPtr<UE::Sequencer::NavigationTool::FNavigationTool> NavigationTool;
+
+	enum class ETimeControllerState
+	{
+		ReadyToPlay,
+		PreparingToPlay
+	};
+	ETimeControllerState TimeControllerState;
 };
diff --git a/Engine/Source/Runtime/MovieScene/Private/MovieSceneSequencePlayer.cpp b/Engine/Source/Runtime/MovieScene/Private/MovieSceneSequencePlayer.cpp
--- a/Engine/Source/Runtime/MovieScene/Private/MovieSceneSequencePlayer.cpp
+++ b/Engine/Source/Runtime/MovieScene/Private/MovieSceneSequencePlayer.cpp
@@ -175,7 +175,7 @@ UMovieSceneSequencePlayer::UMovieSceneSequencePlayer(const FObjectInitializer& I
 	, Status(EMovieScenePlayerStatus::Stopped)
 	, bReversePlayback(false)
 	, bPendingOnStartedPlaying(false)
-	, bIsTimeControllerReadyToPlay(true)
+	, TimeControllerState(ETimeControllerState::ReadyToPlay)
 	, bIsAsyncUpdate(false)
 	, bSkipNextUpdate(false)
 	, bUpdateNetSync(false)
@@ -316,7 +316,7 @@ void UMovieSceneSequencePlayer::PlayInternal()
 		return;
 	}
 
-	if (!IsPlaying() && Sequence && CanPlay() && bIsTimeControllerReadyToPlay)
+	if (!IsPlaying() && Sequence && CanPlay() && TimeControllerState == ETimeControllerState::ReadyToPlay)
 	{
 		const FString SequenceName = GetSequenceName(true);
 		UE_LOG(LogMovieScene, Verbose, TEXT("PlayInternal - %s (current status: %s)"), *SequenceName, *UEnum::GetValueAsString(Status));
@@ -352,11 +352,11 @@ void UMovieSceneSequencePlayer::PlayInternal()
 		if (TimeController->IsReadyToPlay())
 		{
 			StartTimeControllerAndBroadcastPlayState();
-			bIsTimeControllerReadyToPlay = true;
+			TimeControllerState = ETimeControllerState::ReadyToPlay;
 		}
 		else
 		{
-			bIsTimeControllerReadyToPlay = false;
+			TimeControllerState = ETimeControllerState::PreparingToPlay;
 		}
 	}
 }
@@ -455,9 +455,9 @@ void UMovieSceneSequencePlayer::StopInternal(FFrameTime TimeToResetTo)
 	UE_LOG(LogMovieScene, Verbose, TEXT("StopInternal - %s (at: %s, current status: %s)"), *SequenceName, *LexToString(TimeToResetTo), *UEnum::GetValueAsString(Status));
 
 	// We may have been waiting for our time controller to tell us it was ready to play,
-	// so bIsTimeControllerReadyToPlay may be false. Here we reset that back to its
-	// initialized state of true.
-	bIsTimeControllerReadyToPlay = true;
+	// so TimeControllerState may be ETimeControllerState::PreparingToPlay. Here we reset that back to its
+	// initialized state of ETimeControllerState::ReadyToPlay.
+	TimeControllerState = ETimeControllerState::ReadyToPlay;
 
 	if (IsPlaying() || IsPaused())
 	{
@@ -1065,13 +1065,13 @@ void UMovieSceneSequencePlayer::Update(const float DeltaSeconds)
 
 	UpdateNetworkSync();
 
-	if (!bIsTimeControllerReadyToPlay && TimeController->IsReadyToPlay())
+	if (TimeControllerState == ETimeControllerState::PreparingToPlay && TimeController->IsReadyToPlay())
 	{
 		StartTimeControllerAndBroadcastPlayState();
-		bIsTimeControllerReadyToPlay = true;
+		TimeControllerState = ETimeControllerState::ReadyToPlay;
 	}
 
-	if (bIsTimeControllerReadyToPlay && IsPlaying())
+	if (TimeControllerState == ETimeControllerState::ReadyToPlay && IsPlaying())
 	{
 		// Delta seconds has already been multiplied by GetEffectiveTimeDilation at this point, so don't pass that through to Tick
 		float PlayRate = bReversePlayback ? -PlaybackSettings.PlayRate : PlaybackSettings.PlayRate;
diff --git a/Engine/Source/Runtime/MovieScene/Public/MovieSceneSequencePlayer.h b/Engine/Source/Runtime/MovieScene/Public/MovieSceneSequencePlayer.h
--- a/Engine/Source/Runtime/MovieScene/Public/MovieSceneSequencePlayer.h
+++ b/Engine/Source/Runtime/MovieScene/Public/MovieSceneSequencePlayer.h
@@ -671,8 +671,12 @@ protected:
 	/** Set to true to invoke OnStartedPlaying on first update tick for started playing */
 	uint32 bPendingOnStartedPlaying : 1;
 
-	/** set to true when Time Controller has been told to PrepareToPlay but has not yet become IsReadyToPlay */
-	uint32 bIsTimeControllerReadyToPlay : 1;
+	enum class ETimeControllerState
+	{
+		ReadyToPlay,
+		PreparingToPlay
+	};
+	ETimeControllerState TimeControllerState;
 
 	/** Set to true when the player is currently in the main level update */
 	uint32 bIsAsyncUpdate : 1;
