Sequencer: fix crash when reinitializing a Sequencer editor on a new sequence while the previous one had legacy tracks

We cleared the SharedPlaybackState reference from the Sequencer player's RootEvaluationTemplateInstance before destroying the sequence instance, because we didn't want to trip the leak-detection assert in the sequence instance destructor. However, in some cases, destroying a sequence instance first requires a "final update". This update could run legacy track eval-templates, and those called back into the player's RootEvaluationTemplateInstance... so they need it to still be functional by having a valid SharedPlaybackState.

Now we clear the SharedPlaybackState in a callback that gets executed after the "final update", but before the destruction of the sequence instance.

The change also adds a check for the place where it was crashing.

diff --git a/Engine/Source/Runtime/MovieScene/Private/EntitySystem/MovieSceneEntitySystemRunner.cpp b/Engine/Source/Runtime/MovieScene/Private/EntitySystem/MovieSceneEntitySystemRunner.cpp
--- a/Engine/Source/Runtime/MovieScene/Private/EntitySystem/MovieSceneEntitySystemRunner.cpp
+++ b/Engine/Source/Runtime/MovieScene/Private/EntitySystem/MovieSceneEntitySystemRunner.cpp
@@ -203,6 +203,10 @@ bool FMovieSceneEntitySystemRunner::QueueFinalUpdateAndDestroy(FInstanceHandle I
 	return QueueFinalUpdateImpl(InInstanceHandle, FSimpleDelegate(), true);
 }
 
+bool FMovieSceneEntitySystemRunner::QueueFinalUpdateAndDestroy(FInstanceHandle InInstanceHandle, FSimpleDelegate&& InOnLastFlushDelegate)
+{
+	return QueueFinalUpdateImpl(InInstanceHandle, MoveTemp(InOnLastFlushDelegate), true);
+}
 
 bool FMovieSceneEntitySystemRunner::QueueFinalUpdateImpl(FInstanceHandle InInstanceHandle, FSimpleDelegate&& InOnLastFlushDelegate, bool bDestroyInstance)
 {
@@ -268,6 +272,14 @@ void FMovieSceneEntitySystemRunner::AbandonAndDestroyInstance(FInstanceHandle In
 	}
 }
 
+void FMovieSceneEntitySystemRunner::AbandonAndDestroyInstance(FInstanceHandle Instance, FSimpleDelegate&& InOnLastFlushDelegate)
+{
+	if (QueueFinalUpdateAndDestroy(Instance, MoveTemp(InOnLastFlushDelegate)))
+	{
+		Flush();
+	}
+}
+
 UE::MovieScene::ERunnerFlushResult FMovieSceneEntitySystemRunner::StartEvaluation(UMovieSceneEntitySystemLinker* Linker)
 {
 	using namespace UE::MovieScene;
diff --git a/Engine/Source/Runtime/MovieScene/Private/Evaluation/MovieSceneEvaluationTemplateInstance.cpp b/Engine/Source/Runtime/MovieScene/Private/Evaluation/MovieSceneEvaluationTemplateInstance.cpp
--- a/Engine/Source/Runtime/MovieScene/Private/Evaluation/MovieSceneEvaluationTemplateInstance.cpp
+++ b/Engine/Source/Runtime/MovieScene/Private/Evaluation/MovieSceneEvaluationTemplateInstance.cpp
@@ -127,16 +127,32 @@ void FMovieSceneRootEvaluationTemplateInstance::Initialize(UMovieSceneSequence&
 
 	if (bReinitialize)
 	{
-		// Tear down previous sequence first.
 		FRootInstanceHandle OldRootInstanceHandle = SharedPlaybackState ?
 			SharedPlaybackState->GetRootInstanceHandle() : FRootInstanceHandle();
-		SharedPlaybackState.Reset();
 
 		if (OldRootInstanceHandle.IsValid() && bOwnsSharedPlaybackState)
 		{
 			if (PreviousRunner)
 			{
-				PreviousRunner->AbandonAndDestroyInstance(OldRootInstanceHandle);
+				// Reset out SharedPlaybackState in a callback executed soon after the last update.
+				// This is because:
+				//
+				// 1) We want it to still be valid during the last update, as some legacy things might go through
+				//    the IMovieScenePlayer::GetEvaluationTemplate() method, which is us, and we need a valid
+				//    SharedPlaybackState to still be functional during that time.
+				//
+				// 2) We also want it to be null when the runner finally destroys the FSequenceInstance, because
+				//    its destructor has a leak-detection assert that checks that nobody has lingering references
+				//    to the SharedPlaybackState.
+				//
+				auto OnFlushBeforeDestroy = [this]()
+				{
+					SharedPlaybackState.Reset();
+				};
+
+				PreviousRunner->AbandonAndDestroyInstance(
+						OldRootInstanceHandle,
+						FSimpleDelegate::CreateLambda(OnFlushBeforeDestroy));
 			}
 			else if (EntitySystemLinker)
 			{
diff --git a/Engine/Source/Runtime/MovieScene/Private/Evaluation/PersistentEvaluationData.cpp b/Engine/Source/Runtime/MovieScene/Private/Evaluation/PersistentEvaluationData.cpp
--- a/Engine/Source/Runtime/MovieScene/Private/Evaluation/PersistentEvaluationData.cpp
+++ b/Engine/Source/Runtime/MovieScene/Private/Evaluation/PersistentEvaluationData.cpp
@@ -7,9 +7,17 @@
 
 FPersistentEvaluationData::FPersistentEvaluationData(IMovieScenePlayer& InPlayer)
 	: Player(InPlayer)
-	, EntityData(Player.GetEvaluationState()->PersistentEntityData)
-	, SharedData(Player.GetEvaluationState()->PersistentSharedData)
-{}
+	, EntityData(nullptr)
+	, SharedData(nullptr)
+{
+
+	FMovieSceneEvaluationState* EvaluationState = Player.GetEvaluationState();
+	if (ensure(EvaluationState))
+	{
+		EntityData = &EvaluationState->PersistentEntityData;
+		SharedData = &EvaluationState->PersistentSharedData;
+	}
+}
 
 const FMovieSceneSequenceInstanceData* FPersistentEvaluationData::GetInstanceData() const
 {
diff --git a/Engine/Source/Runtime/MovieScene/Public/EntitySystem/MovieSceneEntitySystemRunner.h b/Engine/Source/Runtime/MovieScene/Public/EntitySystem/MovieSceneEntitySystemRunner.h
--- a/Engine/Source/Runtime/MovieScene/Public/EntitySystem/MovieSceneEntitySystemRunner.h
+++ b/Engine/Source/Runtime/MovieScene/Public/EntitySystem/MovieSceneEntitySystemRunner.h
@@ -100,11 +100,13 @@ public:
 	MOVIESCENE_API bool QueueFinalUpdate(FInstanceHandle Instance);
 	MOVIESCENE_API bool QueueFinalUpdate(FInstanceHandle Instance, FSimpleDelegate&& InOnLastFlushDelegate);
 	MOVIESCENE_API bool QueueFinalUpdateAndDestroy(FInstanceHandle Instance);
+	MOVIESCENE_API bool QueueFinalUpdateAndDestroy(FInstanceHandle Instance, FSimpleDelegate&& InOnLastFlushDelegate);
 
 	/**
 	 * Abandon the specified instance handle and destroy it immediately. May flush this runner if necessary
 	 */
 	MOVIESCENE_API void AbandonAndDestroyInstance(FInstanceHandle Instance);
+	MOVIESCENE_API void AbandonAndDestroyInstance(FInstanceHandle Instance, FSimpleDelegate&& InOnLastFlushDelegate);
 
 	/**
 	 * Flushes the update queue and applies any outstanding evaluation logic
diff --git a/Engine/Source/Runtime/MovieScene/Public/Evaluation/PersistentEvaluationData.h b/Engine/Source/Runtime/MovieScene/Public/Evaluation/PersistentEvaluationData.h
--- a/Engine/Source/Runtime/MovieScene/Public/Evaluation/PersistentEvaluationData.h
+++ b/Engine/Source/Runtime/MovieScene/Public/Evaluation/PersistentEvaluationData.h
@@ -196,7 +196,9 @@ public:
 	{
 		MOVIESCENE_DETAILED_SCOPE_CYCLE_COUNTER(MovieSceneEval_PersistentData_Access)
 
-		if (TUniquePtr<IPersistentEvaluationData>* Existing = SharedData.Find(InKey))
+		check(SharedData);
+
+		if (TUniquePtr<IPersistentEvaluationData>* Existing = SharedData->Find(InKey))
 		{
 			return static_cast<T&>(*Existing->Get());
 		}
@@ -209,8 +211,10 @@ public:
 	{
 		MOVIESCENE_DETAILED_SCOPE_CYCLE_COUNTER(MovieSceneEval_PersistentData_Access)
 
+		check(SharedData);
+
 		T* Ptr = new T;
-		SharedData.Add(InKey, TUniquePtr<IPersistentEvaluationData>(Ptr));
+		SharedData->Add(InKey, TUniquePtr<IPersistentEvaluationData>(Ptr));
 		return *Ptr;
 	}
 
@@ -219,8 +223,12 @@ public:
 	{
 		MOVIESCENE_DETAILED_SCOPE_CYCLE_COUNTER(MovieSceneEval_PersistentData_Access)
 
-		TUniquePtr<IPersistentEvaluationData>* Existing = SharedData.Find(InKey);
-		return Existing ? static_cast<T*>(Existing->Get()) : nullptr;
+		if (ensure(SharedData))
+		{
+			TUniquePtr<IPersistentEvaluationData>* Existing = SharedData->Find(InKey);
+			return Existing ? static_cast<T*>(Existing->Get()) : nullptr;
+		}
+		return nullptr;
 	}
 
 	template<typename T>
@@ -228,8 +236,12 @@ public:
 	{
 		MOVIESCENE_DETAILED_SCOPE_CYCLE_COUNTER(MovieSceneEval_PersistentData_Access)
 
-		const TUniquePtr<IPersistentEvaluationData>* Existing = SharedData.Find(InKey);
-		return Existing ? static_cast<const T*>(Existing->Get()) : nullptr;
+		if (ensure(SharedData))
+		{
+			const TUniquePtr<IPersistentEvaluationData>* Existing = SharedData->Find(InKey);
+			return Existing ? static_cast<const T*>(Existing->Get()) : nullptr;
+		}
+		return nullptr;
 	}
 
 	template<typename T>
@@ -256,7 +268,10 @@ public:
 	{
 		MOVIESCENE_DETAILED_SCOPE_CYCLE_COUNTER(MovieSceneEval_PersistentData_Access)
 		
-		SharedData.Remove(InKey);
+		if (ensure(SharedData))
+		{
+			SharedData->Remove(InKey);
+		}
 	}
 
 private:
@@ -267,7 +282,9 @@ private:
 	{
 		MOVIESCENE_DETAILED_SCOPE_CYCLE_COUNTER(MovieSceneEval_PersistentData_Access)
 
-		if (TUniquePtr<IPersistentEvaluationData>* Existing = EntityData.Find(InKey))
+		check(EntityData);
+
+		if (TUniquePtr<IPersistentEvaluationData>* Existing = EntityData->Find(InKey))
 		{
 			return static_cast<T&>(*Existing->Get());
 		}
@@ -280,8 +297,10 @@ private:
 	{
 		MOVIESCENE_DETAILED_SCOPE_CYCLE_COUNTER(MovieSceneEval_PersistentData_Access)
 
+		check(EntityData);
+
 		T* Ptr = new T;
-		EntityData.Add(InKey, TUniquePtr<IPersistentEvaluationData>(Ptr));
+		EntityData->Add(InKey, TUniquePtr<IPersistentEvaluationData>(Ptr));
 		return *Ptr;
 	}
 
@@ -290,8 +309,12 @@ private:
 	{
 		MOVIESCENE_DETAILED_SCOPE_CYCLE_COUNTER(MovieSceneEval_PersistentData_Access)
 
-		const TUniquePtr<IPersistentEvaluationData>* Existing = EntityData.Find(InKey);
-		return Existing ? static_cast<T*>(Existing->Get()) : nullptr;
+		if (ensure(EntityData))
+		{
+			const TUniquePtr<IPersistentEvaluationData>* Existing = EntityData->Find(InKey);
+			return Existing ? static_cast<T*>(Existing->Get()) : nullptr;
+		}
+		return nullptr;
 	}
 
 	template<typename T>
@@ -308,7 +331,10 @@ private:
 	{
 		MOVIESCENE_DETAILED_SCOPE_CYCLE_COUNTER(MovieSceneEval_PersistentData_Access)
 		
-		EntityData.Remove(InKey);
+		if (ensure(EntityData))
+		{
+			EntityData->Remove(InKey);
+		}
 	}
 
 private:
@@ -317,10 +343,10 @@ private:
 	IMovieScenePlayer& Player;
 
 	/** Persistent data that's associated with a template entity (such as a track or a section) */
-	TMap<FMovieSceneEvaluationKey, TUniquePtr<IPersistentEvaluationData>>& EntityData;
+	TMap<FMovieSceneEvaluationKey, TUniquePtr<IPersistentEvaluationData>>* EntityData;
 
 	/** Persistent data that's shared across multiple template entities */
-	TMap<FSharedPersistentDataKey, TUniquePtr<IPersistentEvaluationData>>& SharedData;
+	TMap<FSharedPersistentDataKey, TUniquePtr<IPersistentEvaluationData>>* SharedData;
 
 	// The keys themselves are mutable since this a proxy representation of the data above.
 	// For a const FPersistentEvaluationData& we can change the ptr to the persistent data, but we can't change the data itself (as with a const T*)
