Sequencer: remove ensure() because unit tests want to use invalid players for some stuff

diff --git a/Engine/Source/Runtime/MovieScene/Private/Evaluation/PersistentEvaluationData.cpp b/Engine/Source/Runtime/MovieScene/Private/Evaluation/PersistentEvaluationData.cpp
--- a/Engine/Source/Runtime/MovieScene/Private/Evaluation/PersistentEvaluationData.cpp
+++ b/Engine/Source/Runtime/MovieScene/Private/Evaluation/PersistentEvaluationData.cpp
@@ -12,7 +12,7 @@ FPersistentEvaluationData::FPersistentEvaluationData(IMovieScenePlayer& InPlayer
 {
 
 	FMovieSceneEvaluationState* EvaluationState = Player.GetEvaluationState();
-	if (ensure(EvaluationState))
+	if (EvaluationState)
 	{
 		EntityData = &EvaluationState->PersistentEntityData;
 		SharedData = &EvaluationState->PersistentSharedData;
