Sequencer: fix pre-roll/post-roll sometimes still evaluating sections that aren't set to evaluate there

diff --git a/Engine/Source/Runtime/MovieScene/Private/EntitySystem/MovieSceneEntityLedger.cpp b/Engine/Source/Runtime/MovieScene/Private/EntitySystem/MovieSceneEntityLedger.cpp
--- a/Engine/Source/Runtime/MovieScene/Private/EntitySystem/MovieSceneEntityLedger.cpp
+++ b/Engine/Source/Runtime/MovieScene/Private/EntitySystem/MovieSceneEntityLedger.cpp
@@ -326,12 +326,24 @@ void FEntityLedger::ImportEntity(UMovieSceneEntitySystemLinker* Linker, const FE
 	Params.EntityMetaData = EntityField->FindMetaData(Query);
 	Params.SharedMetaData = EntityField->FindSharedMetaData(Query);
 
+	bool bUnsupportedPrePostRoll = false;
 	if (ImportParams.bPreRoll && (Params.EntityMetaData == nullptr || Params.EntityMetaData->bEvaluateInSequencePreRoll == false))
 	{
-		return;
+		bUnsupportedPrePostRoll = true;
+	}
+	else if (ImportParams.bPostRoll && (Params.EntityMetaData == nullptr || Params.EntityMetaData->bEvaluateInSequencePostRoll == false))
+	{
+		bUnsupportedPrePostRoll = true;
 	}
-	if (ImportParams.bPostRoll && (Params.EntityMetaData == nullptr || Params.EntityMetaData->bEvaluateInSequencePostRoll == false))
+	if (bUnsupportedPrePostRoll)
 	{
+		if (EntityData.EntityID)
+		{
+			// If we are reimporting an existing entity and we entered a pre/post-roll time range
+			// that this entity shouldn't update in, we need to delete that entity.
+			Linker->EntityManager.AddComponents(EntityData.EntityID, FBuiltInComponentTypes::Get()->FinishedMask, EEntityRecursion::Full);
+			EntityData.EntityID = FMovieSceneEntityID();
+		}
 		return;
 	}
 
