Sequencer: The view and working ranges don't need to be saved

New behavior is:
- The zoom level will not be saved into the asset. Every time you open a level sequence for the first time in the editor session, it will be zoomed to the playback range +/- 0.5 seconds (that's the default padding).
- During the editor session, however, the zoom level will be remembered. If you open a sequence, zoom far in, close sequencer (but not the editor), open the sequence again, it'll be zoomed far in again.
- Same goes for a root sequence and subsequences. If you dive into a shot or subsequence and zoom far in, pop back to the root and back into the subsequence, it should remember that you were zoomed far in.

diff --git a/Engine/Source/Runtime/MovieScene/Private/MovieScene.cpp b/Engine/Source/Runtime/MovieScene/Private/MovieScene.cpp
--- a/Engine/Source/Runtime/MovieScene/Private/MovieScene.cpp
+++ b/Engine/Source/Runtime/MovieScene/Private/MovieScene.cpp
@@ -1472,8 +1472,18 @@ void UMovieScene::UpgradeTimeRanges()
 
 	if (EditorData.ViewStart >= EditorData.ViewEnd)
 	{
-		EditorData.ViewStart = PlaybackRange.Value.GetLowerBoundValue() / TickResolution;
-		EditorData.ViewEnd   = PlaybackRange.Value.GetUpperBoundValue() / TickResolution;
+		// Update the working and view ranges to encompass the new range
+		const double RangeStartSeconds = PlaybackRange.Value.GetLowerBoundValue() / TickResolution;
+		const double RangeEndSeconds = PlaybackRange.Value.GetUpperBoundValue() / TickResolution;
+
+		// Initialize the working and view range with a little bit more space
+		const double OutputChange = (RangeEndSeconds - RangeStartSeconds) * 0.1;
+
+		double ExpandedStart = RangeStartSeconds - OutputChange;
+		double ExpandedEnd = RangeEndSeconds + OutputChange;
+
+		EditorData.ViewStart = ExpandedStart;
+		EditorData.ViewEnd   = ExpandedEnd;
 	}
 
 	if (SelectionRange.Value.GetLowerBound().IsOpen() || SelectionRange.Value.GetUpperBound().IsOpen())
diff --git a/Engine/Source/Runtime/MovieScene/Public/MovieScene.h b/Engine/Source/Runtime/MovieScene/Public/MovieScene.h
--- a/Engine/Source/Runtime/MovieScene/Public/MovieScene.h
+++ b/Engine/Source/Runtime/MovieScene/Public/MovieScene.h
@@ -110,17 +110,17 @@ struct FMovieSceneEditorData
 	TArray<FString> PinnedNodes;
 
 	/** The last view-range start that the user was observing */
-	UPROPERTY()
+	UPROPERTY(Transient)
 	double ViewStart;
 	/** The last view-range end that the user was observing */
-	UPROPERTY()
+	UPROPERTY(Transient)
 	double ViewEnd;
 
 	/** User-defined working range start in which the entire sequence should reside. */
-	UPROPERTY()
+	UPROPERTY(Transient)
 	double WorkStart;
 	/** User-defined working range end in which the entire sequence should reside. */
-	UPROPERTY()
+	UPROPERTY(Transient)
 	double WorkEnd;
 
 	/** Deprecated */
