PR #13312: UMovieSceneSequence: Fixed nullptr crash

diff --git a/Engine/Source/Runtime/MovieScene/Private/MovieSceneSequence.cpp b/Engine/Source/Runtime/MovieScene/Private/MovieSceneSequence.cpp
--- a/Engine/Source/Runtime/MovieScene/Private/MovieSceneSequence.cpp
+++ b/Engine/Source/Runtime/MovieScene/Private/MovieSceneSequence.cpp
@@ -388,10 +388,14 @@ FMovieSceneObjectBindingID UMovieSceneSequence::FindBindingByTag(FName InBinding
 
 const TArray<FMovieSceneObjectBindingID>& UMovieSceneSequence::FindBindingsByTag(FName InBindingName) const
 {
-	const FMovieSceneObjectBindingIDs* BindingIDs = GetMovieScene()->AllTaggedBindings().Find(InBindingName);
-	if (BindingIDs)
+	const UMovieScene* MovieScene = GetMovieScene();
+	if (MovieScene)
 	{
-		return BindingIDs->IDs;
+		const FMovieSceneObjectBindingIDs* BindingIDs = MovieScene->AllTaggedBindings().Find(InBindingName);
+		if (BindingIDs)
+		{
+			return BindingIDs->IDs;
+		}
 	}
 
 	static TArray<FMovieSceneObjectBindingID> EmptyBindings;
