Sequencer: Remove Modify from SetPreRollFrames / SetPostRollFrames

In some cases, sections want to set the default pre/post roll frames on PostLoad() or constructor. This was causing Modify to be called when it shouldn't be. The caller should be responsible for calling Modify()

diff --git a/Engine/Source/Runtime/MovieScene/Public/MovieSceneSection.h b/Engine/Source/Runtime/MovieScene/Public/MovieSceneSection.h
--- a/Engine/Source/Runtime/MovieScene/Public/MovieSceneSection.h
+++ b/Engine/Source/Runtime/MovieScene/Public/MovieSceneSection.h
@@ -592,13 +592,13 @@ public:
 
 	/** Gets the number of frames to prepare this section for evaluation before it actually starts. */
 	UFUNCTION(BlueprintCallable, Category = "Sequencer|Section")
-	void SetPreRollFrames(int32 InPreRollFrames) { if (TryModify()) { PreRollFrames = InPreRollFrames; } }
+	void SetPreRollFrames(int32 InPreRollFrames) { PreRollFrames = InPreRollFrames; }
 	UFUNCTION(BlueprintPure, Category = "Sequencer|Section")
 	int32 GetPreRollFrames() const { return PreRollFrames.Value; }
 
 	/** Gets/sets the number of frames to continue 'postrolling' this section for after evaluation has ended. */
 	UFUNCTION(BlueprintCallable, Category = "Sequencer|Section")
-	void SetPostRollFrames(int32 InPostRollFrames) { if (TryModify()) { PostRollFrames = InPostRollFrames; } }
+	void SetPostRollFrames(int32 InPostRollFrames) { PostRollFrames = InPostRollFrames; }
 	UFUNCTION(BlueprintPure, Category = "Sequencer|Section")
 	int32 GetPostRollFrames() const { return PostRollFrames.Value; }
 
