Sequencer: properly implement end-at-time feature

diff --git a/Engine/Source/Runtime/MovieScene/Private/Evaluation/MovieScenePlaybackManager.cpp b/Engine/Source/Runtime/MovieScene/Private/Evaluation/MovieScenePlaybackManager.cpp
--- a/Engine/Source/Runtime/MovieScene/Private/Evaluation/MovieScenePlaybackManager.cpp
+++ b/Engine/Source/Runtime/MovieScene/Private/Evaluation/MovieScenePlaybackManager.cpp
@@ -113,33 +113,59 @@ void FMovieScenePlaybackManager::InternalUpdateToTick(const FFrameNumber NextTic
 	// IMPORTANT: we assume that LastValidTick is less than the duration (current implementation is 
 	//			  duration minus one tick).
 
-	const bool bIsPlayingForward = IsPlayingForward();
-	const FFrameNumber LoopStartTick = bIsPlayingForward ? EffectiveStartTick : LastValidTick;
-	const FFrameNumber LoopLastTick = bIsPlayingForward ? LastValidTick : EffectiveStartTick;
-	// If the start/end offsets make the duration 0, we treat each loop as one frame long.
+	bool bIsPlayingForward = IsPlayingForward();
+	FFrameNumber LoopStartTick = bIsPlayingForward ? EffectiveStartTick : LastValidTick;
+	FFrameNumber LoopLastTick = bIsPlayingForward ? LastValidTick : EffectiveStartTick;
+	FFrameNumber LastLoopLastTick = PlaybackEndTick.Get(LoopLastTick);
+	// If the start/end offsets make the duration 0, we treat each loop as one tick long.
 	const FFrameNumber LoopDurationTicks = FMath::Max(FFrameNumber(1), EffectiveDurationTicks);
 
-	// Figure out if we crossed the loop-end boundary.
-	if ((bIsPlayingForward && NextTick > LoopLastTick) ||
-			((!bIsPlayingForward && NextTick < LoopLastTick)))
+	// Figure out if we crossed the loop-end boundary, or the "stop at" time on the last loop.
+	const bool bCrossedLoopEnd = (
+			(bIsPlayingForward && NextTick > LoopLastTick) ||
+			(!bIsPlayingForward && NextTick < LoopLastTick)
+			);
+	const bool bReachedLastLoopEnd = (
+			(NumLoopsToPlay > 0 && NumLoopsCompleted >= NumLoopsToPlay - 1) &&
+			(
+			 (bIsPlayingForward && NextTick >= LastLoopLastTick) ||
+			 (!bIsPlayingForward && NextTick <= LastLoopLastTick)
+			));
+	if (bReachedLastLoopEnd)
+	{
+		// Special case when we reach or pass the last tick of the last loop. Play the last
+		// bit of the animation, and stop.
+		OutContexts.Add(
+				FMovieSceneContext(PlaybackPosition.PlayTo(LastLoopLastTick), PlaybackStatus)
+				.SetHasJumped(bShouldJump)
+				);
+
+		++NumLoopsCompleted;
+		ensure(NumLoopsToPlay > 0 && NumLoopsCompleted == NumLoopsToPlay);
+
+		PlaybackStatus = EMovieScenePlayerStatus::Stopped;
+	}
+	else if (bCrossedLoopEnd)
 	{
 		// Compute how many times we crossed the loop-end boundary.
+		// NOTE: we completed one or more loops. Any extra loop may be the last loop, and therefore
+		//       may need to consider a different end time (see LastLoopLastTick).
 		const FFrameNumber LoopRelativeTick = NextTick - LoopStartTick;
 		const int32 NumLoopingsOver = FMath::Abs(LoopRelativeTick.Value) / LoopDurationTicks.Value;
 		ensure(NumLoopingsOver > 0);
 
-		// Massage this a bit with the following rules:
-		//
-		// 1) Don't go over the number of loops to play unless we're playing indefinitely.
-		//
-		// 2) Add an extra completed loop if we reached the end (as opposed to crossing it) because
-		//    that should count as "completing a loop".
-		//
+		// Get the actual number of loops we want to consider completed, and see if this will bring
+		// us inside or past the last loop.
 		const int32 NumLoopsNewlyCompleted = ((NumLoopsToPlay > 0) ?
 				FMath::Min(NumLoopingsOver, NumLoopsToPlay - NumLoopsCompleted) :
 				NumLoopingsOver);
+		const bool bCompletesLastLoop = (
+				NumLoopsToPlay > 0 && NumLoopsCompleted + NumLoopsNewlyCompleted >= NumLoopsToPlay);
+		ensure(NumLoopsNewlyCompleted > 0);
 
 		// Play the last bit of the loop if we are looping and doing any sort of dissections.
+		// We know this isn't the last loop (otherwise we'd be in the bReachedLastLoopEnd block)
+		// so we don't need to be worried about any custom end time.
 		if (DissectLooping != EMovieSceneLoopDissection::None)
 		{
 			OutContexts.Add(
@@ -151,78 +177,112 @@ void FMovieScenePlaybackManager::InternalUpdateToTick(const FFrameNumber NextTic
 		// See if we need to generate more update ranges for the loops. This can happen if we had a
 		// large delta-time, and the duration of a loop is pretty short (i.e. we could have looped 
 		// several times in one update).
-		if (NumLoopingsOver > 1)
+		if (NumLoopsNewlyCompleted > 1 && DissectLooping == EMovieSceneLoopDissection::DissectAll)
 		{
-			const int32 ExtraLoops = NumLoopingsOver - 1;
+			// Add an explicit update for each loop, from start to end. Do one less extra loop if
+			// the last extra loop is the final loop. We will handle the last loop later (see below).
+			const int32 ExtraLoops = NumLoopsNewlyCompleted - 1 - (bCompletesLastLoop ? 1 : 0);
+			ensure(ExtraLoops >= 0);
 
-			if (DissectLooping == EMovieSceneLoopDissection::DissectAll)
+			if (!bPingPongPlayback)
 			{
-				// If we dissect the looping, we add an explicit update for each loop, from start to end.
-				if (!bPingPongPlayback)
+				for (int32 Index = 0; Index < ExtraLoops; ++Index)
 				{
-					for (int32 Index = 0; Index < ExtraLoops; ++Index)
+					PlaybackPosition.Reset(LoopStartTick);
+					OutContexts.Add(
+							FMovieSceneContext(PlaybackPosition.PlayTo(LoopLastTick), PlaybackStatus)
+							.SetHasJumped(true)
+							);
+				}
+			}
+			else
+			{
+				ReversePlayDirection();
+
+				for (int32 Index = 0; Index < ExtraLoops; ++Index)
+				{
+					if (IsPlayingForward())
 					{
-						PlaybackPosition.Reset(LoopStartTick);
+						PlaybackPosition.Reset(EffectiveStartTick);
 						OutContexts.Add(
-								FMovieSceneContext(PlaybackPosition.PlayTo(LoopLastTick), PlaybackStatus)
+								FMovieSceneContext(PlaybackPosition.PlayTo(LastValidTick), PlaybackStatus)
 								.SetHasJumped(true)
 								);
 					}
-				}
-				else
-				{
-					ReversePlayDirection();
-
-					for (int32 Index = 0; Index < ExtraLoops; ++Index)
+					else
 					{
-						if (IsPlayingForward())
-						{
-							PlaybackPosition.Reset(EffectiveStartTick);
-							OutContexts.Add(
-									FMovieSceneContext(PlaybackPosition.PlayTo(LastValidTick), PlaybackStatus)
-									.SetHasJumped(true)
-									);
-						}
-						else
-						{
-							PlaybackPosition.Reset(LastValidTick);
-							OutContexts.Add(
-									FMovieSceneContext(PlaybackPosition.PlayTo(EffectiveStartTick), PlaybackStatus)
-									.SetHasJumped(true)
-									);
-						}
-
-						ReversePlayDirection();
+						PlaybackPosition.Reset(LastValidTick);
+						OutContexts.Add(
+								FMovieSceneContext(PlaybackPosition.PlayTo(EffectiveStartTick), PlaybackStatus)
+								.SetHasJumped(true)
+								);
 					}
-				}
-			}
-			else
-			{
-				// We are not dissecting loops so don't emit extra evaluation contexts.
-				// If we are ping-pong'ing, we at least need to keep track of which way we are now going.
-				if (bPingPongPlayback && (ExtraLoops + 1) % 2 != 0)
-				{
+
 					ReversePlayDirection();
 				}
 			}
 		}
-		else if (bPingPongPlayback)
+
+		// If we are ping-pong'ing, keep track of the direction to go next. We don't need to do
+		// this if we were dissecting each loop, because then we already updated the play direction
+		// while dissecting (see above).
+		if (bPingPongPlayback && DissectLooping != EMovieSceneLoopDissection::DissectAll)
 		{
-			ReversePlayDirection();
+			// If we are reaching the end of the last loop, don't flip an extra time. Just finish
+			// playing in that loop's direction.
+			const int32 NumPingPongs = NumLoopsNewlyCompleted - (bCompletesLastLoop ? 1 : 0);
+			if (NumPingPongs % 2 != 0)
+			{
+				ReversePlayDirection();
+			}
 		}
 
 		// Complete the loops we said we completed.
 		NumLoopsCompleted += NumLoopsNewlyCompleted;
 
-		if (NumLoopsToPlay > 0 && NumLoopsCompleted >= NumLoopsToPlay)
+		// Re-query the direction of play in case we were ping-pong'ing, and update loop times.
+		bIsPlayingForward = IsPlayingForward();
+		LoopStartTick = bIsPlayingForward ? EffectiveStartTick : LastValidTick;
+		LoopLastTick = bIsPlayingForward ? LastValidTick : EffectiveStartTick;
+		LastLoopLastTick = PlaybackEndTick.Get(LoopLastTick);
+
+		if (bCompletesLastLoop)
 		{
-			// If we have played all the loops we needed to play, we can stop playback.
-			// However, if we don't dissect looping, we didn't finish the loop, so do it now.
+			// We completed one loop, plus one or more extra loops that brought us to the end.
+			ensure(NumLoopsNewlyCompleted > 1);
+
 			if (DissectLooping == EMovieSceneLoopDissection::None)
 			{
+				// If we didn't dissect anything, we play to the end time if that's possible from
+				// where we were last, or reset and play the last loop from the start if not.
+				const FFrameNumber CurrentTick = PlaybackPosition.GetCurrentPosition().FrameNumber;
+				const bool bCanPlayToEndFromCurrentTick = 
+						(bIsPlayingForward && CurrentTick < LastLoopLastTick) ||
+						(!bIsPlayingForward && CurrentTick > LastLoopLastTick);
+				if (bCanPlayToEndFromCurrentTick)
+				{
+					OutContexts.Add(
+							FMovieSceneContext(PlaybackPosition.PlayTo(LastLoopLastTick), PlaybackStatus)
+							.SetHasJumped(bShouldJump)
+							);
+				}
+				else
+				{
+					PlaybackPosition.Reset(LoopStartTick);
+					OutContexts.Add(
+							FMovieSceneContext(PlaybackPosition.PlayTo(LastLoopLastTick), PlaybackStatus)
+							.SetHasJumped(true)
+							);
+				}
+			}
+			else
+			{
+				// If we dissect all loops, we left the last loop earlier to do now.
+				// If we dissect only one loop, we play the last loop now.
+				PlaybackPosition.Reset(LoopStartTick);
 				OutContexts.Add(
-						FMovieSceneContext(PlaybackPosition.PlayTo(LoopLastTick), PlaybackStatus)
-						.SetHasJumped(bShouldJump)
+						FMovieSceneContext(PlaybackPosition.PlayTo(LastLoopLastTick), PlaybackStatus)
+						.SetHasJumped(true)
 						);
 			}
 
@@ -239,26 +299,31 @@ void FMovieScenePlaybackManager::InternalUpdateToTick(const FFrameNumber NextTic
 			// loop = [-30, -10], next time = -35, relative time = -25, mod(-25, 20) = -5
 			const FFrameNumber OverplayTicks = LoopRelativeTick.Value % LoopDurationTicks.Value;
 
-			// Don't use LoopStartTick/LoopLastTick here because if we are ping-pong'ing, we would
-			// be going the other way. Use the most recent value of PlayDirection (via IsPlayingForward())
-			// for the same reason.
-			// Also, reverse OverplayTicks when ping-pong'in since we're playing this overplay in the
-			// reverse direction.
-			const FFrameTime EffectiveOverplayTicks = 
-				(IsPlayingForward() ? EffectiveStartTick : LastValidTick) + 
-				((!bPingPongPlayback) ? OverplayTicks : (-OverplayTicks));
-			PlaybackPosition.Reset(IsPlayingForward() ? EffectiveStartTick : LastValidTick);
+			// We reverse OverplayTicks when ping-pong'in since we're playing this overplay in the
+			// reverse direction (otherwise, it already has the correct sign).
+			FFrameTime EffectiveOverplayTick = LoopStartTick + ((!bPingPongPlayback) ? OverplayTicks : (-OverplayTicks));
+
+			// If we have reached the last loop, clamp our overplay with any custom end time.
+			const bool bIsLastLoop = (
+					NumLoopsToPlay > 0 && NumLoopsCompleted >= NumLoopsToPlay - 1);
+			if (bIsLastLoop)
+			{
+				EffectiveOverplayTick = bIsPlayingForward ?
+					FMath::Min(EffectiveOverplayTick, FFrameTime(LastLoopLastTick)) :
+					FMath::Max(EffectiveOverplayTick, FFrameTime(LastLoopLastTick));
+			}
+
+			// Evaluate the overplay.
+			PlaybackPosition.Reset(bIsPlayingForward ? EffectiveStartTick : LastValidTick);
 			OutContexts.Add(
-					FMovieSceneContext(PlaybackPosition.PlayTo(EffectiveOverplayTicks), PlaybackStatus)
+					FMovieSceneContext(PlaybackPosition.PlayTo(EffectiveOverplayTick), PlaybackStatus)
 					.SetHasJumped(true)
 					);
 
-			// If the overplay leads us exactly to the last tick of the loop, let's count that
-			// as a completed loop... but only if that finishes playback. Otherwise, we'll wait
-			// for the next update to loop over in order to avoid counting that loop twice.
-			if (EffectiveOverplayTicks == LoopLastTick 
-					&& NumLoopsToPlay > 0
-					&& NumLoopsCompleted == NumLoopsToPlay - 1)
+			// If the overplay leads us at/beyond the last tick of the last loop, let's count that
+			// as a completed loop and finish playback. Otherwise, we'll wait for the next update 
+			// to loop over in order to avoid counting that loop twice.
+			if (bIsLastLoop && EffectiveOverplayTick == LastLoopLastTick)
 			{
 				++NumLoopsCompleted;
 
@@ -273,18 +338,6 @@ void FMovieScenePlaybackManager::InternalUpdateToTick(const FFrameNumber NextTic
 				FMovieSceneContext(PlaybackPosition.PlayTo(NextTick), PlaybackStatus)
 				.SetHasJumped(bShouldJump)
 				);
-
-		// If we were updated right up to the last tick of the loop, let's count that as a completed
-		// loop... but only if that finishes playback. Otherwise, we'll wait for the next update to
-		// loop over in order to avoid counting that loop twice.
-		if (NextTick == LoopLastTick
-				&& NumLoopsToPlay > 0
-				&& NumLoopsCompleted == NumLoopsToPlay - 1)
-		{
-			++NumLoopsCompleted;
-
-			PlaybackStatus = EMovieScenePlayerStatus::Stopped;
-		}
 	}
 
 	ensure(OutContexts.Num() > 0);
@@ -312,7 +365,9 @@ void FMovieScenePlaybackManager::SetCurrentTime(const FFrameTime& InFrameTime)
 {
 	const FFrameRate TickResolution = PlaybackPosition.GetOutputRate();
 	const FFrameNumber CurrentTick = ConvertFrameTime(InFrameTime, DisplayRate, TickResolution).RoundToFrame();
-	PlaybackPosition.Reset(CurrentTick);
+	const FFrameNumber EffectiveStartTick = SequenceStartTick + StartOffsetTicks;
+	const FFrameNumber LastValidTick = GetLastValidTick();
+	PlaybackPosition.Reset(FMath::Clamp(CurrentTick, EffectiveStartTick, LastValidTick));
 }
 
 TRange<FFrameTime> FMovieScenePlaybackManager::GetEffectivePlaybackRange() const
@@ -337,7 +392,7 @@ FFrameTime FMovieScenePlaybackManager::GetEffectiveStartTime() const
 FFrameTime FMovieScenePlaybackManager::GetEffectiveEndTime() const
 {
 	const FFrameRate TickResolution = PlaybackPosition.GetOutputRate();
-	return ConvertFrameTime(SequenceEndTick + EndOffsetTicks, TickResolution, DisplayRate);
+	return ConvertFrameTime(SequenceEndTick - EndOffsetTicks, TickResolution, DisplayRate);
 }
 
 void FMovieScenePlaybackManager::SetStartOffset(const FFrameTime& InStartOffset)
@@ -390,6 +445,30 @@ FFrameTime FMovieScenePlaybackManager::GetEndOffset() const
 	return ConvertFrameTime(EndOffsetTicks, TickResolution, DisplayRate);
 }
 
+TOptional<FFrameTime> FMovieScenePlaybackManager::GetPlaybackEndTime() const
+{
+	if (const FFrameNumber* EndTickValue = PlaybackEndTick.GetPtrOrNull())
+	{
+		const FFrameRate TickResolution = PlaybackPosition.GetOutputRate();
+		return ConvertFrameTime(*EndTickValue, TickResolution, DisplayRate);
+	}
+	return TOptional<FFrameTime>();
+}
+
+void FMovieScenePlaybackManager::SetPlaybackEndTime(const FFrameTime& InEndTime)
+{
+	const FFrameRate TickResolution = PlaybackPosition.GetOutputRate();
+	PlaybackEndTick = FMath::Clamp(
+			ConvertFrameTime(InEndTime, DisplayRate, TickResolution).RoundToFrame(),
+			SequenceStartTick + StartOffsetTicks,
+			SequenceEndTick - EndOffsetTicks);
+}
+
+void FMovieScenePlaybackManager::ClearPlaybackEndTime()
+{
+	PlaybackEndTick.Reset();
+}
+
 void FMovieScenePlaybackManager::SetNumLoopsToPlay(int32 InNumLoopsToPlay)
 {
 	NumLoopsToPlay = InNumLoopsToPlay;
diff --git a/Engine/Source/Runtime/MovieScene/Private/Tests/MovieScenePlaybackManagerTests.cpp b/Engine/Source/Runtime/MovieScene/Private/Tests/MovieScenePlaybackManagerTests.cpp
--- a/Engine/Source/Runtime/MovieScene/Private/Tests/MovieScenePlaybackManagerTests.cpp
+++ b/Engine/Source/Runtime/MovieScene/Private/Tests/MovieScenePlaybackManagerTests.cpp
@@ -14,6 +14,9 @@
 namespace UE::MovieScene::Test
 {
 
+/**
+ * Default test sequence is 10 seconds long, at 30fps, with 60000 ticks/second.
+ */
 struct FMakeTestSequenceParams
 {
 	int32 StartTick = 0;
@@ -176,7 +179,7 @@ bool FMovieScenePlaybackManagerTestLoopingNoDissection::RunTest(const FString& P
 		Contexts.Reset();
 		Manager.UpdateTo(30 * 5, Contexts);
 		UTEST_EQUAL("NumContexts", Contexts.Num(), 1);
-		UTEST_EQUAL("Context0", Contexts[0].GetRange(), MakeHullTimeRange(TR * 5, TR * 10));
+		UTEST_EQUAL("Context0", Contexts[0].GetRange(), MakeHullTimeRange(TR * 5, LastValidTick));
 		UTEST_FALSE("Context0_Jumped", Contexts[0].HasJumped());
 		UTEST_EQUAL("Context0_Backwards", Contexts[0].GetDirection(), EPlayDirection::Backwards);
 		UTEST_EQUAL("NumLoops", Manager.GetNumLoopsCompleted(), 0);
@@ -289,7 +292,7 @@ bool FMovieScenePlaybackManagerTestLoopingDissectOne::RunTest(const FString& Par
 		Contexts.Reset();
 		Manager.UpdateTo(30 * 5, Contexts);
 		UTEST_EQUAL("NumContexts", Contexts.Num(), 1);
-		UTEST_EQUAL("Context0", Contexts[0].GetRange(), MakeHullTimeRange(TR * 5, TR * 10));
+		UTEST_EQUAL("Context0", Contexts[0].GetRange(), MakeHullTimeRange(TR * 5, LastValidTick));
 		UTEST_FALSE("Context0_Jumped", Contexts[0].HasJumped());
 		UTEST_EQUAL("Context0_Backwards", Contexts[0].GetDirection(), EPlayDirection::Backwards);
 		UTEST_EQUAL("NumLoops", Manager.GetNumLoopsCompleted(), 0);
@@ -410,7 +413,7 @@ bool FMovieScenePlaybackManagerTestLoopingDissectAll::RunTest(const FString& Par
 		Contexts.Reset();
 		Manager.UpdateTo(30 * 5, Contexts);
 		UTEST_EQUAL("NumContexts", Contexts.Num(), 1);
-		UTEST_EQUAL("Context0", Contexts[0].GetRange(), MakeHullTimeRange(TR * 5, TR * 10));
+		UTEST_EQUAL("Context0", Contexts[0].GetRange(), MakeHullTimeRange(TR * 5, LastValidTick));
 		UTEST_FALSE("Context0_Jumped", Contexts[0].HasJumped());
 		UTEST_EQUAL("NumLoops", Manager.GetNumLoopsCompleted(), 0);
 		UTEST_EQUAL("Status", Manager.GetPlaybackStatus(), EMovieScenePlayerStatus::Playing);
@@ -525,6 +528,390 @@ bool FMovieScenePlaybackManagerTestUpdateToEnd::RunTest(const FString& Parameter
 
 	return true;
 }
+IMPLEMENT_SIMPLE_AUTOMATION_TEST(
+		FMovieScenePlaybackManagerTestCustomEndTime, 
+		"System.Engine.Sequencer.PlaybackManager.CustomEndTime", 
+		EAutomationTestFlags::EditorContext | EAutomationTestFlags::EngineFilter)
+bool FMovieScenePlaybackManagerTestCustomEndTime::RunTest(const FString& Parameters)
+{
+	using namespace UE::MovieScene;
+	using namespace UE::MovieScene::Test;
+
+	FMakeTestSequenceParams MakeParams;
+	UMovieSceneSequence* Sequence = MakeTestSequence(MakeParams);
+
+	const int32 FPS = 30;
+	const int32 TR = 60000;
+	const FFrameTime LastValidTick(10 * TR - 1);
+
+	FMovieScenePlaybackManager Manager(Sequence);
+	Manager.SetPlaybackEndTime(6 * FPS); // End with the frame at 6 seconds.
+
+	FMovieScenePlaybackManager::FContexts Contexts;
+
+	Manager.SetPlayDirection(EPlayDirection::Forwards);
+	Manager.SetDissectLooping(EMovieSceneLoopDissection::DissectOne);
+
+	Manager.SetCurrentTime(0);
+	Manager.SetPlaybackStatus(EMovieScenePlayerStatus::Playing);
+
+	{
+		Manager.SetNumLoopsToPlay(1);
+		Contexts.Reset();
+		Manager.UpdateTo(8 * FPS, Contexts);
+		UTEST_EQUAL("NumContexts", Contexts.Num(), 1);
+		UTEST_EQUAL("Context0", Contexts[0].GetRange(), MakeHullTimeRange(0, TR * 6));
+		UTEST_EQUAL("NumLoops", Manager.GetNumLoopsCompleted(), 1);
+		UTEST_EQUAL("Status", Manager.GetPlaybackStatus(), EMovieScenePlayerStatus::Stopped);
+	}
+
+	Manager.ResetNumLoopsCompleted();
+	Manager.SetCurrentTime(0);
+	Manager.SetPlaybackStatus(EMovieScenePlayerStatus::Playing);
+
+	{
+		Manager.SetNumLoopsToPlay(2);
+		Contexts.Reset();
+		Manager.UpdateTo(8 * FPS, Contexts);
+		UTEST_EQUAL("NumContexts", Contexts.Num(), 1);
+		UTEST_EQUAL("Context0", Contexts[0].GetRange(), MakeHullTimeRange(0, TR * 8));
+		UTEST_EQUAL("NumLoops", Manager.GetNumLoopsCompleted(), 0);
+		UTEST_EQUAL("Status", Manager.GetPlaybackStatus(), EMovieScenePlayerStatus::Playing);
+
+		Contexts.Reset();
+		Manager.UpdateTo(12 * FPS, Contexts);
+		UTEST_EQUAL("NumContexts", Contexts.Num(), 2);
+		UTEST_EQUAL("Context0", Contexts[0].GetRange(), MakeContinuedTimeRange(TR * 8, LastValidTick));
+		UTEST_EQUAL("Context1", Contexts[1].GetRange(), MakeHullTimeRange(0, TR * 2));
+		UTEST_EQUAL("NumLoops", Manager.GetNumLoopsCompleted(), 1);
+		UTEST_EQUAL("Status", Manager.GetPlaybackStatus(), EMovieScenePlayerStatus::Playing);
+
+		Contexts.Reset();
+		Manager.UpdateTo(8 * FPS, Contexts);
+		UTEST_EQUAL("NumContexts", Contexts.Num(), 1);
+		UTEST_EQUAL("Context0", Contexts[0].GetRange(), MakeContinuedTimeRange(TR * 2, TR * 6));
+		UTEST_EQUAL("NumLoops", Manager.GetNumLoopsCompleted(), 2);
+		UTEST_EQUAL("Status", Manager.GetPlaybackStatus(), EMovieScenePlayerStatus::Stopped);
+	}
+
+	Manager.ResetNumLoopsCompleted();
+	Manager.SetCurrentTime(0);
+	Manager.SetPlaybackStatus(EMovieScenePlayerStatus::Playing);
+
+	{
+		Manager.SetNumLoopsToPlay(3);
+		Contexts.Reset();
+		Manager.UpdateTo(40 * FPS, Contexts);
+		UTEST_EQUAL("NumContexts", Contexts.Num(), 2);
+		UTEST_EQUAL("Context0", Contexts[0].GetRange(), MakeHullTimeRange(0, LastValidTick));
+		UTEST_EQUAL("Context1", Contexts[1].GetRange(), MakeHullTimeRange(0, TR * 6));
+		UTEST_EQUAL("NumLoops", Manager.GetNumLoopsCompleted(), 3);
+		UTEST_EQUAL("Status", Manager.GetPlaybackStatus(), EMovieScenePlayerStatus::Stopped);
+	}
+
+	Manager.ResetNumLoopsCompleted();
+	Manager.SetCurrentTime(0);
+	Manager.SetPlaybackStatus(EMovieScenePlayerStatus::Playing);
+
+	{
+		Manager.SetNumLoopsToPlay(4);
+		Contexts.Reset();
+		Manager.UpdateTo(12 * FPS, Contexts);
+		UTEST_EQUAL("NumContexts", Contexts.Num(), 2);
+		UTEST_EQUAL("Context0", Contexts[0].GetRange(), MakeHullTimeRange(0, LastValidTick));
+		UTEST_EQUAL("Context1", Contexts[1].GetRange(), MakeHullTimeRange(0, TR * 2));
+		UTEST_EQUAL("NumLoops", Manager.GetNumLoopsCompleted(), 1);
+		UTEST_EQUAL("Status", Manager.GetPlaybackStatus(), EMovieScenePlayerStatus::Playing);
+
+		Contexts.Reset();
+		Manager.UpdateTo(29 * FPS, Contexts);
+		UTEST_EQUAL("NumContexts", Contexts.Num(), 2);
+		UTEST_EQUAL("Context0", Contexts[0].GetRange(), MakeContinuedTimeRange(TR * 2, LastValidTick));
+		UTEST_EQUAL("Context1", Contexts[1].GetRange(), MakeHullTimeRange(0, TR * 6));
+		UTEST_EQUAL("NumLoops", Manager.GetNumLoopsCompleted(), 4);
+		UTEST_EQUAL("Status", Manager.GetPlaybackStatus(), EMovieScenePlayerStatus::Stopped);
+	}
+
+	Manager.SetPlayDirection(EPlayDirection::Forwards);
+	Manager.SetDissectLooping(EMovieSceneLoopDissection::DissectAll);
+
+	Manager.ResetNumLoopsCompleted();
+	Manager.SetCurrentTime(0);
+	Manager.SetPlaybackStatus(EMovieScenePlayerStatus::Playing);
+
+	{
+		Manager.SetNumLoopsToPlay(1);
+		Contexts.Reset();
+		Manager.UpdateTo(8 * FPS, Contexts);
+		UTEST_EQUAL("NumContexts", Contexts.Num(), 1);
+		UTEST_EQUAL("Context0", Contexts[0].GetRange(), MakeHullTimeRange(0, TR * 6));
+		UTEST_EQUAL("NumLoops", Manager.GetNumLoopsCompleted(), 1);
+		UTEST_EQUAL("Status", Manager.GetPlaybackStatus(), EMovieScenePlayerStatus::Stopped);
+	}
+
+	Manager.ResetNumLoopsCompleted();
+	Manager.SetCurrentTime(0);
+	Manager.SetPlaybackStatus(EMovieScenePlayerStatus::Playing);
+
+	{
+		Manager.SetNumLoopsToPlay(2);
+		Contexts.Reset();
+		Manager.UpdateTo(8 * FPS, Contexts);
+		UTEST_EQUAL("NumContexts", Contexts.Num(), 1);
+		UTEST_EQUAL("Context0", Contexts[0].GetRange(), MakeHullTimeRange(0, TR * 8));
+		UTEST_EQUAL("NumLoops", Manager.GetNumLoopsCompleted(), 0);
+		UTEST_EQUAL("Status", Manager.GetPlaybackStatus(), EMovieScenePlayerStatus::Playing);
+
+		Contexts.Reset();
+		Manager.UpdateTo(12 * FPS, Contexts);
+		UTEST_EQUAL("NumContexts", Contexts.Num(), 2);
+		UTEST_EQUAL("Context0", Contexts[0].GetRange(), MakeContinuedTimeRange(TR * 8, LastValidTick));
+		UTEST_EQUAL("Context1", Contexts[1].GetRange(), MakeHullTimeRange(0, TR * 2));
+		UTEST_EQUAL("NumLoops", Manager.GetNumLoopsCompleted(), 1);
+		UTEST_EQUAL("Status", Manager.GetPlaybackStatus(), EMovieScenePlayerStatus::Playing);
+
+		Contexts.Reset();
+		Manager.UpdateTo(8 * FPS, Contexts);
+		UTEST_EQUAL("NumContexts", Contexts.Num(), 1);
+		UTEST_EQUAL("Context0", Contexts[0].GetRange(), MakeContinuedTimeRange(TR * 2, TR * 6));
+		UTEST_EQUAL("NumLoops", Manager.GetNumLoopsCompleted(), 2);
+		UTEST_EQUAL("Status", Manager.GetPlaybackStatus(), EMovieScenePlayerStatus::Stopped);
+	}
+
+	Manager.ResetNumLoopsCompleted();
+	Manager.SetCurrentTime(0);
+	Manager.SetPlaybackStatus(EMovieScenePlayerStatus::Playing);
+
+	{
+		Manager.SetNumLoopsToPlay(3);
+		Contexts.Reset();
+		Manager.UpdateTo(40 * FPS, Contexts);
+		UTEST_EQUAL("NumContexts", Contexts.Num(), 3);
+		UTEST_EQUAL("Context0", Contexts[0].GetRange(), MakeHullTimeRange(0, LastValidTick));
+		UTEST_EQUAL("Context1", Contexts[1].GetRange(), MakeHullTimeRange(0, LastValidTick));
+		UTEST_EQUAL("Context2", Contexts[2].GetRange(), MakeHullTimeRange(0, TR * 6));
+		UTEST_EQUAL("NumLoops", Manager.GetNumLoopsCompleted(), 3);
+		UTEST_EQUAL("Status", Manager.GetPlaybackStatus(), EMovieScenePlayerStatus::Stopped);
+	}
+
+	Manager.ResetNumLoopsCompleted();
+	Manager.SetCurrentTime(0);
+	Manager.SetPlaybackStatus(EMovieScenePlayerStatus::Playing);
+
+	{
+		Manager.SetNumLoopsToPlay(4);
+		Contexts.Reset();
+		Manager.UpdateTo(12 * FPS, Contexts);
+		UTEST_EQUAL("NumContexts", Contexts.Num(), 2);
+		UTEST_EQUAL("Context0", Contexts[0].GetRange(), MakeHullTimeRange(0, LastValidTick));
+		UTEST_EQUAL("Context1", Contexts[1].GetRange(), MakeHullTimeRange(0, TR * 2));
+		UTEST_EQUAL("NumLoops", Manager.GetNumLoopsCompleted(), 1);
+		UTEST_EQUAL("Status", Manager.GetPlaybackStatus(), EMovieScenePlayerStatus::Playing);
+
+		Contexts.Reset();
+		Manager.UpdateTo(29 * FPS, Contexts);
+		UTEST_EQUAL("NumContexts", Contexts.Num(), 3);
+		UTEST_EQUAL("Context0", Contexts[0].GetRange(), MakeContinuedTimeRange(TR * 2, LastValidTick));
+		UTEST_EQUAL("Context1", Contexts[1].GetRange(), MakeHullTimeRange(0, LastValidTick));
+		UTEST_EQUAL("Context2", Contexts[2].GetRange(), MakeHullTimeRange(0, TR * 6));
+		UTEST_EQUAL("NumLoops", Manager.GetNumLoopsCompleted(), 4);
+		UTEST_EQUAL("Status", Manager.GetPlaybackStatus(), EMovieScenePlayerStatus::Stopped);
+	}
+
+	Manager.SetPlayDirection(EPlayDirection::Backwards);
+	Manager.SetDissectLooping(EMovieSceneLoopDissection::DissectOne);
+
+	Manager.ResetNumLoopsCompleted();
+	Manager.SetCurrentTime(10 * FPS);
+	Manager.SetPlaybackStatus(EMovieScenePlayerStatus::Playing);
+
+	{
+		Manager.SetNumLoopsToPlay(1);
+		Contexts.Reset();
+		Manager.UpdateTo(2 * FPS, Contexts);
+		UTEST_EQUAL("NumContexts", Contexts.Num(), 1);
+		UTEST_EQUAL("Context0", Contexts[0].GetRange(), MakeHullTimeRange(TR * 6, LastValidTick));
+		UTEST_EQUAL("Context0_Direction", Contexts[0].GetDirection(), EPlayDirection::Backwards);
+		UTEST_EQUAL("NumLoops", Manager.GetNumLoopsCompleted(), 1);
+		UTEST_EQUAL("Status", Manager.GetPlaybackStatus(), EMovieScenePlayerStatus::Stopped);
+	}
+
+	Manager.ResetNumLoopsCompleted();
+	Manager.SetCurrentTime(10 * FPS);
+	Manager.SetPlaybackStatus(EMovieScenePlayerStatus::Playing);
+
+	{
+		Manager.SetNumLoopsToPlay(2);
+		Contexts.Reset();
+		Manager.UpdateTo(2 * FPS, Contexts);
+		UTEST_EQUAL("NumContexts", Contexts.Num(), 1);
+		UTEST_EQUAL("Context0", Contexts[0].GetRange(), MakeHullTimeRange(TR * 2, LastValidTick));
+		UTEST_EQUAL("Context0_Direction", Contexts[0].GetDirection(), EPlayDirection::Backwards);
+		UTEST_EQUAL("NumLoops", Manager.GetNumLoopsCompleted(), 0);
+		UTEST_EQUAL("Status", Manager.GetPlaybackStatus(), EMovieScenePlayerStatus::Playing);
+
+		Contexts.Reset();
+		Manager.UpdateTo(-3 * FPS, Contexts);
+		UTEST_EQUAL("NumContexts", Contexts.Num(), 2);
+		UTEST_EQUAL("Context0", Contexts[0].GetRange(), MakeDiscreteTimeRange(0, TR * 2));
+		UTEST_EQUAL("Context0_Direction", Contexts[0].GetDirection(), EPlayDirection::Backwards);
+		UTEST_EQUAL("Context1", Contexts[1].GetRange(), MakeHullTimeRange(TR * 7, LastValidTick));
+		UTEST_EQUAL("Context1_Direction", Contexts[1].GetDirection(), EPlayDirection::Backwards);
+		UTEST_EQUAL("NumLoops", Manager.GetNumLoopsCompleted(), 1);
+		UTEST_EQUAL("Status", Manager.GetPlaybackStatus(), EMovieScenePlayerStatus::Playing);
+
+		Contexts.Reset();
+		Manager.UpdateTo(4 * FPS, Contexts);
+		UTEST_EQUAL("NumContexts", Contexts.Num(), 1);
+		UTEST_EQUAL("Context0", Contexts[0].GetRange(), MakeDiscreteTimeRange(TR * 6, TR * 7));
+		UTEST_EQUAL("Context0_Direction", Contexts[0].GetDirection(), EPlayDirection::Backwards);
+		UTEST_EQUAL("NumLoops", Manager.GetNumLoopsCompleted(), 2);
+		UTEST_EQUAL("Status", Manager.GetPlaybackStatus(), EMovieScenePlayerStatus::Stopped);
+	}
+
+	Manager.ResetNumLoopsCompleted();
+	Manager.SetCurrentTime(10 * FPS);
+	Manager.SetPlaybackStatus(EMovieScenePlayerStatus::Playing);
+
+	{
+		Manager.SetNumLoopsToPlay(3);
+		Contexts.Reset();
+		Manager.UpdateTo(-30 * FPS, Contexts);
+		UTEST_EQUAL("NumContexts", Contexts.Num(), 2);
+		UTEST_EQUAL("Context0", Contexts[0].GetRange(), MakeHullTimeRange(0, LastValidTick));
+		UTEST_EQUAL("Context0_Direction", Contexts[0].GetDirection(), EPlayDirection::Backwards);
+		UTEST_EQUAL("Context1", Contexts[1].GetRange(), MakeHullTimeRange(TR * 6, LastValidTick));
+		UTEST_EQUAL("Context1_Direction", Contexts[1].GetDirection(), EPlayDirection::Backwards);
+		UTEST_EQUAL("NumLoops", Manager.GetNumLoopsCompleted(), 3);
+		UTEST_EQUAL("Status", Manager.GetPlaybackStatus(), EMovieScenePlayerStatus::Stopped);
+	}
+
+	Manager.ResetNumLoopsCompleted();
+	Manager.SetCurrentTime(10 * FPS);
+	Manager.SetPlaybackStatus(EMovieScenePlayerStatus::Playing);
+
+	{
+		Manager.SetNumLoopsToPlay(4);
+		Contexts.Reset();
+		Manager.UpdateTo(-2 * FPS, Contexts);
+		UTEST_EQUAL("NumContexts", Contexts.Num(), 2);
+		UTEST_EQUAL("Context0", Contexts[0].GetRange(), MakeHullTimeRange(0, LastValidTick));
+		UTEST_EQUAL("Context0_Direction", Contexts[0].GetDirection(), EPlayDirection::Backwards);
+		UTEST_EQUAL("Context1", Contexts[1].GetRange(), MakeHullTimeRange(TR * 8, LastValidTick));
+		UTEST_EQUAL("Context1_Direction", Contexts[1].GetDirection(), EPlayDirection::Backwards);
+		UTEST_EQUAL("NumLoops", Manager.GetNumLoopsCompleted(), 1);
+		UTEST_EQUAL("Status", Manager.GetPlaybackStatus(), EMovieScenePlayerStatus::Playing);
+
+		Contexts.Reset();
+		Manager.UpdateTo(-29 * FPS, Contexts);
+		UTEST_EQUAL("NumContexts", Contexts.Num(), 2);
+		UTEST_EQUAL("Context0", Contexts[0].GetRange(), MakeDiscreteTimeRange(0, TR * 8));
+		UTEST_EQUAL("Context0_Direction", Contexts[0].GetDirection(), EPlayDirection::Backwards);
+		UTEST_EQUAL("Context1", Contexts[1].GetRange(), MakeHullTimeRange(TR * 6, LastValidTick));
+		UTEST_EQUAL("Context1_Direction", Contexts[1].GetDirection(), EPlayDirection::Backwards);
+		UTEST_EQUAL("NumLoops", Manager.GetNumLoopsCompleted(), 4);
+		UTEST_EQUAL("Status", Manager.GetPlaybackStatus(), EMovieScenePlayerStatus::Stopped);
+	}
+
+	Manager.SetDissectLooping(EMovieSceneLoopDissection::DissectOne);
+	Manager.SetPingPongPlayback(true);
+
+	Manager.SetPlayDirection(EPlayDirection::Forwards);
+	Manager.ResetNumLoopsCompleted();
+	Manager.SetCurrentTime(0);
+	Manager.SetPlaybackStatus(EMovieScenePlayerStatus::Playing);
+
+	{
+		Manager.SetNumLoopsToPlay(1);
+		Contexts.Reset();
+		Manager.UpdateTo(8 * FPS, Contexts);
+		UTEST_EQUAL("NumContexts", Contexts.Num(), 1);
+		UTEST_EQUAL("Context0", Contexts[0].GetRange(), MakeHullTimeRange(0, TR * 6));
+		UTEST_EQUAL("NumLoops", Manager.GetNumLoopsCompleted(), 1);
+		UTEST_EQUAL("Status", Manager.GetPlaybackStatus(), EMovieScenePlayerStatus::Stopped);
+	}
+
+	Manager.SetPlayDirection(EPlayDirection::Forwards);
+	Manager.ResetNumLoopsCompleted();
+	Manager.SetCurrentTime(0);
+	Manager.SetPlaybackStatus(EMovieScenePlayerStatus::Playing);
+
+	{
+		Manager.SetNumLoopsToPlay(2);
+		Contexts.Reset();
+		Manager.UpdateTo(8 * FPS, Contexts);
+		UTEST_EQUAL("NumContexts", Contexts.Num(), 1);
+		UTEST_EQUAL("Context0", Contexts[0].GetRange(), MakeHullTimeRange(0, TR * 8));
+		UTEST_EQUAL("NumLoops", Manager.GetNumLoopsCompleted(), 0);
+		UTEST_EQUAL("Status", Manager.GetPlaybackStatus(), EMovieScenePlayerStatus::Playing);
+
+		Contexts.Reset();
+		Manager.UpdateTo(12 * FPS, Contexts);
+		UTEST_EQUAL("NumContexts", Contexts.Num(), 2);
+		UTEST_EQUAL("Context0", Contexts[0].GetRange(), MakeContinuedTimeRange(TR * 8, LastValidTick));
+		UTEST_EQUAL("Context0_Direction", Contexts[0].GetDirection(), EPlayDirection::Forwards);
+		UTEST_EQUAL("Context1", Contexts[1].GetRange(), MakeHullTimeRange(TR * 8 - 1, LastValidTick));
+		UTEST_EQUAL("Context1_Direction", Contexts[1].GetDirection(), EPlayDirection::Backwards);
+		UTEST_EQUAL("NumLoops", Manager.GetNumLoopsCompleted(), 1);
+		UTEST_EQUAL("Status", Manager.GetPlaybackStatus(), EMovieScenePlayerStatus::Playing);
+
+		Contexts.Reset();
+		Manager.UpdateTo(2 * FPS, Contexts);
+		UTEST_EQUAL("NumContexts", Contexts.Num(), 1);
+		UTEST_EQUAL("Context0", Contexts[0].GetRange(), MakeDiscreteTimeRange(TR * 6, TR * 8 - 1));
+		UTEST_EQUAL("Context0_Direction", Contexts[0].GetDirection(), EPlayDirection::Backwards);
+		UTEST_EQUAL("NumLoops", Manager.GetNumLoopsCompleted(), 2);
+		UTEST_EQUAL("Status", Manager.GetPlaybackStatus(), EMovieScenePlayerStatus::Stopped);
+	}
+
+	Manager.SetPlayDirection(EPlayDirection::Forwards);
+	Manager.ResetNumLoopsCompleted();
+	Manager.SetCurrentTime(0);
+	Manager.SetPlaybackStatus(EMovieScenePlayerStatus::Playing);
+
+	{
+		Manager.SetNumLoopsToPlay(3);
+		Contexts.Reset();
+		Manager.UpdateTo(40 * FPS, Contexts);
+		UTEST_EQUAL("NumContexts", Contexts.Num(), 2);
+		UTEST_EQUAL("Context0", Contexts[0].GetRange(), MakeHullTimeRange(0, LastValidTick));
+		UTEST_EQUAL("Context0_Direction", Contexts[0].GetDirection(), EPlayDirection::Forwards);
+		UTEST_EQUAL("Context1", Contexts[1].GetRange(), MakeHullTimeRange(0, TR * 6));
+		UTEST_EQUAL("Context1_Direction", Contexts[1].GetDirection(), EPlayDirection::Forwards);
+		UTEST_EQUAL("NumLoops", Manager.GetNumLoopsCompleted(), 3);
+		UTEST_EQUAL("Status", Manager.GetPlaybackStatus(), EMovieScenePlayerStatus::Stopped);
+	}
+
+	Manager.SetPlayDirection(EPlayDirection::Forwards);
+	Manager.ResetNumLoopsCompleted();
+	Manager.SetCurrentTime(0);
+	Manager.SetPlaybackStatus(EMovieScenePlayerStatus::Playing);
+
+	{
+		Manager.SetNumLoopsToPlay(4);
+		Contexts.Reset();
+		Manager.UpdateTo(12 * FPS, Contexts);
+		UTEST_EQUAL("NumContexts", Contexts.Num(), 2);
+		UTEST_EQUAL("Context0", Contexts[0].GetRange(), MakeHullTimeRange(0, LastValidTick));
+		UTEST_EQUAL("Context0_Direction", Contexts[0].GetDirection(), EPlayDirection::Forwards);
+		UTEST_EQUAL("Context1", Contexts[1].GetRange(), MakeHullTimeRange(TR * 8 - 1, LastValidTick));
+		UTEST_EQUAL("Context1_Direction", Contexts[1].GetDirection(), EPlayDirection::Backwards);
+		UTEST_EQUAL("NumLoops", Manager.GetNumLoopsCompleted(), 1);
+		UTEST_EQUAL("Status", Manager.GetPlaybackStatus(), EMovieScenePlayerStatus::Playing);
+
+		Contexts.Reset();
+		Manager.UpdateTo(-30 * FPS, Contexts);
+		UTEST_EQUAL("NumContexts", Contexts.Num(), 2);
+		UTEST_EQUAL("Context0", Contexts[0].GetRange(), MakeDiscreteTimeRange(0, TR * 8 - 1));
+		UTEST_EQUAL("Context0_Direction", Contexts[0].GetDirection(), EPlayDirection::Backwards);
+		UTEST_EQUAL("Context1", Contexts[1].GetRange(), MakeHullTimeRange(TR * 6, LastValidTick));
+		UTEST_EQUAL("Context1_Direction", Contexts[1].GetDirection(), EPlayDirection::Backwards);
+		UTEST_EQUAL("NumLoops", Manager.GetNumLoopsCompleted(), 4);
+		UTEST_EQUAL("Status", Manager.GetPlaybackStatus(), EMovieScenePlayerStatus::Stopped);
+	}
+
+	return true;
+}
 
 #undef LOCTEXT_NAMESPACE
 
diff --git a/Engine/Source/Runtime/MovieScene/Public/Evaluation/MovieScenePlaybackManager.h b/Engine/Source/Runtime/MovieScene/Public/Evaluation/MovieScenePlaybackManager.h
--- a/Engine/Source/Runtime/MovieScene/Public/Evaluation/MovieScenePlaybackManager.h
+++ b/Engine/Source/Runtime/MovieScene/Public/Evaluation/MovieScenePlaybackManager.h
@@ -5,6 +5,7 @@
 #include "EntitySystem/MovieSceneSharedPlaybackState.h"
 #include "Evaluation/MovieScenePlayback.h"
 #include "Evaluation/MovieSceneSequenceTransform.h"
+#include "Misc/Optional.h"
 #include "MovieSceneFwd.h"
 
 class UMovieSceneSequence;
@@ -167,6 +168,13 @@ public:
 	/** Sets the end time of the playback in display rate. */
 	MOVIESCENE_API void SetEndOffsetAsTime(const FFrameTime& InEndTime);
 
+	/** Gets the time to stop at on the last loop of playback. */
+	MOVIESCENE_API TOptional<FFrameTime> GetPlaybackEndTime() const;
+	/** Sets the time to stop at on the last loop of playback. */
+	MOVIESCENE_API void SetPlaybackEndTime(const FFrameTime& InEndTime);
+	/** Removes any previously specified time to stop at. */
+	MOVIESCENE_API void ClearPlaybackEndTime();
+
 	/** Gets whether the playback time transform should be applied to time updates. */
 	bool ShouldTransformPlaybackTime() const { return bTransformPlaybackTime; }
 	/** Sets whether the playback time transform should be applied to time updates. */
@@ -209,6 +217,8 @@ private:
 	FFrameNumber StartOffsetTicks;
 	FFrameNumber EndOffsetTicks;
 
+	TOptional<FFrameNumber> PlaybackEndTick;
+
 	int32 NumLoopsToPlay = 1;
 	int32 NumLoopsCompleted = 0;
 
@@ -218,7 +228,6 @@ private:
 
 	EMovieSceneLoopDissection DissectLooping = EMovieSceneLoopDissection::DissectOne;
 
-	bool bSupportsTimeWarp = true;
 	bool bPingPongPlayback = false;
 	bool bTransformPlaybackTime = false;
 };
diff --git a/Engine/Source/Runtime/UMG/Private/Animation/WidgetAnimationState.cpp b/Engine/Source/Runtime/UMG/Private/Animation/WidgetAnimationState.cpp
--- a/Engine/Source/Runtime/UMG/Private/Animation/WidgetAnimationState.cpp
+++ b/Engine/Source/Runtime/UMG/Private/Animation/WidgetAnimationState.cpp
@@ -57,6 +57,7 @@ namespace UE::UMG
 FWidgetAnimationState::FWidgetAnimationState()
 {
 	UserTag = NAME_None;
+	BlockedDeltaTimeCompensation = 0.f;
 	bRestoreState = false;
 	bIsBeginningPlay = false;
 	bIsStopping = false;
@@ -197,6 +198,17 @@ void FWidgetAnimationState::Tick(float InDeltaSeconds)
 		return;
 	}
 
+	if (bIsStopping || bIsBeginningPlay)
+	{
+		return;
+	}
+
+	if (EntitySystemRunner->IsCurrentlyEvaluating())
+	{
+		BlockedDeltaTimeCompensation += InDeltaSeconds;
+		return;
+	}
+
 	const EMovieScenePlayerStatus::Type PreviousPlaybackStatus = PlaybackManager.GetPlaybackStatus();
 	if (PreviousPlaybackStatus == EMovieScenePlayerStatus::Playing)
 	{
@@ -445,17 +457,19 @@ void FWidgetAnimationState::Play(const FWidgetAnimationStatePlayParams& PlayPara
 	}
 
 	PlaybackManager.SetEndOffset(FFrameTime(0));
-	if (PlayParams.EndAtTime.IsSet())
-	{
-		const FFrameTime EndFrameTime = DisplayRate.AsFrameTime(PlayParams.EndAtTime.GetValue());
-		PlaybackManager.SetEndOffsetAsTime(EndFrameTime);
-	}
 	if (PlayParams.EndOffset.IsSet())
 	{
 		const FFrameTime EndOffsetTime = DisplayRate.AsFrameTime(PlayParams.EndOffset.GetValue());
 		PlaybackManager.SetEndOffset(PlaybackManager.GetEndOffset() + EndOffsetTime);
 	}
 
+	PlaybackManager.ClearPlaybackEndTime();
+	if (PlayParams.EndAtTime.IsSet())
+	{
+		const FFrameTime EndFrameTime = DisplayRate.AsFrameTime(PlayParams.EndAtTime.GetValue());
+		PlaybackManager.SetPlaybackEndTime(EndFrameTime);
+	}
+
 	PlayMode = PlayParams.PlayMode;
 	if (PlayMode == EUMGSequencePlayMode::Reverse)
 	{
diff --git a/Engine/Source/Runtime/UMG/Public/Animation/WidgetAnimationState.h b/Engine/Source/Runtime/UMG/Public/Animation/WidgetAnimationState.h
--- a/Engine/Source/Runtime/UMG/Public/Animation/WidgetAnimationState.h
+++ b/Engine/Source/Runtime/UMG/Public/Animation/WidgetAnimationState.h
@@ -180,6 +180,9 @@ private:
 	 */
 	FName UserTag;
 
+	/** Accumulated delta-time we were not able to evaluate with due to ongoing budgeted evaluation. */
+	float BlockedDeltaTimeCompensation;
+
 	/** Whether to restore pre-animated state. */
 	bool bRestoreState : 1;
 
