Replace some usages of FORCEINLINE with inline in MovieScene modules.

Includes modules:MovieSceneAnimMixer, TimecodeSynchronizer, MovieRenderPipelineCore, MovieRenderPipelineRenderPasses, MovieSceneTextTrack, Media, MediaAssets, MediaUtils, MovieScene, MovieSceneTracks, TimeManagement

Reviewers were automatically determined. Please comment if anyone else should be added.

#rnx

diff --git a/Engine/Plugins/Experimental/MovieSceneAnimMixer/Source/MovieSceneAnimMixer/Public/Systems/MovieSceneAnimBlueprintTargetSystem.h b/Engine/Plugins/Experimental/MovieSceneAnimMixer/Source/MovieSceneAnimMixer/Public/Systems/MovieSceneAnimBlueprintTargetSystem.h
--- a/Engine/Plugins/Experimental/MovieSceneAnimMixer/Source/MovieSceneAnimMixer/Public/Systems/MovieSceneAnimBlueprintTargetSystem.h
+++ b/Engine/Plugins/Experimental/MovieSceneAnimMixer/Source/MovieSceneAnimMixer/Public/Systems/MovieSceneAnimBlueprintTargetSystem.h
@@ -20,7 +20,7 @@ struct MOVIESCENEANIMMIXER_API FMovieSceneAnimBlueprintTarget : public FMovieSce
 	UPROPERTY(EditAnywhere, Category="Animation")
 	FName BlueprintNodeName;
 
-	FORCEINLINE friend uint32 GetTypeHash(const FMovieSceneAnimBlueprintTarget& Target)
+	inline friend uint32 GetTypeHash(const FMovieSceneAnimBlueprintTarget& Target)
 	{
 		return HashCombine(GetTypeHash(FMovieSceneAnimBlueprintTarget::StaticStruct()), GetTypeHash(Target.BlueprintNodeName));
 	}
diff --git a/Engine/Plugins/Experimental/MovieSceneAnimMixer/Source/MovieSceneAnimMixer/Public/Systems/MovieSceneAnimInstanceTargetSystem.h b/Engine/Plugins/Experimental/MovieSceneAnimMixer/Source/MovieSceneAnimMixer/Public/Systems/MovieSceneAnimInstanceTargetSystem.h
--- a/Engine/Plugins/Experimental/MovieSceneAnimMixer/Source/MovieSceneAnimMixer/Public/Systems/MovieSceneAnimInstanceTargetSystem.h
+++ b/Engine/Plugins/Experimental/MovieSceneAnimMixer/Source/MovieSceneAnimMixer/Public/Systems/MovieSceneAnimInstanceTargetSystem.h
@@ -22,7 +22,7 @@ struct MOVIESCENEANIMMIXER_API FMovieSceneAnimInstanceTarget : public FMovieScen
 {
 	GENERATED_BODY() 
 
-	FORCEINLINE friend uint32 GetTypeHash(const FMovieSceneAnimInstanceTarget& Target)
+	inline friend uint32 GetTypeHash(const FMovieSceneAnimInstanceTarget& Target)
 	{
 		return GetTypeHash(FMovieSceneAnimInstanceTarget::StaticStruct());
 	}
diff --git a/Engine/Plugins/Experimental/MovieSceneAnimMixer/Source/MovieSceneAnimMixer/Public/Systems/MovieSceneAnimMixerSystem.h b/Engine/Plugins/Experimental/MovieSceneAnimMixer/Source/MovieSceneAnimMixer/Public/Systems/MovieSceneAnimMixerSystem.h
--- a/Engine/Plugins/Experimental/MovieSceneAnimMixer/Source/MovieSceneAnimMixer/Public/Systems/MovieSceneAnimMixerSystem.h
+++ b/Engine/Plugins/Experimental/MovieSceneAnimMixer/Source/MovieSceneAnimMixer/Public/Systems/MovieSceneAnimMixerSystem.h
@@ -43,7 +43,7 @@ struct FMovieSceneAnimMixerEntry
 	UE::MovieScene::FInstanceHandle InstanceHandle;
 
 
-	FORCEINLINE bool operator<(const FMovieSceneAnimMixerEntry& RHS) const
+	inline bool operator<(const FMovieSceneAnimMixerEntry& RHS) const
 	{
 		if (Priority == RHS.Priority)
 		{
diff --git a/Engine/Plugins/Experimental/MovieSceneAnimMixer/Source/MovieSceneAnimMixer/Public/Systems/MovieSceneAnimNextTargetSystem.h b/Engine/Plugins/Experimental/MovieSceneAnimMixer/Source/MovieSceneAnimMixer/Public/Systems/MovieSceneAnimNextTargetSystem.h
--- a/Engine/Plugins/Experimental/MovieSceneAnimMixer/Source/MovieSceneAnimMixer/Public/Systems/MovieSceneAnimNextTargetSystem.h
+++ b/Engine/Plugins/Experimental/MovieSceneAnimMixer/Source/MovieSceneAnimMixer/Public/Systems/MovieSceneAnimNextTargetSystem.h
@@ -44,7 +44,7 @@ struct MOVIESCENEANIMMIXER_API FMovieSceneAnimNextInjectionTarget : public FMovi
 	UPROPERTY(EditAnywhere, Category="Animation", meta = (AllowedType = "FAnimNextAnimGraph"))
 	FAnimNextVariableReference InjectionSite;
 
-	FORCEINLINE friend uint32 GetTypeHash(const FMovieSceneAnimNextInjectionTarget& Target)
+	inline friend uint32 GetTypeHash(const FMovieSceneAnimNextInjectionTarget& Target)
 	{
 		return HashCombine(GetTypeHash(FMovieSceneAnimNextInjectionTarget::StaticStruct()), GetTypeHash(Target.InjectionSite));
 	}
diff --git a/Engine/Plugins/Media/TimecodeSynchronizer/Source/TimecodeSynchronizer/Public/TimecodeSynchronizer.h b/Engine/Plugins/Media/TimecodeSynchronizer/Source/TimecodeSynchronizer/Public/TimecodeSynchronizer.h
--- a/Engine/Plugins/Media/TimecodeSynchronizer/Source/TimecodeSynchronizer/Public/TimecodeSynchronizer.h
+++ b/Engine/Plugins/Media/TimecodeSynchronizer/Source/TimecodeSynchronizer/Public/TimecodeSynchronizer.h
@@ -157,47 +157,47 @@ public:
 	/** Updates the internal state of this source, returning whether or not the source is ready (e.g. IsReady() == true). */
 	const bool UpdateSourceState(const FFrameRate& SynchronizerFrameRate);
 
-	FORCEINLINE const UTimeSynchronizationSource* GetInputSource() const
+	inline const UTimeSynchronizationSource* GetInputSource() const
 	{
 		return InputSource;
 	}
 
-	FORCEINLINE bool IsInputSourceValid() const
+	inline bool IsInputSourceValid() const
 	{
 		return nullptr != InputSource;
 	}
 
-	FORCEINLINE FString GetDisplayName() const
+	inline FString GetDisplayName() const
 	{
 		return InputSource->GetDisplayName();
 	}
 
 	/** Whether or not this source is ready. */
-	FORCEINLINE bool IsReady() const
+	inline bool IsReady() const
 	{
 		return bIsReady;
 	}
 
 	/** Whether or not this source can be synchronized. */
-	FORCEINLINE bool CanBeSynchronized() const
+	inline bool CanBeSynchronized() const
 	{
 		return bCanBeSynchronized;
 	}
 
 	/** Gets the FrameRate of the source. */
-	FORCEINLINE const FFrameRate& GetFrameRate() const
+	inline const FFrameRate& GetFrameRate() const
 	{
 		return FrameRate;
 	}
 
 	/** Gets the state of the Source relative to its own frame rate. */
-	FORCEINLINE const FTimecodeSourceState& GetInputSourceState() const
+	inline const FTimecodeSourceState& GetInputSourceState() const
 	{
 		return InputSourceState;
 	}
 
 	/** Gets the state of the Source relative to the Synchronizer's frame rate. */
-	FORCEINLINE const FTimecodeSourceState& GetSynchronizerRelativeState() const
+	inline const FTimecodeSourceState& GetSynchronizerRelativeState() const
 	{
 		return SynchronizerRelativeState;
 	}
@@ -370,7 +370,7 @@ private:
 		Synchronized,								// all sources are running and synchronized
 	};
 
-	FORCEINLINE static FString SynchronizationStateToString(ESynchronizationState InState)
+	inline static FString SynchronizationStateToString(ESynchronizationState InState)
 	{
 		switch (InState)
 		{
diff --git a/Engine/Plugins/MovieScene/MovieRenderPipeline/Source/MovieRenderPipelineCore/Public/Graph/MovieGraphDefaultRenderer.h b/Engine/Plugins/MovieScene/MovieRenderPipeline/Source/MovieRenderPipelineCore/Public/Graph/MovieGraphDefaultRenderer.h
--- a/Engine/Plugins/MovieScene/MovieRenderPipeline/Source/MovieRenderPipelineCore/Public/Graph/MovieGraphDefaultRenderer.h
+++ b/Engine/Plugins/MovieScene/MovieRenderPipeline/Source/MovieRenderPipelineCore/Public/Graph/MovieGraphDefaultRenderer.h
@@ -266,7 +266,7 @@ namespace UE::MovieGraph::DefaultRenderer
 			return LastCompletionEvent;
 		}
 
-		FORCEINLINE TStatId GetStatId() const
+		inline TStatId GetStatId() const
 		{
 			RETURN_QUICK_DECLARE_CYCLE_STAT(FMovieGraphAccumulationTask, STATGROUP_ThreadPoolAsyncTasks);
 		}
diff --git a/Engine/Plugins/MovieScene/MovieRenderPipeline/Source/MovieRenderPipelineCore/Public/MoviePipelineVideoOutputBase.h b/Engine/Plugins/MovieScene/MovieRenderPipeline/Source/MovieRenderPipelineCore/Public/MoviePipelineVideoOutputBase.h
--- a/Engine/Plugins/MovieScene/MovieRenderPipeline/Source/MovieRenderPipelineCore/Public/MoviePipelineVideoOutputBase.h
+++ b/Engine/Plugins/MovieScene/MovieRenderPipeline/Source/MovieRenderPipelineCore/Public/MoviePipelineVideoOutputBase.h
@@ -60,7 +60,7 @@ public:
 		return LastCompletionEvent;
 	}
 
-	FORCEINLINE TStatId GetStatId() const
+	inline TStatId GetStatId() const
 	{
 		RETURN_QUICK_DECLARE_CYCLE_STAT(FMoviePipelineBackgroundMediaTasks, STATGROUP_ThreadPoolAsyncTasks);
 	}
diff --git a/Engine/Plugins/MovieScene/MovieRenderPipeline/Source/MovieRenderPipelineCore/Public/MovieRenderPipelineDataTypes.h b/Engine/Plugins/MovieScene/MovieRenderPipeline/Source/MovieRenderPipelineCore/Public/MovieRenderPipelineDataTypes.h
--- a/Engine/Plugins/MovieScene/MovieRenderPipeline/Source/MovieRenderPipelineCore/Public/MovieRenderPipelineDataTypes.h
+++ b/Engine/Plugins/MovieScene/MovieRenderPipeline/Source/MovieRenderPipelineCore/Public/MovieRenderPipelineDataTypes.h
@@ -648,7 +648,7 @@ public:
 			TimeDilation = 1.f;
 		}
 
-		FORCEINLINE bool IsTimeDilated() const { return !FMath::IsNearlyEqual(TimeDilation, 1.f); }
+		inline bool IsTimeDilated() const { return !FMath::IsNearlyEqual(TimeDilation, 1.f); }
 	};
 
 	FMoviePipelineFrameOutputState()
@@ -668,8 +668,8 @@ public:
 	}
 
 	/** Is this the first temporal sample for the output frame? */
-	FORCEINLINE bool IsFirstTemporalSample() const { return TemporalSampleIndex == 0; }
-	FORCEINLINE bool IsLastTemporalSample() const { return TemporalSampleIndex == (TemporalSampleCount - 1); }
+	inline bool IsFirstTemporalSample() const { return TemporalSampleIndex == 0; }
+	inline bool IsLastTemporalSample() const { return TemporalSampleIndex == (TemporalSampleCount - 1); }
 
 	/**
 	* The expected output frame count that the render is working towards creating.
@@ -979,12 +979,12 @@ public:
 
 
 	/** Get a 0-(TileCount-1) version of TileIndex. */
-	FORCEINLINE int32 GetTileIndex() const
+	inline int32 GetTileIndex() const
 	{
 		return (TileIndexes.Y * TileCounts.Y) + TileIndexes.X;
 	}
 
-	FORCEINLINE int32 GetTileCount() const
+	inline int32 GetTileCount() const
 	{
 		return (TileCounts.X * TileCounts.Y);
 	}
@@ -1180,25 +1180,25 @@ struct FImagePixelDataPayload : IImagePixelDataPayload, public TSharedFromThis<F
 	}
 
 	/** Is this the first tile of an image and we should start accumulating? */
-	FORCEINLINE bool IsFirstTile() const
+	inline bool IsFirstTile() const
 	{
 		return SampleState.TileIndexes.X == 0 && SampleState.TileIndexes.Y == 0 && SampleState.SpatialSampleIndex == 0;
 	}
 
 	/** Is this the last tile of an image and we should finish accumulating? */
-	FORCEINLINE bool IsLastTile() const
+	inline bool IsLastTile() const
 	{
 		return SampleState.TileIndexes.X == SampleState.TileCounts.X - 1 &&
 			   SampleState.TileIndexes.Y == SampleState.TileCounts.Y - 1 &&
 			   SampleState.SpatialSampleIndex == SampleState.SpatialSampleCount - 1;
 	}
 
-	FORCEINLINE bool IsFirstTemporalSample() const
+	inline bool IsFirstTemporalSample() const
 	{
 		return SampleState.TemporalSampleIndex == 0;
 	}
 
-	FORCEINLINE bool IsLastTemporalSample() const
+	inline bool IsLastTemporalSample() const
 	{
 		return SampleState.TemporalSampleIndex == SampleState.TemporalSampleCount - 1;
 	}
diff --git a/Engine/Plugins/MovieScene/MovieRenderPipeline/Source/MovieRenderPipelineRenderPasses/Public/MoviePipelineImagePassBase.h b/Engine/Plugins/MovieScene/MovieRenderPipeline/Source/MovieRenderPipelineRenderPasses/Public/MoviePipelineImagePassBase.h
--- a/Engine/Plugins/MovieScene/MovieRenderPipeline/Source/MovieRenderPipelineRenderPasses/Public/MoviePipelineImagePassBase.h
+++ b/Engine/Plugins/MovieScene/MovieRenderPipeline/Source/MovieRenderPipelineRenderPasses/Public/MoviePipelineImagePassBase.h
@@ -66,7 +66,7 @@ public:
 		return LastCompletionEvent;
 	}
 
-	FORCEINLINE TStatId GetStatId() const
+	inline TStatId GetStatId() const
 	{
 		RETURN_QUICK_DECLARE_CYCLE_STAT(FMoviePipelineBackgroundAccumulateTask, STATGROUP_ThreadPoolAsyncTasks);
 	}
diff --git a/Engine/Plugins/MovieScene/MovieSceneTextTrack/Source/MovieSceneTextTrack/Public/MovieSceneTextChannel.h b/Engine/Plugins/MovieScene/MovieSceneTextTrack/Source/MovieSceneTextTrack/Public/MovieSceneTextChannel.h
--- a/Engine/Plugins/MovieScene/MovieSceneTextTrack/Source/MovieSceneTextTrack/Public/MovieSceneTextChannel.h
+++ b/Engine/Plugins/MovieScene/MovieSceneTextTrack/Source/MovieSceneTextTrack/Public/MovieSceneTextChannel.h
@@ -33,7 +33,7 @@ struct MOVIESCENETEXTTRACK_API FMovieSceneTextChannel : public FMovieSceneChanne
 	 * Access a mutable interface for this channel's data
 	 * @return An object that is able to manipulate this channel's data
 	 */
-	FORCEINLINE TMovieSceneChannelData<FText> GetData()
+	inline TMovieSceneChannelData<FText> GetData()
 	{
 		return TMovieSceneChannelData<FText>(&Times, &Values, this, &KeyHandles);
 	}
@@ -42,18 +42,18 @@ struct MOVIESCENETEXTTRACK_API FMovieSceneTextChannel : public FMovieSceneChanne
 	 * Access a constant interface for this channel's data
 	 * @return An object that is able to interrogate this channel's data
 	 */
-	FORCEINLINE TMovieSceneChannelData<const FText> GetData() const
+	inline TMovieSceneChannelData<const FText> GetData() const
 	{
 		return TMovieSceneChannelData<const FText>(&Times, &Values);
 	}
 
-	FORCEINLINE TConstArrayView<FFrameNumber> GetTimes() const
+	inline TConstArrayView<FFrameNumber> GetTimes() const
 	{
 		return Times;
 	}
 
 	/** Returns whether this channel has any values */
-	FORCEINLINE bool HasAnyData() const
+	inline bool HasAnyData() const
 	{
 		return Times.Num() != 0 || bHasDefaultValue == true;
 	}
@@ -70,7 +70,7 @@ struct MOVIESCENETEXTTRACK_API FMovieSceneTextChannel : public FMovieSceneChanne
 	 * @param InTimes Times to add
 	 * @param InValues Values to add
 	 */
-	FORCEINLINE void AddKeys(const TArray<FFrameNumber>& InTimes, const TArray<FText>& InValues)
+	inline void AddKeys(const TArray<FFrameNumber>& InTimes, const TArray<FText>& InValues)
 	{
 		check(InTimes.Num() == InValues.Num());
 		Times.Append(InTimes);
@@ -81,7 +81,7 @@ struct MOVIESCENETEXTTRACK_API FMovieSceneTextChannel : public FMovieSceneChanne
 	 * Set this channel's default value that should be used when no keys are present
 	 * @param InDefaultValue The desired default value
 	 */
-	FORCEINLINE void SetDefault(FText InDefaultValue)
+	inline void SetDefault(FText InDefaultValue)
 	{
 		bHasDefaultValue = true;
 		DefaultValue = InDefaultValue;
@@ -91,13 +91,13 @@ struct MOVIESCENETEXTTRACK_API FMovieSceneTextChannel : public FMovieSceneChanne
 	 * Get this channel's default value that will be used when no keys are present
 	 * @return (Optional) The channel's default value
 	 */
-	FORCEINLINE TOptional<FText> GetDefault() const
+	inline TOptional<FText> GetDefault() const
 	{
 		return bHasDefaultValue ? TOptional<FText>(DefaultValue) : TOptional<FText>();
 	}
 
 	/** Remove this channel's default value causing the channel to have no effect where no keys are present */
-	FORCEINLINE void RemoveDefault()
+	inline void RemoveDefault()
 	{
 		bHasDefaultValue = false;
 	}
diff --git a/Engine/Source/Runtime/Media/Public/IMediaControls.h b/Engine/Source/Runtime/Media/Public/IMediaControls.h
--- a/Engine/Source/Runtime/Media/Public/IMediaControls.h
+++ b/Engine/Source/Runtime/Media/Public/IMediaControls.h
@@ -267,7 +267,7 @@ public:
 	 * @return true if the media is being paused, false otherwise.
 	 * @see Play, Stop
 	 */
-	FORCEINLINE bool Pause()
+	inline bool Pause()
 	{
 		return SetRate(0.0f);
 	}
@@ -280,7 +280,7 @@ public:
 	 * @return true if playback is starting, false otherwise.
 	 * @see Pause, Stop
 	 */
-	FORCEINLINE bool Play()
+	inline bool Play()
 	{
 		return SetRate(1.0f);
 	}
diff --git a/Engine/Source/Runtime/Media/Public/IMediaTextureSample.h b/Engine/Source/Runtime/Media/Public/IMediaTextureSample.h
--- a/Engine/Source/Runtime/Media/Public/IMediaTextureSample.h
+++ b/Engine/Source/Runtime/Media/Public/IMediaTextureSample.h
@@ -155,7 +155,7 @@ struct FMediaTextureTilingDescription
 	FIntPoint TileSize = FIntPoint::ZeroValue;
 	int32 TileBorderSize = 0;
 
-	FORCEINLINE bool IsValid() const
+	inline bool IsValid() const
 	{
 		return TileNum.X > 0 && TileNum.Y > 0 && TileSize.X > 0 && TileSize.Y > 0;
 	}
diff --git a/Engine/Source/Runtime/Media/Public/IMediaView.h b/Engine/Source/Runtime/Media/Public/IMediaView.h
--- a/Engine/Source/Runtime/Media/Public/IMediaView.h
+++ b/Engine/Source/Runtime/Media/Public/IMediaView.h
@@ -15,36 +15,36 @@ struct FMediaTileCoordinate
 
 	FMediaTileCoordinate() = default;
 
-	FORCEINLINE explicit FMediaTileCoordinate(uint16 InX, uint16 InY)
+	inline explicit FMediaTileCoordinate(uint16 InX, uint16 InY)
 		: X(InX)
 		, Y(InY)
 	{
 	}
 
-	FORCEINLINE explicit FMediaTileCoordinate(int32 InX, int32 InY)
+	inline explicit FMediaTileCoordinate(int32 InX, int32 InY)
 		: X(IntCastChecked<uint16>(InX))
 		, Y(IntCastChecked<uint16>(InY))
 	{
 	}
 
-	FORCEINLINE FMediaTileCoordinate(EForceInit)
+	inline FMediaTileCoordinate(EForceInit)
 		: X(0u)
 		, Y(0u)
 	{
 	}
 
-	FORCEINLINE bool operator==(const FMediaTileCoordinate& Other) const
+	inline bool operator==(const FMediaTileCoordinate& Other) const
 	{
 		return X == Other.X && Y == Other.Y;
 	}
 
 
-	FORCEINLINE bool operator!=(const FMediaTileCoordinate& Other) const
+	inline bool operator!=(const FMediaTileCoordinate& Other) const
 	{
 		return X != Other.X || Y != Other.Y;
 	}
 
-	FORCEINLINE static FMediaTileCoordinate Zero()
+	inline static FMediaTileCoordinate Zero()
 	{
 		return FMediaTileCoordinate(EForceInit::ForceInit);
 	}
@@ -53,7 +53,7 @@ struct FMediaTileCoordinate
 template <> struct TIsPODType<FMediaTileCoordinate> { enum { Value = true }; };
 
 
-FORCEINLINE uint32 GetTypeHash(const FMediaTileCoordinate& Coord)
+inline uint32 GetTypeHash(const FMediaTileCoordinate& Coord)
 {
 	return FCrc::TypeCrc32(Coord);
 }
diff --git a/Engine/Source/Runtime/MediaAssets/Public/MediaPlayer.h b/Engine/Source/Runtime/MediaAssets/Public/MediaPlayer.h
--- a/Engine/Source/Runtime/MediaAssets/Public/MediaPlayer.h
+++ b/Engine/Source/Runtime/MediaAssets/Public/MediaPlayer.h
@@ -154,7 +154,7 @@ public:
 		return(PlayerFacade.Pin());
 	}
 
-	friend FORCEINLINE uint32 GetTypeHash(const FMediaPlayerProxy& InProxy)
+	friend inline uint32 GetTypeHash(const FMediaPlayerProxy& InProxy)
 	{
 		return InProxy.TypeHash;
 	}
diff --git a/Engine/Source/Runtime/MediaUtils/Public/MediaSampleCache.h b/Engine/Source/Runtime/MediaUtils/Public/MediaSampleCache.h
--- a/Engine/Source/Runtime/MediaUtils/Public/MediaSampleCache.h
+++ b/Engine/Source/Runtime/MediaUtils/Public/MediaSampleCache.h
@@ -34,17 +34,17 @@ class FMediaSampleCache
 	struct TSampleKeyFuncs
 		: BaseKeyFuncs<TSharedPtr<SampleType, ESPMode::ThreadSafe>, FTimespan, false>
 	{
-		static FORCEINLINE FTimespan GetSetKey(const TSharedPtr<SampleType, ESPMode::ThreadSafe>& Element)
+		static inline FTimespan GetSetKey(const TSharedPtr<SampleType, ESPMode::ThreadSafe>& Element)
 		{
 			return Element->GetTime().Time;
 		}
 
-		static FORCEINLINE bool Matches(FTimespan A, FTimespan B)
+		static inline bool Matches(FTimespan A, FTimespan B)
 		{
 			return A == B;
 		}
 
-		static FORCEINLINE uint32 GetKeyHash(FTimespan Key)
+		static inline uint32 GetKeyHash(FTimespan Key)
 		{
 			return GetTypeHash(Key);
 		}
diff --git a/Engine/Source/Runtime/MovieScene/Public/Channels/MovieSceneAudioTriggerChannel.h b/Engine/Source/Runtime/MovieScene/Public/Channels/MovieSceneAudioTriggerChannel.h
--- a/Engine/Source/Runtime/MovieScene/Public/Channels/MovieSceneAudioTriggerChannel.h
+++ b/Engine/Source/Runtime/MovieScene/Public/Channels/MovieSceneAudioTriggerChannel.h
@@ -40,7 +40,7 @@ struct FMovieSceneAudioTriggerChannel : public FMovieSceneChannel
 	 *
 	 * @return An object that is able to manipulate this channel's data
 	 */
-	FORCEINLINE TMovieSceneChannelData<bool> GetData()
+	inline TMovieSceneChannelData<bool> GetData()
 	{
 		return TMovieSceneChannelData<bool>(&Times, &Values, this, &KeyHandles);
 	}
@@ -50,7 +50,7 @@ struct FMovieSceneAudioTriggerChannel : public FMovieSceneChannel
 	 *
 	 * @return An object that is able to interrogate this channel's data
 	 */
-	FORCEINLINE TMovieSceneChannelData<const bool> GetData() const
+	inline TMovieSceneChannelData<const bool> GetData() const
 	{
 		return TMovieSceneChannelData<const bool>(&Times, &Values);
 	}
@@ -58,7 +58,7 @@ struct FMovieSceneAudioTriggerChannel : public FMovieSceneChannel
 	/**
 	 * Const access to this channel's times
 	 */
-	FORCEINLINE TArrayView<const FFrameNumber> GetTimes() const
+	inline TArrayView<const FFrameNumber> GetTimes() const
 	{
 		return Times;
 	}
@@ -66,7 +66,7 @@ struct FMovieSceneAudioTriggerChannel : public FMovieSceneChannel
 	/**
 	 * Check whether this channel has any data
 	 */
-	FORCEINLINE bool HasAnyData() const
+	inline bool HasAnyData() const
 	{
 		return Times.Num() != 0;
 	}
diff --git a/Engine/Source/Runtime/MovieScene/Public/Channels/MovieSceneBoolChannel.h b/Engine/Source/Runtime/MovieScene/Public/Channels/MovieSceneBoolChannel.h
--- a/Engine/Source/Runtime/MovieScene/Public/Channels/MovieSceneBoolChannel.h
+++ b/Engine/Source/Runtime/MovieScene/Public/Channels/MovieSceneBoolChannel.h
@@ -48,7 +48,7 @@ struct FMovieSceneBoolChannel : public FMovieSceneChannel
 	 *
 	 * @return An object that is able to manipulate this channel's data
 	 */
-	virtual FORCEINLINE TMovieSceneChannelData<bool> GetData()
+	virtual inline TMovieSceneChannelData<bool> GetData()
 	{
 		return TMovieSceneChannelData<bool>(&Times, &Values, this, &KeyHandles);
 	}
@@ -58,7 +58,7 @@ struct FMovieSceneBoolChannel : public FMovieSceneChannel
 	 *
 	 * @return An object that is able to interrogate this channel's data
 	 */
-	virtual FORCEINLINE TMovieSceneChannelData<const bool> GetData() const
+	virtual inline TMovieSceneChannelData<const bool> GetData() const
 	{
 		return TMovieSceneChannelData<const bool>(&Times, &Values);
 	}
@@ -66,7 +66,7 @@ struct FMovieSceneBoolChannel : public FMovieSceneChannel
 	/**
 	 * Const access to this channel's times
 	 */
-	FORCEINLINE TArrayView<const FFrameNumber> GetTimes() const
+	inline TArrayView<const FFrameNumber> GetTimes() const
 	{
 		return Times;
 	}
@@ -74,7 +74,7 @@ struct FMovieSceneBoolChannel : public FMovieSceneChannel
 	/**
 	 * Const access to this channel's values
 	 */
-	FORCEINLINE TArrayView<const bool> GetValues() const
+	inline TArrayView<const bool> GetValues() const
 	{
 		return Values;
 	}
@@ -82,7 +82,7 @@ struct FMovieSceneBoolChannel : public FMovieSceneChannel
 	/**
 	 * Check whether this channel has any data
 	 */
-	FORCEINLINE bool HasAnyData() const
+	inline bool HasAnyData() const
 	{
 		return Times.Num() != 0 || bHasDefaultValue == true;
 	}
@@ -101,7 +101,7 @@ struct FMovieSceneBoolChannel : public FMovieSceneChannel
 	 * @param InTimes Times to add
 	 * @param InValues Values to add
 	 */
-	FORCEINLINE void AddKeys(const TArray<FFrameNumber>& InTimes, const TArray<bool>& InValues)
+	inline void AddKeys(const TArray<FFrameNumber>& InTimes, const TArray<bool>& InValues)
 	{
 		check(InTimes.Num() == InValues.Num());
 		Times.Append(InTimes);
@@ -134,7 +134,7 @@ public:
 	 *
 	 * @param InDefaultValue The desired default value
 	 */
-	FORCEINLINE void SetDefault(bool InDefaultValue)
+	inline void SetDefault(bool InDefaultValue)
 	{
 		bHasDefaultValue = true;
 		DefaultValue = InDefaultValue;
@@ -145,7 +145,7 @@ public:
 	 *
 	 * @return (Optional) The channel's default value
 	 */
-	FORCEINLINE TOptional<bool> GetDefault() const
+	inline TOptional<bool> GetDefault() const
 	{
 		return bHasDefaultValue ? TOptional<bool>(DefaultValue) : TOptional<bool>();
 	}
@@ -153,7 +153,7 @@ public:
 	/**
 	 * Remove this channel's default value causing the channel to have no effect where no keys are present
 	 */
-	FORCEINLINE void RemoveDefault()
+	inline void RemoveDefault()
 	{
 		bHasDefaultValue = false;
 	}
diff --git a/Engine/Source/Runtime/MovieScene/Public/Channels/MovieSceneByteChannel.h b/Engine/Source/Runtime/MovieScene/Public/Channels/MovieSceneByteChannel.h
--- a/Engine/Source/Runtime/MovieScene/Public/Channels/MovieSceneByteChannel.h
+++ b/Engine/Source/Runtime/MovieScene/Public/Channels/MovieSceneByteChannel.h
@@ -54,7 +54,7 @@ struct FMovieSceneByteChannel : public FMovieSceneChannel
 	 *
 	 * @return An object that is able to manipulate this channel's data
 	 */
-	FORCEINLINE TMovieSceneChannelData<uint8> GetData()
+	inline TMovieSceneChannelData<uint8> GetData()
 	{
 		return TMovieSceneChannelData<uint8>(&Times, &Values, this, &KeyHandles);
 	}
@@ -64,7 +64,7 @@ struct FMovieSceneByteChannel : public FMovieSceneChannel
 	 *
 	 * @return An object that is able to interrogate this channel's data
 	 */
-	FORCEINLINE TMovieSceneChannelData<const uint8> GetData() const
+	inline TMovieSceneChannelData<const uint8> GetData() const
 	{
 		return TMovieSceneChannelData<const uint8>(&Times, &Values);
 	}
@@ -72,7 +72,7 @@ struct FMovieSceneByteChannel : public FMovieSceneChannel
 	/**
 	 * Const access to this channel's times
 	 */
-	FORCEINLINE TArrayView<const FFrameNumber> GetTimes() const
+	inline TArrayView<const FFrameNumber> GetTimes() const
 	{
 		return Times;
 	}
@@ -80,7 +80,7 @@ struct FMovieSceneByteChannel : public FMovieSceneChannel
 	/**
 	 * Const access to this channel's values
 	 */
-	FORCEINLINE TArrayView<const uint8> GetValues() const
+	inline TArrayView<const uint8> GetValues() const
 	{
 		return Values;
 	}
@@ -88,7 +88,7 @@ struct FMovieSceneByteChannel : public FMovieSceneChannel
 	/**
 	 * Check whether this channel has any data
 	 */
-	FORCEINLINE bool HasAnyData() const
+	inline bool HasAnyData() const
 	{
 		return Times.Num() != 0 || bHasDefaultValue == true;
 	}
@@ -107,7 +107,7 @@ struct FMovieSceneByteChannel : public FMovieSceneChannel
 	 * @param InTimes Times to add
 	 * @param InValues Values to add
 	 */
-	FORCEINLINE void AddKeys(const TArray<FFrameNumber>& InTimes, const TArray<uint8>& InValues)
+	inline void AddKeys(const TArray<FFrameNumber>& InTimes, const TArray<uint8>& InValues)
 	{
 		check(InTimes.Num() == InValues.Num());
 		Times.Append(InTimes);
@@ -138,7 +138,7 @@ public:
 	 *
 	 * @param InDefaultValue The desired default value
 	 */
-	FORCEINLINE void SetDefault(uint8 InDefaultValue)
+	inline void SetDefault(uint8 InDefaultValue)
 	{
 		bHasDefaultValue = true;
 		DefaultValue = InDefaultValue;
@@ -149,7 +149,7 @@ public:
 	 *
 	 * @return (Optional) The channel's default value
 	 */
-	FORCEINLINE TOptional<uint8> GetDefault() const
+	inline TOptional<uint8> GetDefault() const
 	{
 		return bHasDefaultValue ? TOptional<uint8>(DefaultValue) : TOptional<uint8>();
 	}
@@ -157,7 +157,7 @@ public:
 	/**
 	 * Remove this channel's default value causing the channel to have no effect where no keys are present
 	 */
-	FORCEINLINE void RemoveDefault()
+	inline void RemoveDefault()
 	{
 		bHasDefaultValue = false;
 	}
diff --git a/Engine/Source/Runtime/MovieScene/Public/Channels/MovieSceneChannel.h b/Engine/Source/Runtime/MovieScene/Public/Channels/MovieSceneChannel.h
--- a/Engine/Source/Runtime/MovieScene/Public/Channels/MovieSceneChannel.h
+++ b/Engine/Source/Runtime/MovieScene/Public/Channels/MovieSceneChannel.h
@@ -72,20 +72,20 @@ struct FCycleParams
 		, bMirrorCurve(false)
 	{}
 
-	FORCEINLINE void ComputePreValueOffset(double FirstValue, double LastValue)
+	inline void ComputePreValueOffset(double FirstValue, double LastValue)
 	{
 		// CycleCount is negative for pre-extrap
 		ValueOffset = (LastValue - FirstValue) * CycleCount;
 	}
-	FORCEINLINE void ComputePostValueOffset(double FirstValue, double LastValue)
+	inline void ComputePostValueOffset(double FirstValue, double LastValue)
 	{
 		ValueOffset = (LastValue - FirstValue) * CycleCount;
 	}
-	FORCEINLINE bool ShouldMirrorCurve() const
+	inline bool ShouldMirrorCurve() const
 	{
 		return bMirrorCurve;
 	}
-	FORCEINLINE void Oscillate(int32 MinFrame, int32 MaxFrame)
+	inline void Oscillate(int32 MinFrame, int32 MaxFrame)
 	{
 		if (FMath::Abs(CycleCount) % 2 == 1)
 		{
diff --git a/Engine/Source/Runtime/MovieScene/Public/Channels/MovieSceneChannelData.h b/Engine/Source/Runtime/MovieScene/Public/Channels/MovieSceneChannelData.h
--- a/Engine/Source/Runtime/MovieScene/Public/Channels/MovieSceneChannelData.h
+++ b/Engine/Source/Runtime/MovieScene/Public/Channels/MovieSceneChannelData.h
@@ -109,7 +109,7 @@ struct FMovieSceneChannelData
 	/**
 	 * Read-only access to this channel's key times.
 	 */
-	FORCEINLINE TArrayView<const FFrameNumber> GetTimes() const
+	inline TArrayView<const FFrameNumber> GetTimes() const
 	{
 		return *Times;
 	}
@@ -118,7 +118,7 @@ struct FMovieSceneChannelData
 	 * Mutable access to this channel's key times.
 	 * @note: *Warning*: any usage *must* keep times sorted. Any reordering of times will not be reflected in the values array.
 	 */
-	FORCEINLINE TArrayView<FFrameNumber> GetTimes()
+	inline TArrayView<FFrameNumber> GetTimes()
 	{
 		return *Times;
 	}
@@ -294,7 +294,7 @@ struct TMovieSceneChannelData : FMovieSceneChannelData
 	/**
 	 * Read-only access to this channel's values
 	 */
-	FORCEINLINE TArrayView<const ValueType> GetValues() const
+	inline TArrayView<const ValueType> GetValues() const
 	{
 		return *Values;
 	}
@@ -302,7 +302,7 @@ struct TMovieSceneChannelData : FMovieSceneChannelData
 	/**
 	 * Mutable access to this channel's values
 	 */
-	FORCEINLINE TArrayView<ValueType> GetValues()
+	inline TArrayView<ValueType> GetValues()
 	{
 		return *Values;
 	}
@@ -722,7 +722,7 @@ struct TMovieSceneChannelData<const ValueType>
 	/**
 	 * Read-only access to this channel's key times.
 	 */
-	FORCEINLINE TArrayView<const FFrameNumber> GetTimes() const
+	inline TArrayView<const FFrameNumber> GetTimes() const
 	{
 		return *Times;
 	}
@@ -730,7 +730,7 @@ struct TMovieSceneChannelData<const ValueType>
 	/**
 	 * Read-only access to this channel's values
 	 */
-	FORCEINLINE TArrayView<const ValueType> GetValues() const
+	inline TArrayView<const ValueType> GetValues() const
 	{
 		return *Values;
 	}
diff --git a/Engine/Source/Runtime/MovieScene/Public/Channels/MovieSceneDoubleChannel.h b/Engine/Source/Runtime/MovieScene/Public/Channels/MovieSceneDoubleChannel.h
--- a/Engine/Source/Runtime/MovieScene/Public/Channels/MovieSceneDoubleChannel.h
+++ b/Engine/Source/Runtime/MovieScene/Public/Channels/MovieSceneDoubleChannel.h
@@ -145,7 +145,7 @@ struct FMovieSceneDoubleChannel : public FMovieSceneChannel
 	 *
 	 * @return An object that is able to manipulate this channel's data
 	 */
-	FORCEINLINE TMovieSceneChannelData<FMovieSceneDoubleValue> GetData()
+	inline TMovieSceneChannelData<FMovieSceneDoubleValue> GetData()
 	{
 		return TMovieSceneChannelData<FMovieSceneDoubleValue>(&Times, &Values, this, &KeyHandles);
 	}
@@ -155,7 +155,7 @@ struct FMovieSceneDoubleChannel : public FMovieSceneChannel
 	 *
 	 * @return An object that is able to interrogate this channel's data
 	 */
-	FORCEINLINE TMovieSceneChannelData<const FMovieSceneDoubleValue> GetData() const
+	inline TMovieSceneChannelData<const FMovieSceneDoubleValue> GetData() const
 	{
 		return TMovieSceneChannelData<const FMovieSceneDoubleValue>(&Times, &Values);
 	}
@@ -163,7 +163,7 @@ struct FMovieSceneDoubleChannel : public FMovieSceneChannel
 	/**
 	 * Const access to this channel's times
 	 */
-	FORCEINLINE TArrayView<const FFrameNumber> GetTimes() const
+	inline TArrayView<const FFrameNumber> GetTimes() const
 	{
 		return Times;
 	}
@@ -171,7 +171,7 @@ struct FMovieSceneDoubleChannel : public FMovieSceneChannel
 	/**
 	 * Const access to this channel's values
 	 */
-	FORCEINLINE TArrayView<const FMovieSceneDoubleValue> GetValues() const
+	inline TArrayView<const FMovieSceneDoubleValue> GetValues() const
 	{
 		return Values;
 	}
@@ -277,7 +277,7 @@ public:
 	/**
 	 * Check whether this channel has any data
 	 */
-	FORCEINLINE bool HasAnyData() const
+	inline bool HasAnyData() const
 	{
 		return Times.Num() != 0 || bHasDefaultValue == true;
 	}
@@ -287,7 +287,7 @@ public:
 	 *
 	 * @param InDefaultValue The desired default value
 	 */
-	FORCEINLINE void SetDefault(double InDefaultValue)
+	inline void SetDefault(double InDefaultValue)
 	{
 		bHasDefaultValue = true;
 		DefaultValue = InDefaultValue;
@@ -298,7 +298,7 @@ public:
 	 *
 	 * @return (Optional) The channel's default value
 	 */
-	FORCEINLINE TOptional<double> GetDefault() const
+	inline TOptional<double> GetDefault() const
 	{
 		return bHasDefaultValue ? TOptional<double>(DefaultValue) : TOptional<double>();
 	}
@@ -306,7 +306,7 @@ public:
 	/**
 	 * Remove this channel's default value causing the channel to have no effect where no keys are present
 	 */
-	FORCEINLINE void RemoveDefault()
+	inline void RemoveDefault()
 	{
 		bHasDefaultValue = false;
 	}
diff --git a/Engine/Source/Runtime/MovieScene/Public/Channels/MovieSceneFloatChannel.h b/Engine/Source/Runtime/MovieScene/Public/Channels/MovieSceneFloatChannel.h
--- a/Engine/Source/Runtime/MovieScene/Public/Channels/MovieSceneFloatChannel.h
+++ b/Engine/Source/Runtime/MovieScene/Public/Channels/MovieSceneFloatChannel.h
@@ -136,7 +136,7 @@ struct FMovieSceneFloatChannel : public FMovieSceneChannel
 	 *
 	 * @return An object that is able to manipulate this channel's data
 	 */
-	FORCEINLINE TMovieSceneChannelData<FMovieSceneFloatValue> GetData()
+	inline TMovieSceneChannelData<FMovieSceneFloatValue> GetData()
 	{
 		return TMovieSceneChannelData<FMovieSceneFloatValue>(&Times, &Values, this, &KeyHandles);
 	}
@@ -146,7 +146,7 @@ struct FMovieSceneFloatChannel : public FMovieSceneChannel
 	 *
 	 * @return An object that is able to interrogate this channel's data
 	 */
-	FORCEINLINE TMovieSceneChannelData<const FMovieSceneFloatValue> GetData() const
+	inline TMovieSceneChannelData<const FMovieSceneFloatValue> GetData() const
 	{
 		return TMovieSceneChannelData<const FMovieSceneFloatValue>(&Times, &Values);
 	}
@@ -154,7 +154,7 @@ struct FMovieSceneFloatChannel : public FMovieSceneChannel
 	/**
 	 * Const access to this channel's times
 	 */
-	FORCEINLINE TArrayView<const FFrameNumber> GetTimes() const
+	inline TArrayView<const FFrameNumber> GetTimes() const
 	{
 		return Times;
 	}
@@ -162,7 +162,7 @@ struct FMovieSceneFloatChannel : public FMovieSceneChannel
 	/**
 	 * Const access to this channel's values
 	 */
-	FORCEINLINE TArrayView<const FMovieSceneFloatValue> GetValues() const
+	inline TArrayView<const FMovieSceneFloatValue> GetValues() const
 	{
 		return Values;
 	}
@@ -216,7 +216,7 @@ public:
 	/**
 	 * Check whether this channel has any data
 	 */
-	FORCEINLINE bool HasAnyData() const
+	inline bool HasAnyData() const
 	{
 		return Times.Num() != 0 || bHasDefaultValue == true;
 	}
@@ -226,7 +226,7 @@ public:
 	 *
 	 * @param InDefaultValue The desired default value
 	 */
-	FORCEINLINE void SetDefault(float InDefaultValue)
+	inline void SetDefault(float InDefaultValue)
 	{
 		bHasDefaultValue = true;
 		DefaultValue = InDefaultValue;
@@ -237,7 +237,7 @@ public:
 	 *
 	 * @return (Optional) The channel's default value
 	 */
-	FORCEINLINE TOptional<float> GetDefault() const
+	inline TOptional<float> GetDefault() const
 	{
 		return bHasDefaultValue ? TOptional<float>(DefaultValue) : TOptional<float>();
 	}
@@ -245,7 +245,7 @@ public:
 	/**
 	 * Remove this channel's default value causing the channel to have no effect where no keys are present
 	 */
-	FORCEINLINE void RemoveDefault()
+	inline void RemoveDefault()
 	{
 		bHasDefaultValue = false;
 	}
diff --git a/Engine/Source/Runtime/MovieScene/Public/Channels/MovieSceneIntegerChannel.h b/Engine/Source/Runtime/MovieScene/Public/Channels/MovieSceneIntegerChannel.h
--- a/Engine/Source/Runtime/MovieScene/Public/Channels/MovieSceneIntegerChannel.h
+++ b/Engine/Source/Runtime/MovieScene/Public/Channels/MovieSceneIntegerChannel.h
@@ -53,7 +53,7 @@ struct FMovieSceneIntegerChannel : public FMovieSceneChannel
 	 *
 	 * @return An object that is able to manipulate this channel's data
 	 */
-	FORCEINLINE TMovieSceneChannelData<int32> GetData()
+	inline TMovieSceneChannelData<int32> GetData()
 	{
 		return TMovieSceneChannelData<int32>(&Times, &Values, this, &KeyHandles);
 	}
@@ -63,7 +63,7 @@ struct FMovieSceneIntegerChannel : public FMovieSceneChannel
 	 *
 	 * @return An object that is able to interrogate this channel's data
 	 */
-	FORCEINLINE TMovieSceneChannelData<const int32> GetData() const
+	inline TMovieSceneChannelData<const int32> GetData() const
 	{
 		return TMovieSceneChannelData<const int32>(&Times, &Values);
 	}
@@ -71,7 +71,7 @@ struct FMovieSceneIntegerChannel : public FMovieSceneChannel
 	/**
 	 * Const access to this channel's times
 	 */
-	FORCEINLINE TArrayView<const FFrameNumber> GetTimes() const
+	inline TArrayView<const FFrameNumber> GetTimes() const
 	{
 		return Times;
 	}
@@ -79,7 +79,7 @@ struct FMovieSceneIntegerChannel : public FMovieSceneChannel
 	/**
 	 * Const access to this channel's values
 	 */
-	FORCEINLINE TArrayView<const int32> GetValues() const
+	inline TArrayView<const int32> GetValues() const
 	{
 		return Values;
 	}
@@ -87,7 +87,7 @@ struct FMovieSceneIntegerChannel : public FMovieSceneChannel
 	/**
 	 * Check whether this channel has any data
 	 */
-	FORCEINLINE bool HasAnyData() const
+	inline bool HasAnyData() const
 	{
 		return Times.Num() != 0 || bHasDefaultValue == true;
 	}
@@ -113,7 +113,7 @@ struct FMovieSceneIntegerChannel : public FMovieSceneChannel
 	/**
 	 * Set the channel's times and values to the requested values
 	 */
-	FORCEINLINE void Set(TArray<FFrameNumber> InTimes, TArray<int32> InValues)
+	inline void Set(TArray<FFrameNumber> InTimes, TArray<int32> InValues)
 	{
 		check(InTimes.Num() == InValues.Num());
 
@@ -132,7 +132,7 @@ struct FMovieSceneIntegerChannel : public FMovieSceneChannel
 	 * @param InTimes Times to add
 	 * @param InValues Values to add
 	 */
-	FORCEINLINE void AddKeys(const TArray<FFrameNumber>& InTimes, const TArray<int32>& InValues)
+	inline void AddKeys(const TArray<FFrameNumber>& InTimes, const TArray<int32>& InValues)
 	{
 		check(InTimes.Num() == InValues.Num());
 		int32 Index = Times.Num();
@@ -170,7 +170,7 @@ public:
 	 *
 	 * @param InDefaultValue The desired default value
 	 */
-	FORCEINLINE void SetDefault(int32 InDefaultValue)
+	inline void SetDefault(int32 InDefaultValue)
 	{
 		bHasDefaultValue = true;
 		DefaultValue = InDefaultValue;
@@ -181,7 +181,7 @@ public:
 	 *
 	 * @return (Optional) The channel's default value
 	 */
-	FORCEINLINE TOptional<int32> GetDefault() const
+	inline TOptional<int32> GetDefault() const
 	{
 		return bHasDefaultValue ? TOptional<int32>(DefaultValue) : TOptional<int32>();
 	}
@@ -189,7 +189,7 @@ public:
 	/**
 	 * Remove this channel's default value causing the channel to have no effect where no keys are present
 	 */
-	FORCEINLINE void RemoveDefault()
+	inline void RemoveDefault()
 	{
 		bHasDefaultValue = false;
 	}
diff --git a/Engine/Source/Runtime/MovieScene/Public/Channels/MovieSceneObjectPathChannel.h b/Engine/Source/Runtime/MovieScene/Public/Channels/MovieSceneObjectPathChannel.h
--- a/Engine/Source/Runtime/MovieScene/Public/Channels/MovieSceneObjectPathChannel.h
+++ b/Engine/Source/Runtime/MovieScene/Public/Channels/MovieSceneObjectPathChannel.h
@@ -78,12 +78,12 @@ struct FMovieSceneObjectPathChannel : public FMovieSceneChannel
 		: PropertyClass(nullptr)
 	{}
 
-	FORCEINLINE void SetPropertyClass(UClass* InPropertyClass)
+	inline void SetPropertyClass(UClass* InPropertyClass)
 	{
 		PropertyClass = InPropertyClass;
 	}
 
-	FORCEINLINE UClass* GetPropertyClass() const
+	inline UClass* GetPropertyClass() const
 	{
 		return PropertyClass;
 	}
@@ -93,7 +93,7 @@ struct FMovieSceneObjectPathChannel : public FMovieSceneChannel
 	 *
 	 * @return An object that is able to manipulate this channel's data
 	 */
-	FORCEINLINE TMovieSceneChannelData<FMovieSceneObjectPathChannelKeyValue> GetData()
+	inline TMovieSceneChannelData<FMovieSceneObjectPathChannelKeyValue> GetData()
 	{
 		return TMovieSceneChannelData<FMovieSceneObjectPathChannelKeyValue>(&Times, &Values, this, &KeyHandles);
 	}
@@ -103,7 +103,7 @@ struct FMovieSceneObjectPathChannel : public FMovieSceneChannel
 	 *
 	 * @return An object that is able to interrogate this channel's data
 	 */
-	FORCEINLINE TMovieSceneChannelData<const FMovieSceneObjectPathChannelKeyValue> GetData() const
+	inline TMovieSceneChannelData<const FMovieSceneObjectPathChannelKeyValue> GetData() const
 	{
 		return TMovieSceneChannelData<const FMovieSceneObjectPathChannelKeyValue>(&Times, &Values);
 	}
@@ -143,12 +143,12 @@ public:
 	 *
 	 * @param InDefaultValue The desired default value
 	 */
-	FORCEINLINE void SetDefault(UObject* InDefaultValue)
+	inline void SetDefault(UObject* InDefaultValue)
 	{
 		DefaultValue = InDefaultValue;
 	}
 
-	FORCEINLINE void SetDefault(FMovieSceneObjectPathChannelKeyValue InDefaultValue)
+	inline void SetDefault(FMovieSceneObjectPathChannelKeyValue InDefaultValue)
 	{
 		DefaultValue = InDefaultValue;
 	}
@@ -158,7 +158,7 @@ public:
 	 *
 	 * @return (Optional) The channel's default value
 	 */
-	FORCEINLINE const FMovieSceneObjectPathChannelKeyValue& GetDefault() const
+	inline const FMovieSceneObjectPathChannelKeyValue& GetDefault() const
 	{
 		return DefaultValue;
 	}
@@ -166,7 +166,7 @@ public:
 	/**
 	 * Remove this channel's default value causing the channel to have no effect where no keys are present
 	 */
-	FORCEINLINE void RemoveDefault()
+	inline void RemoveDefault()
 	{
 		DefaultValue = nullptr;
 	}
diff --git a/Engine/Source/Runtime/MovieScene/Public/Compilation/MovieSceneTemplateInterrogation.h b/Engine/Source/Runtime/MovieScene/Public/Compilation/MovieSceneTemplateInterrogation.h
--- a/Engine/Source/Runtime/MovieScene/Public/Compilation/MovieSceneTemplateInterrogation.h
+++ b/Engine/Source/Runtime/MovieScene/Public/Compilation/MovieSceneTemplateInterrogation.h
@@ -136,21 +136,21 @@ struct TMovieSceneInterrogationIterator
 		return &static_cast<const TMovieSceneInterrogationToken<DataType>&>(Container.TokenData[Index].GetValue()).Data;
 	}
 
-	FORCEINLINE explicit operator bool() const
+	inline explicit operator bool() const
 	{
 		return Container.TokenData.IsValidIndex(Index);
 	}
 
-	FORCEINLINE bool operator!() const 
+	inline bool operator!() const 
 	{
 		return !(bool)*this;
 	}
 
-	FORCEINLINE friend bool operator==(const TMovieSceneInterrogationIterator& LHS, const TMovieSceneInterrogationIterator& RHS) { return &LHS.Container == &RHS.Container && LHS.Index == RHS.Index; }
-	FORCEINLINE friend bool operator!=(const TMovieSceneInterrogationIterator& LHS, const TMovieSceneInterrogationIterator& RHS) { return &LHS.Container != &RHS.Container || LHS.Index != RHS.Index; }
+	inline friend bool operator==(const TMovieSceneInterrogationIterator& LHS, const TMovieSceneInterrogationIterator& RHS) { return &LHS.Container == &RHS.Container && LHS.Index == RHS.Index; }
+	inline friend bool operator!=(const TMovieSceneInterrogationIterator& LHS, const TMovieSceneInterrogationIterator& RHS) { return &LHS.Container != &RHS.Container || LHS.Index != RHS.Index; }
 
-	FORCEINLINE friend TMovieSceneInterrogationIterator<DataType> begin(const TMovieSceneInterrogationIterator<DataType>& In) { return In; }
-	FORCEINLINE friend TMovieSceneInterrogationIterator<DataType> end(  const TMovieSceneInterrogationIterator<DataType>& In) { TMovieSceneInterrogationIterator<DataType> New = In; New.Index = -1; return New; }
+	inline friend TMovieSceneInterrogationIterator<DataType> begin(const TMovieSceneInterrogationIterator<DataType>& In) { return In; }
+	inline friend TMovieSceneInterrogationIterator<DataType> end(  const TMovieSceneInterrogationIterator<DataType>& In) { TMovieSceneInterrogationIterator<DataType> New = In; New.Index = -1; return New; }
 
 	void NextElement()
 	{
diff --git a/Engine/Source/Runtime/MovieScene/Public/Containers/SparseBitSet.h b/Engine/Source/Runtime/MovieScene/Public/Containers/SparseBitSet.h
--- a/Engine/Source/Runtime/MovieScene/Public/Containers/SparseBitSet.h
+++ b/Engine/Source/Runtime/MovieScene/Public/Containers/SparseBitSet.h
@@ -386,7 +386,7 @@ private:
 	template<typename, typename>
 	friend struct TFixedSparseBitSet;
 
-	FORCEINLINE void CheckIndex(uint32 BitIndex) const
+	inline void CheckIndex(uint32 BitIndex) const
 	{
 		checkfSlow(BitIndex < MaxNumBits, TEXT("Invalid index (%d) specified for a sparse bitset of maximum size (%d)"), BitIndex, MaxNumBits);
 	}
diff --git a/Engine/Source/Runtime/MovieScene/Public/EntitySystem/BuiltInComponentTypes.h b/Engine/Source/Runtime/MovieScene/Public/EntitySystem/BuiltInComponentTypes.h
--- a/Engine/Source/Runtime/MovieScene/Public/EntitySystem/BuiltInComponentTypes.h
+++ b/Engine/Source/Runtime/MovieScene/Public/EntitySystem/BuiltInComponentTypes.h
@@ -808,7 +808,7 @@ public:
 
 	static MOVIESCENE_API FBuiltInComponentTypes* Get();
 
-	FORCEINLINE static bool IsBoundObjectGarbage(UObject* InObject)
+	inline static bool IsBoundObjectGarbage(UObject* InObject)
 	{
 		return InObject == nullptr || !IsValidChecked(InObject) || InObject->IsUnreachable();
 	}
diff --git a/Engine/Source/Runtime/MovieScene/Public/EntitySystem/MovieSceneCachedEntityFilterResult.h b/Engine/Source/Runtime/MovieScene/Public/EntitySystem/MovieSceneCachedEntityFilterResult.h
--- a/Engine/Source/Runtime/MovieScene/Public/EntitySystem/MovieSceneCachedEntityFilterResult.h
+++ b/Engine/Source/Runtime/MovieScene/Public/EntitySystem/MovieSceneCachedEntityFilterResult.h
@@ -49,7 +49,7 @@ struct FCachedEntityManagerState
 	/**
 	 * Retrieve the serial number for this cache
 	 */
-	FORCEINLINE uint64 GetSerial() const
+	inline uint64 GetSerial() const
 	{
 		return LastSystemVersion;
 	}
diff --git a/Engine/Source/Runtime/MovieScene/Public/EntitySystem/MovieSceneComponentAccessors.h b/Engine/Source/Runtime/MovieScene/Public/EntitySystem/MovieSceneComponentAccessors.h
--- a/Engine/Source/Runtime/MovieScene/Public/EntitySystem/MovieSceneComponentAccessors.h
+++ b/Engine/Source/Runtime/MovieScene/Public/EntitySystem/MovieSceneComponentAccessors.h
@@ -31,7 +31,7 @@ struct FComponentAccess
 {
 	FComponentTypeID ComponentType;
 
-	FORCEINLINE void PreLockComponentData(const FEntityAllocation* Allocation, FPreLockedDataPtr* PrelockedComponentData) const
+	inline void PreLockComponentData(const FEntityAllocation* Allocation, FPreLockedDataPtr* PrelockedComponentData) const
 	{
 		// Always relative to the component data which is a separate allocation
 		const FComponentHeader& Header = Allocation->GetComponentHeaderChecked(ComponentType);
@@ -53,7 +53,7 @@ struct FOptionalComponentAccess
 	FComponentTypeID ComponentType;
 	FComponentTypeIDFilter Condition;
 
-	FORCEINLINE void PreLockComponentData(FEntityAllocationIteratorItem Item, FPreLockedDataPtr* PrelockedComponentData) const
+	inline void PreLockComponentData(FEntityAllocationIteratorItem Item, FPreLockedDataPtr* PrelockedComponentData) const
 	{
 		if (!Condition.Passes(Item.GetAllocationType()))
 		{
@@ -86,18 +86,18 @@ struct FEntityIDAccess
 	static constexpr int32 PreLockedDataNum = 1;
 
 
-	FORCEINLINE void PreLockComponentData(const FEntityAllocation* Allocation, FPreLockedDataPtr* PrelockedComponentData) const
+	inline void PreLockComponentData(const FEntityAllocation* Allocation, FPreLockedDataPtr* PrelockedComponentData) const
 	{
 		// Must be relative to the base allocation ptr
 		PrelockedComponentData->Reset(Allocation, Allocation->GetRawEntityIDs());
 	}
-	FORCEINLINE TRead<FMovieSceneEntityID> ResolvePreLockedComponentData(const FEntityAllocation* Allocation, const FPreLockedDataPtr* Ptr, FEntityAllocationWriteContext WriteContext) const
+	inline TRead<FMovieSceneEntityID> ResolvePreLockedComponentData(const FEntityAllocation* Allocation, const FPreLockedDataPtr* Ptr, FEntityAllocationWriteContext WriteContext) const
 	{
 		// Resolve using the same base ptr as PreLockComponentData
 		return TRead<FMovieSceneEntityID>(Ptr->Resolve<FMovieSceneEntityID>(Allocation));
 	}
 
-	FORCEINLINE TRead<FMovieSceneEntityID> LockComponentData(const FEntityAllocation* Allocation, FEntityAllocationWriteContext WriteContext) const
+	inline TRead<FMovieSceneEntityID> LockComponentData(const FEntityAllocation* Allocation, FEntityAllocationWriteContext WriteContext) const
 	{
 		return TRead<FMovieSceneEntityID>(Allocation->GetRawEntityIDs());
 	}
@@ -108,7 +108,7 @@ struct FFilterMatchPassthrough
 	using AccessType = bool;
 	//static constexpr int32 PreLockedDataNum = 1;
 
-	FORCEINLINE void PreLockComponentData(FEntityAllocationIteratorItem Item, FPreLockedDataPtr* PrelockedComponentData) const
+	inline void PreLockComponentData(FEntityAllocationIteratorItem Item, FPreLockedDataPtr* PrelockedComponentData) const
 	{
 		// Must be relative to the base allocation ptr
 		const bool bFilterMatch = Filter.Match(Item.GetAllocationType());
@@ -124,12 +124,12 @@ struct FFilterMatchPassthrough
 			PrelockedComponentData->Reset(Allocation, Allocation);
 		}
 	}
-	FORCEINLINE bool ResolvePreLockedComponentData(const FEntityAllocation* Allocation, const FPreLockedDataPtr* Ptr, FEntityAllocationWriteContext WriteContext) const
+	inline bool ResolvePreLockedComponentData(const FEntityAllocation* Allocation, const FPreLockedDataPtr* Ptr, FEntityAllocationWriteContext WriteContext) const
 	{
 		return Ptr->Resolve<FEntityAllocation>(Allocation) != nullptr;
 	}
 
-	FORCEINLINE bool LockComponentData(FEntityAllocationIteratorItem Item, FEntityAllocationWriteContext WriteContext) const
+	inline bool LockComponentData(FEntityAllocationIteratorItem Item, FEntityAllocationWriteContext WriteContext) const
 	{
 		return Filter.Match(Item.GetAllocationType());
 	}
@@ -147,12 +147,12 @@ struct TReadAccess : FReadAccess
 		: FReadAccess{ InComponentTypeID }
 	{}
 
-	FORCEINLINE TRead<T> ResolvePreLockedComponentData(const FEntityAllocation* Allocation, const FPreLockedDataPtr* Ptr, FEntityAllocationWriteContext WriteContext) const
+	inline TRead<T> ResolvePreLockedComponentData(const FEntityAllocation* Allocation, const FPreLockedDataPtr* Ptr, FEntityAllocationWriteContext WriteContext) const
 	{
 		// Resolve using the same base ptr as PreLockComponentData
 		return TRead<T>(Ptr->Resolve<const T>(Allocation->GetComponentDataAddress()));
 	}
-	FORCEINLINE TComponentLock<TRead<T>> LockComponentData(const FEntityAllocation* Allocation, FEntityAllocationWriteContext WriteContext) const
+	inline TComponentLock<TRead<T>> LockComponentData(const FEntityAllocation* Allocation, FEntityAllocationWriteContext WriteContext) const
 	{
 		return Allocation->ReadComponents(ComponentType.ReinterpretCast<T>());
 	}
@@ -167,18 +167,18 @@ struct FErasedReadAccess : FReadAccess
 		: FReadAccess{ InComponentTypeID }
 	{}
 
-	FORCEINLINE void PreLockComponentData(const FEntityAllocation* Allocation, FPreLockedDataPtr* PrelockedComponentData) const
+	inline void PreLockComponentData(const FEntityAllocation* Allocation, FPreLockedDataPtr* PrelockedComponentData) const
 	{
 		// Erased reads always pass the header since they need the size information
 		const FComponentHeader& Header = Allocation->GetComponentHeaderChecked(ComponentType);
 		PrelockedComponentData->Reset(Allocation, &Header);
 	}
-	FORCEINLINE FReadErased ResolvePreLockedComponentData(const FEntityAllocation* Allocation, const FPreLockedDataPtr* Ptr, FEntityAllocationWriteContext WriteContext) const
+	inline FReadErased ResolvePreLockedComponentData(const FEntityAllocation* Allocation, const FPreLockedDataPtr* Ptr, FEntityAllocationWriteContext WriteContext) const
 	{
 		// Resolve using the same base ptr as PreLockComponentData
 		return FReadErased(Ptr->Resolve<const FComponentHeader>(Allocation));
 	}
-	FORCEINLINE FComponentReader LockComponentData(const FEntityAllocation* Allocation, FEntityAllocationWriteContext WriteContext) const
+	inline FComponentReader LockComponentData(const FEntityAllocation* Allocation, FEntityAllocationWriteContext WriteContext) const
 	{
 		return Allocation->ReadComponentsErased(ComponentType);
 	}
@@ -194,7 +194,7 @@ struct FErasedOptionalReadAccess : FReadAccess
 		, Condition(InCondition)
 	{}
 
-	FORCEINLINE void PreLockComponentData(FEntityAllocationIteratorItem Item, FPreLockedDataPtr* PrelockedComponentData) const
+	inline void PreLockComponentData(FEntityAllocationIteratorItem Item, FPreLockedDataPtr* PrelockedComponentData) const
 	{
 		if (!Condition.Passes(Item.GetAllocationType()))
 		{
@@ -209,12 +209,12 @@ struct FErasedOptionalReadAccess : FReadAccess
 			PrelockedComponentData->Reset(Allocation, Header);
 		}
 	}
-	FORCEINLINE FReadErasedOptional ResolvePreLockedComponentData(const FEntityAllocation* Allocation, const FPreLockedDataPtr* Ptr, FEntityAllocationWriteContext WriteContext) const
+	inline FReadErasedOptional ResolvePreLockedComponentData(const FEntityAllocation* Allocation, const FPreLockedDataPtr* Ptr, FEntityAllocationWriteContext WriteContext) const
 	{
 		// Resolve using the same base ptr as PreLockComponentData
 		return FReadErasedOptional(Ptr->Resolve<const FComponentHeader>(Allocation));
 	}
-	FORCEINLINE FOptionalComponentReader LockComponentData(FEntityAllocationIteratorItem Item, FEntityAllocationWriteContext WriteContext) const
+	inline FOptionalComponentReader LockComponentData(FEntityAllocationIteratorItem Item, FEntityAllocationWriteContext WriteContext) const
 	{
 		if (!Condition.Passes(Item.GetAllocationType()))
 		{
@@ -233,18 +233,18 @@ struct FErasedWriteAccess : FWriteAccess
 		: FWriteAccess{ InComponentTypeID }
 	{}
 
-	FORCEINLINE void PreLockComponentData(const FEntityAllocation* Allocation, FPreLockedDataPtr* PrelockedComponentData) const
+	inline void PreLockComponentData(const FEntityAllocation* Allocation, FPreLockedDataPtr* PrelockedComponentData) const
 	{
 		const FComponentHeader& Header = Allocation->GetComponentHeaderChecked(ComponentType);
 		// Erased writes always pass the header since they need the size information
 		PrelockedComponentData->Reset(Allocation, &Header);
 	}
-	FORCEINLINE FWriteErased ResolvePreLockedComponentData(const FEntityAllocation* Allocation, const FPreLockedDataPtr* Ptr, FEntityAllocationWriteContext WriteContext) const
+	inline FWriteErased ResolvePreLockedComponentData(const FEntityAllocation* Allocation, const FPreLockedDataPtr* Ptr, FEntityAllocationWriteContext WriteContext) const
 	{
 		// Resolve using the same base ptr as PreLockComponentData
 		return FWriteErased(Ptr->Resolve<FComponentHeader>(Allocation));
 	}
-	FORCEINLINE FComponentWriter LockComponentData(const FEntityAllocation* Allocation, FEntityAllocationWriteContext WriteContext) const
+	inline FComponentWriter LockComponentData(const FEntityAllocation* Allocation, FEntityAllocationWriteContext WriteContext) const
 	{
 		return Allocation->WriteComponentsErased(ComponentType, WriteContext);
 	}
@@ -263,12 +263,12 @@ struct TWriteAccess : FWriteAccess
 		: FWriteAccess{ InComponentTypeID }
 	{}
 
-	FORCEINLINE TWrite<T> ResolvePreLockedComponentData(const FEntityAllocation* Allocation, const FPreLockedDataPtr* Ptr, FEntityAllocationWriteContext WriteContext) const
+	inline TWrite<T> ResolvePreLockedComponentData(const FEntityAllocation* Allocation, const FPreLockedDataPtr* Ptr, FEntityAllocationWriteContext WriteContext) const
 	{
 		// Resolve using the same base ptr as PreLockComponentData
 		return TWrite<T>(Ptr->Resolve<T>(Allocation->GetComponentDataAddress()));
 	}
-	FORCEINLINE TComponentLock<TWrite<T>> LockComponentData(const FEntityAllocation* Allocation, FEntityAllocationWriteContext WriteContext) const
+	inline TComponentLock<TWrite<T>> LockComponentData(const FEntityAllocation* Allocation, FEntityAllocationWriteContext WriteContext) const
 	{
 		return Allocation->WriteComponents(ComponentType.ReinterpretCast<T>(), WriteContext);
 	}
@@ -287,12 +287,12 @@ struct TOptionalReadAccess : FOptionalReadAccess
 		: FOptionalReadAccess{InComponentTypeID, InConditionType}
 	{}
 
-	FORCEINLINE TReadOptional<T> ResolvePreLockedComponentData(const FEntityAllocation* Allocation, const FPreLockedDataPtr* Ptr, FEntityAllocationWriteContext WriteContext) const
+	inline TReadOptional<T> ResolvePreLockedComponentData(const FEntityAllocation* Allocation, const FPreLockedDataPtr* Ptr, FEntityAllocationWriteContext WriteContext) const
 	{
 		// Resolve using the same base ptr as PreLockComponentData
 		return TReadOptional<T>(Ptr->Resolve<T>(Allocation->GetComponentDataAddress()));
 	}
-	FORCEINLINE TComponentLock<TReadOptional<T>> LockComponentData(const FEntityAllocation* Allocation, FEntityAllocationWriteContext WriteContext) const
+	inline TComponentLock<TReadOptional<T>> LockComponentData(const FEntityAllocation* Allocation, FEntityAllocationWriteContext WriteContext) const
 	{
 		return Allocation->TryReadComponents(ComponentType.ReinterpretCast<T>());
 	}
@@ -311,12 +311,12 @@ struct TOptionalWriteAccess : FOptionalWriteAccess
 		: FOptionalWriteAccess{ InComponentTypeID, InConditionType }
 	{}
 
-	FORCEINLINE TWriteOptional<T> ResolvePreLockedComponentData(const FEntityAllocation* Allocation, const FPreLockedDataPtr* Ptr, FEntityAllocationWriteContext WriteContext) const
+	inline TWriteOptional<T> ResolvePreLockedComponentData(const FEntityAllocation* Allocation, const FPreLockedDataPtr* Ptr, FEntityAllocationWriteContext WriteContext) const
 	{
 		// Resolve using the same base ptr as PreLockComponentData
 		return TWriteOptional<T>(Ptr->Resolve<T>(Allocation->GetComponentDataAddress()));
 	}
-	FORCEINLINE TComponentLock<TWriteOptional<T>> LockComponentData(const FEntityAllocation* Allocation, FEntityAllocationWriteContext WriteContext) const
+	inline TComponentLock<TWriteOptional<T>> LockComponentData(const FEntityAllocation* Allocation, FEntityAllocationWriteContext WriteContext) const
 	{
 		return Allocation->TryWriteComponents(ComponentType.ReinterpretCast<T>(), WriteContext);
 	}
diff --git a/Engine/Source/Runtime/MovieScene/Public/EntitySystem/MovieSceneComponentPtr.h b/Engine/Source/Runtime/MovieScene/Public/EntitySystem/MovieSceneComponentPtr.h
--- a/Engine/Source/Runtime/MovieScene/Public/EntitySystem/MovieSceneComponentPtr.h
+++ b/Engine/Source/Runtime/MovieScene/Public/EntitySystem/MovieSceneComponentPtr.h
@@ -41,45 +41,45 @@ struct TComponentPtr
 		return *this;
 	}
 
-	FORCEINLINE explicit operator bool() const
+	inline explicit operator bool() const
 	{
 		return IsValid();
 	}
 
-	FORCEINLINE bool IsValid() const
+	inline bool IsValid() const
 	{
 		return ComponentPtr != nullptr;
 	}
 
 	/** Explicitly convert this component data to its underlying pointer type */
-	FORCEINLINE T* AsPtr() const
+	inline T* AsPtr() const
 	{
 		return ComponentPtr;
 	}
 
 	/** Retrieve this component data as an array view of the specified size (usually the size will be FEntityAllocation::Num()) */
-	FORCEINLINE TArrayView<T> AsArray(int32 Num) const
+	inline TArrayView<T> AsArray(int32 Num) const
 	{
 		return MakeArrayView(ComponentPtr, Num);
 	}
 
 	/** Retrieve a slice of this component data as an array view of the specified size and offset (usually the size will be FEntityAllocation::Num()) */
-	FORCEINLINE TArrayView<T> Slice(int32 StartIndex, int32 Num) const
+	inline TArrayView<T> Slice(int32 StartIndex, int32 Num) const
 	{
 		return MakeArrayView(ComponentPtr + StartIndex, Num);
 	}
 
-	FORCEINLINE T* operator->() const
+	inline T* operator->() const
 	{
 		return ComponentPtr;
 	}
 
-	FORCEINLINE ValueType operator*() const
+	inline ValueType operator*() const
 	{
 		return *ComponentPtr;
 	}
 
-	FORCEINLINE ValueType operator[](int32 Index) const
+	inline ValueType operator[](int32 Index) const
 	{
 		return ComponentPtr[Index];
 	}
@@ -89,7 +89,7 @@ struct TComponentPtr
 		return (*this)[Index];
 	}
 
-	FORCEINLINE operator T*() const
+	inline operator T*() const
 	{
 		return ComponentPtr;
 	}
@@ -137,19 +137,19 @@ struct FReadErased
 		, Sizeof(InHeader->Sizeof)
 	{}
 
-	FORCEINLINE const void* AsPtr() const
+	inline const void* AsPtr() const
 	{
 		return ComponentPtr;
 	}
-	FORCEINLINE const void* operator[](int32 Index) const
+	inline const void* operator[](int32 Index) const
 	{
 		return ComponentPtr + Sizeof*Index;
 	}
-	FORCEINLINE const void* ComponentAtIndex(int32 Index) const
+	inline const void* ComponentAtIndex(int32 Index) const
 	{
 		return ComponentPtr + Sizeof*Index;
 	}
-	FORCEINLINE bool IsValid() const
+	inline bool IsValid() const
 	{
 		return ComponentPtr != nullptr;
 	}
@@ -236,15 +236,15 @@ struct FWriteErased
 		: ComponentPtr(InHeader->Components + ComponentOffset*InHeader->Sizeof)
 		, Sizeof(InHeader->Sizeof)
 	{}
-	FORCEINLINE void* AsPtr() const
+	inline void* AsPtr() const
 	{
 		return ComponentPtr;
 	}
-	FORCEINLINE void* operator[](int32 Index) const
+	inline void* operator[](int32 Index) const
 	{
 		return ComponentPtr + Sizeof*Index;
 	}
-	FORCEINLINE bool IsValid() const
+	inline bool IsValid() const
 	{
 		return ComponentPtr != nullptr;
 	}
diff --git a/Engine/Source/Runtime/MovieScene/Public/EntitySystem/MovieSceneEntityIDs.h b/Engine/Source/Runtime/MovieScene/Public/EntitySystem/MovieSceneEntityIDs.h
--- a/Engine/Source/Runtime/MovieScene/Public/EntitySystem/MovieSceneEntityIDs.h
+++ b/Engine/Source/Runtime/MovieScene/Public/EntitySystem/MovieSceneEntityIDs.h
@@ -73,53 +73,53 @@ struct FComponentMask
 		return Bits.CompareSetBits(Other.Bits, false);
 	}
 
-	FORCEINLINE FComponentMask& CombineWithBitwiseOR(const FComponentMask& Other, EBitwiseOperatorFlags InFlags)
+	inline FComponentMask& CombineWithBitwiseOR(const FComponentMask& Other, EBitwiseOperatorFlags InFlags)
 	{
 		Bits.CombineWithBitwiseOR(Other.Bits, InFlags);
 		return *this;
 	}
 
-	FORCEINLINE FComponentMask& CombineWithBitwiseXOR(const FComponentMask& Other, EBitwiseOperatorFlags InFlags)
+	inline FComponentMask& CombineWithBitwiseXOR(const FComponentMask& Other, EBitwiseOperatorFlags InFlags)
 	{
 		Bits.CombineWithBitwiseXOR(Other.Bits, InFlags);
 		return *this;
 	}
 	
-	FORCEINLINE FComponentMask& CombineWithBitwiseAND(const FComponentMask& Other, EBitwiseOperatorFlags InFlags)
+	inline FComponentMask& CombineWithBitwiseAND(const FComponentMask& Other, EBitwiseOperatorFlags InFlags)
 	{
 		Bits.CombineWithBitwiseAND(Other.Bits, InFlags);
 		return *this;
 	}
 
-	FORCEINLINE static FComponentMask BitwiseAND(const FComponentMask& A, const FComponentMask& B, EBitwiseOperatorFlags InFlags)
+	inline static FComponentMask BitwiseAND(const FComponentMask& A, const FComponentMask& B, EBitwiseOperatorFlags InFlags)
 	{
 		FComponentMask Result;
 		Result.Bits = TBitArray<FComponentMaskAllocator>::BitwiseAND(A.Bits, B.Bits, InFlags);
 		return Result;
 	}
 
-	FORCEINLINE static FComponentMask BitwiseOR(const FComponentMask& A, const FComponentMask& B, EBitwiseOperatorFlags InFlags)
+	inline static FComponentMask BitwiseOR(const FComponentMask& A, const FComponentMask& B, EBitwiseOperatorFlags InFlags)
 	{
 		FComponentMask Result;
 		Result.Bits = TBitArray<FComponentMaskAllocator>::BitwiseOR(A.Bits, B.Bits, InFlags);
 		return Result;
 	}
 
-	FORCEINLINE static FComponentMask BitwiseXOR(const FComponentMask& A, const FComponentMask& B, EBitwiseOperatorFlags InFlags)
+	inline static FComponentMask BitwiseXOR(const FComponentMask& A, const FComponentMask& B, EBitwiseOperatorFlags InFlags)
 	{
 		FComponentMask Result;
 		Result.Bits = TBitArray<FComponentMaskAllocator>::BitwiseXOR(A.Bits, B.Bits, InFlags);
 		return Result;
 	}
 
-	FORCEINLINE static FComponentMask BitwiseNOT(const FComponentMask& A)
+	inline static FComponentMask BitwiseNOT(const FComponentMask& A)
 	{
 		FComponentMask Result = A;
 		Result.Bits.BitwiseNOT();
 		return Result;
 	}
 
-	FORCEINLINE void BitwiseNOT()
+	inline void BitwiseNOT()
 	{
 		Bits.BitwiseNOT();
 	}
@@ -139,27 +139,27 @@ struct FComponentMask
 	 */
 	FComponentTypeID First() const;
 
-	FORCEINLINE void  Reset()
+	inline void  Reset()
 	{
 		Bits.Reset();
 	}
 
-	FORCEINLINE int32 Num() const
+	inline int32 Num() const
 	{
 		return Bits.Num();
 	}
 
-	FORCEINLINE int32 Find(bool bValue) const
+	inline int32 Find(bool bValue) const
 	{
 		return Bits.Find(bValue);
 	}
 
-	FORCEINLINE int32 NumComponents() const
+	inline int32 NumComponents() const
 	{
 		return Bits.CountSetBits();
 	}
 
-	FORCEINLINE void PadToNum(int32 Num, bool bPadValue)
+	inline void PadToNum(int32 Num, bool bPadValue)
 	{
 		Bits.PadToNum(Num, bPadValue);
 	}
diff --git a/Engine/Source/Runtime/MovieScene/Public/EntitySystem/MovieSceneEntitySystemTask.h b/Engine/Source/Runtime/MovieScene/Public/EntitySystem/MovieSceneEntitySystemTask.h
--- a/Engine/Source/Runtime/MovieScene/Public/EntitySystem/MovieSceneEntitySystemTask.h
+++ b/Engine/Source/Runtime/MovieScene/Public/EntitySystem/MovieSceneEntitySystemTask.h
@@ -108,17 +108,17 @@ struct TEntityTaskTraits : TDefaultEntityTaskTraits<T>
 struct FEntityIterationResult
 {
 	template<typename T>
-	friend FORCEINLINE FEntityIterationResult operator,(T, FEntityIterationResult)
+	friend inline FEntityIterationResult operator,(T, FEntityIterationResult)
 	{
 		return FEntityIterationResult { true };
 	}
 
-	friend FORCEINLINE FEntityIterationResult operator,(bool In, FEntityIterationResult)
+	friend inline FEntityIterationResult operator,(bool In, FEntityIterationResult)
 	{
 		return FEntityIterationResult{ In };
 	}
 
-	FORCEINLINE explicit operator bool() const
+	inline explicit operator bool() const
 	{
 		return Value;
 	}
@@ -961,7 +961,7 @@ public:
 	 * Get the accessor for a specific index within this task
 	 */
 	template<int Index>
-	FORCEINLINE auto GetAccessor() const
+	inline auto GetAccessor() const
 	{
 		static_assert(Index < sizeof...(T), "Invalid component index specified");
 		return Accessors.template Get<Index>();
@@ -1006,11 +1006,11 @@ public:
 struct FNoopTask
 {
 	template<typename... T>
-	FORCEINLINE static void ForEachAllocation(T&&...)
+	inline static void ForEachAllocation(T&&...)
 	{}
 
 	template<typename... T>
-	FORCEINLINE static void ForEachEntity(T&&...)
+	inline static void ForEachEntity(T&&...)
 	{}
 };
 
@@ -1926,19 +1926,19 @@ template<int32 NumComponents>
 struct TEntityTaskCaller<NumComponents, false>
 {
 	template<typename TaskImpl>
-	FORCEINLINE static void ForEachEntityImpl(TaskImpl& TaskImplInstance, ...)
+	inline static void ForEachEntityImpl(TaskImpl& TaskImplInstance, ...)
 	{
 		static_assert(!std::is_same_v<TaskImpl, TaskImpl>, "non-expanded entity iteration is not supported");
 	}
 
 	template<typename TaskImpl, typename... AccessorTypes>
-	FORCEINLINE static void ForEachAllocationImpl(TaskImpl& TaskImplInstance, FEntityAllocationIteratorItem Item, TArrayView<const FPreLockedDataPtr> PreLockedData, FEntityAllocationWriteContext WriteContext, const TEntityTaskComponents<AccessorTypes...>& Components)
+	inline static void ForEachAllocationImpl(TaskImpl& TaskImplInstance, FEntityAllocationIteratorItem Item, TArrayView<const FPreLockedDataPtr> PreLockedData, FEntityAllocationWriteContext WriteContext, const TEntityTaskComponents<AccessorTypes...>& Components)
 	{
 		TaskImplInstance.ForEachAllocation(Item, Components, PreLockedData);
 	}
 
 	template<typename TaskImpl, typename... AccessorTypes>
-	FORCEINLINE static void ForEachAllocationImpl(TaskImpl& TaskImplInstance, FEntityAllocationIteratorItem Item, const TEntityTaskComponents<AccessorTypes...>& Components)
+	inline static void ForEachAllocationImpl(TaskImpl& TaskImplInstance, FEntityAllocationIteratorItem Item, const TEntityTaskComponents<AccessorTypes...>& Components)
 	{
 		TaskImplInstance.ForEachAllocation(Item, Components);
 	}
diff --git a/Engine/Source/Runtime/MovieScene/Public/EntitySystem/MovieSceneInitialValueCache.h b/Engine/Source/Runtime/MovieScene/Public/EntitySystem/MovieSceneInitialValueCache.h
--- a/Engine/Source/Runtime/MovieScene/Public/EntitySystem/MovieSceneInitialValueCache.h
+++ b/Engine/Source/Runtime/MovieScene/Public/EntitySystem/MovieSceneInitialValueCache.h
@@ -355,7 +355,7 @@ private:
 		}
 	};
 
-	FORCEINLINE void CleanupStaleEntries()
+	inline void CleanupStaleEntries()
 	{
 		if (!bLUTContainsInvalidEntries)
 		{
diff --git a/Engine/Source/Runtime/MovieScene/Public/EntitySystem/MovieSceneOverlappingEntityTracker.h b/Engine/Source/Runtime/MovieScene/Public/EntitySystem/MovieSceneOverlappingEntityTracker.h
--- a/Engine/Source/Runtime/MovieScene/Public/EntitySystem/MovieSceneOverlappingEntityTracker.h
+++ b/Engine/Source/Runtime/MovieScene/Public/EntitySystem/MovieSceneOverlappingEntityTracker.h
@@ -30,32 +30,32 @@ struct FEntityOutputAggregate
 
 struct FGarbageTraits
 {
-	FORCEINLINE static constexpr bool IsGarbage(...)
+	inline static constexpr bool IsGarbage(...)
 	{
 		return false;
 	}
 	template<typename T>
-	FORCEINLINE static std::enable_if_t<TPointerIsConvertibleFromTo<T, const UObject>::Value, bool> IsGarbage(T* InObject)
+	inline static std::enable_if_t<TPointerIsConvertibleFromTo<T, const UObject>::Value, bool> IsGarbage(T* InObject)
 	{
 		constexpr bool bTypeDependentFalse = !std::is_same_v<T, T>;
 		static_assert(bTypeDependentFalse, "Raw object pointers are no longer supported. Please use TObjectPtr<T> instead.");
 	}
-	FORCEINLINE static bool IsGarbage(FObjectKey InObject)
+	inline static bool IsGarbage(FObjectKey InObject)
 	{
 		return FBuiltInComponentTypes::IsBoundObjectGarbage(InObject.ResolveObjectPtr());
 	}
 	template<typename T>
-	FORCEINLINE static bool IsGarbage(TObjectPtr<T>& InObject)
+	inline static bool IsGarbage(TObjectPtr<T>& InObject)
 	{
 		return FBuiltInComponentTypes::IsBoundObjectGarbage(InObject);
 	}
-	FORCEINLINE static bool IsGarbage(FObjectComponent& InComponent)
+	inline static bool IsGarbage(FObjectComponent& InComponent)
 	{
 		return FBuiltInComponentTypes::IsBoundObjectGarbage(InComponent.GetObject());
 	}
 	
 	template<typename T>
-	FORCEINLINE static void AddReferencedObjects(FReferenceCollector& ReferenceCollector, T* In)
+	inline static void AddReferencedObjects(FReferenceCollector& ReferenceCollector, T* In)
 	{
 		if constexpr (THasAddReferencedObjectForComponent<T>::Value)
 		{
diff --git a/Engine/Source/Runtime/MovieScene/Public/EntitySystem/MovieScenePartialProperties.h b/Engine/Source/Runtime/MovieScene/Public/EntitySystem/MovieScenePartialProperties.h
--- a/Engine/Source/Runtime/MovieScene/Public/EntitySystem/MovieScenePartialProperties.h
+++ b/Engine/Source/Runtime/MovieScene/Public/EntitySystem/MovieScenePartialProperties.h
@@ -18,7 +18,7 @@ namespace MovieScene
 {
 
 template<typename CompositeType>
-FORCEINLINE void PatchComposite(uint8* OutValueBase, const CompositeType* Composite, uint16 PtrOffset)
+inline void PatchComposite(uint8* OutValueBase, const CompositeType* Composite, uint16 PtrOffset)
 {
 	if (Composite)
 	{
@@ -27,7 +27,7 @@ FORCEINLINE void PatchComposite(uint8* OutValueBase, const CompositeType* Compos
 }
 
 template<typename... CompositeTypes, int ...CompositeIndices>
-FORCEINLINE void PatchCompositeValueImpl(const TIntegerSequence<int, CompositeIndices...>&, TArrayView<const FPropertyCompositeDefinition> CompositeDefinitions, uint8* OutValueBase, const CompositeTypes*... Composites)
+inline void PatchCompositeValueImpl(const TIntegerSequence<int, CompositeIndices...>&, TArrayView<const FPropertyCompositeDefinition> CompositeDefinitions, uint8* OutValueBase, const CompositeTypes*... Composites)
 {
 	int Tmp[] = {
 		( PatchComposite(OutValueBase, Composites, CompositeDefinitions[CompositeIndices].CompositeOffset), 0 )..., 0
diff --git a/Engine/Source/Runtime/MovieScene/Public/EntitySystem/MovieSceneSystemTaskDependencies.h b/Engine/Source/Runtime/MovieScene/Public/EntitySystem/MovieSceneSystemTaskDependencies.h
--- a/Engine/Source/Runtime/MovieScene/Public/EntitySystem/MovieSceneSystemTaskDependencies.h
+++ b/Engine/Source/Runtime/MovieScene/Public/EntitySystem/MovieSceneSystemTaskDependencies.h
@@ -59,7 +59,7 @@ struct FSystemTaskPrerequisites
 		return &AllTasks;
 	}
 
-	FORCEINLINE void FilterByComponent(FGraphEventArray& OutArray, FComponentTypeID ComponentType) const
+	inline void FilterByComponent(FGraphEventArray& OutArray, FComponentTypeID ComponentType) const
 	{
 		FilterByComponent(OutArray, { ComponentType });
 	}
diff --git a/Engine/Source/Runtime/MovieScene/Public/Evaluation/Blending/MovieSceneBlendType.h b/Engine/Source/Runtime/MovieScene/Public/Evaluation/Blending/MovieSceneBlendType.h
--- a/Engine/Source/Runtime/MovieScene/Public/Evaluation/Blending/MovieSceneBlendType.h
+++ b/Engine/Source/Runtime/MovieScene/Public/Evaluation/Blending/MovieSceneBlendType.h
@@ -163,11 +163,11 @@ public:
 	static FMovieSceneBlendTypeFieldIterator Begin(FMovieSceneBlendTypeField InField);
 	static FMovieSceneBlendTypeFieldIterator End(FMovieSceneBlendTypeField InField);
 
-	FORCEINLINE explicit operator bool() const 	{ return Offset >= 0 && Offset <= MaxValidOffset(); }
-	FORCEINLINE bool operator!() const 			{ return Offset < 0 || Offset > MaxValidOffset(); }
-	FORCEINLINE void operator++() 				{ IterateToNext(); }
-	FORCEINLINE bool operator==(const FMovieSceneBlendTypeFieldIterator& RHS) const { return Field == RHS.Field && Offset == RHS.Offset; }
-	FORCEINLINE bool operator!=(const FMovieSceneBlendTypeFieldIterator& RHS) const { return Field != RHS.Field || Offset != RHS.Offset; }
+	inline explicit operator bool() const 	{ return Offset >= 0 && Offset <= MaxValidOffset(); }
+	inline bool operator!() const 			{ return Offset < 0 || Offset > MaxValidOffset(); }
+	inline void operator++() 				{ IterateToNext(); }
+	inline bool operator==(const FMovieSceneBlendTypeFieldIterator& RHS) const { return Field == RHS.Field && Offset == RHS.Offset; }
+	inline bool operator!=(const FMovieSceneBlendTypeFieldIterator& RHS) const { return Field != RHS.Field || Offset != RHS.Offset; }
 
 	MOVIESCENE_API EMovieSceneBlendType operator*();
 
@@ -177,7 +177,7 @@ private:
 
 	MOVIESCENE_API void IterateToNext();
 
-	FORCEINLINE static int8 MaxValidOffset() { return (int8)(StaticEnum<EMovieSceneBlendType>()->NumEnums() - 2); }
+	inline static int8 MaxValidOffset() { return (int8)(StaticEnum<EMovieSceneBlendType>()->NumEnums() - 2); }
 
 	FMovieSceneBlendTypeField Field;
 	int8 Offset;
diff --git a/Engine/Source/Runtime/MovieScene/Public/Evaluation/MovieSceneEvalTemplateBase.h b/Engine/Source/Runtime/MovieScene/Public/Evaluation/MovieSceneEvalTemplateBase.h
--- a/Engine/Source/Runtime/MovieScene/Public/Evaluation/MovieSceneEvalTemplateBase.h
+++ b/Engine/Source/Runtime/MovieScene/Public/Evaluation/MovieSceneEvalTemplateBase.h
@@ -50,12 +50,12 @@ public:
 	/**
 	 * Check whether this entity requires set up when it is first evaluated
 	 */
-	FORCEINLINE bool RequiresSetup() const { return (OverrideMask & RequiresSetupFlag) != 0; }
+	inline bool RequiresSetup() const { return (OverrideMask & RequiresSetupFlag) != 0; }
 
 	/**
 	 * Check whether this entity requires tear up when it no longer being evaluated
 	 */
-	FORCEINLINE bool RequiresTearDown() const { return (OverrideMask & RequiresTearDownFlag) != 0; }
+	inline bool RequiresTearDown() const { return (OverrideMask & RequiresTearDownFlag) != 0; }
 
 	/**
 	 * Called before this template is evaluated for the first time, or since OnEndEvaluation has been called
@@ -63,7 +63,7 @@ public:
 	 * @param PersistentData		Persistent data proxy that may contain data pertaining to this entity
 	 * @param Player				The player that is responsible for playing back this template
 	 */
-	FORCEINLINE void OnBeginEvaluation(FPersistentEvaluationData& PersistentData, IMovieScenePlayer& Player) const
+	inline void OnBeginEvaluation(FPersistentEvaluationData& PersistentData, IMovieScenePlayer& Player) const
 	{
 		if (RequiresSetup())
 		{
@@ -77,7 +77,7 @@ public:
 	 * @param PersistentData		Persistent data proxy that may contain data pertaining to this entity
 	 * @param Player				The player that is responsible for playing back this template
 	 */
-	FORCEINLINE void OnEndEvaluation(FPersistentEvaluationData& PersistentData, IMovieScenePlayer& Player) const
+	inline void OnEndEvaluation(FPersistentEvaluationData& PersistentData, IMovieScenePlayer& Player) const
 	{
 		if (RequiresTearDown())
 		{
diff --git a/Engine/Source/Runtime/MovieScene/Public/Evaluation/MovieSceneEvaluationField.h b/Engine/Source/Runtime/MovieScene/Public/Evaluation/MovieSceneEvaluationField.h
--- a/Engine/Source/Runtime/MovieScene/Public/Evaluation/MovieSceneEvaluationField.h
+++ b/Engine/Source/Runtime/MovieScene/Public/Evaluation/MovieSceneEvaluationField.h
@@ -163,9 +163,9 @@ struct FMovieSceneEvaluationFieldEntityQuery
 
 struct FMovieSceneEvaluationFieldEntityKeyFuncs : BaseKeyFuncs<FMovieSceneEvaluationFieldEntityQuery,FMovieSceneEvaluationFieldEntityKey,false>
 {
-	static FORCEINLINE FMovieSceneEvaluationFieldEntityKey GetSetKey(const FMovieSceneEvaluationFieldEntityQuery& Element)        { return Element.Entity.Key; }
-	static FORCEINLINE bool Matches(const FMovieSceneEvaluationFieldEntityKey& A, const FMovieSceneEvaluationFieldEntityKey& B)   { return A == B;      }
-	static FORCEINLINE uint32 GetKeyHash(const FMovieSceneEvaluationFieldEntityKey& Key)                                          { return GetTypeHash(Key); }
+	static inline FMovieSceneEvaluationFieldEntityKey GetSetKey(const FMovieSceneEvaluationFieldEntityQuery& Element)        { return Element.Entity.Key; }
+	static inline bool Matches(const FMovieSceneEvaluationFieldEntityKey& A, const FMovieSceneEvaluationFieldEntityKey& B)   { return A == B;      }
+	static inline uint32 GetKeyHash(const FMovieSceneEvaluationFieldEntityKey& Key)                                          { return GetTypeHash(Key); }
 };
 using FMovieSceneEvaluationFieldEntitySet = TSet<FMovieSceneEvaluationFieldEntityQuery, FMovieSceneEvaluationFieldEntityKeyFuncs>;
 
diff --git a/Engine/Source/Runtime/MovieScene/Public/Evaluation/MovieSceneEvaluationState.h b/Engine/Source/Runtime/MovieScene/Public/Evaluation/MovieSceneEvaluationState.h
--- a/Engine/Source/Runtime/MovieScene/Public/Evaluation/MovieSceneEvaluationState.h
+++ b/Engine/Source/Runtime/MovieScene/Public/Evaluation/MovieSceneEvaluationState.h
@@ -265,15 +265,15 @@ private:
 		typedef typename TTypeTraits<FGuid>::ConstPointerType KeyInitType;
 		typedef const TPairInitializer<typename TTypeTraits<FGuid>::ConstInitType, typename TTypeTraits<ValueType>::ConstInitType>& ElementInitType;
 
-		static FORCEINLINE KeyInitType GetSetKey(ElementInitType Element)
+		static inline KeyInitType GetSetKey(ElementInitType Element)
 		{
 			return Element.Key;
 		}
-		static FORCEINLINE bool Matches(KeyInitType A,KeyInitType B)
+		static inline bool Matches(KeyInitType A,KeyInitType B)
 		{
 			return A == B;
 		}
-		static FORCEINLINE uint32 GetKeyHash(KeyInitType Key)
+		static inline uint32 GetKeyHash(KeyInitType Key)
 		{
 			return Key.A ^ Key.B ^ Key.C ^ Key.D;
 		}
@@ -370,7 +370,7 @@ struct FMovieSceneEvaluationState : public UE::MovieScene::IPlaybackCapability
 	 *
 	 * @param InSequenceID		The sequence ID to lookup
 	 */
-	FORCEINLINE FMovieSceneObjectCache* FindObjectCache(FMovieSceneSequenceIDRef SequenceID)
+	inline FMovieSceneObjectCache* FindObjectCache(FMovieSceneSequenceIDRef SequenceID)
 	{
 		if (FVersionedObjectCache* Cache = ObjectCaches.Find(SequenceID))
 		{
@@ -384,7 +384,7 @@ struct FMovieSceneEvaluationState : public UE::MovieScene::IPlaybackCapability
 	 *
 	 * @param InSequenceID		The sequence ID to lookup
 	 */
-	FORCEINLINE const FMovieSceneObjectCache* FindObjectCache(FMovieSceneSequenceIDRef SequenceID) const
+	inline const FMovieSceneObjectCache* FindObjectCache(FMovieSceneSequenceIDRef SequenceID) const
 	{
 		if (const FVersionedObjectCache* Cache = ObjectCaches.Find(SequenceID))
 		{
@@ -398,7 +398,7 @@ struct FMovieSceneEvaluationState : public UE::MovieScene::IPlaybackCapability
 	 *
 	 * @param InSequenceID		The sequence ID to lookup
 	 */
-	FORCEINLINE FMovieSceneObjectCache& GetObjectCache(FMovieSceneSequenceIDRef SequenceID)
+	inline FMovieSceneObjectCache& GetObjectCache(FMovieSceneSequenceIDRef SequenceID)
 	{
 		FVersionedObjectCache* Cache = ObjectCaches.Find(SequenceID);
 		if (!Cache)
diff --git a/Engine/Source/Runtime/MovieScene/Public/Evaluation/MovieSceneEvaluationTrack.h b/Engine/Source/Runtime/MovieScene/Public/Evaluation/MovieSceneEvaluationTrack.h
--- a/Engine/Source/Runtime/MovieScene/Public/Evaluation/MovieSceneEvaluationTrack.h
+++ b/Engine/Source/Runtime/MovieScene/Public/Evaluation/MovieSceneEvaluationTrack.h
@@ -76,7 +76,7 @@ public:
 	/**
 	 * Get the object binding ID that this track belongs to
 	 */
-	FORCEINLINE const FGuid& GetObjectBindingID() const
+	inline const FGuid& GetObjectBindingID() const
 	{
 		return ObjectBindingID;
 	}
@@ -86,7 +86,7 @@ public:
 	 * NOTE that this is intended for use during the compilation phase in-editor.
 	 * Beware of using this to modify templates afterwards as it will almost certainly break evaluation.
 	 */
-	FORCEINLINE TArrayView<FMovieSceneEvalTemplatePtr> GetChildTemplates()
+	inline TArrayView<FMovieSceneEvalTemplatePtr> GetChildTemplates()
 	{
 		return TArrayView<FMovieSceneEvalTemplatePtr>(ChildTemplates.GetData(), ChildTemplates.Num());
 	}
@@ -94,7 +94,7 @@ public:
 	/**
 	 * Get this track's child templates
 	 */
-	FORCEINLINE TArrayView<const FMovieSceneEvalTemplatePtr> GetChildTemplates() const
+	inline TArrayView<const FMovieSceneEvalTemplatePtr> GetChildTemplates() const
 	{
 		return TArrayView<const FMovieSceneEvalTemplatePtr>(ChildTemplates.GetData(), ChildTemplates.Num());
 	}
@@ -102,7 +102,7 @@ public:
 	/**
 	 * Get the template from the given template index
 	 */
-	FORCEINLINE const FMovieSceneEvalTemplate& GetChildTemplate(int32 TemplateIndex) const
+	inline const FMovieSceneEvalTemplate& GetChildTemplate(int32 TemplateIndex) const
 	{
 		return *ChildTemplates[TemplateIndex];
 	}
@@ -110,7 +110,7 @@ public:
 	/**
 	 * Check whether we have a valid child template for the specified index
 	 */
-	FORCEINLINE bool HasChildTemplate(int32 TemplateIndex) const
+	inline bool HasChildTemplate(int32 TemplateIndex) const
 	{
 		return ChildTemplates.IsValidIndex(TemplateIndex) && ChildTemplates[TemplateIndex].IsValid();
 	}
@@ -365,7 +365,7 @@ public:
 	 * @param PersistentData		Persistent data proxy that may contain data pertaining to this entity
 	 * @param Player				The player that is responsible for playing back this template
 	 */
-	FORCEINLINE void OnBeginEvaluation(FPersistentEvaluationData& PersistentData, IMovieScenePlayer& Player) const
+	inline void OnBeginEvaluation(FPersistentEvaluationData& PersistentData, IMovieScenePlayer& Player) const
 	{
 		if (!TrackTemplate.IsValid())
 		{
diff --git a/Engine/Source/Runtime/MovieScene/Public/Evaluation/MovieSceneExecutionTokens.h b/Engine/Source/Runtime/MovieScene/Public/Evaluation/MovieSceneExecutionTokens.h
--- a/Engine/Source/Runtime/MovieScene/Public/Evaluation/MovieSceneExecutionTokens.h
+++ b/Engine/Source/Runtime/MovieScene/Public/Evaluation/MovieSceneExecutionTokens.h
@@ -48,7 +48,7 @@ struct FMovieSceneExecutionTokens
 	 * Add a new IMovieSceneExecutionToken derived token to the stack
 	 */
 	template<typename T>
-	FORCEINLINE typename TEnableIf<TPointerIsConvertibleFromTo<typename TRemoveReference<T>::Type, const IMovieSceneExecutionToken>::Value>::Type Add(T&& InToken)
+	inline typename TEnableIf<TPointerIsConvertibleFromTo<typename TRemoveReference<T>::Type, const IMovieSceneExecutionToken>::Value>::Type Add(T&& InToken)
 	{
 		check(Scope.Key.IsValid() && Operand.IsValid());
 		OrderedTokens.Add(FEntry(Operand, Scope, Context, Forward<T>(InToken)));
@@ -58,7 +58,7 @@ struct FMovieSceneExecutionTokens
 	* Add a new shared execution token to the sorted array of tokens
 	*/
 	template<typename T>
-	FORCEINLINE typename TEnableIf<TPointerIsConvertibleFromTo<typename TRemoveReference<T>::Type, const IMovieSceneSharedExecutionToken>::Value>::Type AddShared(FMovieSceneSharedDataId ID, T&& InToken)
+	inline typename TEnableIf<TPointerIsConvertibleFromTo<typename TRemoveReference<T>::Type, const IMovieSceneSharedExecutionToken>::Value>::Type AddShared(FMovieSceneSharedDataId ID, T&& InToken)
 	{
 		checkf(!SharedTokens.Contains(ID), TEXT("Already added a shared token of this type"));
 		SharedTokens.Add(ID, MoveTemp(InToken));
@@ -116,7 +116,7 @@ public:
 	/**
 	 * Internal: Set the operand we're currently operating on
 	 */
-	FORCEINLINE void SetOperand(const FMovieSceneEvaluationOperand& InOperand)
+	inline void SetOperand(const FMovieSceneEvaluationOperand& InOperand)
 	{
 		Operand = InOperand;
 	}
@@ -124,7 +124,7 @@ public:
 	/**
 	 * Internal: Set the current scope
 	 */
-	FORCEINLINE void SetCurrentScope(const FMovieSceneEvaluationScope& InScope)
+	inline void SetCurrentScope(const FMovieSceneEvaluationScope& InScope)
 	{
 		Scope = InScope;
 	}
@@ -132,7 +132,7 @@ public:
 	/**
 	 * Internal: Set the current context
 	 */
-	FORCEINLINE void SetContext(const FMovieSceneContext& InContext)
+	inline void SetContext(const FMovieSceneContext& InContext)
 	{
 		Context = InContext;
 	}
@@ -140,7 +140,7 @@ public:
 	/**
 	 * Get the current evaluation scope
 	 */
-	FORCEINLINE FMovieSceneEvaluationScope GetCurrentScope() const
+	inline FMovieSceneEvaluationScope GetCurrentScope() const
 	{
 		return Scope;
 	}
diff --git a/Engine/Source/Runtime/MovieScene/Public/Evaluation/MovieScenePlayback.h b/Engine/Source/Runtime/MovieScene/Public/Evaluation/MovieScenePlayback.h
--- a/Engine/Source/Runtime/MovieScene/Public/Evaluation/MovieScenePlayback.h
+++ b/Engine/Source/Runtime/MovieScene/Public/Evaluation/MovieScenePlayback.h
@@ -59,7 +59,7 @@ struct FMovieSceneEvaluationRange
 	/**
 	 * Get the range that we should be evaluating
 	 */
-	FORCEINLINE TRange<FFrameTime> GetRange() const
+	inline TRange<FFrameTime> GetRange() const
 	{
 		return EvaluationRange;
 	}
@@ -67,7 +67,7 @@ struct FMovieSceneEvaluationRange
 	/**
 	 * Get the range of frame numbers traversed over this evaluation range, not including partial frames
 	 */
-	FORCEINLINE TRange<FFrameNumber> GetFrameNumberRange() const
+	inline TRange<FFrameNumber> GetFrameNumberRange() const
 	{
 		return TimeRangeToNumberRange(EvaluationRange);
 	}
@@ -81,7 +81,7 @@ struct FMovieSceneEvaluationRange
 	/**
 	 * Get the direction to evaluate our range
 	 */
-	FORCEINLINE EPlayDirection GetDirection() const 
+	inline EPlayDirection GetDirection() const 
 	{
 		return Direction;
 	}
@@ -89,7 +89,7 @@ struct FMovieSceneEvaluationRange
 	/**
 	 * Get the current time of evaluation.
 	 */
-	FORCEINLINE FFrameTime GetTime() const
+	inline FFrameTime GetTime() const
 	{
 		if (TimeOverride != TNumericLimits<int32>::Lowest())
 		{
@@ -103,7 +103,7 @@ struct FMovieSceneEvaluationRange
 	 * Get the current time to use for looking up within an evaluation field.
 	 * Subtly different from GetTime in that it returns the previous tick for exclusive boundaries
 	 */
-	FORCEINLINE FFrameNumber GetEvaluationFieldTime() const
+	inline FFrameNumber GetEvaluationFieldTime() const
 	{
 		TRange<FFrameNumber> Range = GetFrameNumberRange();
 
@@ -120,7 +120,7 @@ struct FMovieSceneEvaluationRange
 	/**
 	 * Get the absolute amount of time that has passed since the last update (will always be >= 0)
 	 */
-	FORCEINLINE FFrameTime GetDelta() const
+	inline FFrameTime GetDelta() const
 	{
 		return EvaluationRange.Size<FFrameTime>();
 	}
@@ -128,7 +128,7 @@ struct FMovieSceneEvaluationRange
 	/**
 	 * Get the previous time of evaluation. Should not generally be used. Prefer GetRange instead.
 	 */
-	FORCEINLINE FFrameTime GetPreviousTime() const
+	inline FFrameTime GetPreviousTime() const
 	{
 		return Direction == EPlayDirection::Forwards ? EvaluationRange.GetLowerBoundValue() : EvaluationRange.GetUpperBoundValue();
 	}
@@ -136,7 +136,7 @@ struct FMovieSceneEvaluationRange
 	/**
 	 * Get the current time offset by the specified amount in the direction of play
 	 */
-	FORCEINLINE FFrameTime GetOffsetTime(FFrameTime InOffset) const
+	inline FFrameTime GetOffsetTime(FFrameTime InOffset) const
 	{
 		FFrameTime Now = GetTime();
 		return Direction == EPlayDirection::Forwards ? Now + InOffset : Now - InOffset;
@@ -145,7 +145,7 @@ struct FMovieSceneEvaluationRange
 	/**
 	 * Override the time that we're actually evaluating at
 	 */
-	FORCEINLINE void OverrideTime(FFrameNumber InTimeOverride)
+	inline void OverrideTime(FFrameNumber InTimeOverride)
 	{
 		TimeOverride = InTimeOverride;
 	}
@@ -154,7 +154,7 @@ struct FMovieSceneEvaluationRange
 	 * Get the framerate that this context's times are in
 	 * @return The framerate that all times are relative to
 	 */
-	FORCEINLINE FFrameRate GetFrameRate() const
+	inline FFrameRate GetFrameRate() const
 	{
 		return CurrentFrameRate;
 	}
@@ -233,7 +233,7 @@ struct FMovieSceneContext : FMovieSceneEvaluationRange
 	/**
 	 * Get the playback status
 	 */
-	FORCEINLINE EMovieScenePlayerStatus::Type GetStatus() const
+	inline EMovieScenePlayerStatus::Type GetStatus() const
 	{
 		return Status;
 	}
@@ -241,7 +241,7 @@ struct FMovieSceneContext : FMovieSceneEvaluationRange
 	/**
 	 * Check whether we've just jumped to a different time
 	 */
-	FORCEINLINE bool HasJumped() const
+	inline bool HasJumped() const
 	{
 		return bHasJumped;
 	}
@@ -249,7 +249,7 @@ struct FMovieSceneContext : FMovieSceneEvaluationRange
 	/**
 	 * Check whether we're evaluating in silent mode (no audio or mutating eval)
 	 */
-	FORCEINLINE bool IsSilent() const
+	inline bool IsSilent() const
 	{
 		return bSilent;
 	}
@@ -257,7 +257,7 @@ struct FMovieSceneContext : FMovieSceneEvaluationRange
 	/**
 	 * Get the current root to sequence transform for the current sub sequence
 	 */
-	FORCEINLINE const FMovieSceneSequenceTransform& GetRootToSequenceTransform() const
+	inline const FMovieSceneSequenceTransform& GetRootToSequenceTransform() const
 	{
 		return RootToSequenceTransform;
 	}
@@ -265,7 +265,7 @@ struct FMovieSceneContext : FMovieSceneEvaluationRange
 	/**
 	 * Get the current root to sequence warp counter for the current sub sequence
 	 */
-	FORCEINLINE const FMovieSceneWarpCounter& GetRootToSequenceWarpCounter() const
+	inline const FMovieSceneWarpCounter& GetRootToSequenceWarpCounter() const
 	{
 		return RootToSequenceWarpCounter;
 	}
@@ -287,7 +287,7 @@ struct FMovieSceneContext : FMovieSceneEvaluationRange
 	 * @param bInLeadingRegion			Whether we are considered to be in the section's leading (aka preroll) region
 	 * @param bInTrailingRegion			Whether we are considered to be in the section's trailing (aka postroll) region
 	 */
-	FORCEINLINE void ApplySectionPrePostRoll(bool bInLeadingRegion, bool bInTrailingRegion)
+	inline void ApplySectionPrePostRoll(bool bInLeadingRegion, bool bInTrailingRegion)
 	{
 		if (Direction == EPlayDirection::Forwards)
 		{
@@ -490,7 +490,7 @@ struct FMovieScenePlaybackPosition
 	/**
 	 * @return Whether we are evaluating with sub frames, or frame-locked
 	 */
-	FORCEINLINE EMovieSceneEvaluationType GetEvaluationType() const
+	inline EMovieSceneEvaluationType GetEvaluationType() const
 	{
 		return EvaluationType;
 	}
@@ -498,7 +498,7 @@ struct FMovieScenePlaybackPosition
 	/**
 	 * @return The input frame rate that all frame times provided to this class will be interpreted as
 	 */
-	FORCEINLINE FFrameRate GetInputRate() const
+	inline FFrameRate GetInputRate() const
 	{
 		return InputRate;
 	}
@@ -506,7 +506,7 @@ struct FMovieScenePlaybackPosition
 	/**
 	 * @return The output frame rate that all frame times returned from this class will be interpreted as
 	 */
-	FORCEINLINE FFrameRate GetOutputRate() const
+	inline FFrameRate GetOutputRate() const
 	{
 		return OutputRate;
 	}
diff --git a/Engine/Source/Runtime/MovieScene/Public/Evaluation/MovieSceneRootOverridePath.h b/Engine/Source/Runtime/MovieScene/Public/Evaluation/MovieSceneRootOverridePath.h
--- a/Engine/Source/Runtime/MovieScene/Public/Evaluation/MovieSceneRootOverridePath.h
+++ b/Engine/Source/Runtime/MovieScene/Public/Evaluation/MovieSceneRootOverridePath.h
@@ -64,7 +64,7 @@ struct FSubSequencePath
 	 * @param SequenceID			The sequence ID to find a template for
 	 * @return Pointer to a template instance, or nullptr if the ID was not found
 	 */
-	/*FORCEINLINE_DEBUGGABLE*/ FMovieSceneSequenceID ResolveChildSequenceID(FMovieSceneSequenceID SequenceID) const
+	/*inline*/ FMovieSceneSequenceID ResolveChildSequenceID(FMovieSceneSequenceID SequenceID) const
 	{
 		if (LIKELY(PathToRoot.Num() == 0))
 		{
diff --git a/Engine/Source/Runtime/MovieScene/Public/Evaluation/MovieSceneSegment.h b/Engine/Source/Runtime/MovieScene/Public/Evaluation/MovieSceneSegment.h
--- a/Engine/Source/Runtime/MovieScene/Public/Evaluation/MovieSceneSegment.h
+++ b/Engine/Source/Runtime/MovieScene/Public/Evaluation/MovieSceneSegment.h
@@ -112,10 +112,10 @@ struct FSectionEvaluationData
 	}
 
 	/** Check if this is a preroll eval */
-	FORCEINLINE bool IsPreRoll() const { return (Flags & ESectionEvaluationFlags::PreRoll) != ESectionEvaluationFlags::None; }
+	inline bool IsPreRoll() const { return (Flags & ESectionEvaluationFlags::PreRoll) != ESectionEvaluationFlags::None; }
 
 	/** Check if this is a postroll eval */
-	FORCEINLINE bool IsPostRoll() const { return (Flags & ESectionEvaluationFlags::PostRoll) != ESectionEvaluationFlags::None; }
+	inline bool IsPostRoll() const { return (Flags & ESectionEvaluationFlags::PostRoll) != ESectionEvaluationFlags::None; }
 
 	friend FArchive& operator<<(FArchive& Ar, FSectionEvaluationData& Data)
 	{
diff --git a/Engine/Source/Runtime/MovieScene/Public/Evaluation/MovieSceneTrackIdentifier.h b/Engine/Source/Runtime/MovieScene/Public/Evaluation/MovieSceneTrackIdentifier.h
--- a/Engine/Source/Runtime/MovieScene/Public/Evaluation/MovieSceneTrackIdentifier.h
+++ b/Engine/Source/Runtime/MovieScene/Public/Evaluation/MovieSceneTrackIdentifier.h
@@ -29,27 +29,27 @@ struct FMovieSceneTrackIdentifier
 		return *this;
 	}
 
-	FORCEINLINE friend bool operator==(FMovieSceneTrackIdentifier A, FMovieSceneTrackIdentifier B)
+	inline friend bool operator==(FMovieSceneTrackIdentifier A, FMovieSceneTrackIdentifier B)
 	{
 		return A.Value == B.Value;
 	}
 
-	FORCEINLINE friend bool operator!=(FMovieSceneTrackIdentifier A, FMovieSceneTrackIdentifier B)
+	inline friend bool operator!=(FMovieSceneTrackIdentifier A, FMovieSceneTrackIdentifier B)
 	{
 		return A.Value != B.Value;
 	}
 
-	FORCEINLINE friend bool operator<(FMovieSceneTrackIdentifier A, FMovieSceneTrackIdentifier B)
+	inline friend bool operator<(FMovieSceneTrackIdentifier A, FMovieSceneTrackIdentifier B)
 	{
 		return A.Value < B.Value;
 	}
 
-	FORCEINLINE friend bool operator>(FMovieSceneTrackIdentifier A, FMovieSceneTrackIdentifier B)
+	inline friend bool operator>(FMovieSceneTrackIdentifier A, FMovieSceneTrackIdentifier B)
 	{
 		return A.Value > B.Value;
 	}
 
-	FORCEINLINE friend uint32 GetTypeHash(FMovieSceneTrackIdentifier In)
+	inline friend uint32 GetTypeHash(FMovieSceneTrackIdentifier In)
 	{
 		return In.Value;
 	}
diff --git a/Engine/Source/Runtime/MovieScene/Public/IMovieScenePlayer.h b/Engine/Source/Runtime/MovieScene/Public/IMovieScenePlayer.h
--- a/Engine/Source/Runtime/MovieScene/Public/IMovieScenePlayer.h
+++ b/Engine/Source/Runtime/MovieScene/Public/IMovieScenePlayer.h
@@ -318,7 +318,7 @@ public:
 	 * @param InTokenType		Unique marker that identifies the originating token type
 	 * @param InProducer		Producer implementation that defines how to create the preanimated token, if it doesn't already exist
 	 */
-	FORCEINLINE void SavePreAnimatedState(UObject& InObject, FMovieSceneAnimTypeID InTokenType, const IMovieScenePreAnimatedTokenProducer& InProducer)
+	inline void SavePreAnimatedState(UObject& InObject, FMovieSceneAnimTypeID InTokenType, const IMovieScenePreAnimatedTokenProducer& InProducer)
 	{
 		PreAnimatedState.SavePreAnimatedState(InObject, InTokenType, InProducer);
 	}
@@ -330,7 +330,7 @@ public:
 	 * @param InTokenType		Unique marker that identifies the originating token type
 	 * @param InProducer		Producer implementation that defines how to create the preanimated token, if it doesn't already exist
 	 */
-	FORCEINLINE void SavePreAnimatedState(FMovieSceneAnimTypeID InTokenType, const IMovieScenePreAnimatedGlobalTokenProducer& InProducer)
+	inline void SavePreAnimatedState(FMovieSceneAnimTypeID InTokenType, const IMovieScenePreAnimatedGlobalTokenProducer& InProducer)
 	{
 		PreAnimatedState.SavePreAnimatedState(InTokenType, InProducer);
 	}
diff --git a/Engine/Source/Runtime/MovieScene/Public/MovieScene.h b/Engine/Source/Runtime/MovieScene/Public/MovieScene.h
--- a/Engine/Source/Runtime/MovieScene/Public/MovieScene.h
+++ b/Engine/Source/Runtime/MovieScene/Public/MovieScene.h
@@ -215,11 +215,11 @@ public:
 	 * We only need these for being stored in a container, to check if it's the same object.
 	 * Not intended for direct use.
 	 */
-	FORCEINLINE bool operator==(const FMovieSceneSectionGroup& Other) const
+	inline bool operator==(const FMovieSceneSectionGroup& Other) const
 	{
 		return &Sections == &Other.Sections;
 	}
-	FORCEINLINE bool operator!=(const FMovieSceneSectionGroup& Other) const
+	inline bool operator!=(const FMovieSceneSectionGroup& Other) const
 	{
 		return &Sections != &Other.Sections;
 	}
@@ -228,10 +228,10 @@ public:
 	 * DO NOT USE DIRECTLY
 	 * STL-like iterators to enable range-based for loop support.
 	 */
-	FORCEINLINE TArray<TWeakObjectPtr<UMovieSceneSection> >::RangedForIteratorType      begin() { return Sections.begin(); }
-	FORCEINLINE TArray<TWeakObjectPtr<UMovieSceneSection> >::RangedForConstIteratorType begin() const { return Sections.begin(); }
-	FORCEINLINE TArray<TWeakObjectPtr<UMovieSceneSection> >::RangedForIteratorType      end() { return Sections.end(); }
-	FORCEINLINE TArray<TWeakObjectPtr<UMovieSceneSection> >::RangedForConstIteratorType end() const { return Sections.end(); }
+	inline TArray<TWeakObjectPtr<UMovieSceneSection> >::RangedForIteratorType      begin() { return Sections.begin(); }
+	inline TArray<TWeakObjectPtr<UMovieSceneSection> >::RangedForConstIteratorType begin() const { return Sections.begin(); }
+	inline TArray<TWeakObjectPtr<UMovieSceneSection> >::RangedForIteratorType      end() { return Sections.end(); }
+	inline TArray<TWeakObjectPtr<UMovieSceneSection> >::RangedForConstIteratorType end() const { return Sections.end(); }
 };
 
 /**
@@ -282,11 +282,11 @@ private:
 	 * We only need these for being stored in a container, to check if it's the same object.
 	 * Not intended for direct use.
 	 */
-	FORCEINLINE bool operator==(const UMovieSceneNodeGroup& Other) const
+	inline bool operator==(const UMovieSceneNodeGroup& Other) const
 	{
 		return &Nodes == &Other.Nodes;
 	}
-	FORCEINLINE bool operator!=(const UMovieSceneNodeGroup& Other) const
+	inline bool operator!=(const UMovieSceneNodeGroup& Other) const
 	{
 		return &Nodes != &Other.Nodes;
 	}
@@ -342,10 +342,10 @@ public:
 	 * DO NOT USE DIRECTLY
 	 * STL-like iterators to enable range-based for loop support.
 	 */
-	FORCEINLINE auto begin()	{ return NodeGroups.begin(); }
-	FORCEINLINE auto begin()	const { return NodeGroups.begin(); }
-	FORCEINLINE auto end()	{ return NodeGroups.end(); }
-	FORCEINLINE auto end()	const { return NodeGroups.end(); }
+	inline auto begin()	{ return NodeGroups.begin(); }
+	inline auto begin()	const { return NodeGroups.begin(); }
+	inline auto end()	{ return NodeGroups.end(); }
+	inline auto end()	const { return NodeGroups.end(); }
 #endif
 };
 
diff --git a/Engine/Source/Runtime/MovieScene/Public/MovieSceneBinding.h b/Engine/Source/Runtime/MovieScene/Public/MovieSceneBinding.h
--- a/Engine/Source/Runtime/MovieScene/Public/MovieSceneBinding.h
+++ b/Engine/Source/Runtime/MovieScene/Public/MovieSceneBinding.h
@@ -133,9 +133,9 @@ struct FMovieSceneBinding
 	MOVIESCENE_API void SetTracks(TArray<UMovieSceneTrack*>&& InTracks, UMovieScene* Owner);
 
 	/* For sorts so we can search quickly by Guid */
-	FORCEINLINE bool operator<(const FMovieSceneBinding& RHS) const { return ObjectGuid < RHS.ObjectGuid; }
-	FORCEINLINE bool operator<(const FGuid& InGuid) const { return ObjectGuid < InGuid; }
-	FORCEINLINE friend bool operator<(const FGuid& InGuid, const FMovieSceneBinding& RHS) { return InGuid < RHS.GetObjectGuid(); }
+	inline bool operator<(const FMovieSceneBinding& RHS) const { return ObjectGuid < RHS.ObjectGuid; }
+	inline bool operator<(const FGuid& InGuid) const { return ObjectGuid < InGuid; }
+	inline friend bool operator<(const FGuid& InGuid, const FMovieSceneBinding& RHS) { return InGuid < RHS.GetObjectGuid(); }
 
 #if WITH_EDITORONLY_DATA
 	/**
diff --git a/Engine/Source/Runtime/MovieScene/Public/MovieSceneBindingProxy.h b/Engine/Source/Runtime/MovieScene/Public/MovieSceneBindingProxy.h
--- a/Engine/Source/Runtime/MovieScene/Public/MovieSceneBindingProxy.h
+++ b/Engine/Source/Runtime/MovieScene/Public/MovieSceneBindingProxy.h
@@ -32,17 +32,17 @@ struct FMovieSceneBindingProxy
 		, Sequence(InSequence)
 	{}
 
-	FORCEINLINE friend bool operator==(const FMovieSceneBindingProxy &LHS, const FMovieSceneBindingProxy &RHS)
+	inline friend bool operator==(const FMovieSceneBindingProxy &LHS, const FMovieSceneBindingProxy &RHS)
 	{
 		return LHS.BindingID == RHS.BindingID && LHS.Sequence == RHS.Sequence;
 	}
 
-	FORCEINLINE friend bool operator!=(const FMovieSceneBindingProxy &LHS, const FMovieSceneBindingProxy &RHS)
+	inline friend bool operator!=(const FMovieSceneBindingProxy &LHS, const FMovieSceneBindingProxy &RHS)
 	{
 		return LHS.BindingID != RHS.BindingID || LHS.Sequence != RHS.Sequence;
 	}
 
-	FORCEINLINE friend uint32 GetTypeHash(const FMovieSceneBindingProxy& In)
+	inline friend uint32 GetTypeHash(const FMovieSceneBindingProxy& In)
 	{
 		return HashCombine(GetTypeHash(In.BindingID), GetTypeHash(In.Sequence));
 	}
diff --git a/Engine/Source/Runtime/MovieScene/Public/MovieScenePossessable.h b/Engine/Source/Runtime/MovieScene/Public/MovieScenePossessable.h
--- a/Engine/Source/Runtime/MovieScene/Public/MovieScenePossessable.h
+++ b/Engine/Source/Runtime/MovieScene/Public/MovieScenePossessable.h
@@ -192,9 +192,9 @@ public:
 	MOVIESCENE_API bool BindSpawnableObject(FMovieSceneSequenceID SequenceID, UObject* Object, IMovieScenePlayer* Player);
 
 	/* For sorts so we can search quickly by Guid */
-	FORCEINLINE bool operator<(const FMovieScenePossessable& RHS) const { return Guid < RHS.Guid; }
-	FORCEINLINE friend bool operator<(const FGuid& InGuid, const FMovieScenePossessable& RHS) { return InGuid < RHS.GetGuid(); }
-	FORCEINLINE bool operator<(const FGuid& InGuid) const { return Guid < InGuid; }
+	inline bool operator<(const FMovieScenePossessable& RHS) const { return Guid < RHS.Guid; }
+	inline friend bool operator<(const FGuid& InGuid, const FMovieScenePossessable& RHS) { return InGuid < RHS.GetGuid(); }
+	inline bool operator<(const FGuid& InGuid) const { return Guid < InGuid; }
 
 private:
 
diff --git a/Engine/Source/Runtime/MovieScene/Public/MovieSceneSection.h b/Engine/Source/Runtime/MovieScene/Public/MovieSceneSection.h
--- a/Engine/Source/Runtime/MovieScene/Public/MovieSceneSection.h
+++ b/Engine/Source/Runtime/MovieScene/Public/MovieSceneSection.h
@@ -190,11 +190,11 @@ struct FMovieSceneTimecodeSource
 		: Timecode(FTimecode())
 	{}
 
-	FORCEINLINE bool operator==(const FMovieSceneTimecodeSource& Other) const
+	inline bool operator==(const FMovieSceneTimecodeSource& Other) const
 	{
 		return Timecode == Other.Timecode;
 	}
-	FORCEINLINE bool operator!=(const FMovieSceneTimecodeSource& Other) const
+	inline bool operator!=(const FMovieSceneTimecodeSource& Other) const
 	{
 		return Timecode != Other.Timecode;
 	}
diff --git a/Engine/Source/Runtime/MovieScene/Public/MovieSceneSequenceID.h b/Engine/Source/Runtime/MovieScene/Public/MovieSceneSequenceID.h
--- a/Engine/Source/Runtime/MovieScene/Public/MovieSceneSequenceID.h
+++ b/Engine/Source/Runtime/MovieScene/Public/MovieSceneSequenceID.h
@@ -21,54 +21,54 @@ struct FMovieSceneSequenceID
 		: Value(InValue)
 	{}
 
-	FORCEINLINE friend bool operator==(FMovieSceneSequenceID LHS, FMovieSceneSequenceID RHS)
+	inline friend bool operator==(FMovieSceneSequenceID LHS, FMovieSceneSequenceID RHS)
 	{
 		return LHS.Value == RHS.Value;
 	}
 	
-	FORCEINLINE friend bool operator!=(FMovieSceneSequenceID LHS, FMovieSceneSequenceID RHS)
+	inline friend bool operator!=(FMovieSceneSequenceID LHS, FMovieSceneSequenceID RHS)
 	{
 		return LHS.Value != RHS.Value;
 	}
 
-	FORCEINLINE friend bool operator<(FMovieSceneSequenceID LHS, FMovieSceneSequenceID RHS)
+	inline friend bool operator<(FMovieSceneSequenceID LHS, FMovieSceneSequenceID RHS)
 	{
 		return LHS.Value < RHS.Value;
 	}
 
-	FORCEINLINE friend bool operator>(FMovieSceneSequenceID LHS, FMovieSceneSequenceID RHS)
+	inline friend bool operator>(FMovieSceneSequenceID LHS, FMovieSceneSequenceID RHS)
 	{
 		return LHS.Value > RHS.Value;
 	}
 
-	FORCEINLINE friend uint32 GetTypeHash(FMovieSceneSequenceID In)
+	inline friend uint32 GetTypeHash(FMovieSceneSequenceID In)
 	{
 		return GetTypeHash(In.Value);
 	}
 
-	FORCEINLINE FMovieSceneSequenceID AccumulateParentID(FMovieSceneSequenceID InParentID) const
+	inline FMovieSceneSequenceID AccumulateParentID(FMovieSceneSequenceID InParentID) const
 	{
 		return Value == 0 ? InParentID : FMovieSceneSequenceID(HashCombine(Value, InParentID.Value));
 	}
 
-	FORCEINLINE bool Serialize(FArchive& Ar)
+	inline bool Serialize(FArchive& Ar)
 	{
 		Ar << Value;
 		return true;
 	}
 
-	FORCEINLINE friend FArchive& operator<<(FArchive& Ar, FMovieSceneSequenceID& SequenceID)
+	inline friend FArchive& operator<<(FArchive& Ar, FMovieSceneSequenceID& SequenceID)
 	{
 		SequenceID.Serialize(Ar);
 		return Ar;
 	}
 
-	FORCEINLINE uint32 GetInternalValue() const
+	inline uint32 GetInternalValue() const
 	{
 		return Value;
 	}
 
-	FORCEINLINE bool IsValid() const
+	inline bool IsValid() const
 	{
 		return Value != -1;
 	}
diff --git a/Engine/Source/Runtime/MovieScene/Public/MovieSceneSpawnable.h b/Engine/Source/Runtime/MovieScene/Public/MovieSceneSpawnable.h
--- a/Engine/Source/Runtime/MovieScene/Public/MovieSceneSpawnable.h
+++ b/Engine/Source/Runtime/MovieScene/Public/MovieSceneSpawnable.h
@@ -246,9 +246,9 @@ public:
 	MOVIESCENE_API void AutoSetNetAddressableName();
 
 	/* For sorts and BinarySearch so we can search quickly by Guid */
-	FORCEINLINE bool operator<(const FMovieSceneSpawnable& RHS) const { return Guid < RHS.Guid; }
-	FORCEINLINE friend bool operator<(const FGuid& InGuid, const FMovieSceneSpawnable& RHS) { return InGuid < RHS.GetGuid(); }
-	FORCEINLINE bool operator<(const FGuid& InGuid) const { return Guid < InGuid; }
+	inline bool operator<(const FMovieSceneSpawnable& RHS) const { return Guid < RHS.Guid; }
+	inline friend bool operator<(const FGuid& InGuid, const FMovieSceneSpawnable& RHS) { return InGuid < RHS.GetGuid(); }
+	inline bool operator<(const FGuid& InGuid) const { return Guid < InGuid; }
 
 	/** Array of tags that can be used for grouping and categorizing. */
 	UPROPERTY(EditAnywhere, AdvancedDisplay, Category=Actor)
diff --git a/Engine/Source/Runtime/MovieSceneTracks/Public/Channels/MovieSceneCameraShakeSourceTriggerChannel.h b/Engine/Source/Runtime/MovieSceneTracks/Public/Channels/MovieSceneCameraShakeSourceTriggerChannel.h
--- a/Engine/Source/Runtime/MovieSceneTracks/Public/Channels/MovieSceneCameraShakeSourceTriggerChannel.h
+++ b/Engine/Source/Runtime/MovieSceneTracks/Public/Channels/MovieSceneCameraShakeSourceTriggerChannel.h
@@ -65,12 +65,12 @@ struct FMovieSceneCameraShakeSourceTriggerChannel : public FMovieSceneChannel
 
 	typedef FMovieSceneCameraShakeSourceTrigger CurveValueType;
 
-	FORCEINLINE TMovieSceneChannelData<FMovieSceneCameraShakeSourceTrigger> GetData()
+	inline TMovieSceneChannelData<FMovieSceneCameraShakeSourceTrigger> GetData()
 	{
 		return TMovieSceneChannelData<FMovieSceneCameraShakeSourceTrigger>(&KeyTimes, &KeyValues, this, &KeyHandles);
 	}
 
-	FORCEINLINE TMovieSceneChannelData<const FMovieSceneCameraShakeSourceTrigger> GetData() const
+	inline TMovieSceneChannelData<const FMovieSceneCameraShakeSourceTrigger> GetData() const
 	{
 		return TMovieSceneChannelData<const FMovieSceneCameraShakeSourceTrigger>(&KeyTimes, &KeyValues);
 	}
diff --git a/Engine/Source/Runtime/MovieSceneTracks/Public/Channels/MovieSceneEventChannel.h b/Engine/Source/Runtime/MovieSceneTracks/Public/Channels/MovieSceneEventChannel.h
--- a/Engine/Source/Runtime/MovieSceneTracks/Public/Channels/MovieSceneEventChannel.h
+++ b/Engine/Source/Runtime/MovieSceneTracks/Public/Channels/MovieSceneEventChannel.h
@@ -33,7 +33,7 @@ struct FMovieSceneEventChannel : public FMovieSceneChannel
 	 *
 	 * @return An object that is able to manipulate this channel's data
 	 */
-	FORCEINLINE TMovieSceneChannelData<FMovieSceneEvent> GetData()
+	inline TMovieSceneChannelData<FMovieSceneEvent> GetData()
 	{
 		return TMovieSceneChannelData<FMovieSceneEvent>(&KeyTimes, &KeyValues, this, &KeyHandles);
 	}
@@ -43,7 +43,7 @@ struct FMovieSceneEventChannel : public FMovieSceneChannel
 	 *
 	 * @return An object that is able to interrogate this channel's data
 	 */
-	FORCEINLINE TMovieSceneChannelData<const FMovieSceneEvent> GetData() const
+	inline TMovieSceneChannelData<const FMovieSceneEvent> GetData() const
 	{
 		return TMovieSceneChannelData<const FMovieSceneEvent>(&KeyTimes, &KeyValues);
 	}
diff --git a/Engine/Source/Runtime/MovieSceneTracks/Public/Channels/MovieSceneStringChannel.h b/Engine/Source/Runtime/MovieSceneTracks/Public/Channels/MovieSceneStringChannel.h
--- a/Engine/Source/Runtime/MovieSceneTracks/Public/Channels/MovieSceneStringChannel.h
+++ b/Engine/Source/Runtime/MovieSceneTracks/Public/Channels/MovieSceneStringChannel.h
@@ -47,7 +47,7 @@ struct FMovieSceneStringChannel : public FMovieSceneChannel
 	 *
 	 * @return An object that is able to manipulate this channel's data
 	 */
-	FORCEINLINE TMovieSceneChannelData<FString> GetData()
+	inline TMovieSceneChannelData<FString> GetData()
 	{
 		return TMovieSceneChannelData<FString>(&Times, &Values, this, &KeyHandles);
 	}
@@ -57,7 +57,7 @@ struct FMovieSceneStringChannel : public FMovieSceneChannel
 	 *
 	 * @return An object that is able to interrogate this channel's data
 	 */
-	FORCEINLINE TMovieSceneChannelData<const FString> GetData() const
+	inline TMovieSceneChannelData<const FString> GetData() const
 	{
 		return TMovieSceneChannelData<const FString>(&Times, &Values);
 	}
@@ -65,7 +65,7 @@ struct FMovieSceneStringChannel : public FMovieSceneChannel
 	/**
      * Const access to this channel's times
      */
-	FORCEINLINE TArrayView<const FFrameNumber> GetTimes() const
+	inline TArrayView<const FFrameNumber> GetTimes() const
 	{
 		return Times;
 	}
@@ -73,7 +73,7 @@ struct FMovieSceneStringChannel : public FMovieSceneChannel
 	/**
 	 * Returns whether this channel has any values
 	 */
-	FORCEINLINE bool HasAnyData() const
+	inline bool HasAnyData() const
 	{
 		return Times.Num() != 0 || bHasDefaultValue == true;
 	}
@@ -91,7 +91,7 @@ struct FMovieSceneStringChannel : public FMovieSceneChannel
 	 * @param InTimes Times to add
 	 * @param InValues Values to add
 	 */
-	FORCEINLINE void AddKeys(const TArray<FFrameNumber>& InTimes, const TArray<FString>& InValues)
+	inline void AddKeys(const TArray<FFrameNumber>& InTimes, const TArray<FString>& InValues)
 	{
 		check(InTimes.Num() == InValues.Num());
 		Times.Append(InTimes);
@@ -124,7 +124,7 @@ public:
 	 *
 	 * @param InDefaultValue The desired default value
 	 */
-	FORCEINLINE void SetDefault(FString InDefaultValue)
+	inline void SetDefault(FString InDefaultValue)
 	{
 		bHasDefaultValue = true;
 		DefaultValue = InDefaultValue;
@@ -135,7 +135,7 @@ public:
 	 *
 	 * @return (Optional) The channel's default value
 	 */
-	FORCEINLINE TOptional<FString> GetDefault() const
+	inline TOptional<FString> GetDefault() const
 	{
 		return bHasDefaultValue ? TOptional<FString>(DefaultValue) : TOptional<FString>();
 	}
@@ -143,7 +143,7 @@ public:
 	/**
 	 * Remove this channel's default value causing the channel to have no effect where no keys are present
 	 */
-	FORCEINLINE void RemoveDefault()
+	inline void RemoveDefault()
 	{
 		bHasDefaultValue = false;
 	}
diff --git a/Engine/Source/Runtime/MovieSceneTracks/Public/MovieSceneMixedAnimationTarget.h b/Engine/Source/Runtime/MovieSceneTracks/Public/MovieSceneMixedAnimationTarget.h
--- a/Engine/Source/Runtime/MovieSceneTracks/Public/MovieSceneMixedAnimationTarget.h
+++ b/Engine/Source/Runtime/MovieSceneTracks/Public/MovieSceneMixedAnimationTarget.h
@@ -19,7 +19,7 @@ struct FMovieSceneMixedAnimationTarget
 
 	virtual ~FMovieSceneMixedAnimationTarget() {}
 
-	FORCEINLINE friend uint32 GetTypeHash(const FMovieSceneMixedAnimationTarget& Target)
+	inline friend uint32 GetTypeHash(const FMovieSceneMixedAnimationTarget& Target)
 	{
 		return GetTypeHash(FMovieSceneMixedAnimationTarget::StaticStruct());
 	}
diff --git a/Engine/Source/Runtime/MovieSceneTracks/Public/Sections/MovieSceneActorReferenceSection.h b/Engine/Source/Runtime/MovieSceneTracks/Public/Sections/MovieSceneActorReferenceSection.h
--- a/Engine/Source/Runtime/MovieSceneTracks/Public/Sections/MovieSceneActorReferenceSection.h
+++ b/Engine/Source/Runtime/MovieSceneTracks/Public/Sections/MovieSceneActorReferenceSection.h
@@ -84,7 +84,7 @@ struct FMovieSceneActorReferenceData : public FMovieSceneChannel
 	 *
 	 * @return An object that is able to manipulate this channel's data
 	 */
-	FORCEINLINE TMovieSceneChannelData<FMovieSceneActorReferenceKey> GetData()
+	inline TMovieSceneChannelData<FMovieSceneActorReferenceKey> GetData()
 	{
 		return TMovieSceneChannelData<FMovieSceneActorReferenceKey>(&KeyTimes, &KeyValues, this, &KeyHandles);
 	}
@@ -94,7 +94,7 @@ struct FMovieSceneActorReferenceData : public FMovieSceneChannel
 	 *
 	 * @return An object that is able to interrogate this channel's data
 	 */
-	FORCEINLINE TMovieSceneChannelData<const FMovieSceneActorReferenceKey> GetData() const
+	inline TMovieSceneChannelData<const FMovieSceneActorReferenceKey> GetData() const
 	{
 		return TMovieSceneChannelData<const FMovieSceneActorReferenceKey>(&KeyTimes, &KeyValues);
 	}
@@ -133,7 +133,7 @@ public:
 	 *
 	 * @param InDefaultValue The desired default value
 	 */
-	FORCEINLINE void SetDefault(FMovieSceneActorReferenceKey InDefaultValue)
+	inline void SetDefault(FMovieSceneActorReferenceKey InDefaultValue)
 	{
 		DefaultValue = InDefaultValue;
 	}
@@ -143,7 +143,7 @@ public:
 	 *
 	 * @return (Optional) The channel's default value
 	 */
-	FORCEINLINE const FMovieSceneActorReferenceKey& GetDefault() const
+	inline const FMovieSceneActorReferenceKey& GetDefault() const
 	{
 		return DefaultValue;
 	}
diff --git a/Engine/Source/Runtime/MovieSceneTracks/Public/Sections/MovieSceneAudioSection.h b/Engine/Source/Runtime/MovieSceneTracks/Public/Sections/MovieSceneAudioSection.h
--- a/Engine/Source/Runtime/MovieSceneTracks/Public/Sections/MovieSceneAudioSection.h
+++ b/Engine/Source/Runtime/MovieSceneTracks/Public/Sections/MovieSceneAudioSection.h
@@ -209,7 +209,7 @@ public:
 private:
 
 	template<typename ChannelType, typename ForEachFunction>
-	FORCEINLINE static void ForEachInternal(ForEachFunction InFunction, const TMap<FName, ChannelType>& InMapToIterate) 
+	inline static void ForEachInternal(ForEachFunction InFunction, const TMap<FName, ChannelType>& InMapToIterate) 
 	{
 		for (auto& Item : InMapToIterate)
 		{
diff --git a/Engine/Source/Runtime/MovieSceneTracks/Public/Sections/MovieSceneEventSection.h b/Engine/Source/Runtime/MovieSceneTracks/Public/Sections/MovieSceneEventSection.h
--- a/Engine/Source/Runtime/MovieSceneTracks/Public/Sections/MovieSceneEventSection.h
+++ b/Engine/Source/Runtime/MovieSceneTracks/Public/Sections/MovieSceneEventSection.h
@@ -154,7 +154,7 @@ struct FMovieSceneEventSectionData : public FMovieSceneChannel
 	 *
 	 * @return An object that is able to manipulate this channel's data
 	 */
-	FORCEINLINE TMovieSceneChannelData<FEventPayload> GetData()
+	inline TMovieSceneChannelData<FEventPayload> GetData()
 	{
 		return TMovieSceneChannelData<FEventPayload>(&Times, &KeyValues, this, &KeyHandles);
 	}
@@ -164,7 +164,7 @@ struct FMovieSceneEventSectionData : public FMovieSceneChannel
 	 *
 	 * @return An object that is able to interrogate this channel's data
 	 */
-	FORCEINLINE TMovieSceneChannelData<const FEventPayload> GetData() const
+	inline TMovieSceneChannelData<const FEventPayload> GetData() const
 	{
 		return TMovieSceneChannelData<const FEventPayload>(&Times, &KeyValues);
 	}
diff --git a/Engine/Source/Runtime/TimeManagement/Public/CommonFrameRates.h b/Engine/Source/Runtime/TimeManagement/Public/CommonFrameRates.h
--- a/Engine/Source/Runtime/TimeManagement/Public/CommonFrameRates.h
+++ b/Engine/Source/Runtime/TimeManagement/Public/CommonFrameRates.h
@@ -39,21 +39,21 @@ struct FCommonFrameRates
 {
 	typedef __underlying_type(ECommonFrameRate) NumericType;
 
-	FORCEINLINE static FFrameRate FPS_12()  { return AllFrameRates[(NumericType)ECommonFrameRate::FPS_12].FrameRate;  }
-	FORCEINLINE static FFrameRate FPS_15()  { return AllFrameRates[(NumericType)ECommonFrameRate::FPS_15].FrameRate;  }
-	FORCEINLINE static FFrameRate FPS_24()  { return AllFrameRates[(NumericType)ECommonFrameRate::FPS_24].FrameRate;  }
-	FORCEINLINE static FFrameRate FPS_25()  { return AllFrameRates[(NumericType)ECommonFrameRate::FPS_25].FrameRate;  }
-	FORCEINLINE static FFrameRate FPS_30()  { return AllFrameRates[(NumericType)ECommonFrameRate::FPS_30].FrameRate;  }
-	FORCEINLINE static FFrameRate FPS_48()  { return AllFrameRates[(NumericType)ECommonFrameRate::FPS_48].FrameRate;  }
-	FORCEINLINE static FFrameRate FPS_50()  { return AllFrameRates[(NumericType)ECommonFrameRate::FPS_50].FrameRate;  }
-	FORCEINLINE static FFrameRate FPS_60()  { return AllFrameRates[(NumericType)ECommonFrameRate::FPS_60].FrameRate;  }
-	FORCEINLINE static FFrameRate FPS_100() { return AllFrameRates[(NumericType)ECommonFrameRate::FPS_100].FrameRate; }
-	FORCEINLINE static FFrameRate FPS_120() { return AllFrameRates[(NumericType)ECommonFrameRate::FPS_120].FrameRate; }
-	FORCEINLINE static FFrameRate FPS_240() { return AllFrameRates[(NumericType)ECommonFrameRate::FPS_240].FrameRate; }
+	inline static FFrameRate FPS_12()  { return AllFrameRates[(NumericType)ECommonFrameRate::FPS_12].FrameRate;  }
+	inline static FFrameRate FPS_15()  { return AllFrameRates[(NumericType)ECommonFrameRate::FPS_15].FrameRate;  }
+	inline static FFrameRate FPS_24()  { return AllFrameRates[(NumericType)ECommonFrameRate::FPS_24].FrameRate;  }
+	inline static FFrameRate FPS_25()  { return AllFrameRates[(NumericType)ECommonFrameRate::FPS_25].FrameRate;  }
+	inline static FFrameRate FPS_30()  { return AllFrameRates[(NumericType)ECommonFrameRate::FPS_30].FrameRate;  }
+	inline static FFrameRate FPS_48()  { return AllFrameRates[(NumericType)ECommonFrameRate::FPS_48].FrameRate;  }
+	inline static FFrameRate FPS_50()  { return AllFrameRates[(NumericType)ECommonFrameRate::FPS_50].FrameRate;  }
+	inline static FFrameRate FPS_60()  { return AllFrameRates[(NumericType)ECommonFrameRate::FPS_60].FrameRate;  }
+	inline static FFrameRate FPS_100() { return AllFrameRates[(NumericType)ECommonFrameRate::FPS_100].FrameRate; }
+	inline static FFrameRate FPS_120() { return AllFrameRates[(NumericType)ECommonFrameRate::FPS_120].FrameRate; }
+	inline static FFrameRate FPS_240() { return AllFrameRates[(NumericType)ECommonFrameRate::FPS_240].FrameRate; }
 
-	FORCEINLINE static FFrameRate NTSC_24() { return AllFrameRates[(NumericType)ECommonFrameRate::NTSC_24].FrameRate; }
-	FORCEINLINE static FFrameRate NTSC_30() { return AllFrameRates[(NumericType)ECommonFrameRate::NTSC_30].FrameRate; }
-	FORCEINLINE static FFrameRate NTSC_60() { return AllFrameRates[(NumericType)ECommonFrameRate::NTSC_60].FrameRate; }
+	inline static FFrameRate NTSC_24() { return AllFrameRates[(NumericType)ECommonFrameRate::NTSC_24].FrameRate; }
+	inline static FFrameRate NTSC_30() { return AllFrameRates[(NumericType)ECommonFrameRate::NTSC_30].FrameRate; }
+	inline static FFrameRate NTSC_60() { return AllFrameRates[(NumericType)ECommonFrameRate::NTSC_60].FrameRate; }
 
 	static TIMEMANAGEMENT_API TArrayView<const FCommonFrameRateInfo> GetAll();
 
diff --git a/Engine/Source/Runtime/TimeManagement/Public/TimeSynchronizationSource.h b/Engine/Source/Runtime/TimeManagement/Public/TimeSynchronizationSource.h
--- a/Engine/Source/Runtime/TimeManagement/Public/TimeSynchronizationSource.h
+++ b/Engine/Source/Runtime/TimeManagement/Public/TimeSynchronizationSource.h
@@ -118,7 +118,7 @@ public:
 	 * @param UpperBound		The upper bound of times to check.
 	 * @param RolloverModulus	Rollover frame value.
 	 */
-	FORCEINLINE static bool IsFrameBetweenWithRolloverModulus(const FFrameTime& ToCheck, const FFrameTime& LowerBound, const FFrameTime& UpperBound, const FFrameTime& RolloverModulus)
+	inline static bool IsFrameBetweenWithRolloverModulus(const FFrameTime& ToCheck, const FFrameTime& LowerBound, const FFrameTime& UpperBound, const FFrameTime& RolloverModulus)
 	{
 		if (LowerBound <= UpperBound)
 		{
@@ -131,7 +131,7 @@ public:
 	}
 
 	/** Convenience method to convert a FrameTime and FrameRate to a timecode value. */
-	FORCEINLINE static FTimecode ConvertFrameTimeToTimecode(const FFrameTime& FrameTime, const FFrameRate& FrameRate)
+	inline static FTimecode ConvertFrameTimeToTimecode(const FFrameTime& FrameTime, const FFrameRate& FrameRate)
 	{
 		return FTimecode::FromFrameNumber(FrameTime.GetFrame(), FrameRate);
 	}
@@ -144,7 +144,7 @@ public:
 	 * @param Offset			The offset to add.
 	 * @param RolloverModulus	Rollover frame value.
 	 */
-	FORCEINLINE static FFrameTime AddOffsetWithRolloverModulus(const FFrameTime& FrameTime, const int32 Offset, const FFrameTime& RolloverModulus)
+	inline static FFrameTime AddOffsetWithRolloverModulus(const FFrameTime& FrameTime, const int32 Offset, const FFrameTime& RolloverModulus)
 	{
 		const FFrameTime WithOffset = FrameTime + Offset;
 		const int32 RolloverFrameValue = RolloverModulus.GetFrame().Value;
