Sequencer: Fix muting the time warp curve

Check IsEvalDisabled when generating the time warp transform

diff --git a/Engine/Source/Runtime/MovieScene/Private/Tracks/MovieSceneTimeWarpTrack.cpp b/Engine/Source/Runtime/MovieScene/Private/Tracks/MovieSceneTimeWarpTrack.cpp
--- a/Engine/Source/Runtime/MovieScene/Private/Tracks/MovieSceneTimeWarpTrack.cpp
+++ b/Engine/Source/Runtime/MovieScene/Private/Tracks/MovieSceneTimeWarpTrack.cpp
@@ -69,6 +69,11 @@ int32 UMovieSceneTimeWarpTrack::GetTimeWarpSortOrder() const
 
 FMovieSceneNestedSequenceTransform UMovieSceneTimeWarpTrack::GenerateTransform() const
 {
+	if (IsEvalDisabled())
+	{
+		return FMovieSceneNestedSequenceTransform();
+	}
+
 	for (UMovieSceneSection* Section : GetAllSections())
 	{
 		UMovieSceneTimeWarpSection* TimeWarpSection = Cast<UMovieSceneTimeWarpSection>(Section);
