ShotManagement: Show Level Sequence Favorite Rating information as plain text in asset tooltip. Additionally fixes issue where Author and Created metadata for CineAssemblys.

diff --git a/Engine/Plugins/VirtualProduction/CinematicAssemblyTools/Source/CineAssemblyTools/Private/CineAssemblyToolsModule.cpp b/Engine/Plugins/VirtualProduction/CinematicAssemblyTools/Source/CineAssemblyTools/Private/CineAssemblyToolsModule.cpp
--- a/Engine/Plugins/VirtualProduction/CinematicAssemblyTools/Source/CineAssemblyTools/Private/CineAssemblyToolsModule.cpp
+++ b/Engine/Plugins/VirtualProduction/CinematicAssemblyTools/Source/CineAssemblyTools/Private/CineAssemblyToolsModule.cpp
@@ -2,4 +2,53 @@
 
 #include "Modules/ModuleManager.h"
 
-IMPLEMENT_MODULE(FDefaultModuleImpl, CineAssemblyTools);
+#include "CineAssembly.h"
+#include "MovieSceneMetaData.h"
+#include "MetaData/MovieSceneShotMetaData.h"
+
+#define LOCTEXT_NAMESPACE "CineAssemblyToolsModule"
+
+/**
+ * Implements the CineAssemblyTools module.
+ */
+class FCineAssemblyToolsModule : public IModuleInterface
+{
+public:
+
+	//~ Begin IModuleInterface interface
+	void StartupModule() override
+	{
+#if WITH_EDITOR
+		// Add empty movie scene meta data to the UCineAssembly CDO to ensure that
+		// asset registry tooltips show up in the editor
+		CineAssemblyCDO = GetMutableDefault<UCineAssembly>();
+
+		if (UCineAssembly* CDO = CineAssemblyCDO.Get())
+		{
+			UMovieSceneMetaData* MetaData = CDO->FindOrAddMetaData<UMovieSceneMetaData>();
+			MetaData->SetFlags(RF_Transient);
+			UMovieSceneShotMetaData* ShotMetaData = CDO->FindOrAddMetaData<UMovieSceneShotMetaData>();
+			ShotMetaData->SetFlags(RF_Transient);
+		}
+#endif
+	}
+
+	void ShutdownModule() override
+	{
+#if WITH_EDITOR
+		if (UCineAssembly* CDO = CineAssemblyCDO.Get())
+		{
+			CDO->RemoveMetaData<UMovieSceneShotMetaData>();
+			CDO->RemoveMetaData<UMovieSceneMetaData>();
+		}
+#endif
+	}
+	//~ End IModuleInterface interface
+
+private:
+	TWeakObjectPtr<UCineAssembly> CineAssemblyCDO;
+};
+
+IMPLEMENT_MODULE(FCineAssemblyToolsModule, CineAssemblyTools);
+
+#undef LOCTEXT_NAMESPACE
diff --git a/Engine/Source/Runtime/LevelSequence/Private/LevelSequenceModule.cpp b/Engine/Source/Runtime/LevelSequence/Private/LevelSequenceModule.cpp
--- a/Engine/Source/Runtime/LevelSequence/Private/LevelSequenceModule.cpp
+++ b/Engine/Source/Runtime/LevelSequence/Private/LevelSequenceModule.cpp
@@ -7,6 +7,7 @@
 #include "LevelSequenceActor.h"
 #include "LevelSequenceActorSpawner.h"
 #include "LevelSequencePlayer.h"
+#include "MetaData/MovieSceneShotMetaData.h"
 #include "MovieSceneMetaData.h"
 #include "UObject/UObjectBaseUtility.h"
 #include "IUniversalObjectLocatorModule.h"
@@ -31,6 +32,8 @@ void FLevelSequenceModule::StartupModule()
 	// asset registry tooltips show up in the editor
 	UMovieSceneMetaData* MetaData = GetMutableDefault<ULevelSequence>()->FindOrAddMetaData<UMovieSceneMetaData>();
 	MetaData->SetFlags(RF_Transient);
+	UMovieSceneShotMetaData* ShotMetaData = GetMutableDefault<ULevelSequence>()->FindOrAddMetaData<UMovieSceneShotMetaData>();
+	ShotMetaData->SetFlags(RF_Transient);
 
 	LevelSequenceCDO = GetMutableDefault<ULevelSequence>();
 #endif
@@ -41,6 +44,7 @@ void FLevelSequenceModule::ShutdownModule()
 #if WITH_EDITOR
 	if (ULevelSequence* CDO = LevelSequenceCDO.Get())
 	{
+		CDO->RemoveMetaData<UMovieSceneShotMetaData>();
 		CDO->RemoveMetaData<UMovieSceneMetaData>();
 	}
 #endif
diff --git a/Engine/Source/Runtime/MovieScene/Private/MetaData/MovieSceneShotMetaData.cpp b/Engine/Source/Runtime/MovieScene/Private/MetaData/MovieSceneShotMetaData.cpp
--- a/Engine/Source/Runtime/MovieScene/Private/MetaData/MovieSceneShotMetaData.cpp
+++ b/Engine/Source/Runtime/MovieScene/Private/MetaData/MovieSceneShotMetaData.cpp
@@ -62,24 +62,29 @@ bool UMovieSceneShotMetaData::GetFavoriteRatingByAssetData(const FAssetData& InA
 
 void UMovieSceneShotMetaData::ExtendAssetRegistryTags(FAssetRegistryTagsContext Context) const
 {
+	auto BoolToYesNoString = [](bool bInValue) -> FString
+		{
+			return bInValue ? TEXT("Yes") : TEXT("No");
+		};
+
 	Context.AddTag(FAssetRegistryTag(
-		AssetRegistryTag_bIsNoGood, bIsNoGood.IsSet() ? FString::FromInt(*bIsNoGood) : TEXT("Unset"),
+		AssetRegistryTag_bIsNoGood, bIsNoGood.IsSet() ? BoolToYesNoString(*bIsNoGood) : TEXT(""),
 		FAssetRegistryTag::ETagType::TT_Numerical, FAssetRegistryTag::TD_None
 	));
 	Context.AddTag(FAssetRegistryTag(
-		AssetRegistryTag_bIsFlagged, bIsFlagged.IsSet() ? FString::FromInt(*bIsFlagged) : TEXT("Unset"),
+		AssetRegistryTag_bIsFlagged, bIsFlagged.IsSet() ? BoolToYesNoString(*bIsFlagged) : TEXT(""),
 		FAssetRegistryTag::ETagType::TT_Numerical, FAssetRegistryTag::TD_None
 	));
 	Context.AddTag(FAssetRegistryTag(
-		AssetRegistryTag_bIsRecorded, bIsRecorded.IsSet() ? FString::FromInt(*bIsRecorded) : TEXT("Unset"),
+		AssetRegistryTag_bIsRecorded, bIsRecorded.IsSet() ? BoolToYesNoString(*bIsRecorded) : TEXT(""),
 		FAssetRegistryTag::ETagType::TT_Numerical, FAssetRegistryTag::TD_None
 	));
 	Context.AddTag(FAssetRegistryTag(
-		AssetRegistryTag_bIsSubSequence, bIsSubSequence.IsSet() ? FString::FromInt(*bIsSubSequence) : TEXT("Unset"),
+		AssetRegistryTag_bIsSubSequence, bIsSubSequence.IsSet() ? BoolToYesNoString(*bIsSubSequence) : TEXT(""),
 		FAssetRegistryTag::ETagType::TT_Numerical, FAssetRegistryTag::TD_None
 	));
 	Context.AddTag(FAssetRegistryTag(
-		AssetRegistryTag_FavoriteRating, FavoriteRating.IsSet() ? FString::FromInt(*FavoriteRating) : TEXT("Unset"),
+		AssetRegistryTag_FavoriteRating, FavoriteRating.IsSet() ? FString::FromInt(*FavoriteRating) : TEXT(""),
 		FAssetRegistryTag::ETagType::TT_Alphabetical, FAssetRegistryTag::TD_None
 	));
 }
