Fixed V530 warnings in some Sequencer code

diff --git a/Engine/Plugins/Experimental/MovieSceneAnimMixer/Source/MovieSceneAnimMixer/Private/MovieSceneAnimationMixerTrack.cpp b/Engine/Plugins/Experimental/MovieSceneAnimMixer/Source/MovieSceneAnimMixer/Private/MovieSceneAnimationMixerTrack.cpp
--- a/Engine/Plugins/Experimental/MovieSceneAnimMixer/Source/MovieSceneAnimMixer/Private/MovieSceneAnimationMixerTrack.cpp
+++ b/Engine/Plugins/Experimental/MovieSceneAnimMixer/Source/MovieSceneAnimMixer/Private/MovieSceneAnimationMixerTrack.cpp
@@ -195,8 +195,8 @@ UMovieSceneAnimationMixerTrack::UMovieSceneAnimationMixerTrack(const FObjectInit
 
 bool UMovieSceneAnimationMixerTrack::SupportsType(TSubclassOf<UMovieSceneSection> SectionClass) const
 {
-	SectionClass.Get()->ImplementsInterface(IMovieSceneAnimationSectionInterface::UClassType::StaticClass());
-	return true;
+	return SectionClass.Get()->ImplementsInterface(IMovieSceneAnimationSectionInterface::UClassType::StaticClass())
+		|| SectionClass.Get() == UMovieSceneSkeletalAnimationSection::StaticClass();
 }
 
 bool UMovieSceneAnimationMixerTrack::FixRowIndices()
diff --git a/Engine/Source/Runtime/MovieScene/Private/Tests/MovieSceneEntitySystemTypesTests.cpp b/Engine/Source/Runtime/MovieScene/Private/Tests/MovieSceneEntitySystemTypesTests.cpp
--- a/Engine/Source/Runtime/MovieScene/Private/Tests/MovieSceneEntitySystemTypesTests.cpp
+++ b/Engine/Source/Runtime/MovieScene/Private/Tests/MovieSceneEntitySystemTypesTests.cpp
@@ -389,7 +389,8 @@ bool FMovieSceneBlendTargetComponentTest::RunTest(const FString& Parameters)
 
 		this->TestEqual(TEXT("Inline Basic Capacity"), BlendTarget.GetCapacity(), 7);
 		CompareBlendTarget(BlendTarget, { 1000, 500, 100, 10, 0, -10, -100 }, TEXT("Inline Basic"));
-		GetTypeHash(BlendTarget);
+		// Verify that calling GetTypeHash does not segfault on the manual memory management
+		(void)GetTypeHash(BlendTarget);
 	}
 
 	// Test adding duplicates
@@ -412,7 +413,8 @@ bool FMovieSceneBlendTargetComponentTest::RunTest(const FString& Parameters)
 
 		this->TestEqual(TEXT("Inline Duplicate Capacity"), BlendTarget.GetCapacity(), 7);
 		CompareBlendTarget(BlendTarget, { 1000, 500, 100, 10, 0, -10, -100 }, TEXT("Inline Duplicate"));
-		GetTypeHash(BlendTarget);
+		// Verify that calling GetTypeHash does not segfault on the manual memory management
+		(void)GetTypeHash(BlendTarget);
 	}
 
 	// Simple test - use inline allocation
@@ -428,7 +430,8 @@ bool FMovieSceneBlendTargetComponentTest::RunTest(const FString& Parameters)
 
 		this->TestEqual(TEXT("Inline Basic Capacity"), BlendTarget.GetCapacity(), 7);
 		CompareBlendTarget(BlendTarget, { 1000, 500, 100, 10, 0, -10, -100 }, TEXT("Inline Basic"));
-		GetTypeHash(BlendTarget);
+		// Verify that calling GetTypeHash does not segfault on the manual memory management
+		(void)GetTypeHash(BlendTarget);
 	}
 
 	// Test growing the allocation
@@ -451,7 +454,8 @@ bool FMovieSceneBlendTargetComponentTest::RunTest(const FString& Parameters)
 
 		this->TestEqual(TEXT("Basic Heap Capacity"), BlendTarget.GetCapacity(), 32);
 		CompareBlendTarget(BlendTarget, Expected, TEXT("Basic Heap"));
-		GetTypeHash(BlendTarget);
+		// Verify that calling GetTypeHash does not segfault on the manual memory management
+		(void)GetTypeHash(BlendTarget);
 	}
 
 	// Test growing the allocation with duplicates
@@ -479,7 +483,8 @@ bool FMovieSceneBlendTargetComponentTest::RunTest(const FString& Parameters)
 
 		this->TestEqual(TEXT("Duplicate Heap Capacity"), BlendTarget.GetCapacity(), 32);
 		CompareBlendTarget(BlendTarget, Expected, TEXT("Duplicate Heap"));
-		GetTypeHash(BlendTarget);
+		// Verify that calling GetTypeHash does not segfault on the manual memory management
+		(void)GetTypeHash(BlendTarget);
 	}
 
 	// Test move and copy operators for heap allocations
@@ -510,15 +515,17 @@ bool FMovieSceneBlendTargetComponentTest::RunTest(const FString& Parameters)
 			FHierarchicalBlendTarget BlendTargetCopy(BlendTarget[0]);
 			CompareBlendTarget(BlendTargetCopy, Expected[0], TEXT("Copy 1"));
 			CompareBlendTarget(BlendTarget[0], Expected[0], TEXT("Copied From 2"));
-			GetTypeHash(BlendTargetCopy);
-			GetTypeHash(BlendTarget[0]);
+			// Verify that calling GetTypeHash does not segfault on the manual memory management
+			(void)GetTypeHash(BlendTargetCopy);
+			(void)GetTypeHash(BlendTarget[0]);
 
 			// Copy Assignment
 			BlendTargetCopy = BlendTarget[1];
 			CompareBlendTarget(BlendTargetCopy, Expected[1], TEXT("Copy 2"));
 			CompareBlendTarget(BlendTarget[1], Expected[1], TEXT("Copied From 2"));
-			GetTypeHash(BlendTargetCopy);
-			GetTypeHash(BlendTarget[1]);
+			// Verify that calling GetTypeHash does not segfault on the manual memory management
+			(void)GetTypeHash(BlendTargetCopy);
+			(void)GetTypeHash(BlendTarget[1]);
 		}
 
 		// Make some moves
@@ -530,15 +537,17 @@ bool FMovieSceneBlendTargetComponentTest::RunTest(const FString& Parameters)
 			FHierarchicalBlendTarget BlendTargetMoved(MoveTemp(BlendTargetCopy1));
 			CompareBlendTarget(BlendTargetMoved, Expected[0], TEXT("Move 1"));
 			CompareBlendTarget(BlendTargetCopy1, FHierarchicalBlendTarget().AsArray(), TEXT("Moved From 1"));
-			GetTypeHash(BlendTargetMoved);
-			GetTypeHash(BlendTargetCopy1);
+			// Verify that calling GetTypeHash does not segfault on the manual memory management
+			(void)GetTypeHash(BlendTargetMoved);
+			(void)GetTypeHash(BlendTargetCopy1);
 
 			// Move Assignment
 			BlendTargetMoved = MoveTemp(BlendTargetCopy2);
 			CompareBlendTarget(BlendTargetMoved, Expected[1], TEXT("Move 2"));
 			CompareBlendTarget(BlendTargetCopy2, FHierarchicalBlendTarget().AsArray(), TEXT("Moved From 2"));
-			GetTypeHash(BlendTargetMoved);
-			GetTypeHash(BlendTargetCopy2);
+			// Verify that calling GetTypeHash does not segfault on the manual memory management
+			(void)GetTypeHash(BlendTargetMoved);
+			(void)GetTypeHash(BlendTargetCopy2);
 		}
 	}
 
