Sequencer: Change check to ensure and output some more info about the invalid frame rate

#rnx

diff --git a/Engine/Source/Runtime/MovieScene/Private/Variants/MovieSceneTimeWarpVariantPayloads.cpp b/Engine/Source/Runtime/MovieScene/Private/Variants/MovieSceneTimeWarpVariantPayloads.cpp
--- a/Engine/Source/Runtime/MovieScene/Private/Variants/MovieSceneTimeWarpVariantPayloads.cpp
+++ b/Engine/Source/Runtime/MovieScene/Private/Variants/MovieSceneTimeWarpVariantPayloads.cpp
@@ -399,8 +399,8 @@ FMovieSceneTimeWarpFrameRate::FMovieSceneTimeWarpFrameRate(FFrameRate InRate)
 	int32 Denominator = InRate.Denominator;
 
 	// Do not allow 8 most significant bits, offset by the sign bit (our sign bit becomes bit index 23)
-	check( (Numerator   & InvalidBits) == 0 );
-	check( (Denominator & InvalidBits) == 0 );
+	ensureMsgf((Numerator & InvalidBits) == 0, TEXT("Found invalid numerator: %d, in frame rate: %s"), Numerator, *InRate.ToPrettyText().ToString());
+	ensureMsgf((Denominator & InvalidBits) == 0, TEXT("Found invalid denominator: %d, in frame rate: %s"), Denominator, *InRate.ToPrettyText().ToString());
 
 	// Move the sign bit
 	Numerator   |= ( (Numerator   & SignBit32) >> 8 );
