Sequencer: more advanced logging for the ECS runner

diff --git a/Engine/Source/Runtime/Core/Public/Misc/FrameTime.h b/Engine/Source/Runtime/Core/Public/Misc/FrameTime.h
--- a/Engine/Source/Runtime/Core/Public/Misc/FrameTime.h
+++ b/Engine/Source/Runtime/Core/Public/Misc/FrameTime.h
@@ -324,3 +324,10 @@ inline FString LexToString(const FFrameTime InTime)
 {
 	return FString::Printf(TEXT("Frame: %d Subframe: %f"), InTime.GetFrame().Value, InTime.GetSubFrame());
 }
+
+/** Convert a FFrameTime into a string */
+inline FString LexToShortString(const FFrameTime InTime)
+{
+	return FString::Printf(TEXT("%d.%f"), InTime.GetFrame().Value, InTime.GetSubFrame());
+}
+
diff --git a/Engine/Source/Runtime/MovieScene/Private/EntitySystem/MovieSceneEntitySystemRunner.cpp b/Engine/Source/Runtime/MovieScene/Private/EntitySystem/MovieSceneEntitySystemRunner.cpp
--- a/Engine/Source/Runtime/MovieScene/Private/EntitySystem/MovieSceneEntitySystemRunner.cpp
+++ b/Engine/Source/Runtime/MovieScene/Private/EntitySystem/MovieSceneEntitySystemRunner.cpp
@@ -9,7 +9,9 @@
 #include "Evaluation/PreAnimatedState/MovieScenePreAnimatedCaptureSource.h"
 #include "Evaluation/MovieSceneEvaluationTemplateInstance.h"
 #include "IMovieScenePlayer.h"
+#include "Logging/LogVerbosity.h"
 #include "MovieSceneSequence.h"
+#include "MovieSceneTimeHelpers.h"
 #include "Algo/Reverse.h"
 #include "Algo/Sort.h"
 #include "ProfilingDebugging/CountersTrace.h"
@@ -176,6 +178,24 @@ void FMovieSceneEntitySystemRunner::QueueUpdate(const FMovieSceneContext& Contex
 		bRequireFullFlush = true;
 	}
 	UpdateQueue.Add(FUpdateParamsAndContext{ FSimpleDelegate(), Context, { Instance, Flags } });
+
+#if !NO_LOGGING
+	if (LogMovieSceneECS.GetVerbosity() >= ELogVerbosity::VeryVerbose)
+	{
+		using namespace UE::MovieScene;
+
+		const FInstanceRegistry* InstanceRegistry = GetLinker()->GetInstanceRegistry();
+		const FSequenceInstance& SequenceInstance = InstanceRegistry->GetInstance(Instance);
+		const TSharedRef<FSharedPlaybackState> SharedPlaybackState = SequenceInstance.GetSharedPlaybackState();
+		const UMovieSceneSequence* Sequence = SharedPlaybackState->GetSequence(SequenceInstance.GetSequenceID());
+
+		UE_LOG(LogMovieSceneECS, VeryVerbose, TEXT("Queue update: Instance(%d,%d); Status='%s'; Range=%s; Sequence='%s'"),
+				Instance.InstanceID, Instance.InstanceSerial,
+				*UEnum::GetValueAsString(Context.GetStatus()),
+				*LexToString(Context.GetRange()),
+				*GetNameSafe(Sequence));
+	}
+#endif
 }
 
 void FMovieSceneEntitySystemRunner::QueueUpdate(const FMovieSceneContext& Context, FInstanceHandle Instance, FSimpleDelegate&& OnFlushedDelegate, UE::MovieScene::ERunnerUpdateFlags Flags)
@@ -186,6 +206,24 @@ void FMovieSceneEntitySystemRunner::QueueUpdate(const FMovieSceneContext& Contex
 	}
 
 	UpdateQueue.Add(FUpdateParamsAndContext{ OnFlushedDelegate, Context, { Instance, Flags } });
+
+#if !NO_LOGGING
+	if (LogMovieSceneECS.GetVerbosity() >= ELogVerbosity::VeryVerbose)
+	{
+		using namespace UE::MovieScene;
+
+		const FInstanceRegistry* InstanceRegistry = GetLinker()->GetInstanceRegistry();
+		const FSequenceInstance& SequenceInstance = InstanceRegistry->GetInstance(Instance);
+		const TSharedRef<FSharedPlaybackState> SharedPlaybackState = SequenceInstance.GetSharedPlaybackState();
+		const UMovieSceneSequence* Sequence = SharedPlaybackState->GetSequence(SequenceInstance.GetSequenceID());
+
+		UE_LOG(LogMovieSceneECS, VeryVerbose, TEXT("Queue update: Instance(%d,%d); Status='%s'; Range=%s; Sequence='%s'"),
+				Instance.InstanceID, Instance.InstanceSerial,
+				*UEnum::GetValueAsString(Context.GetStatus()),
+				*LexToString(Context.GetRange()),
+				*GetNameSafe(Sequence));
+	}
+#endif
 }
 
 bool FMovieSceneEntitySystemRunner::QueueFinalUpdate(FInstanceHandle InInstanceHandle)
@@ -261,6 +299,21 @@ bool FMovieSceneEntitySystemRunner::QueueFinalUpdateImpl(FInstanceHandle InInsta
 	LastUpdate.OnFlushed = InOnLastFlushDelegate;
 	UpdateQueue.Add(LastUpdate);
 
+#if !NO_LOGGING
+	if (LogMovieSceneECS.GetVerbosity() >= ELogVerbosity::VeryVerbose)
+	{
+		using namespace UE::MovieScene;
+
+		const TSharedRef<FSharedPlaybackState> SharedPlaybackState = Instance.GetSharedPlaybackState();
+		const UMovieSceneSequence* Sequence = SharedPlaybackState->GetSequence(Instance.GetSequenceID());
+
+		UE_LOG(LogMovieSceneECS, VeryVerbose, TEXT("Queue update: Instance(%d,%d); %s; Sequence='%s'"),
+				InInstanceHandle.InstanceID, InInstanceHandle.InstanceSerial,
+				bDestroyInstance ? TEXT("Destroy Instance") : TEXT("Finish Instance"),
+				*GetNameSafe(Sequence));
+	}
+#endif
+
 	return true;
 }
 
@@ -487,6 +540,8 @@ void FMovieSceneEntitySystemRunner::FlushOutstanding(double BudgetMs, UE::MovieS
 		TargetState = FlushState::Everything;
 	}
 
+	UE_LOG(LogMovieSceneECS, VeryVerbose, TEXT("Flushing ECS runner state '0x%#08x' with budget %fms"), FlushState, BudgetMs);
+
 	const double BudgetSeconds = BudgetMs / 1000.f;
 	if (bRequireFullFlush)
 	{
@@ -825,6 +880,32 @@ UE::MovieScene::ERunnerFlushResult FMovieSceneEntitySystemRunner::GameThread_Upd
 	//       to ensure that up-to-date versions are used inside FEntityManager::OnStructureChanged
 	Linker->EntityManager.IncrementSystemSerial();
 
+#if !NO_LOGGING
+	if (LogMovieSceneECS.GetVerbosity() >= ELogVerbosity::VeryVerbose && DissectedUpdates.Num() > 0)
+	{
+		UE_LOG(LogMovieSceneECS, VeryVerbose, TEXT("Executing %d dissected updates"), DissectedUpdates.Num());
+
+		for (int32 Index = 0; Index < DissectedUpdates.Num(); ++Index)
+		{
+			const FDissectedUpdate& DissectedUpdate(DissectedUpdates[Index]);
+			const FSequenceInstance& Instance = InstanceRegistry->GetInstance(DissectedUpdate.InstanceHandle);
+			const TSharedRef<FSharedPlaybackState> SharedPlaybackState = Instance.GetSharedPlaybackState();
+			const UMovieSceneSequence* Sequence = SharedPlaybackState->GetSequence(Instance.GetSequenceID());
+
+			UE_LOG(LogMovieSceneECS, VeryVerbose,
+					TEXT("[%d] Order='%s'; Instance(%d,%d); Status='%s'; Range=%s; OnFlushed='%s'; Flags='%#03x', Sequence='%s'"),
+					Index,
+					DissectedUpdate.Order == MAX_int32 ? TEXT("MAX") : *LexToString(DissectedUpdate.Order),
+					DissectedUpdate.InstanceHandle.InstanceID, DissectedUpdate.InstanceHandle.InstanceSerial,
+					*UEnum::GetValueAsString(DissectedUpdate.Context.GetStatus()),
+					*LexToString(DissectedUpdate.Context.GetRange()),
+					DissectedUpdate.OnFlushed.IsBound() ? TEXT("BOUND") : TEXT("unbound"),
+					DissectedUpdate.UpdateFlags,
+					*GetNameSafe(Sequence));
+		}
+	}
+#endif
+
 	// Update all systems
 	if (DissectedUpdates.Num() != 0)
 	{
@@ -867,6 +948,30 @@ UE::MovieScene::ERunnerFlushResult FMovieSceneEntitySystemRunner::GameThread_Upd
 		DissectedUpdates.RemoveAt(0, Index);
 	}
 
+#if !NO_LOGGING
+	if (LogMovieSceneECS.GetVerbosity() >= ELogVerbosity::VeryVerbose && CurrentInstances.Num() > 0)
+	{
+		UE_LOG(LogMovieSceneECS, VeryVerbose, TEXT("Running evaluation for %d current instances to update"), CurrentInstances.Num());
+
+		for (int32 Index = 0; Index < CurrentInstances.Num(); ++Index)
+		{
+			const FQueuedUpdateParams& UpdateParams(CurrentInstances[Index]);
+			const FSequenceInstance& Instance = InstanceRegistry->GetInstance(UpdateParams.InstanceHandle);
+			const TSharedRef<FSharedPlaybackState> SharedPlaybackState = Instance.GetSharedPlaybackState();
+			const UMovieSceneSequence* Sequence = SharedPlaybackState->GetSequence(Instance.GetSequenceID());
+
+			UE_LOG(LogMovieSceneECS, VeryVerbose,
+					TEXT("[%d] Instance(%d,%d); Status='%s'; Range=%s; Flags='%#03x', Sequence='%s'"),
+					Index,
+					UpdateParams.InstanceHandle.InstanceID, UpdateParams.InstanceHandle.InstanceSerial,
+					*UEnum::GetValueAsString(Instance.GetContext().GetStatus()),
+					*LexToString(Instance.GetContext().GetRange()),
+					UpdateParams.UpdateFlags,
+					*GetNameSafe(Sequence));
+		}
+	}
+#endif
+
 #if DO_GUARD_SLOW
 	// Check that there are no duplicates in the CurrentInstances array
 	{
diff --git a/Engine/Source/Runtime/MovieScene/Public/MovieSceneTimeHelpers.h b/Engine/Source/Runtime/MovieScene/Public/MovieSceneTimeHelpers.h
--- a/Engine/Source/Runtime/MovieScene/Public/MovieSceneTimeHelpers.h
+++ b/Engine/Source/Runtime/MovieScene/Public/MovieSceneTimeHelpers.h
@@ -410,3 +410,24 @@ inline FString LexToString(const TRange<FFrameNumber>& InRange)
 		*FString::Printf(TEXT("%i)"), SourceUpper.GetValue().Value)
 	);
 }
+
+inline FString LexToString(const TRange<FFrameTime>& InRange)
+{
+	TRangeBound<FFrameTime> SourceLower = InRange.GetLowerBound();
+	TRangeBound<FFrameTime> SourceUpper = InRange.GetUpperBound();
+
+	return *FString::Printf(TEXT("%s,%s"),
+		SourceLower.IsOpen() ?
+		TEXT("[-inf") :
+		SourceLower.IsInclusive() ?
+		*FString::Printf(TEXT("[%s"), *LexToShortString(SourceLower.GetValue())) :
+		*FString::Printf(TEXT("(%s"), *LexToShortString(SourceLower.GetValue())),
+
+		SourceUpper.IsOpen() ?
+		TEXT("+inf]") :
+		SourceUpper.IsInclusive() ?
+		*FString::Printf(TEXT("%s]"), *LexToShortString(SourceUpper.GetValue())) :
+		*FString::Printf(TEXT("%s)"), *LexToShortString(SourceUpper.GetValue()))
+	);
+}
+
