Sequencer: add ISequenceVisitor callbacks for when the visit enters/exists a sub-sequence

diff --git a/Engine/Source/Runtime/MovieScene/Private/MovieSceneSequenceVisitor.cpp b/Engine/Source/Runtime/MovieScene/Private/MovieSceneSequenceVisitor.cpp
--- a/Engine/Source/Runtime/MovieScene/Private/MovieSceneSequenceVisitor.cpp
+++ b/Engine/Source/Runtime/MovieScene/Private/MovieSceneSequenceVisitor.cpp
@@ -53,6 +53,8 @@ void VisitSubTrackImpl(const FSequenceVisitParams& InParams, UMovieSceneSubTrack
 			continue;
 		}
 
+		InVisitor.VisitStartSubSequence(SubSection, ObjectBinding, *SubSequenceSpace);
+
 		auto VisitRange = [&InOutRootPath, &InVisitor, &InParams, SubSection, &ObjectBinding, SubSequence, SubSequenceSpace](TRange<FFrameTime> Range)
 		{
 			TRange<FFrameNumber> EffectiveRange = ConvertToDiscreteRange(TRange<FFrameTime>::Intersection(Range, ConvertToFrameTimeRange(SubSequenceSpace->RootClampRange)));
@@ -90,8 +92,9 @@ void VisitSubTrackImpl(const FSequenceVisitParams& InParams, UMovieSceneSubTrack
 
 		FMovieSceneInverseSequenceTransform SequenceToRootTransform = SubSequenceSpace->RootToSequenceTransform.Inverse();
 		SequenceToRootTransform.TransformFiniteRangeWithinRange(ConvertToFrameTimeRange(Entry.Range), VisitRange, SubSequenceSpace->StartBreadcrumbs, SubSequenceSpace->EndBreadcrumbs);
-	}
 
+		InVisitor.VisitEndSubSequence(SubSection, ObjectBinding, *SubSequenceSpace);
+	}
 }
 
 void VisitTrackImpl(const FSequenceVisitParams& InParams, UMovieSceneTrack* InTrack, ISequenceVisitor& InVisitor, const FGuid& ObjectBinding, FSubSequencePath* InOutRootPath, FSubSequenceSpace* SubSequenceSpace)
diff --git a/Engine/Source/Runtime/MovieScene/Public/MovieSceneSequenceVisitor.h b/Engine/Source/Runtime/MovieScene/Public/MovieSceneSequenceVisitor.h
--- a/Engine/Source/Runtime/MovieScene/Public/MovieSceneSequenceVisitor.h
+++ b/Engine/Source/Runtime/MovieScene/Public/MovieSceneSequenceVisitor.h
@@ -9,6 +9,7 @@
 
 class UMovieSceneSection;
 class UMovieSceneSequence;
+class UMovieSceneSubSection;
 class UMovieSceneTrack;
 struct FGuid;
 struct FMovieSceneBinding;
@@ -96,6 +97,9 @@ struct ISequenceVisitor
 	virtual void VisitSection(UMovieSceneTrack*, UMovieSceneSection*, const FGuid&, const FSubSequenceSpace& LocalSpace) {}
 
 	virtual void VisitSubSequence(UMovieSceneSequence*, const FGuid&, const FSubSequenceSpace& LocalSpace) {}
+
+	virtual void VisitStartSubSequence(UMovieSceneSubSection*, const FGuid&, const FSubSequenceSpace& LocalSpace) {}
+	virtual void VisitEndSubSequence(UMovieSceneSubSection*, const FGuid&, const FSubSequenceSpace& LocalSpace) {}
 };
 
 struct FSequenceVisitParams
