Sequencer x Curve Editor: Avoid Modify call when data is read-only / does not change. Whenever a key is added, SetRange is called which adds a lot of data to the undo stack.

diff --git a/Engine/Source/Runtime/MovieScene/Public/MovieSceneSection.h b/Engine/Source/Runtime/MovieScene/Public/MovieSceneSection.h
--- a/Engine/Source/Runtime/MovieScene/Public/MovieSceneSection.h
+++ b/Engine/Source/Runtime/MovieScene/Public/MovieSceneSection.h
@@ -317,6 +317,11 @@ public:
 	 */
 	virtual void SetRange(const TRange<FFrameNumber>& NewRange)
 	{
+		const bool bHasChanged = NewRange != SectionRange.Value; 
+		if (IsReadOnly() || !bHasChanged)
+		{
+			return;
+		}
 		
 		// Skip TryModify for objects that still need initialization (i.e. we're in the object's constructor), because modifying objects in their constructor can lead to non-deterministic cook issues.
 		bool bCanSetRange = true;
