Sequencer: Add support for meta specifiers for UMovieSceneSubSection Translation and Rotation properties

This allows the property to indicate how the spin boxes should increment when scrubbing

diff --git a/Engine/Source/Runtime/MovieScene/Private/Sections/MovieSceneSubSection.cpp b/Engine/Source/Runtime/MovieScene/Private/Sections/MovieSceneSubSection.cpp
--- a/Engine/Source/Runtime/MovieScene/Private/Sections/MovieSceneSubSection.cpp
+++ b/Engine/Source/Runtime/MovieScene/Private/Sections/MovieSceneSubSection.cpp
@@ -62,6 +62,8 @@ struct FSubSectionEditorData
 	
 	FSubSectionEditorData(EMovieSceneTransformChannel Mask, UMovieSceneSubSection* SubSection)
 	{
+		static const TSet<FName> PropertyMetaDataKeys = { "UIMin", "UIMax", "SliderExponent", "LinearDeltaSensitivity", "Delta", "ClampMin", "ClampMax", "ForceUnits", "WheelStep" };
+
 		const EAxisList::Type XAxis = EAxisList::Forward;
 		const EAxisList::Type YAxis = EAxisList::Left;
 		const EAxisList::Type ZAxis = EAxisList::Up;
@@ -71,12 +73,19 @@ struct FSubSectionEditorData
 		const int32 TranslationOrderOffset = 0;
 		const int32 RotationOrderOffset = TranslationOrderOffset + 3;
 
+		const FProperty* TranslationProperty = UMovieSceneSubSection::StaticClass()->FindPropertyByName(UMovieSceneSubSection::GetTranslationPropertyName());
+		const FProperty* RotationProperty = UMovieSceneSubSection::StaticClass()->FindPropertyByName(UMovieSceneSubSection::GetRotationPropertyName());
+
 		MetaData[0].SetIdentifiers("Override.Location.X", AxisDisplayInfo::GetAxisDisplayName(XAxis), LocationGroup);
 		MetaData[0].SubPropertyPath = TEXT("Location.X");
 		MetaData[0].SortOrder = TranslationOrderOffset + ReverseSwizzle[0];
 		MetaData[0].bEnabled = EnumHasAllFlags(Mask, EMovieSceneTransformChannel::TranslationX);
 		MetaData[0].Color = AxisDisplayInfo::GetAxisColor(XAxis);
 		MetaData[0].bCanCollapseToTrack = false;
+		for (const FName& PropertyMetaDataKey : PropertyMetaDataKeys)
+		{
+			MetaData[0].PropertyMetaData.Add(PropertyMetaDataKey, TranslationProperty->GetMetaData(PropertyMetaDataKey));
+		}
 
 		MetaData[1].SetIdentifiers("Override.Location.Y", AxisDisplayInfo::GetAxisDisplayName(YAxis), LocationGroup);
 		MetaData[1].SubPropertyPath = TEXT("Location.Y");
@@ -85,6 +94,10 @@ struct FSubSectionEditorData
 		MetaData[1].Color = AxisDisplayInfo::GetAxisColor(YAxis);
 		MetaData[1].bCanCollapseToTrack = false;
 		MetaData[1].bInvertValue = AxisDisplayInfo::GetAxisDisplayCoordinateSystem() == EAxisList::LeftUpForward;
+		for (const FName& PropertyMetaDataKey : PropertyMetaDataKeys)
+		{
+			MetaData[1].PropertyMetaData.Add(PropertyMetaDataKey, TranslationProperty->GetMetaData(PropertyMetaDataKey));
+		}
 
 		MetaData[2].SetIdentifiers("Override.Location.Z", AxisDisplayInfo::GetAxisDisplayName(ZAxis), LocationGroup);
 		MetaData[2].SubPropertyPath = TEXT("Location.Z");
@@ -92,6 +105,10 @@ struct FSubSectionEditorData
 		MetaData[2].bEnabled = EnumHasAllFlags(Mask, EMovieSceneTransformChannel::TranslationZ);
 		MetaData[2].Color = AxisDisplayInfo::GetAxisColor(ZAxis);
 		MetaData[2].bCanCollapseToTrack = false;
+		for (const FName& PropertyMetaDataKey : PropertyMetaDataKeys)
+		{
+			MetaData[2].PropertyMetaData.Add(PropertyMetaDataKey, TranslationProperty->GetMetaData(PropertyMetaDataKey));
+		}
 
 		MetaData[3].SetIdentifiers("Override.Rotation.X", NSLOCTEXT("MovieSceneSubSection", "RotationX", "Roll"), RotationGroup);
 		MetaData[3].SubPropertyPath = TEXT("Rotation.X");
@@ -99,6 +116,10 @@ struct FSubSectionEditorData
 		MetaData[3].bEnabled = EnumHasAllFlags(Mask, EMovieSceneTransformChannel::RotationX);
 		MetaData[3].Color = AxisDisplayInfo::GetAxisColor(EAxisList::X);
 		MetaData[3].bCanCollapseToTrack = false;
+		for (const FName& PropertyMetaDataKey : PropertyMetaDataKeys)
+		{
+			MetaData[3].PropertyMetaData.Add(PropertyMetaDataKey, RotationProperty->GetMetaData(PropertyMetaDataKey));
+		}
 
 		MetaData[4].SetIdentifiers("Override.Rotation.Y", NSLOCTEXT("MovieSceneSubSection", "RotationY", "Pitch"), RotationGroup);
 		MetaData[4].SubPropertyPath = TEXT("Rotation.Y");
@@ -106,6 +127,10 @@ struct FSubSectionEditorData
 		MetaData[4].bEnabled = EnumHasAllFlags(Mask, EMovieSceneTransformChannel::RotationY);
 		MetaData[4].Color = AxisDisplayInfo::GetAxisColor(EAxisList::Y);
 		MetaData[4].bCanCollapseToTrack = false;
+		for (const FName& PropertyMetaDataKey : PropertyMetaDataKeys)
+		{
+			MetaData[4].PropertyMetaData.Add(PropertyMetaDataKey, RotationProperty->GetMetaData(PropertyMetaDataKey));
+		}
 
 		MetaData[5].SetIdentifiers("Override.Rotation.Z", NSLOCTEXT("MovieSceneSubSection", "RotationZ", "Yaw"), RotationGroup);
 		MetaData[5].SubPropertyPath = TEXT("Rotation.Z");
@@ -113,6 +138,10 @@ struct FSubSectionEditorData
 		MetaData[5].bEnabled = EnumHasAllFlags(Mask, EMovieSceneTransformChannel::RotationZ);
 		MetaData[5].Color = AxisDisplayInfo::GetAxisColor(EAxisList::Z);
 		MetaData[5].bCanCollapseToTrack = false;
+		for (const FName& PropertyMetaDataKey : PropertyMetaDataKeys)
+		{
+			MetaData[5].PropertyMetaData.Add(PropertyMetaDataKey, RotationProperty->GetMetaData(PropertyMetaDataKey));
+		}
 
 		ExternalValues[0].OnGetExternalValue = [SubSection](UObject& InObject, FTrackInstancePropertyBindings* Bindings) { return GetValue(SubSection, 0); };
 		ExternalValues[1].OnGetExternalValue = [SubSection](UObject& InObject, FTrackInstancePropertyBindings* Bindings) { return GetValue(SubSection, 1); };
diff --git a/Engine/Source/Runtime/MovieScene/Public/Sections/MovieSceneSubSection.h b/Engine/Source/Runtime/MovieScene/Public/Sections/MovieSceneSubSection.h
--- a/Engine/Source/Runtime/MovieScene/Public/Sections/MovieSceneSubSection.h
+++ b/Engine/Source/Runtime/MovieScene/Public/Sections/MovieSceneSubSection.h
@@ -225,6 +225,16 @@ public:
 
 	MOVIESCENE_API void DeleteChannels(TArrayView<const FName> ChannelNames);
 
+	MOVIESCENE_API static const FName GetTranslationPropertyName()
+	{
+		return GET_MEMBER_NAME_CHECKED(UMovieSceneSubSection, Translation);
+	}
+
+	MOVIESCENE_API static const FName GetRotationPropertyName()
+	{
+		return GET_MEMBER_NAME_CHECKED(UMovieSceneSubSection, Rotation);
+	}
+
 public:
 
 	//~ UMovieSceneSection interface
@@ -268,10 +278,10 @@ private:
 	UPROPERTY()
 	FMovieSceneSubSectionOriginOverrideMask OriginOverrideMask;
 
-	UPROPERTY()
+	UPROPERTY(meta = (LinearDeltaSensitivity = "1", Delta = "1.0"))
 	FMovieSceneDoubleChannel Translation[3];
 	
-	UPROPERTY()
+	UPROPERTY(meta = (LinearDeltaSensitivity = "1", Delta = "1.0"))
 	FMovieSceneDoubleChannel Rotation[3];
 
 #if WITH_EDITORONLY_DATA
