Sequencer: fix pause-at-frame logic

When the playback would reach a frame we're supposed to pause at, the code would not actually update the sequence's current time, This means the sequence was pausing at the previous frame, which would vary depending on the current game's framerate.

This change fixes that and uses the pause-at time exactly, plus it also adds a new automated test.

#virtualized

diff --git a/Engine/Source/Runtime/MovieScene/Private/MovieSceneSequencePlayer.cpp b/Engine/Source/Runtime/MovieScene/Private/MovieSceneSequencePlayer.cpp
--- a/Engine/Source/Runtime/MovieScene/Private/MovieSceneSequencePlayer.cpp
+++ b/Engine/Source/Runtime/MovieScene/Private/MovieSceneSequencePlayer.cpp
@@ -867,6 +867,25 @@ TOptional<TRange<FFrameTime>> UMovieSceneSequencePlayer::GetPauseRange(const FFr
 	return TOptional<TRange<FFrameTime>>();
 }
 
+TOptional<FFrameTime> UMovieSceneSequencePlayer::GetPauseTimeForNewPosition(const FFrameTime& NewPosition) const
+{
+	TOptional<TRange<FFrameTime>> PauseRange = GetPauseRange(NewPosition);
+
+	if (PauseRange.IsSet())
+	{
+		if (bReversePlayback)
+		{
+			return PauseRange.GetValue().GetLowerBoundValue();
+		}
+		else
+		{
+			return PauseRange.GetValue().GetUpperBoundValue();
+		}
+	}
+
+	return TOptional<FFrameTime>();
+}
+
 UMovieSceneEntitySystemLinker* UMovieSceneSequencePlayer::ConstructEntitySystemLinker()
 {
 	if (ensure(TickManager && RegisteredTickInterval.IsSet()) && !EnumHasAnyFlags(Sequence->GetFlags(), EMovieSceneSequenceFlags::BlockingEvaluation))
@@ -1238,8 +1257,8 @@ void UMovieSceneSequencePlayer::UpdateTimeCursorPosition_Internal(FFrameTime New
 	}
 
 	// If we should pause during this evaluation, we'll handle that below.
-	const TOptional<TRange<FFrameTime>> PauseRange = GetPauseRange(NewPosition);
-	if (!PauseRange.IsSet() && Method == EUpdatePositionMethod::Play && ShouldStopOrLoop(NewPosition))
+	const TOptional<FFrameTime> PauseAtTime = GetPauseTimeForNewPosition(NewPosition);
+	if (!PauseAtTime.IsSet() && Method == EUpdatePositionMethod::Play && ShouldStopOrLoop(NewPosition))
 	{
 		// The actual start time taking into account reverse playback
 		FFrameTime StartTimeWithReversed = bReversePlayback ? GetLastValidTime() : StartTime;
@@ -1359,10 +1378,15 @@ void UMovieSceneSequencePlayer::UpdateTimeCursorPosition_Internal(FFrameTime New
 	}
 	else
 	{
-		// If the desired evaluation will take us past where we want to go we need to use a clipped range provided by the PauseRange, otherwise use the normal one.
-		FMovieSceneEvaluationRange Range = PauseRange.IsSet() 
-			? FMovieSceneEvaluationRange(PauseRange.GetValue(), PlayPosition.GetOutputRate(), bReversePlayback ? EPlayDirection::Backwards : EPlayDirection::Forwards)
-			: UpdatePlayPosition(PlayPosition, NewPosition, Method, bReversePlayback);
+		// If the desired evaluation will take us past where we want to go we need to use a clipped range provided by PauseAtTime, otherwise use the normal one.
+		FFrameTime ActualNewPosition = PauseAtTime.IsSet() ?
+			FFrameRate::TransformTime(PauseAtTime.GetValue(), PlayPosition.GetOutputRate(), PlayPosition.GetInputRate()) :
+			NewPosition;
+		FMovieSceneEvaluationRange Range = UpdatePlayPosition(
+				PlayPosition,
+				ActualNewPosition,
+				Method,
+				bReversePlayback);
 
 		UMovieSceneSequence* MovieSceneSequence = RootTemplateInstance.GetSequence(MovieSceneSequenceID::Root);
 		const bool bIsSequenceBlocking = MovieSceneSequence ? EnumHasAnyFlags(MovieSceneSequence->GetFlags(), EMovieSceneSequenceFlags::BlockingEvaluation) : false;
@@ -1381,7 +1405,7 @@ void UMovieSceneSequencePlayer::UpdateTimeCursorPosition_Internal(FFrameTime New
 		// Now that the evaluation has taken place we call Pause (to trigger delegates, etc.), however if we're running the evaluation
 		// as async we actually need to queue a latent action to do it so that it happens after the data is actually updated. We can't
 		// use the QueueLatentAction inside of Pause() because IsEvaluating() is only set during actual evaluation so it won't be set yet.
-		if (PauseRange.IsSet())
+		if (PauseAtTime.IsSet())
 		{
 			if (Args.bIsAsync)
 			{
diff --git a/Engine/Source/Runtime/MovieScene/Public/MovieSceneSequencePlayer.h b/Engine/Source/Runtime/MovieScene/Public/MovieSceneSequencePlayer.h
--- a/Engine/Source/Runtime/MovieScene/Public/MovieSceneSequencePlayer.h
+++ b/Engine/Source/Runtime/MovieScene/Public/MovieSceneSequencePlayer.h
@@ -539,10 +539,15 @@ protected:
 	MOVIESCENE_API bool ShouldStopOrLoop(FFrameTime NewPosition) const;
 	/** 
 	* If the current sequence should pause (due to NewPosition overshooting a previously set ShouldPause) 
-	* then a range of time that should be evaluated to reach there will be returned. If we should not pause
-	* then the TOptional will be unset.
+	* then a range of time that should be evaluated to reach there will be returned, in tick resolution.
+	* If we should not pause then the TOptional will be unset.
 	* */
 	MOVIESCENE_API TOptional<TRange<FFrameTime>> GetPauseRange(const FFrameTime& NewPosition) const;
+	/**
+	 * Same as GetPauseRange but returns the next time to update to (the upper or lower bound value of
+	 * the range returned by GetPauseRange, depending on the playback direction), in tick resolution.
+	 */
+	MOVIESCENE_API TOptional<FFrameTime> GetPauseTimeForNewPosition(const FFrameTime& NewPosition) const;
 
 	MOVIESCENE_API UWorld* GetPlaybackWorld() const;
 
