Sequencer: Remove time warp decoration from movie scene when removing a time warp track if it's empty

diff --git a/Engine/Source/Runtime/MovieScene/Private/Decorations/MovieSceneScalingAnchors.cpp b/Engine/Source/Runtime/MovieScene/Private/Decorations/MovieSceneScalingAnchors.cpp
--- a/Engine/Source/Runtime/MovieScene/Private/Decorations/MovieSceneScalingAnchors.cpp
+++ b/Engine/Source/Runtime/MovieScene/Private/Decorations/MovieSceneScalingAnchors.cpp
@@ -46,15 +46,31 @@ void UMovieSceneScalingAnchors::OnDecorationAdded(UMovieScene* MovieScene)
 }
 
 void UMovieSceneScalingAnchors::OnDecorationRemoved()
-{
-	if (UMovieSceneTimeWarpDecoration* TimeWarp = GetTypedOuter<UMovieScene>()->FindDecoration<UMovieSceneTimeWarpDecoration>())
+{	
+	UMovieScene* MovieScene = GetTypedOuter<UMovieScene>();
+	if (!MovieScene)
+	{
+		return;
+	}
+
+	if (UMovieSceneTimeWarpDecoration* TimeWarp = MovieScene->FindDecoration<UMovieSceneTimeWarpDecoration>())
 	{
 		TimeWarp->RemoveTimeWarpSource(this);
+
+		if (!TimeWarp->HasAnySources())
+		{
+			MovieScene->RemoveDecoration<UMovieSceneTimeWarpDecoration>();
+		}
 	}
 #if 0
-	if (UMovieSceneTimeWarpRegistryDecoration* TimeWarp = GetTypedOuter<UMovieScene>()->FindDecoration<UMovieSceneTimeWarpRegistryDecoration>())
+	if (UMovieSceneTimeWarpRegistryDecoration* TimeWarp = MovieScene->FindDecoration<UMovieSceneTimeWarpRegistryDecoration>())
 	{
 		TimeWarp->SetDefaultTimeWarp(nullptr);
+
+		if (!TimeWarp->HasAnySources())
+		{
+			MovieScene->RemoveDecoration<UMovieSceneTimeWarpDecoration>();
+		}
 	}
 #endif
 }
@@ -98,6 +114,11 @@ void UMovieSceneScalingAnchors::RemoveScalingDriver(TScriptInterface<IMovieScene
 	ResetScaling();
 }
 
+bool UMovieSceneScalingAnchors::HasAnyScalingDrivers() const
+{
+	return !ScalingDrivers.IsEmpty();
+}
+
 FMovieSceneAnchorsScalingGroup& UMovieSceneScalingAnchors::GetOrCreateScalingGroup(const FGuid& Guid)
 {
 	return ScalingGroups.FindOrAdd(Guid);
diff --git a/Engine/Source/Runtime/MovieScene/Private/Decorations/MovieSceneSectionAnchorsDecoration.cpp b/Engine/Source/Runtime/MovieScene/Private/Decorations/MovieSceneSectionAnchorsDecoration.cpp
--- a/Engine/Source/Runtime/MovieScene/Private/Decorations/MovieSceneSectionAnchorsDecoration.cpp
+++ b/Engine/Source/Runtime/MovieScene/Private/Decorations/MovieSceneSectionAnchorsDecoration.cpp
@@ -25,6 +25,11 @@ void UMovieSceneSectionAnchorsDecoration::OnDestroy(UMovieScene* MovieScene)
 	{
 		Anchors->Modify();
 		Anchors->RemoveScalingDriver(this);
+
+		if (!Anchors->HasAnyScalingDrivers())
+		{
+			MovieScene->RemoveDecoration<UMovieSceneScalingAnchors>();
+		}
 	}
 }
 
diff --git a/Engine/Source/Runtime/MovieScene/Private/Tracks/MovieSceneTimeWarpTrack.cpp b/Engine/Source/Runtime/MovieScene/Private/Tracks/MovieSceneTimeWarpTrack.cpp
--- a/Engine/Source/Runtime/MovieScene/Private/Tracks/MovieSceneTimeWarpTrack.cpp
+++ b/Engine/Source/Runtime/MovieScene/Private/Tracks/MovieSceneTimeWarpTrack.cpp
@@ -39,9 +39,18 @@ void UMovieSceneTimeWarpTrack::OnAddedToMovieSceneImpl(UMovieScene* InMovieScene
 
 void UMovieSceneTimeWarpTrack::OnRemovedFromMovieSceneImpl()
 {
-	if (UMovieSceneTimeWarpDecoration* TimeWarp = GetTypedOuter<UMovieScene>()->FindDecoration<UMovieSceneTimeWarpDecoration>())
+	UMovieScene* MovieScene = GetTypedOuter<UMovieScene>();
+	if (MovieScene)
 	{
-		TimeWarp->RemoveTimeWarpSource(this);
+		if (UMovieSceneTimeWarpDecoration* TimeWarp = MovieScene->FindDecoration<UMovieSceneTimeWarpDecoration>())
+		{
+			TimeWarp->RemoveTimeWarpSource(this);
+
+			if (!TimeWarp->HasAnySources())
+			{
+				MovieScene->RemoveDecoration<UMovieSceneTimeWarpDecoration>();
+			}
+		}
 	}
 }
 
diff --git a/Engine/Source/Runtime/MovieScene/Public/Decorations/MovieSceneScalingAnchors.h b/Engine/Source/Runtime/MovieScene/Public/Decorations/MovieSceneScalingAnchors.h
--- a/Engine/Source/Runtime/MovieScene/Public/Decorations/MovieSceneScalingAnchors.h
+++ b/Engine/Source/Runtime/MovieScene/Public/Decorations/MovieSceneScalingAnchors.h
@@ -117,6 +117,8 @@ public:
 	 */
 	MOVIESCENE_API void RemoveScalingDriver(TScriptInterface<IMovieSceneScalingDriver> InScalingDriver);
 
+	MOVIESCENE_API bool HasAnyScalingDrivers() const;
+
 public:
 
 	/**
diff --git a/Engine/Source/Runtime/MovieScene/Public/Decorations/MovieSceneTimeWarpDecoration.h b/Engine/Source/Runtime/MovieScene/Public/Decorations/MovieSceneTimeWarpDecoration.h
--- a/Engine/Source/Runtime/MovieScene/Public/Decorations/MovieSceneTimeWarpDecoration.h
+++ b/Engine/Source/Runtime/MovieScene/Public/Decorations/MovieSceneTimeWarpDecoration.h
@@ -51,6 +51,8 @@ public:
 
 	MOVIESCENE_API void RemoveTimeWarpSource(TScriptInterface<IMovieSceneTimeWarpSource> InSource);
 
+	MOVIESCENE_API bool HasAnySources() const { return !Sources.IsEmpty(); }
+
 private:
 
 	virtual void ImportEntityImpl(UMovieSceneEntitySystemLinker* EntityLinker, const FEntityImportParams& Params, FImportedEntity* OutImportedEntity) override;
