Sequencer: add support for animating property bag prorperties

A previous change had added the ability to add property tracks for properties inside instanced struct and property bags (if they were tagged with the correct metadata), but the runtime code only supported instanced structs. This adds the missing support for property bags.

diff --git a/Engine/Source/Runtime/MovieScene/Private/MovieSceneCommonHelpers.cpp b/Engine/Source/Runtime/MovieScene/Private/MovieSceneCommonHelpers.cpp
--- a/Engine/Source/Runtime/MovieScene/Private/MovieSceneCommonHelpers.cpp
+++ b/Engine/Source/Runtime/MovieScene/Private/MovieSceneCommonHelpers.cpp
@@ -24,6 +24,7 @@
 #include "EntitySystem/MovieSceneEntitySystemLinker.h"
 #include "Conditions/MovieSceneGroupCondition.h"
 #include "StructUtils/InstancedStruct.h"
+#include "StructUtils/PropertyBag.h"
 
 bool MovieSceneHelpers::IsSectionKeyable(const UMovieSceneSection* Section)
 {
@@ -1309,6 +1310,14 @@ FTrackInstancePropertyBindings::FResolvedProperty FTrackInstancePropertyBindings
 				uint8* InstancedStructMemory = InstancedStruct->GetMutableMemory();
 				return ResolvePropertyRecursive(Bindings, InstancedStructMemory, const_cast<UScriptStruct*>(InstancedStructType), InPropertyNames, Index + 1);
 			}
+			// As above but for a property bag, unless the property path is specifically targeting its inner instanced struct.
+			else if (StructProp->Struct == FInstancedPropertyBag::StaticStruct() && InPropertyNames[Index + 1] != TEXT("Value"))
+			{
+				FInstancedPropertyBag* PropertyBag = StructProp->ContainerPtrToValuePtr<FInstancedPropertyBag>(BasePointer);
+				const UPropertyBag* PropertyBagType = PropertyBag->GetPropertyBagStruct();
+				uint8* PropertyBagMemory = PropertyBag->GetMutableValue().GetMemory();
+				return ResolvePropertyRecursive(Bindings, PropertyBagMemory, const_cast<UPropertyBag*>(PropertyBagType), InPropertyNames, Index + 1);
+			}
 			else
 			{
 				void* StructContainer = StructProp->ContainerPtrToValuePtr<void>(BasePointer);
