Sequencer: Allow returning multiple objects with FMovieSceneBindingResolveResult

FMovieSceneBindingResolveResult's Object is deprecated in favor of an array of Objects

#changelist validated
#virtualized

diff --git a/Engine/Plugins/Importers/USDImporter/Source/USDStage/Private/USDDynamicBindingResolverLibrary.cpp b/Engine/Plugins/Importers/USDImporter/Source/USDStage/Private/USDDynamicBindingResolverLibrary.cpp
--- a/Engine/Plugins/Importers/USDImporter/Source/USDStage/Private/USDDynamicBindingResolverLibrary.cpp
+++ b/Engine/Plugins/Importers/USDImporter/Source/USDStage/Private/USDDynamicBindingResolverLibrary.cpp
@@ -57,11 +57,11 @@ FMovieSceneDynamicBindingResolveResult UUsdDynamicBindingResolverLibrary::Resolv
 					{
 						if (bIsActorBinding)
 						{
-							Result.Object = FoundComponent->GetOwner();
+							Result.Objects.Add(FoundComponent->GetOwner());
 						}
 						else
 						{
-							Result.Object = FoundComponent;
+							Result.Objects.Add(FoundComponent);
 						}
 						break;
 					}
diff --git a/Engine/Source/Runtime/MovieScene/Private/Bindings/MovieSceneReplaceableBinding.cpp b/Engine/Source/Runtime/MovieScene/Private/Bindings/MovieSceneReplaceableBinding.cpp
--- a/Engine/Source/Runtime/MovieScene/Private/Bindings/MovieSceneReplaceableBinding.cpp
+++ b/Engine/Source/Runtime/MovieScene/Private/Bindings/MovieSceneReplaceableBinding.cpp
@@ -112,10 +112,13 @@ FMovieSceneBindingResolveResult UMovieSceneReplaceableBindingBase::ResolveBindin
 			if (World->WorldType == EWorldType::Editor && PreviewSpawnable)
 			{
 				FMovieSceneBindingResolveResult Result = PreviewSpawnable->ResolveBinding(ResolveParams, BindingIndex, SharedPlaybackState);
-				if (AActor* Actor = Cast<AActor>(Result.Object.Get()))
+				for (TObjectPtr<UObject> Object : Result.Objects)
 				{
-					// In addition to the spawnable tag (which the spawnable will have added), we add a replaceable tag
-					Actor->Tags.AddUnique(UE::MovieScene::ReplaceableBinding::SequencerPreviewActorTag);
+					if (AActor* Actor = Cast<AActor>(Object.Get()))
+					{
+						// In addition to the spawnable tag (which the spawnable will have added), we add a replaceable tag
+						Actor->Tags.AddUnique(UE::MovieScene::ReplaceableBinding::SequencerPreviewActorTag);
+					}
 				}
 				return Result;
 			}
diff --git a/Engine/Source/Runtime/MovieScene/Private/Bindings/MovieSceneSpawnableBinding.cpp b/Engine/Source/Runtime/MovieScene/Private/Bindings/MovieSceneSpawnableBinding.cpp
--- a/Engine/Source/Runtime/MovieScene/Private/Bindings/MovieSceneSpawnableBinding.cpp
+++ b/Engine/Source/Runtime/MovieScene/Private/Bindings/MovieSceneSpawnableBinding.cpp
@@ -198,7 +198,7 @@ FMovieSceneBindingResolveResult UMovieSceneSpawnableBindingBase::ResolveBinding(
 	UObject* SpawnedObject = SpawnRegister ? SpawnRegister->FindSpawnedObject(ResolveParams.ObjectBindingID, ResolveParams.SequenceID, BindingIndex).Get() : nullptr;
 	if (SpawnedObject)
 	{
-		Result.Object = SpawnedObject;
+		Result.Objects.Add(SpawnedObject);
 	}
 	return Result;
 }
diff --git a/Engine/Source/Runtime/MovieScene/Private/MovieSceneBindingReferences.cpp b/Engine/Source/Runtime/MovieScene/Private/MovieSceneBindingReferences.cpp
--- a/Engine/Source/Runtime/MovieScene/Private/MovieSceneBindingReferences.cpp
+++ b/Engine/Source/Runtime/MovieScene/Private/MovieSceneBindingReferences.cpp
@@ -179,8 +179,15 @@ UObject* FMovieSceneBindingReferences::ResolveBindingInternal(const FMovieSceneB
 	// If a custom binding is present and we have valid shared playback state, resolve the custom binding
 	if (SortedReferences[InternalIndex].CustomBinding && SharedPlaybackState.IsValid())
 	{
-		UObject* ResolvedObject = SortedReferences[InternalIndex].CustomBinding->ResolveBinding(BindingResolveParams, BindingIndex, SharedPlaybackState.ToSharedRef()).Object;
-		return UE::MovieScene::FindBoundObjectProxy(ResolvedObject);
+		FMovieSceneBindingResolveResult Result = SortedReferences[InternalIndex].CustomBinding->ResolveBinding(BindingResolveParams, BindingIndex, SharedPlaybackState.ToSharedRef());
+		for (TObjectPtr<UObject> ResolvedObject : Result.Objects)
+		{
+			if (ResolvedObject.Get())
+			{
+				return UE::MovieScene::FindBoundObjectProxy(ResolvedObject.Get());
+			}
+		}
+		return nullptr;
 	}
 	else
 	{
diff --git a/Engine/Source/Runtime/MovieScene/Private/MovieSceneDynamicBinding.cpp b/Engine/Source/Runtime/MovieScene/Private/MovieSceneDynamicBinding.cpp
--- a/Engine/Source/Runtime/MovieScene/Private/MovieSceneDynamicBinding.cpp
+++ b/Engine/Source/Runtime/MovieScene/Private/MovieSceneDynamicBinding.cpp
@@ -31,7 +31,7 @@ FMovieSceneDynamicBindingResolveResult UBuiltInDynamicBindingResolverLibrary::Re
 			{
 				APawn* Pawn = PlayerController->GetPawn();
 				FMovieSceneDynamicBindingResolveResult Result;
-				Result.Object = Pawn;
+				Result.Objects.Add(Pawn);
 				Result.bIsPossessedObject = true;
 				return Result;
 			}
diff --git a/Engine/Source/Runtime/MovieScene/Public/Bindings/MovieSceneCustomBinding.h b/Engine/Source/Runtime/MovieScene/Public/Bindings/MovieSceneCustomBinding.h
--- a/Engine/Source/Runtime/MovieScene/Public/Bindings/MovieSceneCustomBinding.h
+++ b/Engine/Source/Runtime/MovieScene/Public/Bindings/MovieSceneCustomBinding.h
@@ -28,9 +28,14 @@ struct FMovieSceneBindingResolveResult
 {
 	GENERATED_BODY()
 
-	/** The resolved object */
-	UPROPERTY(EditAnywhere, BlueprintReadWrite, Category="General")
-	TObjectPtr<UObject> Object = nullptr;
+	/** The resolved objects */
+	UPROPERTY(EditAnywhere, BlueprintReadWrite, Category = "General")
+	TArray<TObjectPtr<UObject>> Objects;
+
+#if WITH_EDITORONLY_DATA
+	UPROPERTY(meta = (DeprecatedProperty, DeprecationMessage = "Use Objects instead."))
+	TObjectPtr<UObject> Object_DEPRECATED = nullptr;
+#endif // WITH_EDITORONLY_DATA
 };
 
 /*
diff --git a/Engine/Source/Runtime/MovieScene/Public/MovieSceneDynamicBinding.h b/Engine/Source/Runtime/MovieScene/Public/MovieSceneDynamicBinding.h
--- a/Engine/Source/Runtime/MovieScene/Public/MovieSceneDynamicBinding.h
+++ b/Engine/Source/Runtime/MovieScene/Public/MovieSceneDynamicBinding.h
@@ -91,9 +91,14 @@ struct FMovieSceneDynamicBindingResolveResult
 {
 	GENERATED_BODY()
 
-	/** The resolved object */
-	UPROPERTY(EditAnywhere, BlueprintReadWrite, Category="General")
-	TObjectPtr<UObject> Object = nullptr;
+	/** The resolved objects */
+	UPROPERTY(EditAnywhere, BlueprintReadWrite, Category = "General")
+	TArray<TObjectPtr<UObject>> Objects;
+
+#if WITH_EDITORONLY_DATA
+	UPROPERTY(meta = (DeprecatedProperty, DeprecationMessage = "This property is deprecated, please use the array of Objects instead."))
+	TObjectPtr<UObject> Object_DEPRECATED = nullptr;
+#endif // WITH_EDITORONLY_DATA
 
 	/**
 	 * Whether the resolved object is external to the sequence
diff --git a/Engine/Source/Runtime/MovieSceneTracks/Private/Bindings/MovieSceneReplaceableDirectorBlueprintBinding.cpp b/Engine/Source/Runtime/MovieSceneTracks/Private/Bindings/MovieSceneReplaceableDirectorBlueprintBinding.cpp
--- a/Engine/Source/Runtime/MovieSceneTracks/Private/Bindings/MovieSceneReplaceableDirectorBlueprintBinding.cpp
+++ b/Engine/Source/Runtime/MovieSceneTracks/Private/Bindings/MovieSceneReplaceableDirectorBlueprintBinding.cpp
@@ -90,7 +90,7 @@ FMovieSceneBindingResolveResult UMovieSceneReplaceableDirectorBlueprintBinding::
 {
 	FMovieSceneBindingResolveResult ResolveResult;
 	FMovieSceneDynamicBindingResolveResult DynamicResolveResult = FMovieSceneDynamicBindingInvoker::ResolveDynamicBinding(SharedPlaybackState, ResolveParams.Sequence, ResolveParams.SequenceID, ResolveParams.ObjectBindingID, DynamicBinding);
-	ResolveResult.Object = DynamicResolveResult.Object;
+	ResolveResult.Objects = DynamicResolveResult.Objects;
 	return ResolveResult;
 }
 
diff --git a/Engine/Source/Runtime/MovieSceneTracks/Private/Bindings/MovieSceneSpawnableDirectorBlueprintBinding.cpp b/Engine/Source/Runtime/MovieSceneTracks/Private/Bindings/MovieSceneSpawnableDirectorBlueprintBinding.cpp
--- a/Engine/Source/Runtime/MovieSceneTracks/Private/Bindings/MovieSceneSpawnableDirectorBlueprintBinding.cpp
+++ b/Engine/Source/Runtime/MovieSceneTracks/Private/Bindings/MovieSceneSpawnableDirectorBlueprintBinding.cpp
@@ -16,7 +16,7 @@ UObject* UMovieSceneSpawnableDirectorBlueprintBinding::SpawnObjectInternal(UWorl
 {
 	FMovieSceneDynamicBindingResolveResult ResolveResult = FMovieSceneDynamicBindingInvoker::ResolveDynamicBinding(SharedPlaybackState, MovieScene.GetTypedOuter<UMovieSceneSequence>(), TemplateID, BindingId, DynamicBinding);
 
-	return ResolveResult.Object.Get();
+	return ResolveResult.Objects.Num() > 0 ? ResolveResult.Objects[0].Get() : nullptr;
 }
 
 void UMovieSceneSpawnableDirectorBlueprintBinding::DestroySpawnedObjectInternal(UObject* Object)
diff --git a/Engine/Source/Runtime/UMG/Private/Animation/WidgetAnimationBinding.cpp b/Engine/Source/Runtime/UMG/Private/Animation/WidgetAnimationBinding.cpp
--- a/Engine/Source/Runtime/UMG/Private/Animation/WidgetAnimationBinding.cpp
+++ b/Engine/Source/Runtime/UMG/Private/Animation/WidgetAnimationBinding.cpp
@@ -22,9 +22,9 @@ UObject* FWidgetAnimationBinding::FindRuntimeObject(const UWidgetTree& WidgetTre
 	if (Sequence && SharedPlaybackState && DynamicBinding.Function)
 	{
 		FMovieSceneDynamicBindingResolveResult ResolveResult = FMovieSceneDynamicBindingInvoker::ResolveDynamicBinding(SharedPlaybackState.ToSharedRef(), const_cast<UMovieSceneSequence*>(Sequence), MovieSceneSequenceID::Root, AnimationGuid, DynamicBinding);
-		if (ResolveResult.Object)
+		if (ResolveResult.Objects.Num())
 		{
-			return ResolveResult.Object;
+			return ResolveResult.Objects[0];
 		}
 	}
 
