Sequencer: Fixed potential for traits to be instantiated using const channel types and failing to use the correct channel traits

diff --git a/Engine/Source/Runtime/MovieScene/Public/Channels/MovieSceneChannelTraits.h b/Engine/Source/Runtime/MovieScene/Public/Channels/MovieSceneChannelTraits.h
--- a/Engine/Source/Runtime/MovieScene/Public/Channels/MovieSceneChannelTraits.h
+++ b/Engine/Source/Runtime/MovieScene/Public/Channels/MovieSceneChannelTraits.h
@@ -63,6 +63,9 @@ struct TMovieSceneChannelTraits : TMovieSceneChannelTraitsBase<ChannelType>
 {
 };
 
+template<typename ChannelType>
+using TGetMovieSceneChannelTraits = TMovieSceneChannelTraits<typename std::remove_cv_t<ChannelType>>;
+
 namespace UE
 {
 namespace MovieScene
@@ -194,12 +197,12 @@ namespace MovieScene
 	 * @param DefaultValue   The new default value
 	 */
 	template<typename ChannelType, typename ValueType>
-	typename TEnableIf<TMovieSceneChannelTraits<ChannelType>::SupportsDefaults>::Type SetChannelDefault(ChannelType* Channel, ValueType&& DefaultValue)
+	typename TEnableIf<TGetMovieSceneChannelTraits<ChannelType>::SupportsDefaults>::Type SetChannelDefault(ChannelType* Channel, ValueType&& DefaultValue)
 	{
 		Channel->SetDefault(DefaultValue);
 	}
 	template<typename ChannelType, typename ValueType>
-	typename TEnableIf<!TMovieSceneChannelTraits<ChannelType>::SupportsDefaults>::Type SetChannelDefault(ChannelType* Channel, ValueType&& DefaultValue)
+	typename TEnableIf<!TGetMovieSceneChannelTraits<ChannelType>::SupportsDefaults>::Type SetChannelDefault(ChannelType* Channel, ValueType&& DefaultValue)
 	{}
 
 	/**
@@ -208,12 +211,12 @@ namespace MovieScene
 	* @param InChannel      The channel to remove the default from.
 	*/
 	template<typename ChannelType>
-	typename TEnableIf<TMovieSceneChannelTraits<ChannelType>::SupportsDefaults>::Type RemoveChannelDefault(ChannelType* Channel)
+	typename TEnableIf<TGetMovieSceneChannelTraits<ChannelType>::SupportsDefaults>::Type RemoveChannelDefault(ChannelType* Channel)
 	{
 		Channel->RemoveDefault();
 	}
 	template<typename ChannelType>
-	typename TEnableIf<!TMovieSceneChannelTraits<ChannelType>::SupportsDefaults>::Type RemoveChannelDefault(ChannelType* Channel)
+	typename TEnableIf<!TGetMovieSceneChannelTraits<ChannelType>::SupportsDefaults>::Type RemoveChannelDefault(ChannelType* Channel)
 	{}
 
 	/**
@@ -223,7 +226,7 @@ namespace MovieScene
 	* @return	If there is a default value or not.
 	*/
 	template<typename ChannelType, typename ValueType>
-	typename TEnableIf<TMovieSceneChannelTraits<ChannelType>::SupportsDefaults, bool>::Type GetChannelDefault(ChannelType* Channel, ValueType& OutDefaultValue)
+	typename TEnableIf<TGetMovieSceneChannelTraits<ChannelType>::SupportsDefaults, bool>::Type GetChannelDefault(ChannelType* Channel, ValueType& OutDefaultValue)
 	{
 		if (Channel->GetDefault().IsSet())
 		{
@@ -234,7 +237,7 @@ namespace MovieScene
 	}
 
 	template<typename ChannelType, typename ValueType>
-	typename TEnableIf<!TMovieSceneChannelTraits<ChannelType>::SupportsDefaults, bool>::Type GetChannelDefault(ChannelType* Channel, ValueType& OutDefaultValue)
+	typename TEnableIf<!TGetMovieSceneChannelTraits<ChannelType>::SupportsDefaults, bool>::Type GetChannelDefault(ChannelType* Channel, ValueType& OutDefaultValue)
 	{
 		return false;
 	}
