Sequencer: Made blend type FMovieSceneBlendTypeField begin/end friend functions

diff --git a/Engine/Source/Runtime/MovieScene/Public/Evaluation/Blending/MovieSceneBlendType.h b/Engine/Source/Runtime/MovieScene/Public/Evaluation/Blending/MovieSceneBlendType.h
--- a/Engine/Source/Runtime/MovieScene/Public/Evaluation/Blending/MovieSceneBlendType.h
+++ b/Engine/Source/Runtime/MovieScene/Public/Evaluation/Blending/MovieSceneBlendType.h
@@ -11,6 +11,8 @@
 
 #include "MovieSceneBlendType.generated.h"
 
+struct FMovieSceneBlendTypeFieldIterator;
+
 /** Movie scene blend type enumeration */
 UENUM(BlueprintType)
 enum class EMovieSceneBlendType : uint8
@@ -150,6 +152,9 @@ struct FMovieSceneBlendTypeField
 		return A.BlendTypeField != B.BlendTypeField;
 	}
 
+	friend MOVIESCENE_API FMovieSceneBlendTypeFieldIterator begin(const FMovieSceneBlendTypeField& Field);
+	friend MOVIESCENE_API FMovieSceneBlendTypeFieldIterator end(const FMovieSceneBlendTypeField& Field);
+
 private:
 	explicit FMovieSceneBlendTypeField(EMovieSceneBlendType InBlendTypeField) : BlendTypeField(InBlendTypeField) {}
 
@@ -183,6 +188,3 @@ private:
 	int8 Offset;
 };
 
-
-MOVIESCENE_API FMovieSceneBlendTypeFieldIterator begin(const FMovieSceneBlendTypeField& Field);
-MOVIESCENE_API FMovieSceneBlendTypeFieldIterator end(const FMovieSceneBlendTypeField& Field);
