MovieSceneEntityBuilder factor out common code to shared base class

diff --git a/Engine/Source/Runtime/MovieScene/Public/EntitySystem/MovieSceneEntityBuilder.h b/Engine/Source/Runtime/MovieScene/Public/EntitySystem/MovieSceneEntityBuilder.h
--- a/Engine/Source/Runtime/MovieScene/Public/EntitySystem/MovieSceneEntityBuilder.h
+++ b/Engine/Source/Runtime/MovieScene/Public/EntitySystem/MovieSceneEntityBuilder.h
@@ -200,33 +200,8 @@ struct TEntityBuilderImpl<TIntegerSequence<int>>
 	TEntityBuilder< U > Append(U&& InOther);
 };
 
-
-template <typename... T, int... Indices>
-struct TEntityBuilderImpl<TIntegerSequence<int, Indices...>, T...> : IEntityBuilder
+struct FEntityBuilderImplBase : public IEntityBuilder
 {
-	virtual FMovieSceneEntityID Create(FEntityManager* EntityManager) override final
-	{
-		return CreateEntity(EntityManager);
-	}
-
-	virtual void GenerateType(FEntityManager* EntityManager, FComponentMask& OutMask, bool& OutAddMutualComponents) override final
-	{
-		VisitTupleElements([&OutMask](auto& In){ In.AccumulateMask(OutMask); }, this->Payload);
-
-		if (bAddMutualComponents)
-		{
-			OutAddMutualComponents = true;
-		}
-	}
-
-	virtual void Initialize(FEntityManager* EntityManager, const FEntityInfo& Entity) override final
-	{
-		if (Entity.Data.Allocation != nullptr)
-		{
-			VisitTupleElements([Entity](auto& In){ In.Apply(Entity.Data.Allocation, Entity.Data.ComponentOffset); }, this->Payload);
-		}
-	}
-
 	/**
 	 * Create a new entity using this builder's definition by moving the payload components into the new entity.
 	 * @note Will invalidate this instance of TEntityBuilder so its payload cannot be used again.
@@ -255,7 +230,35 @@ struct TEntityBuilderImpl<TIntegerSequence<int, Indices...>, T...> : IEntityBuil
 
 		return Entry.EntityID;
 	}
+};
+
+template <typename... T, int... Indices>
+struct TEntityBuilderImpl<TIntegerSequence<int, Indices...>, T...> : FEntityBuilderImplBase
+{
+	virtual FMovieSceneEntityID Create(FEntityManager* EntityManager) override final
+	{
+		return CreateEntity(EntityManager);
+	}
+
+	virtual void GenerateType(FEntityManager* EntityManager, FComponentMask& OutMask, bool& OutAddMutualComponents) override final
+	{
+		VisitTupleElements([&OutMask](auto& In){ In.AccumulateMask(OutMask); }, this->Payload);
+
+		if (bAddMutualComponents)
+		{
+			OutAddMutualComponents = true;
+		}
+	}
+
+	virtual void Initialize(FEntityManager* EntityManager, const FEntityInfo& Entity) override final
+	{
+		if (Entity.Data.Allocation != nullptr)
+		{
+			VisitTupleElements([Entity](auto& In){ In.Apply(Entity.Data.Allocation, Entity.Data.ComponentOffset); }, this->Payload);
+		}
+	}
 
+	using FEntityBuilderImplBase::CreateEntity;
 
 	/**
 	 * Replace the components of an entity with this builder's definition.
@@ -665,4 +668,4 @@ TEntityBuilder<T...> TEntityBuilderImpl<TIntegerSequence<int, Indices...>, T...>
 }
 
 } // namespace MovieScene
-} // namespace UE
\ No newline at end of file
+} // namespace UE
