Sequencer: Improved property traits variants, added type-erased property values, and unified many property tack editors

This is a large changelist comprising many changes which are intertwined enough to make it untenable to unpick:

------------------------------------------------------
Integers now operate on int64

------------------------------------------------------
Added types for representing property values and intermediate property values

This is used to unify all the various boilerplate code we have surrounding property value recomposition and autokeying. This essentially solves the problem of having to reimplement code that queries property interrogators and deconstructs/reconstructs composite properties for every property type by giving us a single API for all these values.

Two types are provided: FSourcePropertyValue (specifically for representing an actual UObject property address), and FIntermediatePropertyValue (for representing the intermeidate or working data type that each property trait type defines). This allows us to completely remove many of our property track editors and unify that code into a single property track editor that can operate entirely using the property traits defined by each property.

------------------------------------------------------
Added unified registered property track editor that provides track editor functionality for any property type registered with FPropertyRegistry.

Registered property track editors are a fallback for any property that isn't explicitly supported by an existing track editor, and allow us to use the existing property registry mechanisms to provide all our track editor functionality for extensible property traits (which are in an upcoming review...)

This allows us to use the same track editors for FVector and its Verse equivalent without having to make the vector track editor link to the Verse module (which is not possible).

------------------------------------------------------
Added variant property traits and unified colors, vectors, and transforms underneath

The old 'direct' and 'indirect' ways of representing a property to Sequencer are no longer expressive enough to support all the different types that we want to represent statically at compile time. In the new design, TPropertyTraits represents a property that only supports a single type. Its intermediate representation to Sequencer may or may not be the same as the property type itself.

TVariantPropertyTraits represents a compile-time collection of types that are semantically grouped together under a single schema (ie, n-dimensional vectors). It utilizes meta-data to store a unique variant index that corresponds to the compile-time position of the animated type in the template parameter list alongside its ECS entities. This variant index supercedes the manual meta-data that was previously implemented for colors and vectors (to store the color type, and vector dimensionality) and is used to statically index into function jump tables to specific TPropertyTraits functions. I tried a few different methods for 'overloading' the template functions (including using template recursion) and this static-function-table yielded the best results.

For example: TVariantPropertyTraits<FIntermediateVector, FVector2f, FVector2d, FVector3f, FVector3d> defines traits for 2-, 3-dimensional single and double precision vectors. FVector2f uses variant index 0, FVector2d uses index 1, and so on.

TDynamicVariantPropertyTraits is an extension of TVariantPropertyTraits but adds a dynamic array of variants that allows runtime extension of traits at a small performance cost. Low-indices are still reserved for  the TVariantPropertyTraits template variants, and high-indices are used as the dynamic variant identifiers. Calling any function on these traits is still O(1), apart from the 'query'like functions that resolve properties or perform matches against each variant (which are O(N) operations).

In combination, these new traits allow us to remove all the manual code for colors and vectors, whilst also supporting runtime extension for Verse types with minimal code.

------------------------------------------------------
Added a function for retriving the logical path to a target object (FSequencerUtilities::GetPathToObject) and a way for an object schema to redirect objects

This is necessary for accurate creation of parent bindings for Scene Graph, and where possession of an AEntityActor should result in a binding to its wrapped verse::entity

diff --git a/Engine/Plugins/Animation/ControlRig/Source/ControlRig/Private/Sequencer/MovieSceneControlRigSystem.cpp b/Engine/Plugins/Animation/ControlRig/Source/ControlRig/Private/Sequencer/MovieSceneControlRigSystem.cpp
--- a/Engine/Plugins/Animation/ControlRig/Source/ControlRig/Private/Sequencer/MovieSceneControlRigSystem.cpp
+++ b/Engine/Plugins/Animation/ControlRig/Source/ControlRig/Private/Sequencer/MovieSceneControlRigSystem.cpp
@@ -1274,7 +1274,7 @@ struct FGatherControlRigParameterValues
 
 	void ForEachAllocation(FEntityAllocationIteratorItem Item, const FAccumulatedControlEntryIndex* EntryIndices,
 		const FName* ParameterNames,
-		const FMovieSceneControlRigSpaceBaseKey* OptSpaceResults, const bool* OptBoolResults, const uint8* OptByteResults, const int32* OptIntegerResults,
+		const FMovieSceneControlRigSpaceBaseKey* OptSpaceResults, const bool* OptBoolResults, const uint8* OptByteResults, const int64* OptIntegerResults,
 		const double* OptDoubleResults0, const double* OptDoubleResults1, const double* OptDoubleResults2,
 		const double* OptDoubleResults3, const double* OptDoubleResults4, const double* OptDoubleResults5,
 		const double* OptDoubleResults6, const double* OptDoubleResults7, const double* OptDoubleResults8
@@ -1314,7 +1314,7 @@ struct FGatherControlRigParameterValues
 			check(OptIntegerResults);
 			for (int32 Index = 0; Index < Num; ++Index)
 			{
-				AccumulatedValues->Store(EntryIndices[Index], OptIntegerResults[Index]);
+				AccumulatedValues->Store(EntryIndices[Index], static_cast<int32>(OptIntegerResults[Index]));
 			}
 		}
 		else if (AllocationType.Contains(TracksComponents->Parameters.Scalar.PropertyTag))
diff --git a/Engine/Plugins/Animation/ControlRig/Source/ControlRig/Public/Sequencer/MovieSceneControlRigSpaceChannel.h b/Engine/Plugins/Animation/ControlRig/Source/ControlRig/Public/Sequencer/MovieSceneControlRigSpaceChannel.h
--- a/Engine/Plugins/Animation/ControlRig/Source/ControlRig/Public/Sequencer/MovieSceneControlRigSpaceChannel.h
+++ b/Engine/Plugins/Animation/ControlRig/Source/ControlRig/Public/Sequencer/MovieSceneControlRigSpaceChannel.h
@@ -65,6 +65,8 @@ struct FSpaceRange
 USTRUCT()
 struct FMovieSceneControlRigSpaceChannel : public FMovieSceneChannel
 {
+	using CurveValueType = FMovieSceneControlRigSpaceBaseKey;
+
 	GENERATED_BODY()
 
 	FMovieSceneControlRigSpaceChannel()
diff --git a/Engine/Plugins/Cameras/GameplayCameras/Source/GameplayCameras/Private/MovieScene/MovieSceneGameplayCamerasComponentTypes.cpp b/Engine/Plugins/Cameras/GameplayCameras/Source/GameplayCameras/Private/MovieScene/MovieSceneGameplayCamerasComponentTypes.cpp
--- a/Engine/Plugins/Cameras/GameplayCameras/Source/GameplayCameras/Private/MovieScene/MovieSceneGameplayCamerasComponentTypes.cpp
+++ b/Engine/Plugins/Cameras/GameplayCameras/Source/GameplayCameras/Private/MovieScene/MovieSceneGameplayCamerasComponentTypes.cpp
@@ -10,6 +10,18 @@
 #include "MovieSceneTracksComponentTypes.h"
 #include "Systems/MovieScenePiecewiseDoubleBlenderSystem.h"
 
+
+namespace UE::MovieScene
+{
+	void UnpackChannelsFromOperational(FCameraFramingZone InZone, const FProperty& Property, FUnpackedChannelValues& OutUnpackedValues)
+	{
+		OutUnpackedValues.Add(UE_MOVIESCENE_UNPACKED_MEMBER(FMovieSceneDoubleChannel, 0, InZone, Left));
+		OutUnpackedValues.Add(UE_MOVIESCENE_UNPACKED_MEMBER(FMovieSceneDoubleChannel, 1, InZone, Top));
+		OutUnpackedValues.Add(UE_MOVIESCENE_UNPACKED_MEMBER(FMovieSceneDoubleChannel, 2, InZone, Right));
+		OutUnpackedValues.Add(UE_MOVIESCENE_UNPACKED_MEMBER(FMovieSceneDoubleChannel, 3, InZone, Bottom));
+	}
+}
+
 namespace UE::Cameras
 {
 
diff --git a/Engine/Plugins/Cameras/GameplayCameras/Source/GameplayCameras/Private/MovieScene/MovieSceneGameplayCamerasComponentTypes.h b/Engine/Plugins/Cameras/GameplayCameras/Source/GameplayCameras/Private/MovieScene/MovieSceneGameplayCamerasComponentTypes.h
--- a/Engine/Plugins/Cameras/GameplayCameras/Source/GameplayCameras/Private/MovieScene/MovieSceneGameplayCamerasComponentTypes.h
+++ b/Engine/Plugins/Cameras/GameplayCameras/Source/GameplayCameras/Private/MovieScene/MovieSceneGameplayCamerasComponentTypes.h
@@ -6,7 +6,13 @@
 #include "EntitySystem/MovieSceneEntityIDs.h"
 #include "EntitySystem/MovieScenePropertySystemTypes.h"
 #include "EntitySystem/MovieScenePropertyTraits.h"
-#include "EntitySystem/MovieScenePropertyMetaDataTraits.h"
+#include "EntitySystem/MovieScenePropertyMetaData.h"
+
+
+namespace UE::MovieScene
+{
+	void UnpackChannelsFromOperational(FCameraFramingZone InZone, const FProperty& Property, FUnpackedChannelValues& OutUnpackedValues);
+}
 
 namespace UE::Cameras
 {
diff --git a/Engine/Plugins/Experimental/MovieSceneAnimMixer/Source/MovieSceneAnimMixer/Public/AnimMixerComponentTypes.h b/Engine/Plugins/Experimental/MovieSceneAnimMixer/Source/MovieSceneAnimMixer/Public/AnimMixerComponentTypes.h
--- a/Engine/Plugins/Experimental/MovieSceneAnimMixer/Source/MovieSceneAnimMixer/Public/AnimMixerComponentTypes.h
+++ b/Engine/Plugins/Experimental/MovieSceneAnimMixer/Source/MovieSceneAnimMixer/Public/AnimMixerComponentTypes.h
@@ -4,7 +4,7 @@
 
 #include "EntitySystem/BuiltInComponentTypes.h"
 #include "EntitySystem/MovieSceneEntityIDs.h"
-#include "EntitySystem/MovieScenePropertyMetaDataTraits.h"
+#include "EntitySystem/MovieScenePropertyMetaData.h"
 #include "EntitySystem/MovieScenePropertySystemTypes.h"
 #include "EntitySystem/MovieScenePropertyTraits.h"
 #include "MovieSceneMixedAnimationTarget.h"
diff --git a/Engine/Plugins/MovieScene/TemplateSequence/Source/TemplateSequence/Private/CameraAnimationSequencePlayer.cpp b/Engine/Plugins/MovieScene/TemplateSequence/Source/TemplateSequence/Private/CameraAnimationSequencePlayer.cpp
--- a/Engine/Plugins/MovieScene/TemplateSequence/Source/TemplateSequence/Private/CameraAnimationSequencePlayer.cpp
+++ b/Engine/Plugins/MovieScene/TemplateSequence/Source/TemplateSequence/Private/CameraAnimationSequencePlayer.cpp
@@ -52,12 +52,12 @@ struct TUpdateInitialPropertyValuesImpl<PropertyTraits, TIntegerSequence<int, Me
 
 		const FPropertyDefinition& PropertyDefinition = BuiltInComponents->PropertyRegistry.GetDefinition(PropertyComponents.CompositeID);
 
-		TGetPropertyValues<PropertyTraits> GetProperties(PropertyDefinition.CustomPropertyRegistration);
+		TGetPropertyValues<PropertyTraits> GetProperties(&PropertyComponents.PropertyTraits, PropertyDefinition.CustomPropertyRegistration);
 
 		FEntityTaskBuilder()
 			.Read(BuiltInComponents->BoundObject)
 			.ReadOneOf(BuiltInComponents->CustomPropertyIndex, BuiltInComponents->FastPropertyOffset, BuiltInComponents->SlowProperty)
-			.ReadAllOf(PropertyComponents.MetaDataComponents.template GetType<MetaDataIndices>()...)
+			.ReadAllOf(PropertyComponents.MetaDataComponents)
 			.Write(PropertyComponents.InitialValue)
 			.FilterAll({ PropertyComponents.PropertyTag })
 			.SetDesiredThread(Linker->EntityManager.GetGatherThread())
diff --git a/Engine/Source/Editor/MovieSceneTools/Private/MovieSceneToolsModule.cpp b/Engine/Source/Editor/MovieSceneTools/Private/MovieSceneToolsModule.cpp
--- a/Engine/Source/Editor/MovieSceneTools/Private/MovieSceneToolsModule.cpp
+++ b/Engine/Source/Editor/MovieSceneTools/Private/MovieSceneToolsModule.cpp
@@ -15,22 +15,27 @@
 #include "K2Node_CustomEvent.h"
 #include "K2Node_FunctionEntry.h"
 #include "KismetCompiler.h"
+
+#include "Tracks/MovieSceneTransformTrack.h"
+
 #include "Sections/MovieSceneEventSectionBase.h"
 #include "Sections/MovieSceneParticleSection.h"
 
-#include "TrackEditors/PropertyTrackEditors/BoolPropertyTrackEditor.h"
+#include "Sections/BoolPropertySection.h"
+#include "Sections/ColorPropertySection.h"
+
+#include "Sections/MovieSceneColorSection.h"
+#include "Tracks/MovieSceneColorTrack.h"
+
+#include "Sections/MovieSceneBoolSection.h"
+#include "Tracks/MovieSceneBoolTrack.h"
+
 #include "TrackEditors/PropertyTrackEditors/BytePropertyTrackEditor.h"
-#include "TrackEditors/PropertyTrackEditors/ColorPropertyTrackEditor.h"
-#include "TrackEditors/PropertyTrackEditors/FloatPropertyTrackEditor.h"
-#include "TrackEditors/PropertyTrackEditors/DoublePropertyTrackEditor.h"
-#include "TrackEditors/PropertyTrackEditors/IntegerPropertyTrackEditor.h"
-#include "TrackEditors/PropertyTrackEditors/VectorPropertyTrackEditor.h"
-#include "TrackEditors/PropertyTrackEditors/TransformPropertyTrackEditor.h"
-#include "TrackEditors/PropertyTrackEditors/EulerTransformPropertyTrackEditor.h"
 #include "TrackEditors/PropertyTrackEditors/RotatorPropertyTrackEditor.h"
 #include "TrackEditors/PropertyTrackEditors/VisibilityPropertyTrackEditor.h"
 #include "TrackEditors/PropertyTrackEditors/ActorReferencePropertyTrackEditor.h"
 #include "TrackEditors/PropertyTrackEditors/StringPropertyTrackEditor.h"
+#include "TrackEditors/PropertyTrackEditors/RegisteredPropertyTrackEditor.h"
 #include "TrackEditors/PropertyTrackEditors/TextPropertyTrackEditor.h"
 
 #include "TrackEditors/TransformTrackEditor.h"
@@ -104,6 +109,8 @@
 #include "Channels/ConstraintChannelCurveModel.h"
 #include "Channels/SCurveEditorKeyBarView.h"
 #include "Sections/MovieSceneEventSection.h"
+#include "Sections/TransformPropertySection.h"
+
 
 #include "Channels/MovieSceneDoublePerlinNoiseChannelContainer.h"
 #include "Channels/MovieSceneFloatPerlinNoiseChannelContainer.h"
@@ -143,6 +150,7 @@
 
 TAutoConsoleVariable<bool> CVarDuplicateLinkedAnimSequence(TEXT("Sequencer.DuplicateLinkedAnimSequence"), false, TEXT("When true when we duplicate a level sequence that has a linked anim sequence it will duplicate and link the anim sequencel, if false we leave any link alone."));
 
+
 void FMovieSceneToolsModule::StartupModule()
 {
 	using namespace UE::Sequencer;
@@ -160,17 +168,19 @@ void FMovieSceneToolsModule::StartupModule()
 
 		ISequencerModule& SequencerModule = FModuleManager::Get().LoadModuleChecked<ISequencerModule>( "Sequencer" );
 
+		RegisterPropertyType(MakeShared<TRegisteredProperty<UMovieSceneBoolTrack, UMovieSceneBoolSection>>([](UMovieSceneSection* SectionObject, UMovieSceneTrack* Track, const FGuid& ObjectBinding, TWeakPtr<ISequencer> WeakSequencer){
+			return TSharedPtr<ISequencerSection>(MakeShared<FBoolPropertySection>(*SectionObject));
+		}));
+		RegisterPropertyType(MakeShared<TRegisteredProperty<UMovieSceneTransformTrack, UMovieScene3DTransformSection>>([](UMovieSceneSection* SectionObject, UMovieSceneTrack* Track, const FGuid& ObjectBinding, TWeakPtr<ISequencer> WeakSequencer){
+			return TSharedPtr<ISequencerSection>(MakeShared<FTransformSection>(*SectionObject, WeakSequencer));
+		}));
+		RegisterPropertyType(MakeShared<TRegisteredProperty<UMovieSceneColorTrack, UMovieSceneColorSection>>([](UMovieSceneSection* SectionObject, UMovieSceneTrack* Track, const FGuid& ObjectBinding, TWeakPtr<ISequencer> WeakSequencer){
+			return TSharedPtr<ISequencerSection>(MakeShared<FColorPropertySection>(*SectionObject, ObjectBinding, WeakSequencer));
+		}));
+
+
 		// register property track editors
-		BoolPropertyTrackCreateEditorHandle = SequencerModule.RegisterPropertyTrackEditor<FBoolPropertyTrackEditor>();
 		BytePropertyTrackCreateEditorHandle = SequencerModule.RegisterPropertyTrackEditor<FBytePropertyTrackEditor>();
-		ColorPropertyTrackCreateEditorHandle = SequencerModule.RegisterPropertyTrackEditor<FColorPropertyTrackEditor>();
-		FloatPropertyTrackCreateEditorHandle = SequencerModule.RegisterPropertyTrackEditor<FFloatPropertyTrackEditor>();
-		DoublePropertyTrackCreateEditorHandle = SequencerModule.RegisterPropertyTrackEditor<FDoublePropertyTrackEditor>();
-		IntegerPropertyTrackCreateEditorHandle = SequencerModule.RegisterPropertyTrackEditor<FIntegerPropertyTrackEditor>();
-		FloatVectorPropertyTrackCreateEditorHandle = SequencerModule.RegisterPropertyTrackEditor<FFloatVectorPropertyTrackEditor>();
-		DoubleVectorPropertyTrackCreateEditorHandle = SequencerModule.RegisterPropertyTrackEditor<FDoubleVectorPropertyTrackEditor>();
-		TransformPropertyTrackCreateEditorHandle = SequencerModule.RegisterPropertyTrackEditor<FTransformPropertyTrackEditor>();
-		EulerTransformPropertyTrackCreateEditorHandle = SequencerModule.RegisterPropertyTrackEditor<FEulerTransformPropertyTrackEditor>();
 		RotatorPropertyTrackCreateEditorHandle = SequencerModule.RegisterPropertyTrackEditor<FRotatorPropertyTrackEditor>();
 		VisibilityPropertyTrackCreateEditorHandle = SequencerModule.RegisterPropertyTrackEditor<FVisibilityPropertyTrackEditor>();
 		ActorReferencePropertyTrackCreateEditorHandle = SequencerModule.RegisterPropertyTrackEditor<FActorReferencePropertyTrackEditor>();
@@ -204,6 +214,7 @@ void FMovieSceneToolsModule::StartupModule()
 		CustomPrimitiveDataTrackCreateEditorHandle = SequencerModule.RegisterTrackEditor(FOnCreateTrackEditor::CreateStatic(&FCustomPrimitiveDataTrackEditor::CreateTrackEditor));
 		BindingLifetimeTrackCreateEditorHandle = SequencerModule.RegisterTrackEditor(FOnCreateTrackEditor::CreateStatic(&FBindingLifetimeTrackEditor::CreateTrackEditor));
 		TimeWarpTrackCreateEditorHandle = SequencerModule.RegisterTrackEditor(FOnCreateTrackEditor::CreateStatic(&FTimeWarpTrackEditor::CreateTrackEditor));
+		RegisteredPropertyTrackCreateEditorHandle = SequencerModule.RegisterTrackEditor(FOnCreateTrackEditor::CreateStatic(&FRegisteredPropertyTrackEditor::CreateTrackEditor));
 
 		// register track models
 		CameraCutTrackModelHandle = SequencerModule.RegisterTrackModel(FOnCreateTrackModel::CreateStatic(&FCameraCutTrackModel::CreateTrackModel));
@@ -356,16 +367,7 @@ void FMovieSceneToolsModule::ShutdownModule()
 	ISequencerModule& SequencerModule = FModuleManager::Get().GetModuleChecked<ISequencerModule>( "Sequencer" );
 
 	// unregister property track editors
-	SequencerModule.UnRegisterTrackEditor( BoolPropertyTrackCreateEditorHandle );
 	SequencerModule.UnRegisterTrackEditor( BytePropertyTrackCreateEditorHandle );
-	SequencerModule.UnRegisterTrackEditor( ColorPropertyTrackCreateEditorHandle );
-	SequencerModule.UnRegisterTrackEditor( FloatPropertyTrackCreateEditorHandle );
-	SequencerModule.UnRegisterTrackEditor( DoublePropertyTrackCreateEditorHandle );
-	SequencerModule.UnRegisterTrackEditor( IntegerPropertyTrackCreateEditorHandle );
-	SequencerModule.UnRegisterTrackEditor( FloatVectorPropertyTrackCreateEditorHandle );
-	SequencerModule.UnRegisterTrackEditor( DoubleVectorPropertyTrackCreateEditorHandle );
-	SequencerModule.UnRegisterTrackEditor( TransformPropertyTrackCreateEditorHandle );
-	SequencerModule.UnRegisterTrackEditor( EulerTransformPropertyTrackCreateEditorHandle );
 	SequencerModule.UnRegisterTrackEditor( RotatorPropertyTrackCreateEditorHandle );
 	SequencerModule.UnRegisterTrackEditor( VisibilityPropertyTrackCreateEditorHandle );
 	SequencerModule.UnRegisterTrackEditor( ActorReferencePropertyTrackCreateEditorHandle );
@@ -399,6 +401,7 @@ void FMovieSceneToolsModule::ShutdownModule()
 	SequencerModule.UnRegisterTrackEditor( CustomPrimitiveDataTrackCreateEditorHandle );
 	SequencerModule.UnRegisterTrackEditor( BindingLifetimeTrackCreateEditorHandle );
 	SequencerModule.UnRegisterTrackEditor( TimeWarpTrackCreateEditorHandle );
+	SequencerModule.UnRegisterTrackEditor( RegisteredPropertyTrackCreateEditorHandle );
 
 	// unregister track models
 	SequencerModule.UnregisterTrackModel( CameraCutTrackModelHandle );
diff --git a/Engine/Source/Editor/MovieSceneTools/Private/TrackEditors/FadeTrackEditor.cpp b/Engine/Source/Editor/MovieSceneTools/Private/TrackEditors/FadeTrackEditor.cpp
--- a/Engine/Source/Editor/MovieSceneTools/Private/TrackEditors/FadeTrackEditor.cpp
+++ b/Engine/Source/Editor/MovieSceneTools/Private/TrackEditors/FadeTrackEditor.cpp
@@ -5,7 +5,6 @@
 #include "Channels/MovieSceneFloatChannel.h"
 #include "Containers/Array.h"
 #include "Delegates/Delegate.h"
-#include "TrackEditors/PropertyTrackEditors/FloatPropertyTrackEditor.h"
 #include "Framework/Commands/UIAction.h"
 #include "Framework/MultiBox/MultiBoxBuilder.h"
 #include "HAL/PlatformCrt.h"
@@ -120,7 +119,7 @@ TSharedRef<ISequencerTrackEditor> FFadeTrackEditor::CreateTrackEditor(TSharedRef
  *****************************************************************************/
 
 FFadeTrackEditor::FFadeTrackEditor(TSharedRef<ISequencer> InSequencer)
-	: FFloatPropertyTrackEditor(InSequencer)
+	: FKeyframeTrackEditor(InSequencer)
 { }
 
 /* ISequencerTrackEditor interface
diff --git a/Engine/Source/Editor/MovieSceneTools/Private/TrackEditors/FadeTrackEditor.h b/Engine/Source/Editor/MovieSceneTools/Private/TrackEditors/FadeTrackEditor.h
--- a/Engine/Source/Editor/MovieSceneTools/Private/TrackEditors/FadeTrackEditor.h
+++ b/Engine/Source/Editor/MovieSceneTools/Private/TrackEditors/FadeTrackEditor.h
@@ -4,7 +4,8 @@
 
 #include "Templates/SharedPointer.h"
 #include "Templates/SubclassOf.h"
-#include "TrackEditors/PropertyTrackEditors/FloatPropertyTrackEditor.h"
+#include "KeyframeTrackEditor.h"
+#include "Tracks/MovieSceneFadeTrack.h"
 
 class FMenuBuilder;
 class ISequencer;
@@ -20,7 +21,7 @@ struct FSlateBrush;
 * A property track editor for fade control.
 */
 class FFadeTrackEditor
-	: public FFloatPropertyTrackEditor
+	: public FKeyframeTrackEditor<UMovieSceneFadeTrack>
 {
 public:
 
diff --git a/Engine/Source/Editor/MovieSceneTools/Private/TrackEditors/SlomoTrackEditor.cpp b/Engine/Source/Editor/MovieSceneTools/Private/TrackEditors/SlomoTrackEditor.cpp
--- a/Engine/Source/Editor/MovieSceneTools/Private/TrackEditors/SlomoTrackEditor.cpp
+++ b/Engine/Source/Editor/MovieSceneTools/Private/TrackEditors/SlomoTrackEditor.cpp
@@ -4,7 +4,6 @@
 
 #include "Containers/Array.h"
 #include "Delegates/Delegate.h"
-#include "TrackEditors/PropertyTrackEditors/FloatPropertyTrackEditor.h"
 #include "Framework/Commands/UIAction.h"
 #include "Framework/MultiBox/MultiBoxBuilder.h"
 #include "ISequencer.h"
@@ -42,7 +41,7 @@ TSharedRef<ISequencerTrackEditor> FSlomoTrackEditor::CreateTrackEditor(TSharedRe
  *****************************************************************************/
 
 FSlomoTrackEditor::FSlomoTrackEditor(TSharedRef<ISequencer> InSequencer)
-	: FFloatPropertyTrackEditor(InSequencer)
+	: FKeyframeTrackEditor(InSequencer)
 { }
 
 
diff --git a/Engine/Source/Editor/MovieSceneTools/Private/TrackEditors/SlomoTrackEditor.h b/Engine/Source/Editor/MovieSceneTools/Private/TrackEditors/SlomoTrackEditor.h
--- a/Engine/Source/Editor/MovieSceneTools/Private/TrackEditors/SlomoTrackEditor.h
+++ b/Engine/Source/Editor/MovieSceneTools/Private/TrackEditors/SlomoTrackEditor.h
@@ -4,7 +4,8 @@
 
 #include "Templates/SharedPointer.h"
 #include "Templates/SubclassOf.h"
-#include "TrackEditors/PropertyTrackEditors/FloatPropertyTrackEditor.h"
+#include "KeyframeTrackEditor.h"
+#include "Tracks/MovieSceneSlomoTrack.h"
 
 class FMenuBuilder;
 class ISequencer;
@@ -17,7 +18,7 @@ struct FSlateBrush;
 * A property track editor for slow motion control.
 */
 class FSlomoTrackEditor
-	: public FFloatPropertyTrackEditor
+	: public FKeyframeTrackEditor<UMovieSceneSlomoTrack>
 {
 public:
 
diff --git a/Engine/Source/Editor/MovieSceneTools/Private/TrackEditors/SpawnTrackEditor.cpp b/Engine/Source/Editor/MovieSceneTools/Private/TrackEditors/SpawnTrackEditor.cpp
--- a/Engine/Source/Editor/MovieSceneTools/Private/TrackEditors/SpawnTrackEditor.cpp
+++ b/Engine/Source/Editor/MovieSceneTools/Private/TrackEditors/SpawnTrackEditor.cpp
@@ -21,6 +21,7 @@
 #include "UObject/UnrealNames.h"
 #include "MovieSceneBindingReferences.h"
 #include "Bindings/MovieSceneSpawnableBinding.h"
+#include "Sections/BoolPropertySection.h"
 
 class ISequencerTrackEditor;
 
@@ -35,15 +36,20 @@ TSharedRef<ISequencerTrackEditor> FSpawnTrackEditor::CreateTrackEditor(TSharedRe
 
 
 FSpawnTrackEditor::FSpawnTrackEditor(TSharedRef<ISequencer> InSequencer)
-	: FBoolPropertyTrackEditor(InSequencer)
+	: FKeyframeTrackEditor(InSequencer)
 { }
 
+TSharedRef<ISequencerSection> FSpawnTrackEditor::MakeSectionInterface(UMovieSceneSection& SectionObject, UMovieSceneTrack& Track, FGuid ObjectBinding)
+{
+	return MakeShared<FBoolPropertySection>(SectionObject);
+}
+
 
 UMovieSceneTrack* FSpawnTrackEditor::AddTrack(UMovieScene* FocusedMovieScene, const FGuid& ObjectHandle, TSubclassOf<UMovieSceneTrack> TrackClass, FName UniqueTypeName)
 {
-	UMovieSceneTrack* NewTrack = FBoolPropertyTrackEditor::AddTrack(FocusedMovieScene, ObjectHandle, TrackClass, UniqueTypeName);
+	UMovieSceneTrack* NewTrack = FKeyframeTrackEditor::AddTrack(FocusedMovieScene, ObjectHandle, TrackClass, UniqueTypeName);
 
-	if (auto* SpawnTrack = Cast<UMovieSceneSpawnTrack>(NewTrack))
+	if (UMovieSceneSpawnTrack* SpawnTrack = Cast<UMovieSceneSpawnTrack>(NewTrack))
 	{
 		SpawnTrack->Modify();
 		SpawnTrack->SetObjectId(ObjectHandle);
diff --git a/Engine/Source/Editor/MovieSceneTools/Private/TrackEditors/TransformTrackEditor.cpp b/Engine/Source/Editor/MovieSceneTools/Private/TrackEditors/TransformTrackEditor.cpp
--- a/Engine/Source/Editor/MovieSceneTools/Private/TrackEditors/TransformTrackEditor.cpp
+++ b/Engine/Source/Editor/MovieSceneTools/Private/TrackEditors/TransformTrackEditor.cpp
@@ -606,16 +606,6 @@ void F3DTransformTrackEditor::BuildObjectBindingTrackMenu(FMenuBuilder& MenuBuil
 	}
 }
 
-
-bool F3DTransformTrackEditor::CanAddTransformTrackForActorHandle( FGuid ObjectBinding ) const
-{
-	if (GetSequencer()->GetFocusedMovieSceneSequence()->GetMovieScene()->FindTrack<UMovieScene3DTransformTrack>(ObjectBinding, TransformPropertyName))
-	{
-		return false;
-	}
-	return true;
-}
-
 EVisibility F3DTransformTrackEditor::IsCameraVisible(FGuid ObjectGuid) const
 {
 	for (auto Object : GetSequencer()->FindObjectsInCurrentSequence(ObjectGuid))
diff --git a/Engine/Source/Editor/MovieSceneTools/Private/TrackEditors/TransformTrackEditor.h b/Engine/Source/Editor/MovieSceneTools/Private/TrackEditors/TransformTrackEditor.h
--- a/Engine/Source/Editor/MovieSceneTools/Private/TrackEditors/TransformTrackEditor.h
+++ b/Engine/Source/Editor/MovieSceneTools/Private/TrackEditors/TransformTrackEditor.h
@@ -82,9 +82,6 @@ private:
 
 	void ProcessKeyOperation(UObject* ObjectToKey, TArrayView<const UE::Sequencer::FKeySectionOperation> SectionsToKey, ISequencer& InSequencer, FFrameNumber KeyTime, TArray<UE::Sequencer::FAddKeyResult>* OutResults = nullptr);
 
-	/** Returns whether or not a transform track can be added for an actor with a specific handle. */
-	bool CanAddTransformTrackForActorHandle(FGuid ActorHandle) const;
-
 	/** Whether the object has an existing transform track */
 	bool HasTransformTrack( UObject& InObject ) const;
 
diff --git a/Engine/Source/Editor/MovieSceneTools/Public/Channels/PerlinNoiseChannelInterface.h b/Engine/Source/Editor/MovieSceneTools/Public/Channels/PerlinNoiseChannelInterface.h
--- a/Engine/Source/Editor/MovieSceneTools/Public/Channels/PerlinNoiseChannelInterface.h
+++ b/Engine/Source/Editor/MovieSceneTools/Public/Channels/PerlinNoiseChannelInterface.h
@@ -207,5 +207,9 @@ struct TPerlinNoiseChannelInterface : ISequencerChannelInterface
 	{
 		return nullptr;
 	}
+	virtual EMovieSceneKeyInterpolation GetInterpolationMode_Raw(FMovieSceneChannel* InChannel, FFrameNumber InTime, EMovieSceneKeyInterpolation InDefaultMode) const override
+	{
+		return InDefaultMode;
+	}
 };
 
diff --git a/Engine/Source/Editor/MovieSceneTools/Public/KeyframeTrackEditor.h b/Engine/Source/Editor/MovieSceneTools/Public/KeyframeTrackEditor.h
--- a/Engine/Source/Editor/MovieSceneTools/Public/KeyframeTrackEditor.h
+++ b/Engine/Source/Editor/MovieSceneTools/Public/KeyframeTrackEditor.h
@@ -14,6 +14,7 @@
 #include "MovieSceneTrackEditor.h"
 #include "Framework/MultiBox/MultiBoxBuilder.h"
 #include "MovieSceneCommonHelpers.h"
+#include "Channels/MovieSceneUnpackedChannelValues.h"
 #include "Channels/MovieSceneChannelTraits.h"
 #include "Channels/MovieSceneChannel.h"
 #include "Channels/MovieSceneChannelProxy.h"
@@ -418,11 +419,26 @@ struct FMovieSceneChannelValueSetter
 
 	/** Templated construction function that can add a key (and potentially also set a default) for the specified channel and value */
 	template<typename ChannelType, typename ValueType>
-	static FMovieSceneChannelValueSetter Create(int32 ChannelIndex, ValueType InNewValue, bool bAddKey)
+	static FMovieSceneChannelValueSetter Create(int32 ChannelIndex, ValueType&& InNewValue, bool bAddKey)
 	{
+		//FMovieSceneChannelValueSetter NewValue;
+		using namespace UE::MovieScene;
+
+		FUnpackedChannelValue Value(TIndexedChannelValue<ChannelType>(Forward<ValueType>(InNewValue), ChannelIndex), NAME_None);
+
 		FMovieSceneChannelValueSetter NewValue;
-		NewValue.Impl = TAddKeyImpl<ChannelType, typename TDecay<ValueType>::Type>(ChannelIndex, bAddKey, Forward<ValueType>(InNewValue));
+		NewValue.Impl = FUnpackedChannelValueSetter(MoveTemp(Value), bAddKey);
+		return MoveTemp(NewValue);
+		//NewValue.Impl = TAddKeyImpl<ChannelType, typename TDecay<ValueType>::Type>(ChannelIndex, bAddKey, Forward<ValueType>(InNewValue));
 		
+		//return MoveTemp(NewValue);
+	}
+
+	/** Templated construction function that can add a key (and potentially also set a default) for the specified channel and value */
+	static FMovieSceneChannelValueSetter Create(UE::MovieScene::FUnpackedChannelValue&& InValue, bool bAddKey)
+	{
+		FMovieSceneChannelValueSetter NewValue;
+		NewValue.Impl = FUnpackedChannelValueSetter(MoveTemp(InValue), bAddKey);
 		return MoveTemp(NewValue);
 	}
 
@@ -441,7 +457,67 @@ private:
 	FMovieSceneChannelValueSetter()
 	{}
 
-	
+	struct FUnpackedChannelValueSetter : IImpl
+	{
+		UE::MovieScene::FUnpackedChannelValue ChannelValue;
+		bool bAddKey;
+
+		FUnpackedChannelValueSetter(UE::MovieScene::FUnpackedChannelValue&& InValue, bool bInAddKey)
+			: ChannelValue(MoveTemp(InValue)), bAddKey(bInAddKey)
+		{}
+
+		virtual bool Apply(UMovieSceneSection* Section, FMovieSceneChannelProxy& Proxy, FFrameNumber InTime, EMovieSceneKeyInterpolation InterpolationMode, bool bKeyEvenIfUnchanged, bool bKeyEvenIfEmpty) const override
+		{
+			bool bKeyCreated = false;
+			using namespace UE::MovieScene;
+
+			FMovieSceneChannel* Channel = ChannelValue->RetrieveChannel(Proxy);
+			if (bAddKey && Channel)
+			{
+				bool bShouldKeyChannel = bKeyEvenIfUnchanged;
+				if (!bShouldKeyChannel)
+				{
+					bShouldKeyChannel = !ChannelValue->AlreadyExistsAtTime(Channel, InTime);
+				}
+
+				if (bShouldKeyChannel)
+				{
+					if (Channel->GetNumKeys() != 0 || bKeyEvenIfEmpty)
+					{
+						if (Section->TryModify())
+						{
+							ChannelValue->AddKey(Channel, InTime, InterpolationMode);
+							bKeyCreated = true;
+						}
+					}
+				}
+			}
+
+			return bKeyCreated;
+		}
+
+		virtual void ApplyDefault(UMovieSceneSection* Section, FMovieSceneChannelProxy& Proxy, EKeyFrameTrackEditorSetDefault SetDefault) const override
+		{
+			using namespace UE::MovieScene;
+
+			const bool bShouldSetDefault = SetDefault == EKeyFrameTrackEditorSetDefault::SetDefault || (SetDefault == EKeyFrameTrackEditorSetDefault::SetDefaultOnAddKeys && bAddKey);
+			FMovieSceneChannel* Channel = ChannelValue->RetrieveChannel(Proxy);
+			if (bShouldSetDefault && Channel && Channel->GetNumKeys() == 0 && !ChannelValue->IsAlreadyDefault(Channel))
+			{
+				if (Section->TryModify())
+				{
+					ChannelValue->SetDefault(Channel);
+				}
+			}
+		}
+
+		virtual bool ModifyByCurrentAndWeight(FMovieSceneChannelProxy& Proxy, FFrameNumber InTime, void* VCurrentValue, float Weight) override
+		{
+			ensure(!"This should never be called any more");
+			return false;
+		}
+	};
+
 	TInlineValue<IImpl> Impl;
 };
 
@@ -647,11 +723,6 @@ private:
 			{
 				FGeneratedTrackKeys GeneratedKeys;
 				OnGenerateKeys(SectionToKey, GeneratedKeys);
-
-				if (!bTrackCreated)
-				{
-					ModifyGeneratedKeysByCurrentAndWeight(Object, Track, SectionToKey, KeyTime, GeneratedKeys, Weight);
-				}
 				KeyPropertyResult |= AddKeysToSection( SectionToKey, KeyTime, GeneratedKeys, KeyMode );
 			}
 		}
@@ -711,21 +782,7 @@ protected:
 		return KeyPropertyResult;
 	}
 
-	/**
-	* Scale the generated keys by the the current value and the weight. This is used for supporting layered workflows
-	* so the keyed value is set correctly.
-	* The issue is that the Generated Property Values are the Globa Values that we want to meet but we need to set the local
-    * value on a particular section based upon
-	* The algorithm to implement should be NewGeneratedValue = (GeneratedTotalValue - CurrentValue) * Weight + CurrentChannelValue
-	* Expect that this will only be called if blended sections are occuring, it will not get called if no blending is happening.
-	* @param Object Object we are keying
-	* @param Track The track we are on
-	* @param SectionToKey the section we are on
-	* @param Time Time we are at
-	* @param GeneratedKeys Array of keys that are modified from the changed property. This is the same one that's calculdate but it's the Total/Global
-	* @param Weight The Weight used in the algorithm above. Note it maybe be greater than one since it's the actual inverse of the weight on that section.
-	* @return Return true if this is actually implemented false if not. Usually only TrackEditors with valid BlendTypes will implement this.
-	*/
+	UE_DEPRECATED(5.7, "This function is no longer called. Please ensure that generated keys populated through GenerateKeysCallback have already been compensated.")
 	virtual bool ModifyGeneratedKeysByCurrentAndWeight(UObject *Object, UMovieSceneTrack *Track, UMovieSceneSection* SectionToKey, FFrameNumber Time, FGeneratedTrackKeys& GeneratedTotalKeys, float Weight) const
 	{
 		return false;
diff --git a/Engine/Source/Editor/MovieSceneTools/Public/MovieSceneToolsModule.h b/Engine/Source/Editor/MovieSceneTools/Public/MovieSceneToolsModule.h
--- a/Engine/Source/Editor/MovieSceneTools/Public/MovieSceneToolsModule.h
+++ b/Engine/Source/Editor/MovieSceneTools/Public/MovieSceneToolsModule.h
@@ -106,16 +106,7 @@ private:
 private:
 
 	/** Registered delegate handles */
-	FDelegateHandle BoolPropertyTrackCreateEditorHandle;
 	FDelegateHandle BytePropertyTrackCreateEditorHandle;
-	FDelegateHandle ColorPropertyTrackCreateEditorHandle;
-	FDelegateHandle FloatPropertyTrackCreateEditorHandle;
-	FDelegateHandle DoublePropertyTrackCreateEditorHandle;
-	FDelegateHandle IntegerPropertyTrackCreateEditorHandle;
-	FDelegateHandle FloatVectorPropertyTrackCreateEditorHandle;
-	FDelegateHandle DoubleVectorPropertyTrackCreateEditorHandle;
-	FDelegateHandle TransformPropertyTrackCreateEditorHandle;
-	FDelegateHandle EulerTransformPropertyTrackCreateEditorHandle;
 	FDelegateHandle RotatorPropertyTrackCreateEditorHandle;
 	FDelegateHandle VisibilityPropertyTrackCreateEditorHandle;
 	FDelegateHandle ActorReferencePropertyTrackCreateEditorHandle;
@@ -148,6 +139,7 @@ private:
 	FDelegateHandle CustomPrimitiveDataTrackCreateEditorHandle;
 	FDelegateHandle BindingLifetimeTrackCreateEditorHandle;
 	FDelegateHandle TimeWarpTrackCreateEditorHandle;
+	FDelegateHandle RegisteredPropertyTrackCreateEditorHandle;
 
 	FDelegateHandle CameraCutTrackModelHandle;
 	FDelegateHandle CinematicShotTrackModelHandle;
diff --git a/Engine/Source/Editor/MovieSceneTools/Public/TrackEditors/PropertyTrackEditors/BoolPropertyTrackEditor.h b/Engine/Source/Editor/MovieSceneTools/Public/TrackEditors/PropertyTrackEditors/BoolPropertyTrackEditor.h
--- a/Engine/Source/Editor/MovieSceneTools/Public/TrackEditors/PropertyTrackEditors/BoolPropertyTrackEditor.h
+++ b/Engine/Source/Editor/MovieSceneTools/Public/TrackEditors/PropertyTrackEditors/BoolPropertyTrackEditor.h
@@ -4,37 +4,15 @@
 
 #include "AnimatedPropertyKey.h"
 #include "Containers/Array.h"
-#include "Containers/ArrayView.h"
-#include "Containers/ContainerAllocationPolicies.h"
-#include "CoreMinimal.h"
-#include "Delegates/Delegate.h"
-#include "HAL/PlatformCrt.h"
-#include "ISequencer.h"
-#include "ISequencerSection.h"
-#include "ISequencerTrackEditor.h"
-#include "KeyframeTrackEditor.h"
-#include "Misc/Guid.h"
 #include "PropertyTrackEditor.h"
-#include "Sections/MovieSceneBoolSection.h"
-#include "Templates/SharedPointer.h"
-#include "Templates/UnrealTemplate.h"
 #include "Tracks/MovieSceneBoolTrack.h"
-#include "UObject/UnrealNames.h"
 
-#define UE_API MOVIESCENETOOLS_API
-
-class FPropertyChangedParams;
-class ISequencer;
-class ISequencerSection;
-class ISequencerTrackEditor;
-class UMovieSceneSection;
-class UMovieSceneTrack;
-struct FGuid;
+UE_DEPRECATED_HEADER(5.7, "This header is no longer in use.")
 
 /**
  * A property track editor for Booleans.
  */
-class FBoolPropertyTrackEditor
+class UE_DEPRECATED(5.7, "This class is no longer in use") FBoolPropertyTrackEditor
 	: public FPropertyTrackEditor<UMovieSceneBoolTrack>
 {
 public:
@@ -55,21 +33,4 @@ public:
 	{
 		return TArray<FAnimatedPropertyKey, TInlineAllocator<1>>({ FAnimatedPropertyKey::FromPropertyTypeName(NAME_BoolProperty) });
 	}
-
-	/**
-	 * Creates an instance of this class (called by a sequencer).
-	 *
-	 * @param OwningSequencer The sequencer instance to be used by this tool.
-	 * @return The new instance of this class.
-	 */
-	static UE_API TSharedRef<ISequencerTrackEditor> CreateTrackEditor( TSharedRef<ISequencer> OwningSequencer );
-
-protected:
-
-	//~ FPropertyTrackEditor interface
-	UE_API virtual FText GetDisplayName() const override;
-	UE_API virtual TSharedRef<ISequencerSection> MakeSectionInterface(UMovieSceneSection& SectionObject, UMovieSceneTrack& Track, FGuid ObjectBinding) override;
-	UE_API virtual void GenerateKeysFromPropertyChanged( const FPropertyChangedParams& PropertyChangedParams, UMovieSceneSection* SectionToKey, FGeneratedTrackKeys& OutGeneratedKeys ) override;
 };
-
-#undef UE_API
diff --git a/Engine/Source/Editor/MovieSceneTools/Public/TrackEditors/SpawnTrackEditor.h b/Engine/Source/Editor/MovieSceneTools/Public/TrackEditors/SpawnTrackEditor.h
--- a/Engine/Source/Editor/MovieSceneTools/Public/TrackEditors/SpawnTrackEditor.h
+++ b/Engine/Source/Editor/MovieSceneTools/Public/TrackEditors/SpawnTrackEditor.h
@@ -11,7 +11,8 @@
 #include "MovieSceneTrack.h"
 #include "Templates/SharedPointer.h"
 #include "Templates/SubclassOf.h"
-#include "TrackEditors/PropertyTrackEditors/BoolPropertyTrackEditor.h"
+#include "KeyframeTrackEditor.h"
+#include "Tracks/MovieSceneSpawnTrack.h"
 #include "UObject/NameTypes.h"
 
 #define UE_API MOVIESCENETOOLS_API
@@ -32,9 +33,7 @@ struct FGuid;
  * A property track editor for controlling the lifetime of a sapwnable object
  */
 class FSpawnTrackEditor
-	// TODO: Don't derive from bool property track editor here since it doesn't actually edit
-	// bool property tracks.
-	: public FBoolPropertyTrackEditor
+	: public FKeyframeTrackEditor<UMovieSceneSpawnTrack>
 {
 public:
 
@@ -64,7 +63,8 @@ public:
 	virtual TSharedPtr<SWidget> BuildOutlinerEditWidget(const FGuid& ObjectBinding, UMovieSceneTrack* Track, const FBuildEditWidgetParams& Params) override  { return TSharedPtr<SWidget>(); }
 	virtual bool HandleAssetAdded(UObject* Asset, const FGuid& TargetObjectGuid) override { return false; }
 	UE_API virtual bool SupportsType(TSubclassOf<UMovieSceneTrack> Type) const override;
-    UE_API virtual bool SupportsSequence(UMovieSceneSequence* InSequence) const override;
+	UE_API virtual bool SupportsSequence(UMovieSceneSequence* InSequence) const override;
+	UE_API virtual TSharedRef<ISequencerSection> MakeSectionInterface(UMovieSceneSection& SectionObject, UMovieSceneTrack& Track, FGuid ObjectBinding) override;
 
 private:
 
diff --git a/Engine/Source/Editor/Sequencer/Private/IKeyArea.cpp b/Engine/Source/Editor/Sequencer/Private/IKeyArea.cpp
--- a/Engine/Source/Editor/Sequencer/Private/IKeyArea.cpp
+++ b/Engine/Source/Editor/Sequencer/Private/IKeyArea.cpp
@@ -515,3 +515,15 @@ bool IKeyArea::ShouldShowCurve() const
 
 	return false;
 }
+
+EMovieSceneKeyInterpolation IKeyArea::GetInterpolationMode(FFrameNumber InTime, EMovieSceneKeyInterpolation DefaultMode) const
+{
+	ISequencerChannelInterface* EditorInterface = FindChannelEditorInterface();
+	FMovieSceneChannel* Channel = ChannelHandle.Get();
+	if (EditorInterface && Channel)
+	{
+		return EditorInterface->GetInterpolationMode_Raw(Channel, InTime, DefaultMode);
+	}
+
+	return DefaultMode;
+}
diff --git a/Engine/Source/Editor/Sequencer/Private/ISequencerObjectSchema.cpp b/Engine/Source/Editor/Sequencer/Private/ISequencerObjectSchema.cpp
--- a/Engine/Source/Editor/Sequencer/Private/ISequencerObjectSchema.cpp
+++ b/Engine/Source/Editor/Sequencer/Private/ISequencerObjectSchema.cpp
@@ -73,4 +73,9 @@ FText IObjectSchema::GetPrettyName(const UObject* Object) const
 	return FText::FromString(Object->GetName());
 }
 
+UObject* IObjectSchema::GetRedirectedBoundObject(UObject* Object) const
+{
+	return Object;
+}
+
 } // namespace UE::Sequencer
\ No newline at end of file
diff --git a/Engine/Source/Editor/Sequencer/Private/KeyPropertyParams.cpp b/Engine/Source/Editor/Sequencer/Private/KeyPropertyParams.cpp
--- a/Engine/Source/Editor/Sequencer/Private/KeyPropertyParams.cpp
+++ b/Engine/Source/Editor/Sequencer/Private/KeyPropertyParams.cpp
@@ -128,30 +128,84 @@ FPropertyChangedParams::FPropertyChangedParams(TArray<UObject*> InObjectsThatCha
 {
 }
 
-template<> bool FPropertyChangedParams::GetPropertyValueImpl<bool>(void* Data, const FPropertyInfo& PropertyInfo)
+TPair<const FProperty*, UE::MovieScene::FSourcePropertyValue> FPropertyChangedParams::GetPropertyAndValue() const
 {
+	using namespace UE::MovieScene;
+
+	check(ObjectsThatChanged[0]);
+
+	void* ContainerPtr = ObjectsThatChanged[0];
+	const FProperty* Property = nullptr;
+	for (int32 i = 0; i < PropertyPath.GetNumProperties(); i++)
+	{
+		const FPropertyInfo& PropertyInfo = PropertyPath.GetPropertyInfo(i);
+		Property = PropertyInfo.Property.Get();
+		if (!Property)
+		{
+			return TPair<const FProperty*, FSourcePropertyValue>();
+		}
+
+		if (const FArrayProperty* ArrayProp = CastField<FArrayProperty>(Property))
+		{
+			// Sometimes property paths have the array property twice, first with no array index,
+			// then a second so we skip over this property if that's the case
+			int32 ArrayIndex = FMath::Max(0, PropertyInfo.ArrayIndex);
+			if (PropertyInfo.ArrayIndex == INDEX_NONE && i < PropertyPath.GetNumProperties()-1) //-V1051
+			{
+				const FPropertyInfo& InnerPropertyInfo = PropertyPath.GetPropertyInfo(i+1);
+				const FProperty* InnerProperty = InnerPropertyInfo.Property.Get();
+				if (InnerProperty && InnerProperty->GetOwner<FProperty>() == ArrayProp)
+				{
+					ArrayIndex = InnerPropertyInfo.ArrayIndex;
+					++i;
+				}
+			}
+
+			FScriptArrayHelper ParentArrayHelper(ArrayProp, ArrayProp->ContainerPtrToValuePtr<void>(ContainerPtr));
+			if (!ParentArrayHelper.IsValidIndex(ArrayIndex))
+			{
+				return TPair<const FProperty*, FSourcePropertyValue>(nullptr, FSourcePropertyValue());
+			}
+			ContainerPtr = ParentArrayHelper.GetRawPtr(ArrayIndex);
+		}
+		else if (const FStructProperty* StructProp = CastField<FStructProperty>(Property))
+		{
+			if (StructProp->Struct == FInstancedStruct::StaticStruct())
+			{
+				FInstancedStruct* InstancedStruct = Property->ContainerPtrToValuePtr<FInstancedStruct>(ContainerPtr);
+				ContainerPtr = InstancedStruct->GetMutableMemory();
+			}
+			else
+			{
+				ContainerPtr = Property->ContainerPtrToValuePtr<void>(ContainerPtr);
+			}
+		}
+		else
+		{
+			ContainerPtr = Property->ContainerPtrToValuePtr<void>(ContainerPtr);
+		}
+	}
+
+
 	// Bool property values are stored in a bit field so using a straight cast of the pointer to get their value does not
 	// work.  Instead use the actual property to get the correct value.
-	if ( const FBoolProperty* BoolProperty = CastField<const FBoolProperty>( PropertyInfo.Property.Get() ) )
+	if (const FBoolProperty* BoolProperty = CastField<const FBoolProperty>(Property))
 	{
-		return BoolProperty->GetPropertyValue(Data);
+		return MakeTuple(BoolProperty, FSourcePropertyValue::FromValue(BoolProperty->GetPropertyValue(ContainerPtr)));
 	}
-	else
+	// Object properties might have various different types of storage, but we always expose them as a raw ptr
+	else if (const FObjectPropertyBase* ObjectProperty = CastField<const FObjectPropertyBase>(Property))
 	{
-		return *((bool*)Data);
+		return MakeTuple(ObjectProperty, FSourcePropertyValue::FromValue(ObjectProperty->GetObjectPropertyValue(ContainerPtr)));
 	}
+
+	check(Property);
+	return MakeTuple(Property, FSourcePropertyValue::FromAddress(ContainerPtr, *Property));
 }
 
-template<> UObject* FPropertyChangedParams::GetPropertyValueImpl<UObject*>(void* Data, const FPropertyInfo& PropertyInfo)
+UE::MovieScene::FSourcePropertyValue FPropertyChangedParams::GetPropertyValue() const
 {
-	if (const FObjectPropertyBase* ObjectProperty = CastField<const FObjectPropertyBase>(PropertyInfo.Property.Get()))
-	{
-		return ObjectProperty->GetObjectPropertyValue(Data);
-	}
-	else
-	{
-		return *((UObject**)Data);
-	}
+	return GetPropertyAndValue().Value;
 }
 
 FString FPropertyChangedParams::GetPropertyPathString() const
diff --git a/Engine/Source/Editor/Sequencer/Private/MovieSceneTrackEditor.cpp b/Engine/Source/Editor/Sequencer/Private/MovieSceneTrackEditor.cpp
--- a/Engine/Source/Editor/Sequencer/Private/MovieSceneTrackEditor.cpp
+++ b/Engine/Source/Editor/Sequencer/Private/MovieSceneTrackEditor.cpp
@@ -170,12 +170,15 @@ void FMovieSceneTrackEditor::AnimatablePropertyChanged( FOnKeyProperty OnKeyProp
 
 FMovieSceneTrackEditor::FFindOrCreateHandleResult FMovieSceneTrackEditor::FindOrCreateHandleToObject( UObject* Object, bool bCreateHandleIfMissing, const FName& CreatedFolderName)
 {
-	FFindOrCreateHandleResult Result;
-	bool bHandleWasValid = GetSequencer()->GetHandleToObject( Object, false ).IsValid();
-	
+	FGuid ObjectBindingID = GetSequencer()->GetHandleToObject(Object, false);
+	if (ObjectBindingID.IsValid() || !bCreateHandleIfMissing)
+	{
+		return FFindOrCreateHandleResult{ ObjectBindingID, false };
+	}
 
-	Result.Handle = GetSequencer()->GetHandleToObject( Object, bCreateHandleIfMissing, CreatedFolderName);
-	Result.bWasCreated = bHandleWasValid == false && Result.Handle.IsValid();
+	FFindOrCreateHandleResult Result;
+	Result.Handle      = GetSequencer()->GetHandleToObject(Object, true, CreatedFolderName);
+	Result.bWasCreated = Result.Handle.IsValid();
 	return Result;
 }
 
diff --git a/Engine/Source/Editor/Sequencer/Private/Sequencer.cpp b/Engine/Source/Editor/Sequencer/Private/Sequencer.cpp
--- a/Engine/Source/Editor/Sequencer/Private/Sequencer.cpp
+++ b/Engine/Source/Editor/Sequencer/Private/Sequencer.cpp
@@ -83,6 +83,7 @@
 #include "SequencerPropertyKeyedStatus.h"
 #include "SequencerSettings.h"
 #include "SequencerLog.h"
+#include "ISequencerObjectSchema.h"
 #include "SequencerEdMode.h"
 #include "MovieSceneBindingProxy.h"
 #include "MovieSceneSequence.h"
@@ -4280,6 +4281,8 @@ UMovieSceneFolder* FSequencer::CreateFoldersRecursively(const TArray<FName>& Fol
 
 FGuid FSequencer::GetHandleToObject( UObject* Object, bool bCreateHandleIfMissing, const FName& CreatedFolderName )
 {
+	using namespace UE::Sequencer;
+
 	if (Object == nullptr)
 	{
 		return FGuid();
@@ -4298,6 +4301,16 @@ FGuid FSequencer::GetHandleToObject( UObject* Object, bool bCreateHandleIfMissin
 		return FGuid();
 	}
 
+	ISequencerModule& Module = FModuleManager::Get().LoadModuleChecked<ISequencerModule>("Sequencer");
+	if (TSharedPtr<IObjectSchema> Schema = Module.FindObjectSchema(Object))
+	{
+		Object = Schema->GetRedirectedBoundObject(Object);
+		if (!Object)
+		{
+			return FGuid();
+		}
+	}
+
 	// Attempt to resolve the object through the movie scene instance first, 
 	FGuid ObjectGuid = bCreateHandleIfMissing ?
 		FindObjectId(*Object, ActiveTemplateIDs.Top()) :
diff --git a/Engine/Source/Editor/Sequencer/Private/SequencerAddKeyOperation.cpp b/Engine/Source/Editor/Sequencer/Private/SequencerAddKeyOperation.cpp
--- a/Engine/Source/Editor/Sequencer/Private/SequencerAddKeyOperation.cpp
+++ b/Engine/Source/Editor/Sequencer/Private/SequencerAddKeyOperation.cpp
@@ -101,6 +101,9 @@ void FAddKeyOperation::AddPreFilteredNodes(TArrayView<const TWeakPtr<FViewModel>
 				{
 					TSharedPtr<ITrackExtension> TrackExtension = Child->CastThisShared<ITrackExtension>();
 					ConsiderKeyableAreas(TrackExtension, *Child);
+
+					// Stop recursively iterating everything within this Child again since we just considered all keyable areas for this and its children
+					Child.IgnoreCurrentChildren();
 				}
 			}
 		}
@@ -120,9 +123,10 @@ bool FAddKeyOperation::ConsiderKeyableAreas(TSharedPtr<ITrackExtension> InTrackM
 		UMovieSceneTrack* Track = Cast<UMovieSceneTrack>(Section->GetOuter());
 		if (Track && Track->GetSectionToKey() == Section)
 		{
-			for (TParentFirstChildIterator<FChannelGroupModel> ChannelGroupModelIt(SectionModel->AsShared(), bIncludeThis); ChannelGroupModelIt; ++ChannelGroupModelIt)
-			{
-				bKeyedAnything |= ProcessKeyArea(InTrackModel, *ChannelGroupModelIt);
+			ISequencerTrackEditor* TrackEditor = InTrackModel->GetTrackEditor().Get();
+			for (TParentFirstChildIterator<FChannelModel> ChannelModelIt(SectionModel->AsShared(), bIncludeThis); ChannelModelIt; ++ChannelModelIt)
+			{	
+				bKeyedAnything |= ProcessKeyArea(TrackEditor, ChannelModelIt->GetKeyArea());
 			}
 		}
 	}
diff --git a/Engine/Source/Editor/Sequencer/Private/SequencerModule.cpp b/Engine/Source/Editor/Sequencer/Private/SequencerModule.cpp
--- a/Engine/Source/Editor/Sequencer/Private/SequencerModule.cpp
+++ b/Engine/Source/Editor/Sequencer/Private/SequencerModule.cpp
@@ -23,6 +23,9 @@
 #include "MovieSceneSequence.h"
 #include "Sections/MovieSceneSubSection.h"
 
+#include "EntitySystem/MovieScenePropertyRegistry.h"
+#include "EntitySystem/BuiltInComponentTypes.h"
+
 #include "MVVM/CurveEditorExtension.h"
 #include "MVVM/CurveEditorIntegrationExtension.h"
 #include "MVVM/FolderModelStorageExtension.h"
@@ -598,6 +601,18 @@ public:
 			return true;
 		}
 
+		const UE::MovieScene::FPropertyRegistry& PropertyRegistry = UE::MovieScene::FBuiltInComponentTypes::Get()->PropertyRegistry;
+
+		// Find the property that applies to this
+		for (const UE::MovieScene::FPropertyDefinition& PropertyDefinition : PropertyRegistry.GetProperties())
+		{
+			if (PropertyDefinition.Handler->SupportsProperty(PropertyDefinition, *Property))
+			{
+				return true;
+			}
+		}
+
+
 		FObjectPropertyBase* ObjectProperty = CastField<FObjectPropertyBase>(Property);
 
 		// Check each level of the property hierarchy
diff --git a/Engine/Source/Editor/Sequencer/Private/SequencerObjectChangeListener.cpp b/Engine/Source/Editor/Sequencer/Private/SequencerObjectChangeListener.cpp
--- a/Engine/Source/Editor/Sequencer/Private/SequencerObjectChangeListener.cpp
+++ b/Engine/Source/Editor/Sequencer/Private/SequencerObjectChangeListener.cpp
@@ -155,6 +155,11 @@ FOnObjectPropertyChanged& FSequencerObjectChangeListener::GetOnAnyPropertyChange
 	return ObjectToPropertyChangedEvent.FindOrAdd(&Object);
 }
 
+void FSequencerObjectChangeListener::AddFindPropertyChangedHandler(FFindAnimatablePropertyChangedHandler&& Handler)
+{
+	FindPropertyChangedHandlers.Emplace(MoveTemp(Handler));
+}
+
 void FSequencerObjectChangeListener::ReportObjectDestroyed(UObject& Object)
 {
 	ObjectToPropertyChangedEvent.Remove(&Object);
@@ -359,6 +364,17 @@ bool FSequencerObjectChangeListener::CanKeyProperty_Internal(FCanKeyPropertyPara
 				return true;
 			}
 
+			for (const FFindAnimatablePropertyChangedHandler& FindHandler : FindPropertyChangedHandlers)
+			{
+				FOnAnimatablePropertyChanged PropertyChanged = FindHandler.Execute(Property);
+				if (PropertyChanged.IsBound())
+				{
+					InOutProperty = Property;
+					InOutDelegate = MoveTemp(PropertyChanged);
+					return true;
+				}
+			}
+
 			const FObjectPropertyBase* ObjectProperty = CastField<FObjectPropertyBase>(Property);
 
 			// Check each level of the property hierarchy
diff --git a/Engine/Source/Editor/Sequencer/Private/SequencerObjectChangeListener.h b/Engine/Source/Editor/Sequencer/Private/SequencerObjectChangeListener.h
--- a/Engine/Source/Editor/Sequencer/Private/SequencerObjectChangeListener.h
+++ b/Engine/Source/Editor/Sequencer/Private/SequencerObjectChangeListener.h
@@ -29,6 +29,7 @@ public:
 	virtual FOnAnimatablePropertyChanged& GetOnAnimatablePropertyChanged(const FProperty* Property) override;
 	virtual FOnPropagateObjectChanges& GetOnPropagateObjectChanges() override;
 	virtual FOnObjectPropertyChanged& GetOnAnyPropertyChanged(UObject& Object) override;
+	virtual void AddFindPropertyChangedHandler(FFindAnimatablePropertyChangedHandler&& Handler) override;
 	virtual void ReportObjectDestroyed(UObject& Object) override;
 	UE_DEPRECATED(5.4, "CanKeyProperty taking in only FCanKeyPropertyParams is deprecated. Please use the alternative CanKeyProperty that takes in FCanKeyPropertyParams and an out FPropertyPath.")
 	virtual bool CanKeyProperty(FCanKeyPropertyParams KeyPropertyParams) const override;
@@ -89,6 +90,10 @@ private:
 	bool CanKeyProperty_Internal(FCanKeyPropertyParams KeyPropertyParams, FOnAnimatablePropertyChanged& InOutDelegate, const FProperty*& InOutProperty, FPropertyPath& InOutPropertyPath) const;
 
 private:
+
+	/** Array of delegates that are called before any other method to attempt to retrieve a property changed handler */
+	TArray<FFindAnimatablePropertyChangedHandler> FindPropertyChangedHandlers;
+
 	/** Mapping of object to a listener used to check for property changes */
 	TMap< TWeakObjectPtr<UObject>, TSharedPtr<class IPropertyChangeListener> > ActivePropertyChangeListeners;
 
diff --git a/Engine/Source/Editor/Sequencer/Private/SequencerUtilities.cpp b/Engine/Source/Editor/Sequencer/Private/SequencerUtilities.cpp
--- a/Engine/Source/Editor/Sequencer/Private/SequencerUtilities.cpp
+++ b/Engine/Source/Editor/Sequencer/Private/SequencerUtilities.cpp
@@ -110,6 +110,47 @@ static void ResetCopiedTracksFlags(UMovieSceneTrack* Track)
 	}
 }
 
+
+namespace UE::Sequencer
+{
+
+	bool IteratePathToObject(UObject* CurrentObject, TFunctionRef<void(UObject*, const IObjectSchema&)> InCallback)
+	{
+		static ISequencerModule& Module = FModuleManager::Get().LoadModuleChecked<ISequencerModule>("Sequencer");
+		if (!CurrentObject)
+		{
+			return false;
+		}
+
+		TSharedPtr<IObjectSchema> Schema = Module.FindObjectSchema(CurrentObject);
+		if (Schema)
+		{
+			CurrentObject = Schema->GetRedirectedBoundObject(CurrentObject);
+
+			// Add the parent first
+			UObject* ParentObject = Schema->GetParentObject(CurrentObject);
+			if (ParentObject && !IteratePathToObject(ParentObject, InCallback))
+			{
+				return false;
+			}
+
+			InCallback(CurrentObject, *Schema);
+			return true;
+		}
+		return false;
+	}
+
+
+	bool GetPathToObject(UObject* InObject, TArray<UObject*>& OutObjectPath)
+	{
+		return IteratePathToObject(InObject, [&](UObject* CurrentObject, const IObjectSchema& Schema)
+		{
+			OutObjectPath.Add(CurrentObject);
+		});
+	}
+
+} // namespace UE::Sequencer
+
 TSharedRef<SWidget> FSequencerUtilities::MakeAddButton(FText HoverText, FOnGetContent MenuContent, const TAttribute<bool>& HoverState, TWeakPtr<ISequencer> InSequencer)
 {
 	TAttribute<bool> IsEnabled = MakeAttributeLambda([InSequencer]() -> bool { return InSequencer.IsValid() ? !InSequencer.Pin()->IsReadOnly() : false; });
@@ -2973,39 +3014,25 @@ FGuid CreateGenericBinding(TSharedPtr<ISequencer> Sequencer, UMovieSceneSequence
 	
 	// Otherwise, create a possessable.
 
-	TArray<TPair<UObject*, FString>> ObjectsToPossess;
 
-	// Build up the list of child->parent bindings required for this object
+
+
+	TArray<TPair<UObject*, const IObjectSchema*>> ObjectsToPossess;
+	auto Iterator = [&ObjectsToPossess, &InParams](UObject* CurrentObject, const IObjectSchema& Schema)
+	{
+		ObjectsToPossess.Add(MakeTuple(CurrentObject, &Schema));
+	};
+
 	if (!InObject)
 	{
-		ObjectsToPossess.Add(MakeTuple(InObject, InParams.BindingNameOverride.IsEmpty() ? TEXT("EmptyBinding") : InParams.BindingNameOverride));
+		ObjectsToPossess.Add(MakeTuple(InObject, nullptr));
 	}
-	else
+	else if (!IteratePathToObject(InObject, Iterator))
 	{
-		UObject* CurrentObject = InObject;
-		while (CurrentObject)
-		{
-			TSharedPtr<IObjectSchema> Schema = Module.FindObjectSchema(CurrentObject);
-			if (Schema)
-			{
-				if (ObjectsToPossess.Num() == 0 && InParams.BindingNameOverride.Len() != 0)
-				{
-					ObjectsToPossess.Add(MakeTuple(CurrentObject, InParams.BindingNameOverride));
-				}
-				else
-				{
-					ObjectsToPossess.Add(MakeTuple(CurrentObject, Schema->GetPrettyName(CurrentObject).ToString()));
-				}
-
-				CurrentObject = Schema->GetParentObject(CurrentObject);
-			}
-			else
-			{
-				break;
-			}
-		}
+		return FGuid();
 	}
 
+
 	// Nothing to possess?
 	if (ObjectsToPossess.Num() == 0)
 	{
@@ -3021,15 +3048,17 @@ FGuid CreateGenericBinding(TSharedPtr<ISequencer> Sequencer, UMovieSceneSequence
 
 	FGuid ParentID;
 
-	// Iterate in reverse (parent -> child)
-	for (int32 Index = ObjectsToPossess.Num() - 1; Index >= 0; --Index)
+	for (int32 Index = 0; Index < ObjectsToPossess.Num(); ++Index)
 	{
-		UObject* CurrentObject = ObjectsToPossess[Index].Key;
+		UObject*             CurrentObject = ObjectsToPossess[Index].Key;
+		const IObjectSchema* Schema        = ObjectsToPossess[Index].Value;
+
+		check(CurrentObject);
 
 		// If we're not purposefully replacing a binding, then check to see if we already have one, and use that
 		if (Sequencer && !InParams.ReplacementGuid.IsValid())
 		{
-			FGuid    ObjectGuid = Sequencer->GetHandleToObject(CurrentObject, false);
+			FGuid    ObjectGuid = Sequencer->FindCachedObjectId(*CurrentObject, Sequencer->GetFocusedTemplateID());
 
 			// If the object already has a binding, use that and move on
 			if (ObjectGuid.IsValid())
@@ -3043,9 +3072,8 @@ FGuid CreateGenericBinding(TSharedPtr<ISequencer> Sequencer, UMovieSceneSequence
 			}
 		}
 
-		// Create a new binding for this object
-		FString CurrentName = MoveTemp(ObjectsToPossess[Index].Value);
 
+		// Create a new binding for this object
 		FMovieScenePossessable* NewPossessable = nullptr;
 		FGuid NewID;
 		if (InParams.ReplacementGuid.IsValid() && !ParentID.IsValid())
@@ -3055,15 +3083,19 @@ FGuid CreateGenericBinding(TSharedPtr<ISequencer> Sequencer, UMovieSceneSequence
 		}
 		if (!NewPossessable)
 		{
-			NewID = OwnerMovieScene->AddPossessable(CurrentName, CurrentObject ? CurrentObject->GetClass() : UObject::StaticClass());
+			FString CurrentName = (Index == ObjectsToPossess.Num()-1 && InParams.BindingNameOverride.Len() != 0)
+				? InParams.BindingNameOverride
+				: Schema->GetPrettyName(CurrentObject).ToString();
+
+			NewID = OwnerMovieScene->AddPossessable(CurrentName, CurrentObject->GetClass());
 			NewPossessable = OwnerMovieScene->FindPossessable(NewID);
 		}
 
 		// If we're not trying to replace a binding, and the object is a spawnable, try and bind to that first
-		if (InParams.ReplacementGuid.IsValid() || (Sequencer && (!CurrentObject || !NewPossessable->BindSpawnableObject(Sequencer->GetFocusedTemplateID(), CurrentObject, Sequencer->GetSharedPlaybackState()))))
+		if (InParams.ReplacementGuid.IsValid() || (Sequencer && !NewPossessable->BindSpawnableObject(Sequencer->GetFocusedTemplateID(), CurrentObject, Sequencer->GetSharedPlaybackState())))
 		{
 			FUniversalObjectLocator Locator;
-			if (CurrentObject && (!OwnerSequence->MakeLocatorForObject(CurrentObject, Context, Locator) || Locator.IsEmpty()))
+			if (!OwnerSequence->MakeLocatorForObject(CurrentObject, Context, Locator) || Locator.IsEmpty())
 			{
 				// Unable to possess this object
 				return FGuid();
@@ -3087,6 +3119,9 @@ FGuid CreateGenericBinding(TSharedPtr<ISequencer> Sequencer, UMovieSceneSequence
 					BindingReferences->AddBinding(NewID, MoveTemp(Locator));
 				}
 			}
+
+			// Fixup the possessed class in case the bound object is actually different than the one we tried to possess
+			NewPossessable->FixupPossessedObjectClass(OwnerSequence, Context);
 		}
 
 		if (ParentID.IsValid())
diff --git a/Engine/Source/Editor/Sequencer/Public/IKeyArea.h b/Engine/Source/Editor/Sequencer/Public/IKeyArea.h
--- a/Engine/Source/Editor/Sequencer/Public/IKeyArea.h
+++ b/Engine/Source/Editor/Sequencer/Public/IKeyArea.h
@@ -297,6 +297,11 @@ public:
 	 */
 	UE_API TUniquePtr<FCurveModel> CreateCurveEditorModel(TSharedRef<ISequencer> InSequencer) const;
 
+	/**
+	 * Retrieve the interpolation mode for this key area's channel for the specified time
+	 */
+	UE_API EMovieSceneKeyInterpolation GetInterpolationMode(FFrameNumber InTime, EMovieSceneKeyInterpolation DefaultMode) const;
+
 public:
 
 	FRichCurve* GetRichCurve() const { return nullptr; }
diff --git a/Engine/Source/Editor/Sequencer/Public/ISequencerChannelInterface.h b/Engine/Source/Editor/Sequencer/Public/ISequencerChannelInterface.h
--- a/Engine/Source/Editor/Sequencer/Public/ISequencerChannelInterface.h
+++ b/Engine/Source/Editor/Sequencer/Public/ISequencerChannelInterface.h
@@ -256,4 +256,16 @@ struct ISequencerChannelInterface
 	 */
 	virtual int32 DrawExtra_Raw(FMovieSceneChannel* InChannel, const UMovieSceneSection* InOwner, const FSequencerChannelPaintArgs& PaintArgs, int32 LayerId) const = 0;
 
+	/**
+	 * Retrieve the interpolation mode for this curve at the current time
+	 *
+	 * @param InChannel          The channel to query
+	 * @param InTime             The time at which to compute the interpolation mode
+	 * @param InDefaultMode      Default interp mode to use
+	 * @return This curve's desired interpolation mode
+	 */
+	virtual EMovieSceneKeyInterpolation GetInterpolationMode_Raw(FMovieSceneChannel* InChannel, FFrameNumber InTime, EMovieSceneKeyInterpolation InDefaultMode) const
+	{
+		return InDefaultMode;
+	}
 };
diff --git a/Engine/Source/Editor/Sequencer/Public/ISequencerObjectChangeListener.h b/Engine/Source/Editor/Sequencer/Public/ISequencerObjectChangeListener.h
--- a/Engine/Source/Editor/Sequencer/Public/ISequencerObjectChangeListener.h
+++ b/Engine/Source/Editor/Sequencer/Public/ISequencerObjectChangeListener.h
@@ -12,6 +12,8 @@ DECLARE_MULTICAST_DELEGATE_OneParam( FOnObjectPropertyChanged, UObject& );
 
 DECLARE_MULTICAST_DELEGATE_OneParam( FOnPropagateObjectChanges, UObject* );
 
+DECLARE_DELEGATE_RetVal_OneParam( FOnAnimatablePropertyChanged, FFindAnimatablePropertyChangedHandler, const FProperty* );
+
 /**
  * Listens for changes objects and calls delegates when those objects change
  */
@@ -37,7 +39,12 @@ public:
 	 * A delegate for when any property of a specific object is changed.
 	 */
 	virtual FOnObjectPropertyChanged& GetOnAnyPropertyChanged(UObject& Object) = 0;
-	
+
+	/**
+	 * Add a delegate that will be called before any other method to locate a property changed handler for a property
+	 */
+	virtual void AddFindPropertyChangedHandler(FFindAnimatablePropertyChangedHandler&&) = 0;
+
 	/**
 	 * Report that an object is about to be destroyed (removes any object change delegates bound to the object)
 	 */
diff --git a/Engine/Source/Editor/Sequencer/Public/ISequencerObjectSchema.h b/Engine/Source/Editor/Sequencer/Public/ISequencerObjectSchema.h
--- a/Engine/Source/Editor/Sequencer/Public/ISequencerObjectSchema.h
+++ b/Engine/Source/Editor/Sequencer/Public/ISequencerObjectSchema.h
@@ -45,6 +45,7 @@ public:
 
 	virtual ~IObjectSchema(){}
 
+	/*~ Begin abstract interface */
 	virtual UObject* GetParentObject(UObject* Object) const = 0;
 
 	virtual FObjectSchemaRelevancy GetRelevancy(const UObject* InObject) const = 0;
@@ -52,6 +53,9 @@ public:
 	virtual TSharedPtr<FExtender> ExtendObjectBindingMenu(TSharedRef<FUICommandList> CommandList, TWeakPtr<ISequencer> WeakSequencer, TArrayView<UObject* const> ContextSensitiveObjects) const = 0;
 
 	UE_API virtual FText GetPrettyName(const UObject* Object) const;
+
+	UE_API virtual UObject* GetRedirectedBoundObject(UObject* Object) const;
+	/*~ End abstract implementation */
 };
 
 } // namespace UE::Sequencer
diff --git a/Engine/Source/Editor/Sequencer/Public/KeyPropertyParams.h b/Engine/Source/Editor/Sequencer/Public/KeyPropertyParams.h
--- a/Engine/Source/Editor/Sequencer/Public/KeyPropertyParams.h
+++ b/Engine/Source/Editor/Sequencer/Public/KeyPropertyParams.h
@@ -11,6 +11,7 @@
 #include "UObject/Field.h"
 #include "UObject/UnrealType.h"
 #include "UObject/WeakFieldPtr.h"
+#include "EntitySystem/MovieSceneIntermediatePropertyValue.h"
 
 #define UE_API SEQUENCER_API
 
@@ -88,63 +89,25 @@ struct FKeyPropertyParams
 class FPropertyChangedParams
 {
 public:
+
+	using FSourcePropertyValue = UE::MovieScene::FSourcePropertyValue;
+
 	UE_API FPropertyChangedParams(TArray<UObject*> InObjectsThatChanged, const FPropertyPath& InPropertyPath, const FPropertyPath& InStructPathToKey, ESequencerKeyMode InKeyMode);
 
+
+	UE_API TPair<const FProperty*, FSourcePropertyValue> GetPropertyAndValue() const;
+	UE_API FSourcePropertyValue GetPropertyValue() const;
+
 	/**
 	 * Gets the value of the property that changed.
 	 */
 	template<typename ValueType>
 	ValueType GetPropertyValue() const
 	{
-		void* ContainerPtr = ObjectsThatChanged[0];
-		for (int32 i = 0; i < PropertyPath.GetNumProperties(); i++)
-		{
-			const FPropertyInfo& PropertyInfo = PropertyPath.GetPropertyInfo(i);
-			if (const FProperty* Property = PropertyInfo.Property.Get())
-			{
-				if (const FArrayProperty* ArrayProp = CastField<FArrayProperty>(Property))
-				{
-					// Sometimes property paths have the array property twice, first with no array index,
-					// then a second so we skip over this property if that's the case
-					int32 ArrayIndex = FMath::Max(0, PropertyInfo.ArrayIndex);
-					if (PropertyInfo.ArrayIndex == INDEX_NONE && i < PropertyPath.GetNumProperties()-1) //-V1051
-					{
-						const FPropertyInfo& InnerPropertyInfo = PropertyPath.GetPropertyInfo(i+1);
-						const FProperty* InnerProperty = InnerPropertyInfo.Property.Get();
-						if (InnerProperty && InnerProperty->GetOwner<FProperty>() == ArrayProp)
-						{
-							ArrayIndex = InnerPropertyInfo.ArrayIndex;
-							++i;
-						}
-					}
-
-					FScriptArrayHelper ParentArrayHelper(ArrayProp, ArrayProp->ContainerPtrToValuePtr<void>(ContainerPtr));
-					if (!ParentArrayHelper.IsValidIndex(ArrayIndex))
-					{
-						return ValueType();
-					}
-					ContainerPtr = ParentArrayHelper.GetRawPtr(ArrayIndex);
-				}
-				else if (const FStructProperty* StructProp = CastField<FStructProperty>(Property))
-				{
-					if (StructProp->Struct == FInstancedStruct::StaticStruct())
-					{
-						FInstancedStruct* InstancedStruct = Property->ContainerPtrToValuePtr<FInstancedStruct>(ContainerPtr);
-						ContainerPtr = InstancedStruct->GetMutableMemory();
-					}
-					else
-					{
-						ContainerPtr = Property->ContainerPtrToValuePtr<void>(ContainerPtr);
-					}
-				}
-				else
-				{
-					ContainerPtr = Property->ContainerPtrToValuePtr<void>(ContainerPtr);
-				}
-			}
-		}
-
-		return GetPropertyValueImpl<ValueType>(ContainerPtr, PropertyPath.GetLeafMostProperty());
+		FSourcePropertyValue Value = GetPropertyValue();
+
+		// This is unsafe, but it is assumed that the user knows the cast is safe in calling this function
+		return *Value.Cast<ValueType>();
 	}
 
 	/** Gets the property path as a period seperated string of property names. */
@@ -162,17 +125,6 @@ public:
 
 	/** Keyframing params */
 	const ESequencerKeyMode KeyMode;
-
-private:
-
-	template<typename ValueType>
-	static ValueType GetPropertyValueImpl(void* Data, const FPropertyInfo& Info)
-	{
-		return *((ValueType*)Data);
-	}
 };
 
-template<> SEQUENCER_API bool FPropertyChangedParams::GetPropertyValueImpl<bool>(void* Data, const FPropertyInfo& Info);
-template<> SEQUENCER_API UObject* FPropertyChangedParams::GetPropertyValueImpl<UObject*>(void* Data, const FPropertyInfo& Info);
-
 #undef UE_API
diff --git a/Engine/Source/Editor/Sequencer/Public/MovieSceneTrackEditor.h b/Engine/Source/Editor/Sequencer/Public/MovieSceneTrackEditor.h
--- a/Engine/Source/Editor/Sequencer/Public/MovieSceneTrackEditor.h
+++ b/Engine/Source/Editor/Sequencer/Public/MovieSceneTrackEditor.h
@@ -154,12 +154,12 @@ public:
 		bool bTrackExisted;
 
 		UMovieScene* MovieScene = GetSequencer()->GetFocusedMovieSceneSequence()->GetMovieScene();
-		Result.Track = MovieScene->FindTrack<TrackClass>();
+		Result.Track = MovieScene->template FindTrack<TrackClass>();
 		bTrackExisted = Result.Track != nullptr;
 
 		if (Result.Track == nullptr)
 		{
-			Result.Track = MovieScene->AddTrack<TrackClass>();
+			Result.Track = MovieScene->template AddTrack<TrackClass>();
 		}
 
 		Result.bWasCreated = bTrackExisted == false && Result.Track != nullptr;
diff --git a/Engine/Source/Editor/Sequencer/Public/SequencerChannelInterface.h b/Engine/Source/Editor/Sequencer/Public/SequencerChannelInterface.h
--- a/Engine/Source/Editor/Sequencer/Public/SequencerChannelInterface.h
+++ b/Engine/Source/Editor/Sequencer/Public/SequencerChannelInterface.h
@@ -237,6 +237,19 @@ struct TSequencerChannelInterfaceCommon : ISequencerChannelInterface
 		return DrawExtra(static_cast<ChannelType*>(InChannel), InOwner, PaintArgs, LayerId);
 	}
 
+	/**
+	 * Retrieve the interpolation mode for this curve at the current time
+	 *
+	 * @param InChannel          The channel to query
+	 * @param InTime             The time at which to compute the interpolation mode
+	 * @param InDefaultMode      Default interp mode to use
+	 * @return This curve's desired interpolation mode
+	 */
+	virtual EMovieSceneKeyInterpolation GetInterpolationMode_Raw(FMovieSceneChannel* InChannel, FFrameNumber InTime, EMovieSceneKeyInterpolation InDefaultMode) const override
+	{
+		using namespace Sequencer;
+		return GetInterpolationMode(static_cast<ChannelType*>(InChannel), InTime, InDefaultMode);
+	}
 
 	/**
 	 * Whether this channel supports curve models
diff --git a/Engine/Source/Editor/Sequencer/Public/SequencerChannelTraits.h b/Engine/Source/Editor/Sequencer/Public/SequencerChannelTraits.h
--- a/Engine/Source/Editor/Sequencer/Public/SequencerChannelTraits.h
+++ b/Engine/Source/Editor/Sequencer/Public/SequencerChannelTraits.h
@@ -451,6 +451,11 @@ namespace Sequencer
 		Section->SetRange(NewRange);
 	}
 
+	inline EMovieSceneKeyInterpolation GetInterpolationMode(const FMovieSceneChannel*, FFrameNumber, EMovieSceneKeyInterpolation InDefaultMode)
+	{
+		return InDefaultMode;
+	}
+
 	/**
 	 * Whether the specified channel handle supports curve models or not
 	 */
diff --git a/Engine/Source/Editor/Sequencer/Public/SequencerUtilities.h b/Engine/Source/Editor/Sequencer/Public/SequencerUtilities.h
--- a/Engine/Source/Editor/Sequencer/Public/SequencerUtilities.h
+++ b/Engine/Source/Editor/Sequencer/Public/SequencerUtilities.h
@@ -83,6 +83,8 @@ struct FCreateBindingParams
 	bool bSetupDefaults = true;
 };
 
+SEQUENCER_API bool GetPathToObject(UObject* InObject, TArray<UObject*>& OutObjectPath);
+
 } // namespace UE::Sequencer
 
 /* Paste folders params */
diff --git a/Engine/Source/Runtime/MovieScene/Private/Channels/MovieSceneIntegerChannel.cpp b/Engine/Source/Runtime/MovieScene/Private/Channels/MovieSceneIntegerChannel.cpp
--- a/Engine/Source/Runtime/MovieScene/Private/Channels/MovieSceneIntegerChannel.cpp
+++ b/Engine/Source/Runtime/MovieScene/Private/Channels/MovieSceneIntegerChannel.cpp
@@ -43,7 +43,7 @@ bool FMovieSceneIntegerChannel::SerializeFromMismatchedTag(const FPropertyTag& T
 	return false;
 }
 
-bool FMovieSceneIntegerChannel::Evaluate(FFrameTime InTime, int32& OutValue) const
+bool FMovieSceneIntegerChannel::Evaluate(FFrameTime InTime, int64& OutValue) const
 {
 	double InterpValue = 0.;
 	bool bReturnValue = EvaluateInterp(InTime, InterpValue);
diff --git a/Engine/Source/Runtime/MovieScene/Private/Channels/MovieSceneObjectPathChannel.cpp b/Engine/Source/Runtime/MovieScene/Private/Channels/MovieSceneObjectPathChannel.cpp
--- a/Engine/Source/Runtime/MovieScene/Private/Channels/MovieSceneObjectPathChannel.cpp
+++ b/Engine/Source/Runtime/MovieScene/Private/Channels/MovieSceneObjectPathChannel.cpp
@@ -173,4 +173,8 @@ void FMovieSceneObjectPathChannel::ClearDefault()
 	RemoveDefault();
 }
 
+void FMovieSceneObjectPathChannel::SetDefault(const FMovieSceneObjectPathChannelKeyValue& InDefaultValue)
+{
+	DefaultValue = InDefaultValue;
+}
 
diff --git a/Engine/Source/Runtime/MovieScene/Private/EntitySystem/BuiltInComponentTypes.cpp b/Engine/Source/Runtime/MovieScene/Private/EntitySystem/BuiltInComponentTypes.cpp
--- a/Engine/Source/Runtime/MovieScene/Private/EntitySystem/BuiltInComponentTypes.cpp
+++ b/Engine/Source/Runtime/MovieScene/Private/EntitySystem/BuiltInComponentTypes.cpp
@@ -355,6 +355,7 @@ FBuiltInComponentTypes::FBuiltInComponentTypes()
 	ComponentRegistry->NewComponentType(&CustomPropertyIndex,   TEXT("Custom Property Index"));			// Not EComponentTypeFlags::Preserved because the system property manager will always ensure that the component is added to the correct entity
 	ComponentRegistry->NewComponentType(&FastPropertyOffset,    TEXT("Fast Property Offset"));			// Not EComponentTypeFlags::Preserved because the system property manager will always ensure that the component is added to the correct entity
 	ComponentRegistry->NewComponentType(&SlowProperty,          TEXT("Slow Property Binding"));			// Not EComponentTypeFlags::Preserved because the system property manager will always ensure that the component is added to the correct entity
+	ComponentRegistry->NewComponentType(&VariantPropertyTypeIndex,TEXT("Variant Property Type Index"));
 	ComponentRegistry->NewComponentType(&BoolResult,            TEXT("Bool Result"));
 	ComponentRegistry->NewComponentType(&ByteResult,            TEXT("Byte Result"));
 	ComponentRegistry->NewComponentType(&IntegerResult,         TEXT("Integer Result"));
diff --git a/Engine/Source/Runtime/MovieScene/Private/EntitySystem/MovieScenePropertyRegistry.cpp b/Engine/Source/Runtime/MovieScene/Private/EntitySystem/MovieScenePropertyRegistry.cpp
--- a/Engine/Source/Runtime/MovieScene/Private/EntitySystem/MovieScenePropertyRegistry.cpp
+++ b/Engine/Source/Runtime/MovieScene/Private/EntitySystem/MovieScenePropertyRegistry.cpp
@@ -8,6 +8,8 @@
 #include "EntitySystem/MovieSceneInitialValueSystem.h"
 #include "MovieSceneCommonHelpers.h"
 
+#include "Algo/Find.h"
+
 
 namespace UE
 {
@@ -208,6 +210,11 @@ TOptional<FResolvedProperty> FPropertyRegistry::ResolveProperty(UObject* Object,
 	return FResolvedProperty(TInPlaceType<TSharedPtr<FTrackInstancePropertyBindings>>(), SlowBindings);
 }
 
+const FPropertyDefinition* FPropertyRegistry::FindPropertyDefinition(FComponentTypeID ComponentTypeID) const
+{
+	return Algo::FindBy(Properties, ComponentTypeID, &FPropertyDefinition::PropertyType);
+}
+
 
 void FPropertyDefinition::SetupInitialValueProcessor() const
 {
diff --git a/Engine/Source/Runtime/MovieScene/Private/MovieSceneCommonHelpers.cpp b/Engine/Source/Runtime/MovieScene/Private/MovieSceneCommonHelpers.cpp
--- a/Engine/Source/Runtime/MovieScene/Private/MovieSceneCommonHelpers.cpp
+++ b/Engine/Source/Runtime/MovieScene/Private/MovieSceneCommonHelpers.cpp
@@ -1,5 +1,6 @@
 // Copyright Epic Games, Inc. All Rights Reserved.
 #include "MovieSceneCommonHelpers.h"
+#include "EntitySystem/MovieSceneIntermediatePropertyValue.h"
 #include "Components/SceneComponent.h"
 #include "GameFramework/Actor.h"
 #include "Camera/CameraComponent.h"
@@ -24,6 +25,7 @@
 #include "EntitySystem/MovieSceneEntitySystemLinker.h"
 #include "Conditions/MovieSceneGroupCondition.h"
 #include "StructUtils/InstancedStruct.h"
+#include "String/ParseTokens.h"
 #include "StructUtils/PropertyBag.h"
 
 bool MovieSceneHelpers::IsSectionKeyable(const UMovieSceneSection* Section)
@@ -1207,10 +1209,10 @@ FTrackInstancePropertyBindings::FTrackInstancePropertyBindings(FName InPropertyN
 	FunctionName = FName(*FunctionString);
 }
 
-void FTrackInstancePropertyBindings::FindProperty(FTrackInstancePropertyBindings& Bindings, void* BasePointer, UStruct* InStruct, const FString& InPropertyName, FResolvedProperty& OutResolvedProperty)
+void FTrackInstancePropertyBindings::FindProperty(void* BasePointer, UStruct* InStruct, FStringView InPropertyName, FResolvedProperty& OutResolvedProperty)
 {
 	// Find a property via UObject reflection first.
-	FProperty* FoundProperty = FindFProperty<FProperty>(InStruct, *InPropertyName);
+	FProperty* FoundProperty = FindFProperty<FProperty>(InStruct, FName(InPropertyName, FNAME_Find));
 
 	if (FoundProperty)
 	{
@@ -1218,21 +1220,21 @@ void FTrackInstancePropertyBindings::FindProperty(FTrackInstancePropertyBindings
 	}
 }
 
-FTrackInstancePropertyBindings::FResolvedProperty FTrackInstancePropertyBindings::FindPropertyAndArrayIndex(FTrackInstancePropertyBindings& Bindings, void* BasePointer, UStruct* InStruct, const FString& PropertyName)
+FTrackInstancePropertyBindings::FResolvedProperty FTrackInstancePropertyBindings::FindPropertyAndArrayIndex(void* BasePointer, UStruct* InStruct, FStringView PropertyName)
 {
 	FResolvedProperty PropertyAndIndex;
 
 	// Calculate the array index if possible.
 	int32 ArrayIndex = -1;
-	if (PropertyName.Len() > 0 && PropertyName.GetCharArray()[PropertyName.Len() - 1] == ']')
+	if (PropertyName.Len() > 0 && PropertyName[PropertyName.Len() - 1] == ']')
 	{
 		int32 OpenIndex = 0;
 		if (PropertyName.FindLastChar('[', OpenIndex))
 		{
 			// We have a property name of the form "Foo[123]". Resolve the property itself ("Foo") and then
 			// parse the array element index (123).
-			FString TruncatedPropertyName = FString::ConstructFromPtrSize(*PropertyName, OpenIndex);
-			FindProperty(Bindings, BasePointer, InStruct, *TruncatedPropertyName, PropertyAndIndex);
+			FStringView TruncatedPropertyName(PropertyName.GetData(), OpenIndex);
+			FindProperty(BasePointer, InStruct, TruncatedPropertyName, PropertyAndIndex);
 
 			const int32 NumberLength = PropertyName.Len() - OpenIndex - 2;
 			if (NumberLength > 0 && NumberLength <= 10)
@@ -1247,14 +1249,53 @@ FTrackInstancePropertyBindings::FResolvedProperty FTrackInstancePropertyBindings
 	}
 
 	// Not index found in this property name, just find the property normally.
-	FindProperty(Bindings, BasePointer, InStruct, *PropertyName, PropertyAndIndex);
-
+	FindProperty(BasePointer, InStruct, PropertyName, PropertyAndIndex);
 	return PropertyAndIndex;
 }
 
-FTrackInstancePropertyBindings::FResolvedProperty FTrackInstancePropertyBindings::ResolvePropertyRecursive(FTrackInstancePropertyBindings& Bindings, void* BasePointer, UStruct* InStruct, TArray<FString>& InPropertyNames, uint32 Index)
+TOptional<TPair<const FProperty*, UE::MovieScene::FSourcePropertyValue>> FTrackInstancePropertyBindings::StaticPropertyAndValue(const UObject* Object, FStringView InPropertyPath)
 {
-	FResolvedProperty ResolvedProperty = FindPropertyAndArrayIndex(Bindings, BasePointer, InStruct, *InPropertyNames[Index]);
+	using namespace UE::MovieScene;
+
+	checkf(Object, TEXT("No object specified"));
+
+	FResolvedProperty Address = ResolveProperty(InPropertyPath, *Object);
+
+	if (const FProperty* Property = Address.Property.Get())
+	{
+		const void* PropertyAddress = Property->ContainerPtrToValuePtr<void>(Address.ContainerAddress);
+
+		// Bool property values are stored in a bit field so using a straight cast of the pointer to get their value does not
+		// work.  Instead use the actual property to get the correct value.
+		if (const FBoolProperty* BoolProperty = CastField<const FBoolProperty>(Property))
+		{
+			return MakeTuple(Property, FSourcePropertyValue::FromValue(BoolProperty->GetPropertyValue(PropertyAddress)));
+		}
+		// Object properties might have various different types of storage, but we always expose them as a raw ptr
+		else if (const FObjectPropertyBase* ObjectProperty = CastField<const FObjectPropertyBase>(Property))
+		{
+			return MakeTuple(Property, FSourcePropertyValue::FromValue(ObjectProperty->GetObjectPropertyValue(PropertyAddress)));
+		}
+
+		return MakeTuple(Property, FSourcePropertyValue::FromAddress(PropertyAddress, *Property));
+	}
+
+	return TOptional<TPair<const FProperty*, UE::MovieScene::FSourcePropertyValue>>();
+}
+
+TOptional<UE::MovieScene::FSourcePropertyValue> FTrackInstancePropertyBindings::StaticValue(const UObject* Object, FStringView InPropertyPath)
+{
+	TOptional<TPair<const FProperty*, UE::MovieScene::FSourcePropertyValue>> Result = StaticPropertyAndValue(Object, InPropertyPath);
+	if (Result.IsSet())
+	{
+		return MoveTemp(Result->Value);
+	}
+	return TOptional<UE::MovieScene::FSourcePropertyValue>();
+}
+
+FTrackInstancePropertyBindings::FResolvedProperty FTrackInstancePropertyBindings::ResolvePropertyRecursive(void* BasePointer, UStruct* InStruct, TArrayView<const FStringView> InPropertyNames, uint32 Index)
+{
+	FResolvedProperty ResolvedProperty = FindPropertyAndArrayIndex(BasePointer, InStruct, InPropertyNames[Index]);
 	
 	FResolvedProperty NewProperty;
 
@@ -1275,7 +1316,7 @@ FTrackInstancePropertyBindings::FResolvedProperty FTrackInstancePropertyBindings
 					{
 						// Move the BasePointer to the array element and keep resolving the property path on it.
 						void* ArrayElement = ArrayHelper.GetRawPtr(ResolvedProperty.ArrayIndex);
-						return ResolvePropertyRecursive(Bindings, ArrayElement, InnerStructProp->Struct, InPropertyNames, Index + 1);
+						return ResolvePropertyRecursive(ArrayElement, InnerStructProp->Struct, InPropertyNames, Index + 1);
 					}
 					else
 					{
@@ -1308,7 +1349,7 @@ FTrackInstancePropertyBindings::FResolvedProperty FTrackInstancePropertyBindings
 				FInstancedStruct* InstancedStruct = StructProp->ContainerPtrToValuePtr<FInstancedStruct>(BasePointer);
 				const UScriptStruct* InstancedStructType = InstancedStruct->GetScriptStruct();
 				uint8* InstancedStructMemory = InstancedStruct->GetMutableMemory();
-				return ResolvePropertyRecursive(Bindings, InstancedStructMemory, const_cast<UScriptStruct*>(InstancedStructType), InPropertyNames, Index + 1);
+				return ResolvePropertyRecursive(InstancedStructMemory, const_cast<UScriptStruct*>(InstancedStructType), InPropertyNames, Index + 1);
 			}
 			// As above but for a property bag, unless the property path is specifically targeting its inner instanced struct.
 			else if (StructProp->Struct == FInstancedPropertyBag::StaticStruct() && InPropertyNames[Index + 1] != TEXT("Value"))
@@ -1316,12 +1357,12 @@ FTrackInstancePropertyBindings::FResolvedProperty FTrackInstancePropertyBindings
 				FInstancedPropertyBag* PropertyBag = StructProp->ContainerPtrToValuePtr<FInstancedPropertyBag>(BasePointer);
 				const UPropertyBag* PropertyBagType = PropertyBag->GetPropertyBagStruct();
 				uint8* PropertyBagMemory = PropertyBag->GetMutableValue().GetMemory();
-				return ResolvePropertyRecursive(Bindings, PropertyBagMemory, const_cast<UPropertyBag*>(PropertyBagType), InPropertyNames, Index + 1);
+				return ResolvePropertyRecursive(PropertyBagMemory, const_cast<UPropertyBag*>(PropertyBagType), InPropertyNames, Index + 1);
 			}
 			else
 			{
 				void* StructContainer = StructProp->ContainerPtrToValuePtr<void>(BasePointer);
-				return ResolvePropertyRecursive(Bindings, StructContainer, StructProp->Struct, InPropertyNames, Index + 1);
+				return ResolvePropertyRecursive(StructContainer, StructProp->Struct, InPropertyNames, Index + 1);
 			}
 		}
 		else
@@ -1338,16 +1379,19 @@ FTrackInstancePropertyBindings::FResolvedProperty FTrackInstancePropertyBindings
 	return NewProperty;
 }
 
-FTrackInstancePropertyBindings::FResolvedProperty FTrackInstancePropertyBindings::ResolveProperty(FTrackInstancePropertyBindings& Bindings, const UObject& InObject)
+FTrackInstancePropertyBindings::FResolvedProperty FTrackInstancePropertyBindings::ResolveProperty(FStringView PropertyPath, const UObject& InObject)
 {
-	TArray<FString> PropertyNames;
+	using namespace UE::String;
+
+	TArray<FStringView, TInlineAllocator<4>> PropertyNames;
 
-	const FString& PropertyPath = Bindings.PropertyPath;
-	PropertyPath.ParseIntoArray(PropertyNames, TEXT("."), true);
+	// Parse property paths into component parts separated by '.'
+	auto ProcessToken = [&PropertyNames](FStringView Token) { PropertyNames.Emplace(Token); };
+	ParseTokens(PropertyPath, TEXT('.'), TFunctionRef<void(FWideStringView)>(ProcessToken), EParseTokensOptions::IgnoreCase | EParseTokensOptions::SkipEmpty);
 
 	if (IsValid(&InObject) && PropertyNames.Num() > 0)
 	{
-		return ResolvePropertyRecursive(Bindings, (void*)&InObject, InObject.GetClass(), PropertyNames, 0);
+		return ResolvePropertyRecursive((void*)&InObject, InObject.GetClass(), PropertyNames, 0);
 	}
 	else
 	{
@@ -1355,10 +1399,9 @@ FTrackInstancePropertyBindings::FResolvedProperty FTrackInstancePropertyBindings
 	}
 }
 
-FProperty* FTrackInstancePropertyBindings::FindProperty(const UObject* Object, const FString& InPropertyPath)
+FProperty* FTrackInstancePropertyBindings::FindProperty(const UObject* Object, FStringView InPropertyPath)
 {
-	FTrackInstancePropertyBindings Temp(NAME_None, InPropertyPath);
-	FResolvedProperty ResolvedProperty = ResolveProperty(Temp, *Object);
+	FResolvedProperty ResolvedProperty = ResolveProperty(InPropertyPath, *Object);
 	return ResolvedProperty.Property.Get();
 }
 
@@ -1415,7 +1458,7 @@ void FTrackInstancePropertyBindings::CacheBinding(const UObject& Object)
 {
 	FResolvedPropertyAndFunction PropAndFunction;
 	{
-		PropAndFunction.ResolvedProperty = ResolveProperty(*this, Object);
+		PropAndFunction.ResolvedProperty = ResolveProperty(PropertyPath, Object);
 
 		UFunction* SetterFunction = Object.FindFunction(FunctionName);
 		if (SetterFunction && SetterFunction->NumParms >= 1)
diff --git a/Engine/Source/Runtime/MovieScene/Public/Channels/MovieSceneIntegerChannel.h b/Engine/Source/Runtime/MovieScene/Public/Channels/MovieSceneIntegerChannel.h
--- a/Engine/Source/Runtime/MovieScene/Public/Channels/MovieSceneIntegerChannel.h
+++ b/Engine/Source/Runtime/MovieScene/Public/Channels/MovieSceneIntegerChannel.h
@@ -99,7 +99,17 @@ struct FMovieSceneIntegerChannel : public FMovieSceneChannel
 	 * @param OutValue   A value to receive the result
 	 * @return true if the channel was evaluated successfully, false otherwise
 	 */
-	MOVIESCENE_API bool Evaluate(FFrameTime InTime, int32& OutValue) const;
+	MOVIESCENE_API bool Evaluate(FFrameTime InTime, int64& OutValue) const;
+	bool Evaluate(FFrameTime InTime, int32& OutValue) const
+	{
+		int64 Value = 0;
+		if (Evaluate(InTime, Value))
+		{
+			OutValue = static_cast<int32>(Value);
+			return true;
+		}
+		return false;
+	}
 
 	/**
 	 * Evaluate the interpolated value when bInterpolateLinearKeys is true
diff --git a/Engine/Source/Runtime/MovieScene/Public/Channels/MovieSceneObjectPathChannel.h b/Engine/Source/Runtime/MovieScene/Public/Channels/MovieSceneObjectPathChannel.h
--- a/Engine/Source/Runtime/MovieScene/Public/Channels/MovieSceneObjectPathChannel.h
+++ b/Engine/Source/Runtime/MovieScene/Public/Channels/MovieSceneObjectPathChannel.h
@@ -29,6 +29,11 @@ struct FMovieSceneObjectPathChannelKeyValue
 	 */
 	MOVIESCENE_API FMovieSceneObjectPathChannelKeyValue& operator=(UObject* NewObject);
 
+	friend bool operator==(const FMovieSceneObjectPathChannelKeyValue& A, const FMovieSceneObjectPathChannelKeyValue& B)
+	{
+		return A.SoftPtr == B.SoftPtr;
+	}
+
 public:
 
 	/**
@@ -49,6 +54,11 @@ public:
 	 */
 	MOVIESCENE_API UObject* Get() const;
 
+	UObject* GetValue() const
+	{
+		return Get();
+	}
+
 private:
 
 	/** Persistent storage of the object by path (which allows us to support cross-level actor references, for instance) */
@@ -147,11 +157,7 @@ public:
 	{
 		DefaultValue = InDefaultValue;
 	}
-
-	inline void SetDefault(FMovieSceneObjectPathChannelKeyValue InDefaultValue)
-	{
-		DefaultValue = InDefaultValue;
-	}
+	MOVIESCENE_API void SetDefault(const FMovieSceneObjectPathChannelKeyValue& InDefaultValue);
 
 	/**
 	 * Get this channel's default value that will be used when no keys are present
@@ -190,6 +196,22 @@ private:
 	FMovieSceneKeyHandleMap KeyHandles;
 };
 
+inline bool GetChannelDefault(const FMovieSceneObjectPathChannel* Channel, FMovieSceneObjectPathChannelKeyValue& OutDefaultValue)
+{
+	OutDefaultValue = Channel->GetDefault();
+	return true;
+}
+inline bool EvaluateChannel(const FMovieSceneObjectPathChannel* InChannel, FFrameTime InTime, FMovieSceneObjectPathChannelKeyValue& OutValue)
+{
+	UObject* Object = nullptr;
+	if (InChannel->Evaluate(InTime, Object))
+	{
+		OutValue = Object;
+		return true;
+	}
+	return false;
+}
+
 template<>
 struct TMovieSceneChannelTraits<FMovieSceneObjectPathChannel> : TMovieSceneChannelTraitsBase<FMovieSceneObjectPathChannel>
 {
@@ -200,11 +222,3 @@ struct TMovieSceneChannelTraits<FMovieSceneObjectPathChannel> : TMovieSceneChann
 
 #endif
 };
-
-inline bool EvaluateChannel(const FMovieSceneObjectPathChannel* InChannel, FFrameTime InTime, FMovieSceneObjectPathChannelKeyValue& OutValue)
-{
-	UObject* Value = nullptr;
-	bool bSuccess = InChannel->Evaluate(InTime, Value);
-	OutValue = FMovieSceneObjectPathChannelKeyValue(Value);
-	return bSuccess;
-}
diff --git a/Engine/Source/Runtime/MovieScene/Public/EntitySystem/BuiltInComponentTypes.h b/Engine/Source/Runtime/MovieScene/Public/EntitySystem/BuiltInComponentTypes.h
--- a/Engine/Source/Runtime/MovieScene/Public/EntitySystem/BuiltInComponentTypes.h
+++ b/Engine/Source/Runtime/MovieScene/Public/EntitySystem/BuiltInComponentTypes.h
@@ -9,6 +9,7 @@
 #include "EntitySystem/MovieSceneInitialValueCache.h"
 #include "EntitySystem/MovieScenePropertyRegistry.h"
 #include "EntitySystem/MovieSceneSequenceInstanceHandle.h"
+#include "EntitySystem/MovieSceneVariantPropertyTypeIndex.h"
 #include "Evaluation/Blending/MovieSceneBlendType.h"
 #include "Evaluation/IMovieSceneEvaluationHook.h"
 #include "Misc/Guid.h"
@@ -658,10 +659,10 @@ public:
 	TComponentTypeID<FSourceIntegerChannel> IntegerChannel;
 
 	// An integer representing the output of an integer track or channel
-	TComponentTypeID<int32> IntegerResult;
+	TComponentTypeID<int64> IntegerResult;
 
 	// An integer representing the base value for the integer channel for the purposes of "additive from base" blending.
-	TComponentTypeID<int32> BaseInteger;
+	TComponentTypeID<int64> BaseInteger;
 
 	// An FMovieSceneBoolChannel
 	TComponentTypeID<FSourceBoolChannel> BoolChannel;
@@ -758,6 +759,16 @@ public:
 	// A property binding that supports setters and notifications
 	TComponentTypeID<TSharedPtr<FTrackInstancePropertyBindings>> SlowProperty;
 
+	/**
+	 * Specifies an index that uniquely identifies the variant type of a property component by its position in
+	 *     either the compile-time list of variant types, or the dynamic list of variants registered at runtime.
+	 * 
+	 * For example: A TDynamicVariantPropertyTraits<FMyIntermediateType, double, float, int, int64> with 2 runtime-registered
+	 *              variants (FFloat16, and FFloat8) will represent the variants with the following indices:
+	 *              double: 0, float: 1, int: 2, int64: 3, FFloat16: 4, and; FFloat8: 5
+	 */
+	TComponentTypeID<FVariantPropertyTypeIndex> VariantPropertyTypeIndex;
+
 	struct
 	{
 		// A tag specifying that an entity wants to restore state on completioon
diff --git a/Engine/Source/Runtime/MovieScene/Public/EntitySystem/IMovieScenePropertyComponentHandler.h b/Engine/Source/Runtime/MovieScene/Public/EntitySystem/IMovieScenePropertyComponentHandler.h
--- a/Engine/Source/Runtime/MovieScene/Public/EntitySystem/IMovieScenePropertyComponentHandler.h
+++ b/Engine/Source/Runtime/MovieScene/Public/EntitySystem/IMovieScenePropertyComponentHandler.h
@@ -4,7 +4,11 @@
 
 #include "Templates/SharedPointer.h"
 #include "EntitySystem/MovieSceneEntityIDs.h"
+#include "EntitySystem/MovieSceneVariantPropertyTypeIndex.h"
+#include "EntitySystem/MovieSceneIntermediatePropertyValue.h"
 
+
+class UMovieSceneSection;
 class UMovieSceneBlenderSystem;
 class UMovieSceneEntitySystemLinker;
 
@@ -19,8 +23,10 @@ struct FInitialValueCache;
 struct FPropertyDefinition;
 struct IPreAnimatedStorage;
 struct FEntityComponentFilter;
+struct FUnpackedChannelValues;
 struct FSystemSubsequentTasks;
 struct FSystemTaskPrerequisites;
+struct FVariantPropertyTypeIndex;
 struct FValueDecompositionParams;
 struct FPreAnimatedStateExtension;
 struct FPropertyCompositeDefinition;
@@ -151,6 +157,10 @@ struct IPropertyComponentHandler
 {
 	virtual ~IPropertyComponentHandler(){}
 
+	virtual bool SupportsProperty(const FPropertyDefinition& Definition, const FProperty& InProperty) const = 0;
+	virtual FIntermediatePropertyValue CoercePropertyValue(const FPropertyDefinition& Definition, const FProperty& InProperty, const FSourcePropertyValue& InPropertyValue) const = 0;
+	virtual void UnpackChannels(const FPropertyDefinition& Definition, const FProperty& Property, const FIntermediatePropertyValueConstRef& Value, FUnpackedChannelValues& OutUnpackedValues) const = 0;
+
 	/**
 	 * Dispatch tasks that apply any entity that matches this property type to their final values
 	 *
@@ -193,7 +203,7 @@ struct IPropertyComponentHandler
 	 * @param InCurrentValue   The current value (of type StorageType) to recompose using. For instance, if a property comprises 3 additive values (a:1, b:2, c:3), and we recompose 'a' with an InCurrentValue of 10, the result for 'a' would be 5.
 	 * @param OutResult        The result to receieve recomposed values, one for every entitiy in Params.Query.Entities. Must be of type StorageType.
 	 */
-	virtual void RecomposeBlendOperational(const FPropertyDefinition& Definition, TArrayView<const FPropertyCompositeDefinition> Composites, const FValueDecompositionParams& Params, UMovieSceneBlenderSystem* Blender, FConstPropertyComponentView InCurrentValue, FPropertyComponentArrayView OutResult) = 0;
+	virtual void RecomposeBlendOperational(const FPropertyDefinition& Definition, TArrayView<const FPropertyCompositeDefinition> Composites, const FValueDecompositionParams& Params, UMovieSceneBlenderSystem* Blender, const FIntermediatePropertyValueConstRef& InCurrentValue, TArrayView<FIntermediatePropertyValue> OutResult) = 0;
 
 	/**
 	 * Run a recomposition using the specified params and values.
diff --git a/Engine/Source/Runtime/MovieScene/Public/EntitySystem/MovieSceneDecompositionQuery.h b/Engine/Source/Runtime/MovieScene/Public/EntitySystem/MovieSceneDecompositionQuery.h
--- a/Engine/Source/Runtime/MovieScene/Public/EntitySystem/MovieSceneDecompositionQuery.h
+++ b/Engine/Source/Runtime/MovieScene/Public/EntitySystem/MovieSceneDecompositionQuery.h
@@ -7,6 +7,7 @@
 #include "Containers/ArrayView.h"
 #include "Containers/ContainerAllocationPolicies.h"
 #include "EntitySystem/MovieSceneEntityIDs.h"
+#include "EntitySystem/MovieSceneIntermediatePropertyValue.h"
 #include "HAL/Platform.h"
 #include "Templates/Tuple.h"
 #include "Templates/TypeCompatibleBytes.h"
@@ -196,6 +197,11 @@ namespace MovieScene
 		TArray<PropertyType, TInlineAllocator<1>> Values;
 	};
 
+	struct FRecompositionResult
+	{
+		TArray<FIntermediatePropertyValue, TInlineAllocator<1>> Values;
+	};
+
 } // namespace MovieScene
 } // namespace UE
 
diff --git a/Engine/Source/Runtime/MovieScene/Public/EntitySystem/MovieSceneEntitySystemTask.h b/Engine/Source/Runtime/MovieScene/Public/EntitySystem/MovieSceneEntitySystemTask.h
--- a/Engine/Source/Runtime/MovieScene/Public/EntitySystem/MovieSceneEntitySystemTask.h
+++ b/Engine/Source/Runtime/MovieScene/Public/EntitySystem/MovieSceneEntitySystemTask.h
@@ -46,6 +46,13 @@ template<typename, typename...> struct TEntityAllocationTask;
 template<typename, typename...> struct TEntityAllocationTaskBase;
 template<typename> struct TUnstructuredTask;
 
+namespace Private
+{
+	template<typename...> struct TComponentTypeIDsImpl;
+}
+template<typename...T>
+using TComponentTypeIDs = Private::TComponentTypeIDsImpl<TMakeIntegerSequence<int, sizeof...(T)>, T...>;
+
 
 template<int32 NumComponents, bool AutoExpandAccessors>
 struct TEntityTaskCaller;
@@ -684,6 +691,22 @@ struct TEntityTaskComponentsImpl<TIntegerSequence<int, Indices...>, T...>
 		return TEntityTaskComponents<T..., TReadAccess<U>... >( CommonParams, Accessors.template Get<Indices>()..., TReadAccess<U>(ComponentTypes)... );
 	}
 
+	/**
+	 * Read all of the specified components and pass them through to the task as individual parameters
+	 *
+	 * @param ComponentTypes  The component types to visit
+	 */
+	template<typename... U>
+	TEntityTaskComponents<T..., TReadAccess<U>...> ReadAllOf(const TComponentTypeIDs<U...>& TypeIDs) const
+	{
+		return ReadAllOf(TypeIDs, TMakeIntegerSequence<int, sizeof...(U)>());
+	}
+	template<typename... U, int... ExtraIndices>
+	TEntityTaskComponents<T..., TReadAccess<U>...> ReadAllOf(const TComponentTypeIDs<U...>& TypeIDs, TIntegerSequence<int, ExtraIndices...>&& In) const
+	{
+		return TEntityTaskComponents<T..., TReadAccess<U>... >( CommonParams, Accessors.template Get<Indices>()..., TReadAccess<U>(TypeIDs.GetType(ExtraIndices).template ReinterpretCast<U>())... );
+	}
+
 	/**
 	 * Read any of the specified components and pass them through to the task as individual optional parameters
 	 *
diff --git a/Engine/Source/Runtime/MovieScene/Public/EntitySystem/MovieScenePartialProperties.h b/Engine/Source/Runtime/MovieScene/Public/EntitySystem/MovieScenePartialProperties.h
--- a/Engine/Source/Runtime/MovieScene/Public/EntitySystem/MovieScenePartialProperties.h
+++ b/Engine/Source/Runtime/MovieScene/Public/EntitySystem/MovieScenePartialProperties.h
@@ -54,8 +54,9 @@ struct TSetPartialPropertyValuesImpl<PropertyTraits, TPropertyMetaData<MetaDataT
 	using FThreeWayAccessor  = TMultiReadOptional<FCustomPropertyIndex, uint16, TSharedPtr<FTrackInstancePropertyBindings>>;
 	using FTwoWayAccessor    = TMultiReadOptional<uint16, TSharedPtr<FTrackInstancePropertyBindings>>;
 
-	explicit TSetPartialPropertyValuesImpl(ICustomPropertyRegistration* InCustomProperties, TArrayView<const FPropertyCompositeDefinition> InCompositeDefinitions)
-		: CustomProperties(InCustomProperties)
+	explicit TSetPartialPropertyValuesImpl(const PropertyTraits* InTraits, const ICustomPropertyRegistration* InCustomProperties, TArrayView<const FPropertyCompositeDefinition> InCompositeDefinitions)
+		: Traits(InTraits)
+		, CustomProperties(InCustomProperties)
 		, CompositeDefinitions(InCompositeDefinitions)
 	{
 		if (CustomProperties)
@@ -76,8 +77,8 @@ private:
 	void ForEachSlow(const FEntityAllocation* Allocation, UObject* const* Objects, const TSharedPtr<FTrackInstancePropertyBindings>* Slow, const MetaDataTypes*... InMetaData, const CompositeTypes*... InCompositeComponents) const;
 
 private:
-
-	ICustomPropertyRegistration* CustomProperties;
+	const PropertyTraits* Traits;
+	const ICustomPropertyRegistration* CustomProperties;
 	FCustomAccessorView CustomAccessors;
 	TArrayView<const FPropertyCompositeDefinition> CompositeDefinitions;
 };
