Sequencer: Improved property traits variants, added type-erased property values, and unified many property tack editors

diff --git a/Engine/Source/Runtime/MovieScene/Public/EntitySystem/MovieScenePartialProperties.inl b/Engine/Source/Runtime/MovieScene/Public/EntitySystem/MovieScenePartialProperties.inl
--- a/Engine/Source/Runtime/MovieScene/Public/EntitySystem/MovieScenePartialProperties.inl
+++ b/Engine/Source/Runtime/MovieScene/Public/EntitySystem/MovieScenePartialProperties.inl
@@ -13,7 +13,12 @@ namespace MovieScene
 
 
 template<typename PropertyTraits, typename ...MetaDataTypes, int ...CompositeIndices, typename ...CompositeTypes>
-void TSetPartialPropertyValuesImpl<PropertyTraits, TPropertyMetaData<MetaDataTypes...>, TIntegerSequence<int, CompositeIndices...>, CompositeTypes...>::ForEachAllocation(const FEntityAllocation* Allocation, TRead<UObject*> BoundObjectComponents, FThreeWayAccessor PropertyBindingComponents, TRead<MetaDataTypes>... InMetaData, TReadOptional<CompositeTypes>... InCompositeComponents) const
+void TSetPartialPropertyValuesImpl<PropertyTraits, TPropertyMetaData<MetaDataTypes...>, TIntegerSequence<int, CompositeIndices...>, CompositeTypes...>::ForEachAllocation(
+	const FEntityAllocation* Allocation,
+	TRead<UObject*> BoundObjectComponents,
+	FThreeWayAccessor PropertyBindingComponents,
+	TRead<MetaDataTypes>... InMetaData,
+	TReadOptional<CompositeTypes>... InCompositeComponents) const
 {
 	// ----------------------------------------------------------------------------------------------------------------------------
 	// For partially animated composites, we first retrieve the current properties for the allocation, then go through and patch in
@@ -37,7 +42,12 @@ void TSetPartialPropertyValuesImpl<PropertyTraits, TPropertyMetaData<MetaDataTyp
 
 
 template<typename PropertyTraits, typename ...MetaDataTypes, int ...CompositeIndices, typename ...CompositeTypes>
-void TSetPartialPropertyValuesImpl<PropertyTraits, TPropertyMetaData<MetaDataTypes...>, TIntegerSequence<int, CompositeIndices...>, CompositeTypes...>::ForEachAllocation(const FEntityAllocation* Allocation, TRead<UObject*> BoundObjectComponents, FTwoWayAccessor PropertyBindingComponents, TRead<MetaDataTypes>... InMetaData, TReadOptional<CompositeTypes>... InCompositeComponents) const
+void TSetPartialPropertyValuesImpl<PropertyTraits, TPropertyMetaData<MetaDataTypes...>, TIntegerSequence<int, CompositeIndices...>, CompositeTypes...>::ForEachAllocation(
+	const FEntityAllocation* Allocation,
+	TRead<UObject*> BoundObjectComponents,
+	FTwoWayAccessor PropertyBindingComponents,
+	TRead<MetaDataTypes>... InMetaData,
+	TReadOptional<CompositeTypes>... InCompositeComponents) const
 {
 	// ----------------------------------------------------------------------------------------------------------------------------
 	// For partially animated composites, we first retrieve the current properties for the allocation, then go through and patch in
@@ -56,7 +66,12 @@ void TSetPartialPropertyValuesImpl<PropertyTraits, TPropertyMetaData<MetaDataTyp
 
 
 template<typename PropertyTraits, typename ...MetaDataTypes, int ...CompositeIndices, typename ...CompositeTypes>
-void TSetPartialPropertyValuesImpl<PropertyTraits, TPropertyMetaData<MetaDataTypes...>, TIntegerSequence<int, CompositeIndices...>, CompositeTypes...>::ForEachCustom(const FEntityAllocation* Allocation, UObject* const* Objects, const FCustomPropertyIndex* Custom, const MetaDataTypes*... InMetaData, const CompositeTypes*... InCompositeComponents) const
+void TSetPartialPropertyValuesImpl<PropertyTraits, TPropertyMetaData<MetaDataTypes...>, TIntegerSequence<int, CompositeIndices...>, CompositeTypes...>::ForEachCustom(
+	const FEntityAllocation* Allocation,
+	UObject* const* Objects,
+	const FCustomPropertyIndex* Custom,
+	const MetaDataTypes*... InMetaData,
+	const CompositeTypes*... InCompositeComponents) const
 {
 	const int32 Num = Allocation->Num();
 
@@ -65,17 +80,22 @@ void TSetPartialPropertyValuesImpl<PropertyTraits, TPropertyMetaData<MetaDataTyp
 		const FCustomPropertyIndex PropertyIndex = Custom[Index];
 
 		StorageType Storage{};
-		PropertyTraits::GetObjectPropertyValue(Objects[Index], InMetaData[Index]..., CustomAccessors[PropertyIndex.Value], Storage);
+		Traits->GetObjectPropertyValue(Objects[Index], InMetaData[Index]..., CustomAccessors[PropertyIndex.Value], Storage);
 
 		PatchCompositeValue(CompositeDefinitions, &Storage, InCompositeComponents ? &InCompositeComponents[Index] : nullptr...);
 
-		PropertyTraits::SetObjectPropertyValue(Objects[Index], InMetaData[Index]..., CustomAccessors[PropertyIndex.Value], Storage);
+		Traits->SetObjectPropertyValue(Objects[Index], InMetaData[Index]..., CustomAccessors[PropertyIndex.Value], Storage);
 	}
 }
 
 
 template<typename PropertyTraits, typename ...MetaDataTypes, int ...CompositeIndices, typename ...CompositeTypes>
-void TSetPartialPropertyValuesImpl<PropertyTraits, TPropertyMetaData<MetaDataTypes...>, TIntegerSequence<int, CompositeIndices...>, CompositeTypes...>::ForEachFast(const FEntityAllocation* Allocation, UObject* const* Objects, const uint16* Fast, const MetaDataTypes*... InMetaData, const CompositeTypes*... InCompositeComponents) const
+void TSetPartialPropertyValuesImpl<PropertyTraits, TPropertyMetaData<MetaDataTypes...>, TIntegerSequence<int, CompositeIndices...>, CompositeTypes...>::ForEachFast(
+	const FEntityAllocation* Allocation,
+	UObject* const* Objects,
+	const uint16* Fast,
+	const MetaDataTypes*... InMetaData,
+	const CompositeTypes*... InCompositeComponents) const
 {
 	const int32 Num = Allocation->Num();
 
@@ -85,17 +105,22 @@ void TSetPartialPropertyValuesImpl<PropertyTraits, TPropertyMetaData<MetaDataTyp
 		checkSlow(PropertyOffset != 0);
 
 		StorageType Storage{};
-		PropertyTraits::GetObjectPropertyValue(Objects[Index], InMetaData[Index]..., PropertyOffset, Storage);
+		Traits->GetObjectPropertyValue(Objects[Index], InMetaData[Index]..., PropertyOffset, Storage);
 
 		PatchCompositeValue(CompositeDefinitions, &Storage, InCompositeComponents ? &InCompositeComponents[Index] : nullptr...);
 
-		PropertyTraits::SetObjectPropertyValue(Objects[Index], InMetaData[Index]..., PropertyOffset, Storage);
+		Traits->SetObjectPropertyValue(Objects[Index], InMetaData[Index]..., PropertyOffset, Storage);
 	}
 }
 
 
 template<typename PropertyTraits, typename ...MetaDataTypes, int ...CompositeIndices, typename ...CompositeTypes>
-void TSetPartialPropertyValuesImpl<PropertyTraits, TPropertyMetaData<MetaDataTypes...>, TIntegerSequence<int, CompositeIndices...>, CompositeTypes...>::ForEachSlow(const FEntityAllocation* Allocation, UObject* const* Objects, const TSharedPtr<FTrackInstancePropertyBindings>* Slow, const MetaDataTypes*... InMetaData, const CompositeTypes*... InCompositeComponents) const
+void TSetPartialPropertyValuesImpl<PropertyTraits, TPropertyMetaData<MetaDataTypes...>, TIntegerSequence<int, CompositeIndices...>, CompositeTypes...>::ForEachSlow(
+	const FEntityAllocation* Allocation,
+	UObject* const* Objects,
+	const TSharedPtr<FTrackInstancePropertyBindings>* Slow,
+	const MetaDataTypes*... InMetaData,
+	const CompositeTypes*... InCompositeComponents) const
 {
 	const int32 Num = Allocation->Num();
 
@@ -104,11 +129,11 @@ void TSetPartialPropertyValuesImpl<PropertyTraits, TPropertyMetaData<MetaDataTyp
 		FTrackInstancePropertyBindings* Bindings = Slow[Index].Get();
 
 		StorageType Storage{};
-		PropertyTraits::GetObjectPropertyValue(Objects[Index], InMetaData[Index]..., Bindings, Storage);
+		Traits->GetObjectPropertyValue(Objects[Index], InMetaData[Index]..., Bindings, Storage);
 
 		PatchCompositeValue(CompositeDefinitions, &Storage, InCompositeComponents ? &InCompositeComponents[Index] : nullptr...);
 
-		PropertyTraits::SetObjectPropertyValue(Objects[Index], InMetaData[Index]..., Bindings, Storage);
+		Traits->SetObjectPropertyValue(Objects[Index], InMetaData[Index]..., Bindings, Storage);
 	}
 }
 
diff --git a/Engine/Source/Runtime/MovieScene/Public/EntitySystem/MovieScenePropertyComponentHandler.h b/Engine/Source/Runtime/MovieScene/Public/EntitySystem/MovieScenePropertyComponentHandler.h
--- a/Engine/Source/Runtime/MovieScene/Public/EntitySystem/MovieScenePropertyComponentHandler.h
+++ b/Engine/Source/Runtime/MovieScene/Public/EntitySystem/MovieScenePropertyComponentHandler.h
@@ -12,12 +12,12 @@
 #include "EntitySystem/MovieScenePropertySystemTypes.inl"
 #include "EntitySystem/MovieSceneOperationalTypeConversions.h"
 #include "EntitySystem/Interrogation/MovieSceneInterrogationExtension.h"
-#include "EntitySystem/Interrogation/MovieSceneInterrogationLinker.h"
 #include "Evaluation/PreAnimatedState/MovieScenePreAnimatedStateExtension.h"
 
 #include "Evaluation/PreAnimatedState/MovieScenePreAnimatedStateStorage.h"
 #include "Evaluation/PreAnimatedState/MovieScenePreAnimatedPropertyStorage.h"
 
+class UMovieSceneTrack;
 
 namespace UE
 {
@@ -25,7 +25,7 @@ namespace MovieScene
 {
 
 
-template<typename PropertyTraits, typename MetaDatatype, typename MetaDataIndices, typename CompositeIndices, typename ...CompositeTypes>
+template<typename PropertyTraits, typename MetaDataType, typename MetaDataIndices, typename CompositeIndices, typename ...CompositeTypes>
 struct TPropertyComponentHandlerImpl;
 
 template<typename PropertyTraits, typename ...CompositeTypes>
@@ -53,11 +53,13 @@ struct TInitialValueProcessorImpl<PropertyTraits, TPropertyMetaData<MetaDataType
 
 	// Transient properties reset on each usage.
 	IInterrogationExtension* Interrogation;
+	const PropertyTraits* Traits;
 	TPropertyValueStorage<PropertyTraits>* CacheStorage;
 
 	FEntityAllocationWriteContext WriteContext;
 
 	TInitialValueProcessorImpl(
+		const PropertyTraits* InTraits,
 		TComponentTypeID<StorageType> InInitialValueType,
 		TArrayView<const FComponentTypeID> InMetaDataComponents,
 		ICustomPropertyRegistration* InCustomPropertyRegistration
@@ -68,6 +70,7 @@ struct TInitialValueProcessorImpl<PropertyTraits, TPropertyMetaData<MetaDataType
 		, MetaDataComponents(InMetaDataComponents[MetaDataIndices].ReinterpretCast<MetaDataTypes>()...)
 		, WriteContext(FEntityAllocationWriteContext::NewAllocation())
 	{
+		Traits = InTraits;
 		Interrogation = nullptr;
 		CacheStorage = nullptr;
 	}
@@ -133,7 +136,7 @@ struct TInitialValueProcessorImpl<PropertyTraits, TPropertyMetaData<MetaDataType
 			const FCustomPropertyIndex* RawIndices = CustomIndices.AsPtr();
 			for (int32 Index = 0; Index < Num; ++Index)
 			{
-				PropertyTraits::GetObjectPropertyValue(BoundObjects[Index], MetaData.template Get<MetaDataIndices>()[Index]..., CustomAccessors[RawIndices[Index].Value], InitialValues[Index]);
+				Traits->GetObjectPropertyValue(BoundObjects[Index], MetaData.template Get<MetaDataIndices>()[Index]..., CustomAccessors[RawIndices[Index].Value], InitialValues[Index]);
 			}
 		}
 
@@ -142,7 +145,7 @@ struct TInitialValueProcessorImpl<PropertyTraits, TPropertyMetaData<MetaDataType
 			const uint16* RawOffsets = FastOffsets.AsPtr();
 			for (int32 Index = 0; Index < Num; ++Index)
 			{
-				PropertyTraits::GetObjectPropertyValue(BoundObjects[Index], MetaData.template Get<MetaDataIndices>()[Index]..., RawOffsets[Index], InitialValues[Index]);
+				Traits->GetObjectPropertyValue(BoundObjects[Index], MetaData.template Get<MetaDataIndices>()[Index]..., RawOffsets[Index], InitialValues[Index]);
 			}
 		}
 
@@ -151,7 +154,7 @@ struct TInitialValueProcessorImpl<PropertyTraits, TPropertyMetaData<MetaDataType
 			const TSharedPtr<FTrackInstancePropertyBindings>* RawProperties = SlowProperties.AsPtr();
 			for (int32 Index = 0; Index < Num; ++Index)
 			{
-				PropertyTraits::GetObjectPropertyValue(BoundObjects[Index], MetaData.template Get<MetaDataIndices>()[Index]..., RawProperties[Index].Get(), InitialValues[Index]);
+				Traits->GetObjectPropertyValue(BoundObjects[Index], MetaData.template Get<MetaDataIndices>()[Index]..., RawProperties[Index].Get(), InitialValues[Index]);
 			}
 		}
 	}
@@ -181,7 +184,7 @@ struct TInitialValueProcessorImpl<PropertyTraits, TPropertyMetaData<MetaDataType
 				else
 				{
 					StorageType Value{};
-					PropertyTraits::GetObjectPropertyValue(BoundObjects[Index], MetaData.template Get<MetaDataIndices>()[Index]..., CustomAccessors[RawIndices[Index].Value], Value);
+					Traits->GetObjectPropertyValue(BoundObjects[Index], MetaData.template Get<MetaDataIndices>()[Index]..., CustomAccessors[RawIndices[Index].Value], Value);
 
 					InitialValues[Index] = Value;
 					InitialValueIndices[Index] = CacheStorage->AddInitialValue(BoundObjects[Index], Value, RawIndices[Index]);
@@ -202,7 +205,7 @@ struct TInitialValueProcessorImpl<PropertyTraits, TPropertyMetaData<MetaDataType
 				else
 				{
 					StorageType Value{};
-					PropertyTraits::GetObjectPropertyValue(BoundObjects[Index], MetaData.template Get<MetaDataIndices>()[Index]..., RawOffsets[Index], Value);
+					Traits->GetObjectPropertyValue(BoundObjects[Index], MetaData.template Get<MetaDataIndices>()[Index]..., RawOffsets[Index], Value);
 
 					InitialValues[Index] = Value;
 					InitialValueIndices[Index] = CacheStorage->AddInitialValue(BoundObjects[Index], Value, RawOffsets[Index]);
@@ -223,7 +226,7 @@ struct TInitialValueProcessorImpl<PropertyTraits, TPropertyMetaData<MetaDataType
 				else
 				{
 					StorageType Value{};
-					PropertyTraits::GetObjectPropertyValue(BoundObjects[Index], MetaData.template Get<MetaDataIndices>()[Index]..., RawProperties[Index].Get(), Value);
+					Traits->GetObjectPropertyValue(BoundObjects[Index], MetaData.template Get<MetaDataIndices>()[Index]..., RawProperties[Index].Get(), Value);
 
 					InitialValues[Index] = Value;
 					InitialValueIndices[Index] = CacheStorage->AddInitialValue(BoundObjects[Index], Value, RawProperties[Index].Get());
@@ -294,16 +297,17 @@ struct TInitialValueProcessorImpl<PropertyTraits, TPropertyMetaData<MetaDataType
 
 			if (!Property.IsSet())
 			{
-				PropertyTraits::GetObjectPropertyValue(Object, MetaData.template Get<MetaDataIndices>()[Index]..., ChannelInfo->PropertyBinding.PropertyPath, CurrentValue.Emplace());
+				FTrackInstancePropertyBindings Bindings(ChannelInfo->PropertyBinding.PropertyName, ChannelInfo->PropertyBinding.PropertyPath.ToString());
+				Traits->GetObjectPropertyValue(Object, MetaData.template Get<MetaDataIndices>()[Index]..., &Bindings, CurrentValue.Emplace());
 			}
 			else if (const FCustomPropertyIndex* Custom = Property->TryGet<FCustomPropertyIndex>())
 			{
-				PropertyTraits::GetObjectPropertyValue(Object, MetaData.template Get<MetaDataIndices>()[Index]..., CustomAccessors[Custom->Value], CurrentValue.Emplace());
+				Traits->GetObjectPropertyValue(Object, MetaData.template Get<MetaDataIndices>()[Index]..., CustomAccessors[Custom->Value], CurrentValue.Emplace());
 			}
 			else
 			{
 				const uint16 FastPtrOffset = Property->Get<uint16>();
-				PropertyTraits::GetObjectPropertyValue(Object, MetaData.template Get<MetaDataIndices>()[Index]..., FastPtrOffset, CurrentValue.Emplace());
+				Traits->GetObjectPropertyValue(Object, MetaData.template Get<MetaDataIndices>()[Index]..., FastPtrOffset, CurrentValue.Emplace());
 			}
 
 			InitialValues[Index] = CurrentValue.GetValue();
@@ -320,10 +324,11 @@ struct TInitialValueProcessor : TInitialValueProcessorImpl<PropertyTraits, typen
 	TInitialValueProcessor() = delete;
 
 	TInitialValueProcessor(
+		const PropertyTraits* InTraits,
 		TComponentTypeID<typename PropertyTraits::StorageType> InInitialValueType,
 		TArrayView<const FComponentTypeID> InMetaDataComponents,
 		ICustomPropertyRegistration* InCustomPropertyRegistration
-	) : Super(InInitialValueType, InMetaDataComponents, InCustomPropertyRegistration)
+	) : Super(InTraits, InInitialValueType, InMetaDataComponents, InCustomPropertyRegistration)
 	{}
 };
 
@@ -338,7 +343,8 @@ constexpr bool IsCompositePropertyTraits(...)
 }
 
 template<typename PropertyTraits, typename ...MetaDataTypes, int ...MetaDataIndices, typename ...CompositeTypes, int ...CompositeIndices>
-struct TPropertyComponentHandlerImpl<PropertyTraits, TPropertyMetaData<MetaDataTypes...>, TIntegerSequence<int, MetaDataIndices...>, TIntegerSequence<int, CompositeIndices...>, CompositeTypes...> : IPropertyComponentHandler
+struct TPropertyComponentHandlerImpl<PropertyTraits, TPropertyMetaData<MetaDataTypes...>, TIntegerSequence<int, MetaDataIndices...>, TIntegerSequence<int, CompositeIndices...>, CompositeTypes...>
+	: IPropertyComponentHandler
 {
 	static constexpr bool bIsComposite = IsCompositePropertyTraits((PropertyTraits*)nullptr);
 
@@ -353,6 +359,20 @@ struct TPropertyComponentHandlerImpl<PropertyTraits, TPropertyMetaData<MetaDataT
 	{
 	}
 
+	virtual bool SupportsProperty(const FPropertyDefinition& Definition, const FProperty& InProperty) const override
+	{
+		return static_cast<const PropertyTraits*>(Definition.TraitsInstance)->SupportsProperty(InProperty);
+	}
+	virtual FIntermediatePropertyValue CoercePropertyValue(const FPropertyDefinition& Definition, const FProperty& InProperty, const FSourcePropertyValue& InPropertyValue) const override
+	{
+		return static_cast<const PropertyTraits*>(Definition.TraitsInstance)->CoercePropertyValue(InProperty, InPropertyValue);
+	}
+	virtual void UnpackChannels(const FPropertyDefinition& Definition, const FProperty& Property, const FIntermediatePropertyValueConstRef& Value, FUnpackedChannelValues& OutUnpackedValues) const override
+	{
+		const PropertyTraits* Traits = static_cast<const PropertyTraits*>(Definition.TraitsInstance);
+		return Traits->UnpackChannels(*Value.Cast<StorageType>(), Property, OutUnpackedValues);
+	}
+
 	virtual TSharedPtr<IPreAnimatedStorage> GetPreAnimatedStateStorage(const FPropertyDefinition& Definition, FPreAnimatedStateExtension* Container) override
 	{
 		TSharedPtr<PreAnimatedStorageType> Existing = Container->FindStorage(StorageID);
@@ -375,6 +395,8 @@ struct TPropertyComponentHandlerImpl<PropertyTraits, TPropertyMetaData<MetaDataT
 	{
 		FBuiltInComponentTypes* BuiltInComponents = FBuiltInComponentTypes::Get();
 
+		const PropertyTraits* TraitsInstance = static_cast<const PropertyTraits*>(Definition.TraitsInstance);
+
 		FEntityTaskBuilder()
 		.Read(BuiltInComponents->BoundObject)
 		.ReadOneOf(BuiltInComponents->CustomPropertyIndex, BuiltInComponents->FastPropertyOffset, BuiltInComponents->SlowProperty)
@@ -385,7 +407,7 @@ struct TPropertyComponentHandlerImpl<PropertyTraits, TPropertyMetaData<MetaDataT
 		.CombineFilter(AdditionalFilter)
 		.SetStat(Definition.StatID)
 		.SetDesiredThread(Linker->EntityManager.GetGatherThread())
-		.template Fork_PerAllocation<CompleteSetterTask>(&Linker->EntityManager, TaskScheduler, Definition.CustomPropertyRegistration);
+		.template Fork_PerAllocation<CompleteSetterTask>(&Linker->EntityManager, TaskScheduler, TraitsInstance, Definition.CustomPropertyRegistration);
 
 		if constexpr (bIsComposite)
 		{
@@ -411,7 +433,7 @@ struct TPropertyComponentHandlerImpl<PropertyTraits, TPropertyMetaData<MetaDataT
 				.FilterOut(CompletePropertyMask)
 				.SetStat(Definition.StatID)
 				.SetDesiredThread(Linker->EntityManager.GetGatherThread())
-				.template Fork_PerAllocation<PartialSetterTask>(&Linker->EntityManager, TaskScheduler, Definition.CustomPropertyRegistration, Composites);
+				.template Fork_PerAllocation<PartialSetterTask>(&Linker->EntityManager, TaskScheduler, TraitsInstance, Definition.CustomPropertyRegistration, Composites);
 			}
 		}
 	}
@@ -420,6 +442,8 @@ struct TPropertyComponentHandlerImpl<PropertyTraits, TPropertyMetaData<MetaDataT
 	{
 		FBuiltInComponentTypes* BuiltInComponents = FBuiltInComponentTypes::Get();
 
+		const PropertyTraits* TraitsInstance = static_cast<const PropertyTraits*>(Definition.TraitsInstance);
+
 		FEntityTaskBuilder()
 		.Read(BuiltInComponents->BoundObject)
 		.ReadOneOf(BuiltInComponents->CustomPropertyIndex, BuiltInComponents->FastPropertyOffset, BuiltInComponents->SlowProperty)
@@ -429,7 +453,7 @@ struct TPropertyComponentHandlerImpl<PropertyTraits, TPropertyMetaData<MetaDataT
 		.FilterNone({ BuiltInComponents->Tags.Ignored })
 		.SetStat(Definition.StatID)
 		.SetDesiredThread(Linker->EntityManager.GetGatherThread())
-		.template Dispatch_PerAllocation<CompleteSetterTask>(&Linker->EntityManager, InPrerequisites, &Subsequents, Definition.CustomPropertyRegistration);
+		.template Dispatch_PerAllocation<CompleteSetterTask>(&Linker->EntityManager, InPrerequisites, &Subsequents, TraitsInstance, Definition.CustomPropertyRegistration);
 
 		if constexpr (bIsComposite)
 		{
@@ -454,7 +478,7 @@ struct TPropertyComponentHandlerImpl<PropertyTraits, TPropertyMetaData<MetaDataT
 				.FilterNone({ BuiltInComponents->Tags.Ignored })
 				.SetStat(Definition.StatID)
 				.SetDesiredThread(Linker->EntityManager.GetGatherThread())
-				.template Dispatch_PerAllocation<PartialSetterTask>(&Linker->EntityManager, InPrerequisites, &Subsequents, Definition.CustomPropertyRegistration, Composites);
+				.template Dispatch_PerAllocation<PartialSetterTask>(&Linker->EntityManager, InPrerequisites, &Subsequents, TraitsInstance, Definition.CustomPropertyRegistration, Composites);
 			}
 		}
 	}
@@ -462,18 +486,19 @@ struct TPropertyComponentHandlerImpl<PropertyTraits, TPropertyMetaData<MetaDataT
 	virtual TSharedPtr<IInitialValueProcessor> MakeInitialValueProcessor(const FPropertyDefinition& Definition) override
 	{
 		return MakeShared<TInitialValueProcessor<PropertyTraits>>(
+			static_cast<const PropertyTraits*>(Definition.TraitsInstance),
 			Definition.InitialValueType.ReinterpretCast<StorageType>(),
 			Definition.MetaDataTypes,
 			Definition.CustomPropertyRegistration
 			);
 	}
 
-	virtual void RecomposeBlendOperational(const FPropertyDefinition& PropertyDefinition, TArrayView<const FPropertyCompositeDefinition> Composites, const FValueDecompositionParams& InParams, UMovieSceneBlenderSystem* Blender, FConstPropertyComponentView InCurrentValue, FPropertyComponentArrayView OutResult) override
+	virtual void RecomposeBlendOperational(const FPropertyDefinition& PropertyDefinition, TArrayView<const FPropertyCompositeDefinition> Composites, const FValueDecompositionParams& InParams, UMovieSceneBlenderSystem* Blender, const FIntermediatePropertyValueConstRef& InCurrentValue, TArrayView<FIntermediatePropertyValue> OutResult) override
 	{
-		RecomposeBlendImpl(PropertyDefinition, Composites, InParams, Blender, InCurrentValue.ReinterpretCast<StorageType>(), OutResult.ReinterpretCast<StorageType>());
+		RecomposeBlendImpl(PropertyDefinition, Composites, InParams, Blender, *InCurrentValue.Cast<StorageType>(), OutResult);
 	}
 
-	void RecomposeBlendImpl(const FPropertyDefinition& PropertyDefinition, TArrayView<const FPropertyCompositeDefinition> Composites, const FValueDecompositionParams& InParams, UMovieSceneBlenderSystem* Blender, const StorageType& InCurrentValue, TArrayView<StorageType> OutResults)
+	void RecomposeBlendImpl(const FPropertyDefinition& PropertyDefinition, TArrayView<const FPropertyCompositeDefinition> Composites, const FValueDecompositionParams& InParams, UMovieSceneBlenderSystem* Blender, const StorageType& InCurrentValue, TArrayView<FIntermediatePropertyValue> OutResults)
 	{
 		check(OutResults.Num() == InParams.Query.Entities.Num());
 
@@ -526,7 +551,7 @@ struct TPropertyComponentHandlerImpl<PropertyTraits, TPropertyMetaData<MetaDataT
 		{
 			FMovieSceneEntityID EntityID = LocalParams.Query.Entities[Index];
 
-			uint8* Result = reinterpret_cast<uint8*>(&OutResults[Index]);
+			uint8* Result = reinterpret_cast<uint8*>(OutResults[Index].Cast<StorageType>());
 
 			for (int32 CompositeIndex = 0; CompositeIndex < NumComposites; ++CompositeIndex)
 			{
@@ -638,6 +663,12 @@ struct TPropertyComponentHandlerImpl<PropertyTraits, TPropertyMetaData<MetaDataT
 
 
 
+struct CPropertyTraitsDefineHandler
+{
+	template<typename T, typename ...Composites>
+	auto Requires() -> decltype(T::MakeHandler());
+};
+
 
 template<typename PropertyTraits>
 struct TPropertyDefinitionBuilder
@@ -688,9 +719,22 @@ struct TPropertyDefinitionBuilder
 		return *this;
 	}
 
+	TPropertyDefinitionBuilder<PropertyTraits>& SetDefaultTrackType(TSubclassOf<UMovieSceneTrack> InDefaultTrackType)
+	{
+		Definition->DefaultTrackType = InDefaultTrackType;
+		return *this;
+	}
+
 	void Commit()
 	{
-		Definition->Handler = TPropertyComponentHandler<PropertyTraits, typename PropertyTraits::StorageType>();
+		if constexpr (TModels_V<CPropertyTraitsDefineHandler, PropertyTraits>)
+		{
+			Definition->Handler = PropertyTraits::template MakeHandler<typename PropertyTraits::StorageType>();
+		}
+		else
+		{
+			Definition->Handler = TPropertyComponentHandler<PropertyTraits, typename PropertyTraits::StorageType>();
+		}
 		Definition->SetupInitialValueProcessor();
 	}
 
@@ -781,9 +825,22 @@ struct TCompositePropertyDefinitionBuilder
 		return *this;
 	}
 
+	TCompositePropertyDefinitionBuilder<PropertyTraits, Composites...>& SetDefaultTrackType(TSubclassOf<UMovieSceneTrack> InDefaultTrackType)
+	{
+		Definition->DefaultTrackType = InDefaultTrackType;
+		return *this;
+	}
+
 	void Commit()
 	{
-		Definition->Handler = TPropertyComponentHandler<PropertyTraits, Composites...>();
+		if constexpr (TModels_V<CPropertyTraitsDefineHandler, PropertyTraits, Composites...>)
+		{
+			Definition->Handler = PropertyTraits::template MakeHandler<Composites...>();
+		}
+		else
+		{
+			Definition->Handler = TPropertyComponentHandler<PropertyTraits, Composites...>();
+		}
 		Definition->SetupInitialValueProcessor();
 	}
 	
@@ -822,6 +879,8 @@ struct FPropertyRecomposerImpl
 	template<typename PropertyTraits>
 	TRecompositionResult<typename PropertyTraits::StorageType> RecomposeBlendOperational(const TPropertyComponents<PropertyTraits>& InComponents, const FDecompositionQuery& InQuery, const typename PropertyTraits::StorageType& InCurrentValue);
 
+	MOVIESCENE_API FRecompositionResult RecomposeBlendOperational(const FPropertyDefinition& InPropertyDefinition, const FDecompositionQuery& InQuery, const FIntermediatePropertyValueConstRef& InCurrentValue);
+
 	FOnGetPropertyRecomposerPropertyInfo OnGetPropertyInfo;
 };
 
@@ -861,8 +920,17 @@ TRecompositionResult<typename PropertyTraits::StorageType> FPropertyRecomposerIm
 
 	TArrayView<const FPropertyCompositeDefinition> Composites = BuiltInComponents->PropertyRegistry.GetComposites(PropertyDefinition);
 
-	PropertyDefinition.Handler->RecomposeBlendOperational(PropertyDefinition, Composites, Params, Blender, InCurrentValue, Result.Values);
+	// Make an array of type-erased values that reference each element of Result.Values.
+	// These will be written to inside RecomposeBlendOperational on success
+	TArray<FIntermediatePropertyValue, TInlineAllocator<16>> ErasedValues;
+	ErasedValues.Reserve(Result.Values.Num());
+	for (int32 Index = 0; Index < Result.Values.Num(); ++Index)
+	{
+		ErasedValues.Emplace(FIntermediatePropertyValue::FromAddress(&Result.Values[Index]));
+	}
 
+	FIntermediatePropertyValueConstRef CurrentValueErased(&InCurrentValue);
+	PropertyDefinition.Handler->RecomposeBlendOperational(PropertyDefinition, Composites, Params, Blender, CurrentValueErased, ErasedValues);
 	return Result;
 }
 
diff --git a/Engine/Source/Runtime/MovieScene/Public/EntitySystem/MovieScenePropertyRegistry.h b/Engine/Source/Runtime/MovieScene/Public/EntitySystem/MovieScenePropertyRegistry.h
--- a/Engine/Source/Runtime/MovieScene/Public/EntitySystem/MovieScenePropertyRegistry.h
+++ b/Engine/Source/Runtime/MovieScene/Public/EntitySystem/MovieScenePropertyRegistry.h
@@ -21,12 +21,14 @@
 #include "Templates/UnrealTemplate.h"
 #include "Templates/UnrealTypeTraits.h"
 #include "UObject/NameTypes.h"
+#include "Templates/SubclassOf.h"
 
 #include <initializer_list>
 
 class FTrackInstancePropertyBindings;
 class UClass;
 class UMovieSceneBlenderSystem;
+class UMovieSceneTrack;
 class UMovieSceneEntitySystemLinker;
 class UObject;
 struct FMovieScenePropertyBinding;
@@ -39,6 +41,7 @@ namespace MovieScene
 struct FFloatDecompositionParams;
 struct FPropertyCompositeDefinition;
 struct FPropertyDefinition;
+struct IPropertyTraits;
 
 DECLARE_CYCLE_STAT(TEXT("Apply properties"), MovieSceneEval_ApplyProperties,  STATGROUP_MovieSceneECS);
 
@@ -67,6 +70,7 @@ struct FPropertyDefinition
 			uint16 InVariableSizeCompositeOffset,
 			FComponentTypeID InPropertyType, FComponentTypeID InInitialValueType)
 		: CustomPropertyRegistration(nullptr)
+		, TraitsInstance(nullptr)
 		, DoubleCompositeMask(0)
 		, VariableSizeCompositeOffset(InVariableSizeCompositeOffset)
 		, CompositeSize(0)
@@ -98,6 +102,12 @@ struct FPropertyDefinition
 	/** Pointer to a custom getter/setter registry for short circuiting the UObject VM. Must outlive this definitions lifetime (usually these are static or singletons) */
 	ICustomPropertyRegistration* CustomPropertyRegistration = nullptr;
 
+	/** Pointer to the traits instance */
+	const IPropertyTraits* TraitsInstance = nullptr;
+
+	/** A default track type to use for this property */
+	TSubclassOf<UMovieSceneTrack> DefaultTrackType;
+
 	/** Stat ID for this property type */
 	TStatId StatID;
 
@@ -191,6 +201,7 @@ public:
 	{
 		DefinePropertyImpl(InOutPropertyComponents, InStatName);
 		FPropertyDefinition* Property = &Properties[InOutPropertyComponents.CompositeID.AsIndex()];
+		Property->TraitsInstance = &InOutPropertyComponents.PropertyTraits;
 		return TCompositePropertyDefinitionBuilder<PropertyTraits>(Property, this);
 	}
 
@@ -205,6 +216,7 @@ public:
 	{
 		DefinePropertyImpl(InOutPropertyComponents, InStatName);
 		FPropertyDefinition* Property = &Properties[InOutPropertyComponents.CompositeID.AsIndex()];
+		Property->TraitsInstance = &InOutPropertyComponents.PropertyTraits;
 		return TPropertyDefinitionBuilder<PropertyTraits>(Property, this);
 	}
 
@@ -224,6 +236,11 @@ public:
 		return Properties;
 	}
 
+	/**
+	 * Access all the properties currently registered
+	 */
+	MOVIESCENE_API const FPropertyDefinition* FindPropertyDefinition(FComponentTypeID ComponentTypeID) const;
+
 	/**
 	 * Retrieve a generic representation of all the composites that contribute to a given property
 	 */
diff --git a/Engine/Source/Runtime/MovieScene/Public/EntitySystem/MovieScenePropertySystemTypes.h b/Engine/Source/Runtime/MovieScene/Public/EntitySystem/MovieScenePropertySystemTypes.h
--- a/Engine/Source/Runtime/MovieScene/Public/EntitySystem/MovieScenePropertySystemTypes.h
+++ b/Engine/Source/Runtime/MovieScene/Public/EntitySystem/MovieScenePropertySystemTypes.h
@@ -10,17 +10,14 @@
 #include "Templates/UnrealTypeTraits.h"
 #include "Containers/StringView.h"
 #include "EntitySystem/MovieSceneComponentAccessors.h"
+#include "EntitySystem/MovieScenePropertyMetaData.h"
 
 class UClass;
 class FTrackInstancePropertyBindings;
 
 
-namespace UE
+namespace UE::MovieScene
 {
-namespace MovieScene
-{
-
-template<typename...> struct TPropertyMetaData;
 
 struct FCustomPropertyIndex
 {
@@ -128,8 +125,12 @@ struct TCustomPropertyAccessorFunctionsImpl<PropertyTraits, TPropertyMetaData<Me
 	SetterFunc Setter;
 };
 
+
 template<typename PropertyTraits>
-using TCustomPropertyAccessorFunctions = TCustomPropertyAccessorFunctionsImpl<PropertyTraits, typename PropertyTraits::MetaDataType>;
+struct TCustomPropertyAccessorFunctions
+	: TCustomPropertyAccessorFunctionsImpl<PropertyTraits, TGetPublicPropertyMetaDataT<PropertyTraits>>
+{
+};
 
 struct FCustomPropertyAccessor
 {
@@ -285,35 +286,31 @@ protected:
 	TArray<TCustomPropertyAccessor<PropertyTraits>, TInlineAllocator<InlineSize>> CustomAccessors;
 };
 
-// #include MovieSceneEntityFactoryTemplates.h for specialization
-template<typename, typename>
-struct TPropertyMetaDataComponentsImpl;
-
-template<typename>
-struct TPropertyMetaDataComponents;
-
-
 /**
  * User-defined property type that is represented as an UE::MovieScene::FPropertyDefinition within UE::MovieScene::FPropertyRegistry
  *
  * This type must be templated on a traits class that defines the storage type for the property, and methods for retrieving and assigning
- * the property value from a UObject*. See TDirectPropertyTraits and TIndirectPropertyTraits for examples. TRuntimePropertyTraits shows
+ * the property value from a UObject*. See TPropertyTraits for an example. TRuntimePropertyTraits shows
  * an example trait for a property that requires additional meta-data components where multiple property types can be represented by
  * a single property definition at runtime (ie, for multi-channel vectors, or colors)
  */
-template<typename PropertyTraits>
+template<typename InPropertyTraits>
 struct TPropertyComponents
 {
+	using MetaDataType = typename InPropertyTraits::MetaDataType;
+
 	FComponentTypeID PropertyTag;
-	TComponentTypeID<typename PropertyTraits::StorageType> InitialValue;
+	TComponentTypeID<typename InPropertyTraits::StorageType> InitialValue;
+
+	TPropertyMetaDataComponents<typename InPropertyTraits::MetaDataType> MetaDataComponents;
 
-	TPropertyMetaDataComponents<typename PropertyTraits::MetaDataType> MetaDataComponents;
+	TCompositePropertyTypeID<InPropertyTraits> CompositeID;
 
-	TCompositePropertyTypeID<PropertyTraits> CompositeID;
+	InPropertyTraits PropertyTraits;
 };
 
 template<typename PropertyTraits, typename MetaDataType>
-struct TSetPropertyValuesImpl{};
+struct TSetPropertyValuesImpl;
 
 
 
@@ -321,10 +318,11 @@ template<typename PropertyTraits, typename ...MetaDataTypes>
 struct TSetPropertyValuesImpl<PropertyTraits, TPropertyMetaData<MetaDataTypes...>>
 {
 	using StorageType = typename PropertyTraits::StorageType;
-	using InParamType     = typename TCallTraits<typename PropertyTraits::StorageType>::ParamType;
+	using InParamType = typename TCallTraits<typename PropertyTraits::StorageType>::ParamType;
 
-	explicit TSetPropertyValuesImpl(ICustomPropertyRegistration* InCustomProperties)
-		: CustomProperties(InCustomProperties)
+	explicit TSetPropertyValuesImpl(const PropertyTraits* InTraits, const ICustomPropertyRegistration* InCustomProperties)
+		: Traits(InTraits)
+		, CustomProperties(InCustomProperties)
 	{
 		if (CustomProperties)
 		{
@@ -354,7 +352,7 @@ struct TSetPropertyValuesImpl<PropertyTraits, TPropertyMetaData<MetaDataTypes...
 	 *     .Read( TComponentTypeID<PropertyType>(...) )
 	 *     .Dispatch_PerEntity<TSetPropertyValues<PropertyType>>(...);
 	 */
-	static void ForEachEntity(UObject* InObject, uint16 PropertyOffset, typename TCallTraits<MetaDataTypes>::ParamType... MetaData, InParamType ValueToSet);
+	void ForEachEntity(UObject* InObject, uint16 PropertyOffset, typename TCallTraits<MetaDataTypes>::ParamType... MetaData, InParamType ValueToSet) const;
 
 
 	/**
@@ -367,7 +365,7 @@ struct TSetPropertyValuesImpl<PropertyTraits, TPropertyMetaData<MetaDataTypes...
 	 *     .Read( TComponentTypeID<PropertyType>(...) )
 	 *     .Dispatch_PerEntity<TSetPropertyValues<PropertyType>>(...);
 	 */
-	static void ForEachEntity(UObject* InObject, const TSharedPtr<FTrackInstancePropertyBindings>& PropertyBindings, typename TCallTraits<MetaDataTypes>::ParamType... MetaData, InParamType ValueToSet);
+	void ForEachEntity(UObject* InObject, const TSharedPtr<FTrackInstancePropertyBindings>& PropertyBindings, typename TCallTraits<MetaDataTypes>::ParamType... MetaData, InParamType ValueToSet) const;
 
 public:
 
@@ -401,7 +399,8 @@ public:
 
 private:
 
-	ICustomPropertyRegistration* CustomProperties;
+	const PropertyTraits* Traits;
+	const ICustomPropertyRegistration* CustomProperties;
 	FCustomAccessorView CustomAccessors;
 };
 
@@ -437,8 +436,8 @@ struct TSetPropertyValues : TSetPropertyValuesImpl<PropertyTraits, typename Prop
 {
 	using ParamType = typename TCallTraits<typename PropertyTraits::StorageType>::ParamType;
 
-	explicit TSetPropertyValues(ICustomPropertyRegistration* InCustomProperties)
-		: TSetPropertyValuesImpl<PropertyTraits, typename PropertyTraits::MetaDataType>(InCustomProperties)
+	explicit TSetPropertyValues(const PropertyTraits* InTraits, const ICustomPropertyRegistration* InCustomProperties)
+		: TSetPropertyValuesImpl<PropertyTraits, typename PropertyTraits::MetaDataType>(InTraits, InCustomProperties)
 	{}
 };
 
@@ -446,7 +445,7 @@ struct TSetPropertyValues : TSetPropertyValuesImpl<PropertyTraits, typename Prop
 
 
 template<typename PropertyTraits, typename MetaDataType>
-struct TGetPropertyValuesImpl{};
+struct TGetPropertyValuesImpl;
 
 
 
@@ -455,8 +454,9 @@ struct TGetPropertyValuesImpl<PropertyTraits, TPropertyMetaData<MetaDataTypes...
 {
 	using StorageType = typename PropertyTraits::StorageType;
 
-	explicit TGetPropertyValuesImpl(ICustomPropertyRegistration* InCustomProperties)
-		: CustomProperties(InCustomProperties)
+	explicit TGetPropertyValuesImpl(const PropertyTraits* InTraits, const ICustomPropertyRegistration* InCustomProperties)
+		: Traits(InTraits)
+		, CustomProperties(InCustomProperties)
 	{
 		if (CustomProperties)
 		{
@@ -537,7 +537,8 @@ public:
 
 private:
 
-	ICustomPropertyRegistration* CustomProperties;
+	const PropertyTraits* Traits;
+	const ICustomPropertyRegistration* CustomProperties;
 	FCustomAccessorView CustomAccessors;
 };
 
@@ -545,8 +546,8 @@ private:
 template<typename PropertyTraits>
 struct TGetPropertyValues : TGetPropertyValuesImpl<PropertyTraits, typename PropertyTraits::MetaDataType>
 {
-	explicit TGetPropertyValues(ICustomPropertyRegistration* InCustomProperties)
-		: TGetPropertyValuesImpl<PropertyTraits, typename PropertyTraits::MetaDataType>(InCustomProperties)
+	explicit TGetPropertyValues(const PropertyTraits* InTraits, const ICustomPropertyRegistration* InCustomProperties)
+		: TGetPropertyValuesImpl<PropertyTraits, typename PropertyTraits::MetaDataType>(InTraits, InCustomProperties)
 	{}
 };
 
@@ -590,8 +591,9 @@ struct TSetCompositePropertyValuesImpl<PropertyTraits, TPropertyMetaData<MetaDat
 {
 	using StorageType = typename PropertyTraits::StorageType;
 
-	explicit TSetCompositePropertyValuesImpl(ICustomPropertyRegistration* InCustomProperties)
-		: CustomProperties(InCustomProperties)
+	explicit TSetCompositePropertyValuesImpl(const PropertyTraits* InTraits, const ICustomPropertyRegistration* InCustomProperties)
+		: Traits(InTraits)
+		, CustomProperties(InCustomProperties)
 	{
 		if (CustomProperties)
 		{
@@ -680,7 +682,8 @@ public:
 
 private:
 
-	ICustomPropertyRegistration* CustomProperties;
+	const PropertyTraits* Traits;
+	const ICustomPropertyRegistration* CustomProperties;
 	FCustomAccessorView CustomAccessors;
 };
 
@@ -688,5 +691,4 @@ template<typename PropertyTraits, typename ...CompositeTypes>
 using TSetCompositePropertyValues = TSetCompositePropertyValuesImpl<PropertyTraits, typename PropertyTraits::MetaDataType, CompositeTypes...>;
 
 
-} // namespace MovieScene
-} // namespace UE
+} // namespace UE::MovieScene
diff --git a/Engine/Source/Runtime/MovieScene/Public/EntitySystem/MovieScenePropertySystemTypes.inl b/Engine/Source/Runtime/MovieScene/Public/EntitySystem/MovieScenePropertySystemTypes.inl
--- a/Engine/Source/Runtime/MovieScene/Public/EntitySystem/MovieScenePropertySystemTypes.inl
+++ b/Engine/Source/Runtime/MovieScene/Public/EntitySystem/MovieScenePropertySystemTypes.inl
@@ -17,35 +17,50 @@
 #include "EntitySystem/MovieScenePropertyBinding.h"
 
 
-namespace UE
-{
-namespace MovieScene
+namespace UE::MovieScene
 {
 
 template<typename PropertyTraits, typename ...MetaDataTypes>
-void TSetPropertyValuesImpl<PropertyTraits, TPropertyMetaData<MetaDataTypes...>>::ForEachEntity(UObject* InObject, FCustomPropertyIndex CustomIndex, typename TCallTraits<MetaDataTypes>::ParamType... MetaData, InParamType ValueToSet) const
+void TSetPropertyValuesImpl<PropertyTraits, TPropertyMetaData<MetaDataTypes...>>::ForEachEntity(
+	UObject* InObject,
+	FCustomPropertyIndex CustomIndex,
+	typename TCallTraits<MetaDataTypes>::ParamType... MetaData,
+	InParamType ValueToSet) const
 {
-	PropertyTraits::SetObjectPropertyValue(InObject, MetaData..., CustomAccessors[CustomIndex.Value], ValueToSet);
+	Traits->SetObjectPropertyValue(InObject, MetaData..., CustomAccessors[CustomIndex.Value], ValueToSet);
 }
 
 template<typename PropertyTraits, typename ...MetaDataTypes>
-void TSetPropertyValuesImpl<PropertyTraits, TPropertyMetaData<MetaDataTypes...>>::ForEachEntity(UObject* InObject, uint16 PropertyOffset, typename TCallTraits<MetaDataTypes>::ParamType... MetaData, InParamType ValueToSet)
+void TSetPropertyValuesImpl<PropertyTraits, TPropertyMetaData<MetaDataTypes...>>::ForEachEntity(
+	UObject* InObject,
+	uint16 PropertyOffset,
+	typename TCallTraits<MetaDataTypes>::ParamType... MetaData,
+	InParamType ValueToSet) const
 {
 	// Would really like to avoid branching here, but if we encounter this data the options are either handle it gracefully, stomp a vtable, or report a fatal error.
 	if (ensureAlwaysMsgf(PropertyOffset != 0, TEXT("Invalid property offset specified (ptr+%d bytes) for property on object %s. This would otherwise overwrite the object's vfptr."), PropertyOffset, *InObject->GetName()))
 	{
-		PropertyTraits::SetObjectPropertyValue(InObject, MetaData..., PropertyOffset, ValueToSet);
+		Traits->SetObjectPropertyValue(InObject, MetaData..., PropertyOffset, ValueToSet);
 	}
 }
 
 template<typename PropertyTraits, typename ...MetaDataTypes>
-void TSetPropertyValuesImpl<PropertyTraits, TPropertyMetaData<MetaDataTypes...>>::ForEachEntity(UObject* InObject, const TSharedPtr<FTrackInstancePropertyBindings>& PropertyBindings, typename TCallTraits<MetaDataTypes>::ParamType... MetaData, InParamType ValueToSet)
+void TSetPropertyValuesImpl<PropertyTraits, TPropertyMetaData<MetaDataTypes...>>::ForEachEntity(
+	UObject* InObject,
+	const TSharedPtr<FTrackInstancePropertyBindings>& PropertyBindings,
+	typename TCallTraits<MetaDataTypes>::ParamType... MetaData,
+	InParamType ValueToSet) const
 {
-	PropertyTraits::SetObjectPropertyValue(InObject, MetaData..., PropertyBindings.Get(), ValueToSet);
+	Traits->SetObjectPropertyValue(InObject, MetaData..., PropertyBindings.Get(), ValueToSet);
 }
 
 template<typename PropertyTraits, typename ...MetaDataTypes>
-void TSetPropertyValuesImpl<PropertyTraits, TPropertyMetaData<MetaDataTypes...>>::ForEachAllocation(const FEntityAllocation* Allocation, TRead<UObject*> BoundObjectComponents, FThreeWayAccessor ResolvedPropertyComponents, TRead<MetaDataTypes>... MetaDataComponents, TRead<StorageType> PropertyValueComponents) const
+void TSetPropertyValuesImpl<PropertyTraits, TPropertyMetaData<MetaDataTypes...>>::ForEachAllocation(
+	const FEntityAllocation* Allocation,
+	TRead<UObject*> BoundObjectComponents,
+	FThreeWayAccessor ResolvedPropertyComponents,
+	TRead<MetaDataTypes>... MetaDataComponents,
+	TRead<StorageType> PropertyValueComponents) const
 {
 	const int32 Num = Allocation->Num();
 	if (const FCustomPropertyIndex* Custom = ResolvedPropertyComponents.template Get<0>())
@@ -72,7 +87,12 @@ void TSetPropertyValuesImpl<PropertyTraits, TPropertyMetaData<MetaDataTypes...>>
 }
 
 template<typename PropertyTraits, typename ...MetaDataTypes>
-void TSetPropertyValuesImpl<PropertyTraits, TPropertyMetaData<MetaDataTypes...>>::ForEachAllocation(const FEntityAllocation* Allocation, TRead<UObject*> BoundObjectComponents, FTwoWayAccessor ResolvedPropertyComponents, TRead<MetaDataTypes>... MetaDataComponents, TRead<StorageType> PropertyValueComponents) const
+void TSetPropertyValuesImpl<PropertyTraits, TPropertyMetaData<MetaDataTypes...>>::ForEachAllocation(
+	const FEntityAllocation* Allocation,
+	TRead<UObject*> BoundObjectComponents,
+	FTwoWayAccessor ResolvedPropertyComponents,
+	TRead<MetaDataTypes>... MetaDataComponents,
+	TRead<StorageType> PropertyValueComponents) const
 {
 	const int32 Num = Allocation->Num();
 	if (const uint16* Fast = ResolvedPropertyComponents.template Get<0>())
@@ -92,29 +112,46 @@ void TSetPropertyValuesImpl<PropertyTraits, TPropertyMetaData<MetaDataTypes...>>
 }
 
 template<typename PropertyTraits, typename ...MetaDataTypes>
-void TGetPropertyValuesImpl<PropertyTraits, TPropertyMetaData<MetaDataTypes...>>::ForEachEntity(UObject* InObject, FCustomPropertyIndex CustomPropertyIndex, typename TCallTraits<MetaDataTypes>::ParamType... MetaData, StorageType& OutValue) const
+void TGetPropertyValuesImpl<PropertyTraits, TPropertyMetaData<MetaDataTypes...>>::ForEachEntity(
+	UObject* InObject,
+	FCustomPropertyIndex CustomPropertyIndex,
+	typename TCallTraits<MetaDataTypes>::ParamType... MetaData,
+	StorageType& OutValue) const
 {
-	PropertyTraits::GetObjectPropertyValue(InObject, MetaData..., CustomAccessors[CustomPropertyIndex.Value], OutValue);
+	Traits->GetObjectPropertyValue(InObject, MetaData..., CustomAccessors[CustomPropertyIndex.Value], OutValue);
 }
 
 template<typename PropertyTraits, typename ...MetaDataTypes>
-void TGetPropertyValuesImpl<PropertyTraits, TPropertyMetaData<MetaDataTypes...>>::ForEachEntity(UObject* InObject, uint16 PropertyOffset, typename TCallTraits<MetaDataTypes>::ParamType... MetaData, StorageType& OutValue) const
+void TGetPropertyValuesImpl<PropertyTraits, TPropertyMetaData<MetaDataTypes...>>::ForEachEntity(
+	UObject* InObject,
+	uint16 PropertyOffset,
+	typename TCallTraits<MetaDataTypes>::ParamType... MetaData,
+	StorageType& OutValue) const
 {
 	// Would really like to avoid branching here, but if we encounter this data the options are either handle it gracefully, stomp a vtable, or report a fatal error.
 	if (ensureAlwaysMsgf(PropertyOffset != 0, TEXT("Invalid property offset specified (ptr+%d bytes) for property on object %s. This would otherwise overwrite the object's vfptr."), PropertyOffset, *InObject->GetName()))
 	{
-		PropertyTraits::GetObjectPropertyValue(InObject, MetaData..., PropertyOffset, OutValue);
+		Traits->GetObjectPropertyValue(InObject, MetaData..., PropertyOffset, OutValue);
 	}
 }
 
 template<typename PropertyTraits, typename ...MetaDataTypes>
-void TGetPropertyValuesImpl<PropertyTraits, TPropertyMetaData<MetaDataTypes...>>::ForEachEntity(UObject* InObject, const TSharedPtr<FTrackInstancePropertyBindings>& PropertyBindings, typename TCallTraits<MetaDataTypes>::ParamType... MetaData, StorageType& OutValue) const
+void TGetPropertyValuesImpl<PropertyTraits, TPropertyMetaData<MetaDataTypes...>>::ForEachEntity(
+	UObject* InObject,
+	const TSharedPtr<FTrackInstancePropertyBindings>& PropertyBindings,
+	typename TCallTraits<MetaDataTypes>::ParamType... MetaData,
+	StorageType& OutValue) const
 {
-	PropertyTraits::GetObjectPropertyValue(InObject, MetaData..., PropertyBindings.Get(), OutValue);
+	Traits->GetObjectPropertyValue(InObject, MetaData..., PropertyBindings.Get(), OutValue);
 }
 
 template<typename PropertyTraits, typename ...MetaDataTypes>
-void TGetPropertyValuesImpl<PropertyTraits, TPropertyMetaData<MetaDataTypes...>>::ForEachAllocation(const FEntityAllocation* Allocation, TRead<UObject*> BoundObjectComponents, FThreeWayAccessor ResolvedPropertyComponents, TRead<MetaDataTypes>... MetaData, TWrite<StorageType> OutValueComponents) const
+void TGetPropertyValuesImpl<PropertyTraits, TPropertyMetaData<MetaDataTypes...>>::ForEachAllocation(
+	const FEntityAllocation* Allocation,
+	TRead<UObject*> BoundObjectComponents,
+	FThreeWayAccessor ResolvedPropertyComponents,
+	TRead<MetaDataTypes>... MetaData,
+	TWrite<StorageType> OutValueComponents) const
 {
 	const int32 Num = Allocation->Num();
 	if (const FCustomPropertyIndex* Custom = ResolvedPropertyComponents.template Get<0>())
@@ -141,7 +178,12 @@ void TGetPropertyValuesImpl<PropertyTraits, TPropertyMetaData<MetaDataTypes...>>
 }
 
 template<typename PropertyTraits, typename ...MetaDataTypes>
-void TGetPropertyValuesImpl<PropertyTraits, TPropertyMetaData<MetaDataTypes...>>::ForEachAllocation(const FEntityAllocation* Allocation, TRead<UObject*> BoundObjectComponents, FTwoWayAccessor ResolvedPropertyComponents, TRead<MetaDataTypes>... MetaData, TWrite<StorageType> OutValueComponents) const
+void TGetPropertyValuesImpl<PropertyTraits, TPropertyMetaData<MetaDataTypes...>>::ForEachAllocation(
+	const FEntityAllocation* Allocation,
+	TRead<UObject*> BoundObjectComponents,
+	FTwoWayAccessor ResolvedPropertyComponents,
+	TRead<MetaDataTypes>... MetaData,
+	TWrite<StorageType> OutValueComponents) const
 {
 	const int32 Num = Allocation->Num();
 	if (const uint16* Fast = ResolvedPropertyComponents.template Get<0>())
@@ -162,32 +204,49 @@ void TGetPropertyValuesImpl<PropertyTraits, TPropertyMetaData<MetaDataTypes...>>
 
 
 template<typename PropertyTraits, typename... MetaDataTypes, typename... CompositeTypes>
-void TSetCompositePropertyValuesImpl<PropertyTraits, TPropertyMetaData<MetaDataTypes...>, CompositeTypes...>::ForEachEntity(UObject* InObject, FCustomPropertyIndex CustomPropertyIndex, typename TCallTraits<MetaDataTypes>::ParamType... MetaData, typename TCallTraits<CompositeTypes>::ParamType... CompositeResults) const
+void TSetCompositePropertyValuesImpl<PropertyTraits, TPropertyMetaData<MetaDataTypes...>, CompositeTypes...>::ForEachEntity(
+	UObject* InObject,
+	FCustomPropertyIndex CustomPropertyIndex,
+	typename TCallTraits<MetaDataTypes>::ParamType... MetaData,
+	typename TCallTraits<CompositeTypes>::ParamType... CompositeResults) const
 {
-	StorageType Result = PropertyTraits::CombineComposites(MetaData..., CompositeResults...);
-	PropertyTraits::SetObjectPropertyValue(InObject, MetaData..., CustomAccessors[CustomPropertyIndex.Value], Result);
+	StorageType Result = Traits->CombineComposites(MetaData..., CompositeResults...);
+	Traits->SetObjectPropertyValue(InObject, MetaData..., CustomAccessors[CustomPropertyIndex.Value], Result);
 }
 
 template<typename PropertyTraits, typename... MetaDataTypes, typename... CompositeTypes>
-void TSetCompositePropertyValuesImpl<PropertyTraits, TPropertyMetaData<MetaDataTypes...>, CompositeTypes...>::ForEachEntity(UObject* InObject, uint16 PropertyOffset, typename TCallTraits<MetaDataTypes>::ParamType... MetaData, typename TCallTraits<CompositeTypes>::ParamType... CompositeResults) const
+void TSetCompositePropertyValuesImpl<PropertyTraits, TPropertyMetaData<MetaDataTypes...>, CompositeTypes...>::ForEachEntity(
+	UObject* InObject,
+	uint16 PropertyOffset,
+	typename TCallTraits<MetaDataTypes>::ParamType... MetaData,
+	typename TCallTraits<CompositeTypes>::ParamType... CompositeResults) const
 {
 	// Would really like to avoid branching here, but if we encounter this data the options are either handle it gracefully, stomp a vtable, or report a fatal error.
 	if (ensureAlwaysMsgf(PropertyOffset != 0, TEXT("Invalid property offset specified (ptr+%d bytes) for property on object %s. This would otherwise overwrite the object's vfptr."), PropertyOffset, *InObject->GetName()))
 	{
-		StorageType Result = PropertyTraits::CombineComposites(MetaData..., CompositeResults...);
-		PropertyTraits::SetObjectPropertyValue(InObject, MetaData..., PropertyOffset, Result);
+		StorageType Result = Traits->CombineComposites(MetaData..., CompositeResults...);
+		Traits->SetObjectPropertyValue(InObject, MetaData..., PropertyOffset, Result);
 	}
 }
 
 template<typename PropertyTraits, typename... MetaDataTypes, typename... CompositeTypes>
-void TSetCompositePropertyValuesImpl<PropertyTraits, TPropertyMetaData<MetaDataTypes...>, CompositeTypes...>::ForEachEntity(UObject* InObject, const TSharedPtr<FTrackInstancePropertyBindings>& PropertyBindings, typename TCallTraits<MetaDataTypes>::ParamType... MetaData, typename TCallTraits<CompositeTypes>::ParamType... CompositeResults) const
+void TSetCompositePropertyValuesImpl<PropertyTraits, TPropertyMetaData<MetaDataTypes...>, CompositeTypes...>::ForEachEntity(
+	UObject* InObject,
+	const TSharedPtr<FTrackInstancePropertyBindings>& PropertyBindings,
+	typename TCallTraits<MetaDataTypes>::ParamType... MetaData,
+	typename TCallTraits<CompositeTypes>::ParamType... CompositeResults) const
 {
-	StorageType Result = PropertyTraits::CombineComposites(MetaData..., CompositeResults...);
-	PropertyTraits::SetObjectPropertyValue(InObject, MetaData..., PropertyBindings.Get(), Result);
+	StorageType Result = Traits->CombineComposites(MetaData..., CompositeResults...);
+	Traits->SetObjectPropertyValue(InObject, MetaData..., PropertyBindings.Get(), Result);
 }
 
 template<typename PropertyTraits, typename... MetaDataTypes, typename... CompositeTypes>
-void TSetCompositePropertyValuesImpl<PropertyTraits, TPropertyMetaData<MetaDataTypes...>, CompositeTypes...>::ForEachAllocation(const FEntityAllocation* Allocation, TRead<UObject*> BoundObjectComponents, FThreeWayAccessor ResolvedPropertyComponents, TRead<MetaDataTypes>... InMetaData, TRead<CompositeTypes>... VariadicComponents) const
+void TSetCompositePropertyValuesImpl<PropertyTraits, TPropertyMetaData<MetaDataTypes...>, CompositeTypes...>::ForEachAllocation(
+	const FEntityAllocation* Allocation,
+	TRead<UObject*> BoundObjectComponents,
+	FThreeWayAccessor ResolvedPropertyComponents,
+	TRead<MetaDataTypes>... InMetaData,
+	TRead<CompositeTypes>... VariadicComponents) const
 {
 	const int32 Num = Allocation->Num();
 	if (const FCustomPropertyIndex* Custom = ResolvedPropertyComponents.template Get<0>())
@@ -214,7 +273,12 @@ void TSetCompositePropertyValuesImpl<PropertyTraits, TPropertyMetaData<MetaDataT
 }
 
 template<typename PropertyTraits, typename... MetaDataTypes, typename... CompositeTypes>
-void TSetCompositePropertyValuesImpl<PropertyTraits, TPropertyMetaData<MetaDataTypes...>, CompositeTypes...>::ForEachAllocation(const FEntityAllocation* Allocation, TRead<UObject*> BoundObjectComponents, FTwoWayAccessor ResolvedPropertyComponents, TRead<MetaDataTypes>... InMetaData, TRead<CompositeTypes>... VariadicComponents) const
+void TSetCompositePropertyValuesImpl<PropertyTraits, TPropertyMetaData<MetaDataTypes...>, CompositeTypes...>::ForEachAllocation(
+	const FEntityAllocation* Allocation,
+	TRead<UObject*> BoundObjectComponents,
+	FTwoWayAccessor ResolvedPropertyComponents,
+	TRead<MetaDataTypes>... InMetaData,
+	TRead<CompositeTypes>... VariadicComponents) const
 {
 	const int32 Num = Allocation->Num();
 	if (const uint16* Fast = ResolvedPropertyComponents.template Get<0>())
@@ -233,5 +297,4 @@ void TSetCompositePropertyValuesImpl<PropertyTraits, TPropertyMetaData<MetaDataT
 	}
 }
 
-} // namespace MovieScene
-} // namespace UE
\ No newline at end of file
+} // namespace UE::MovieScene
\ No newline at end of file
diff --git a/Engine/Source/Runtime/MovieScene/Public/EntitySystem/MovieScenePropertyTraits.h b/Engine/Source/Runtime/MovieScene/Public/EntitySystem/MovieScenePropertyTraits.h
--- a/Engine/Source/Runtime/MovieScene/Public/EntitySystem/MovieScenePropertyTraits.h
+++ b/Engine/Source/Runtime/MovieScene/Public/EntitySystem/MovieScenePropertyTraits.h
@@ -3,8 +3,22 @@
 #pragma once
 
 #include "EntitySystem/MovieScenePropertySystemTypes.h"
+#include "EntitySystem/MovieScenePropertyMetaData.h"
 #include "EntitySystem/MovieSceneOperationalTypeConversions.h"
+#include "EntitySystem/MovieSceneIntermediatePropertyValue.h"
+#include "EntitySystem/MovieSceneVariantPropertyTypeIndex.h"
+#include "EntitySystem/MovieScenePropertySupport.h"
+#include "Channels/MovieSceneUnpackedChannelValues.h"
 #include "MovieSceneCommonHelpers.h"
+#include "Misc/GeneratedTypeName.h"
+
+
+#include "UObject/Class.h"
+#include "UObject/UnrealType.h"
+#include "UObject/TextProperty.h"
+#include "UObject/NoExportTypes.h"
+
+class UMovieSceneTrack;
 
 
 namespace UE
@@ -12,147 +26,582 @@ namespace UE
 namespace MovieScene
 {
 
-/**
- * Property accessor traits that talk directly to the reflected UObject property type
- */
-template<typename ...MetaDataTypes>
-struct TPropertyMetaData
+template<typename PropertyTraits, typename ...CompositeTypes>
+struct TVariantPropertyComponentHandler;
+
+
+template<typename T>
+void UnpackChannelsFromOperational(T&&, const FProperty& Property, FUnpackedChannelValues&)
+{
+	static_assert(!std::is_same_v<T, T>, "Please implement this function for the necessary types to support unpacking the property type into its constituent channels");
+}
+
+struct IPropertyTraits
+{
+	using MetaDataType = TPropertyMetaData<>;
+
+	virtual ~IPropertyTraits()
+	{
+	}
+
+	virtual TSubclassOf<UMovieSceneTrack> GetTrackClass(const FProperty* InProperty) const
+	{
+		return nullptr;
+	}
+	virtual bool InitializeTrackFromProperty(UMovieSceneTrack* InTrack, const FProperty* InProperty) const
+	{
+		return false;
+	}
+};
+
+
+template<typename ...UObjectPropertyTypes>
+struct TSupportedPropertyTypes
 {
-	static constexpr int32 Num = sizeof...(MetaDataTypes);
+
+};
+template<typename ...CompositeTypes>
+struct TCompositeChannelMapping
+{
+	
 };
 
 
 /**
  * Property accessor traits that talk directly to the reflected UObject property type
  */
-template<typename UObjectPropertyType, bool bInIsComposite = true>
-struct TDirectPropertyTraits
+template<typename UObjectPropertyType, typename InMemoryType, bool bInIsComposite = true>
+struct TPropertyTraits : IPropertyTraits
 {
-	static constexpr bool bIsComposite = bInIsComposite;
+	static constexpr bool bIsComposite     = bInIsComposite;
+	static constexpr bool bNeedsConversion = !std::is_same_v<InMemoryType, UObjectPropertyType>;
 
-	using StorageType  = UObjectPropertyType;
-	using MetaDataType = TPropertyMetaData<>;
-	using TraitsType   = TDirectPropertyTraits<UObjectPropertyType>;
+	using TraitsType          = TPropertyTraits<UObjectPropertyType, InMemoryType, bInIsComposite>;
+	using StorageType         = InMemoryType;
+	using StorageTypeParam    = typename TCallTraits<StorageType>::ParamType;
 
-	using ParamType = typename TCallTraits<UObjectPropertyType>::ParamType;
+	/** IPropertyTraits interface  */
+	FORCEINLINE static bool SupportsProperty(const FProperty& InProperty)
+	{
+		return TPropertyMatch<UObjectPropertyType>::SupportsProperty(InProperty);
+	}
+	FORCEINLINE static FIntermediatePropertyValue CoercePropertyValue(const FProperty& InProperty, const FSourcePropertyValue& InPropertyValue)
+	{
+		StorageType Value;
+		ConvertOperationalProperty(*InPropertyValue.Cast<UObjectPropertyType>(), Value);
+		return FIntermediatePropertyValue::FromValue(Value);
+	}
+	FORCEINLINE static void UnpackChannels(const InMemoryType& Value, const FProperty& Property, FUnpackedChannelValues& OutUnpackedValues)
+	{
+		UnpackChannelsFromOperational(Value, Property, OutUnpackedValues);
+	}
 
 	/** Property Value Getters  */
-	static void GetObjectPropertyValue(const UObject* InObject, const FCustomPropertyAccessor& BaseCustomAccessor, StorageType& OutValue)
+	FORCEINLINE static void GetObjectPropertyValue(const UObject* InObject, const FCustomPropertyAccessor& BaseCustomAccessor, StorageType& OutValue)
 	{
 		const TCustomPropertyAccessor<TraitsType>& CustomAccessor = static_cast<const TCustomPropertyAccessor<TraitsType>&>(BaseCustomAccessor);
+
 		OutValue = (*CustomAccessor.Functions.Getter)(InObject);
 	}
-	static void GetObjectPropertyValue(const UObject* InObject, uint16 PropertyOffset, StorageType& OutValue)
+	FORCEINLINE static void GetObjectPropertyValue(const UObject* InObject, uint16 PropertyOffset, StorageType& OutValue)
 	{
 		const UObjectPropertyType* PropertyAddress = reinterpret_cast<const UObjectPropertyType*>( reinterpret_cast<const uint8*>(InObject) + PropertyOffset );
-		OutValue = *PropertyAddress;
+
+		if constexpr (bNeedsConversion)
+		{
+			ConvertOperationalProperty(*PropertyAddress, OutValue);
+		}
+		else
+		{
+			OutValue = *PropertyAddress;
+		}
 	}
-	static void GetObjectPropertyValue(const UObject* InObject, FTrackInstancePropertyBindings* PropertyBindings, StorageType& OutValue)
+	FORCEINLINE static void GetObjectPropertyValue(const UObject* InObject, FTrackInstancePropertyBindings* PropertyBindings, StorageType& OutValue)
 	{
-		OutValue = PropertyBindings->GetCurrentValue<UObjectPropertyType>(*InObject);
+		if constexpr (bNeedsConversion)
+		{
+			UObjectPropertyType Value = PropertyBindings->GetCurrentValue<UObjectPropertyType>(*InObject);
+			ConvertOperationalProperty(Value, OutValue);
+		}
+		else
+		{
+			OutValue = PropertyBindings->GetCurrentValue<UObjectPropertyType>(*InObject);
+		}
 	}
-	static void GetObjectPropertyValue(const UObject* InObject, const FName& PropertyPath, StorageType& OutValue)
+
+	FORCEINLINE static void GetObjectPropertyValue(const UObject* InObject, const FName& PropertyPath, StorageType& OutValue)
 	{
 		TOptional<UObjectPropertyType> Property = FTrackInstancePropertyBindings::StaticValue<UObjectPropertyType>(InObject, *PropertyPath.ToString());
 		if (Property)
 		{
-			OutValue = MoveTemp(Property.GetValue());
+			if constexpr (bNeedsConversion)
+			{
+				ConvertOperationalProperty(Property.GetValue(), OutValue);
+			}
+			else
+			{
+				OutValue = Property.GetValue();
+			}
 		}
 	}
 
 	/** Property Value Setters  */
-	static void SetObjectPropertyValue(UObject* InObject, const FCustomPropertyAccessor& BaseCustomAccessor, ParamType InValue)
+	FORCEINLINE static void SetObjectPropertyValue(UObject* InObject, const FCustomPropertyAccessor& BaseCustomAccessor, StorageTypeParam InValue)
 	{
 		const TCustomPropertyAccessor<TraitsType>& CustomAccessor = static_cast<const TCustomPropertyAccessor<TraitsType>&>(BaseCustomAccessor);
 		(*CustomAccessor.Functions.Setter)(InObject, InValue);
 	}
-	static void SetObjectPropertyValue(UObject* InObject, uint16 PropertyOffset, ParamType InValue)
+	FORCEINLINE static void SetObjectPropertyValue(UObject* InObject, uint16 PropertyOffset, StorageTypeParam InValue)
 	{
 		UObjectPropertyType* PropertyAddress = reinterpret_cast<UObjectPropertyType*>( reinterpret_cast<uint8*>(InObject) + PropertyOffset );
-		*PropertyAddress = InValue;
+		if constexpr (bNeedsConversion)
+		{
+			ConvertOperationalProperty(InValue, *PropertyAddress);
+		}
+		else
+		{
+			*PropertyAddress = InValue;
+		}
 	}
-	static void SetObjectPropertyValue(UObject* InObject, FTrackInstancePropertyBindings* PropertyBindings, ParamType InValue)
+	FORCEINLINE static void SetObjectPropertyValue(UObject* InObject, FTrackInstancePropertyBindings* PropertyBindings, StorageTypeParam InValue)
 	{
-		PropertyBindings->CallFunction<UObjectPropertyType>(*InObject, InValue);
+		if constexpr (bNeedsConversion)
+		{
+			UObjectPropertyType NewValue{};
+			ConvertOperationalProperty(InValue, NewValue);
+
+			PropertyBindings->CallFunction<UObjectPropertyType>(*InObject, NewValue);
+		}
+		else
+		{
+			PropertyBindings->CallFunction<UObjectPropertyType>(*InObject, InValue);
+		}
 	}
 
 	template<typename ...T>
-	static StorageType CombineComposites(T&&... InComposites)
+	FORCEINLINE static StorageType CombineComposites(T&&... InComposites)
 	{
 		return StorageType{ Forward<T>(InComposites)... };
 	}
 };
 
+template<typename StorageType>
+struct TDynamicVariantTraitsBase
+{
+	using StorageTypeParam        = typename TCallTraits<StorageType>::ParamType;
+
+	using CastToOperationalPtr    = void(*)(const void*, StorageType&);
+	using CastToFinalPtr          = void(*)(StorageTypeParam, void*);
+	using RetrieveSlowPropertyPtr = void(*)(const UObject*, FTrackInstancePropertyBindings*, StorageType&);
+	using ApplySlowPropertyPtr    = void(*)(UObject*, FTrackInstancePropertyBindings*, StorageTypeParam);
+	using InitializeNewTrackPtr   = void(*)(UMovieSceneTrack*, const FProperty*);
+
+	UScriptStruct*                Struct              = nullptr;
+	TSubclassOf<UMovieSceneTrack> TrackClass          = nullptr;
+	CastToOperationalPtr          CastToOperational   = nullptr;
+	CastToFinalPtr                CastToFinal         = nullptr;
+	RetrieveSlowPropertyPtr       RetrieveSlowProperty= nullptr;
+	ApplySlowPropertyPtr          ApplySlowProperty   = nullptr;
+	InitializeNewTrackPtr         InitializeNewTrack  = nullptr;
+};
 
-/**
- * Property accessor traits that talk directly to the reflected UObject property type
- */
-template<typename UObjectPropertyType, typename InMemoryType, bool bInIsComposite = true>
-struct TIndirectPropertyTraits
+template<typename UObjectPropertyType, typename StorageType>
+struct TDynamicVariantTraits : TDynamicVariantTraitsBase<StorageType>
 {
-	static constexpr bool bIsComposite = bInIsComposite;
+	using StorageTypeParam = typename TCallTraits<StorageType>::ParamType;
 
-	using StorageType  = InMemoryType;
-	using MetaDataType = TPropertyMetaData<>;
-	using TraitsType   = TIndirectPropertyTraits<UObjectPropertyType, InMemoryType>;
+	using InitializeNewTrackPtr   = void(*)(UMovieSceneTrack*, const FProperty*);
 
-	using PropertyParamType    = typename TCallTraits<UObjectPropertyType>::ParamType;
-	using OperationalParamType = typename TCallTraits<StorageType>::ParamType;
+	TDynamicVariantTraits()
+	{
+		this->Struct               = TBaseStructure<UObjectPropertyType>::Get();
+		this->CastToOperational    = CastToOperationalImpl;
+		this->CastToFinal          = CastToFinalImpl;
+		this->RetrieveSlowProperty = RetrieveSlowPropertyImpl;
+		this->ApplySlowProperty    = ApplySlowPropertyImpl;
+	}
+	TDynamicVariantTraits(UScriptStruct* InStruct)
+	{
+		this->Struct               = InStruct;
+		this->CastToOperational    = CastToOperationalImpl;
+		this->CastToFinal          = CastToFinalImpl;
+		this->RetrieveSlowProperty = RetrieveSlowPropertyImpl;
+		this->ApplySlowProperty    = ApplySlowPropertyImpl;
+	}
 
-	/** Property Value Getters  */
-	static void GetObjectPropertyValue(const UObject* InObject, const FCustomPropertyAccessor& BaseCustomAccessor, StorageType& OutValue)
+	TDynamicVariantTraits<UObjectPropertyType, StorageType>& SetTrackClass(TSubclassOf<UMovieSceneTrack> InTrackClass)
 	{
-		const TCustomPropertyAccessor<TraitsType>& CustomAccessor = static_cast<const TCustomPropertyAccessor<TraitsType>&>(BaseCustomAccessor);
+		this->TrackClass = InTrackClass;
+		return *this;
+	}
 
-		OutValue = (*CustomAccessor.Functions.Getter)(InObject);
+	TDynamicVariantTraits<UObjectPropertyType, StorageType>& SetTrackInitializer(InitializeNewTrackPtr InFunc)
+	{
+		this->InitializeNewTrack = InFunc;
+		return *this;
 	}
-	static void GetObjectPropertyValue(const UObject* InObject, uint16 PropertyOffset, StorageType& OutValue)
+
+	static void CastToOperationalImpl(const void* In, StorageType& OutValue)
 	{
-		const UObjectPropertyType* PropertyAddress = reinterpret_cast<const UObjectPropertyType*>( reinterpret_cast<const uint8*>(InObject) + PropertyOffset );
+		const UObjectPropertyType* PropertyAddress = static_cast<const UObjectPropertyType*>(In);
 		ConvertOperationalProperty(*PropertyAddress, OutValue);
 	}
-	static void GetObjectPropertyValue(const UObject* InObject, FTrackInstancePropertyBindings* PropertyBindings, StorageType& OutValue)
+	static void CastToFinalImpl(StorageTypeParam InValue, void* Out)
+	{
+		UObjectPropertyType* PropertyAddress = static_cast<UObjectPropertyType*>(Out);
+		ConvertOperationalProperty(InValue, *PropertyAddress);
+	}
+
+	static void RetrieveSlowPropertyImpl(const UObject* InObject, FTrackInstancePropertyBindings* PropertyBindings, StorageType& OutValue)
 	{
 		UObjectPropertyType Value = PropertyBindings->GetCurrentValue<UObjectPropertyType>(*InObject);
 		ConvertOperationalProperty(Value, OutValue);
 	}
+	static void ApplySlowPropertyImpl(UObject* InObject, FTrackInstancePropertyBindings* PropertyBindings, StorageTypeParam InValue)
+	{
+		UObjectPropertyType NewValue{};
+		ConvertOperationalProperty(InValue, NewValue);
 
-	static void GetObjectPropertyValue(const UObject* InObject, const FName& PropertyPath, StorageType& OutValue)
+		PropertyBindings->CallFunction<UObjectPropertyType>(*InObject, NewValue);
+	}
+};
+
+template<typename InMemoryType, typename ...UObjectPropertyTypes>
+struct TVariantPropertyTraits
+	: IPropertyTraits
+{
+	static constexpr int32 CompileTimeNum = sizeof...(UObjectPropertyTypes);
+	static constexpr bool  bIsComposite   = true;
+
+	using StorageType    = InMemoryType;
+	using StorageTypeParam = typename TCallTraits<StorageType>::ParamType;
+
+	using MetaDataType   = TPropertyMetaData<FVariantPropertyTypeIndex>;
+	using PublicMetaData = TPropertyMetaData<>;
+
+	using TraitsType     = TVariantPropertyTraits<StorageType, UObjectPropertyTypes...>;
+
+	static bool SupportsProperty(const FProperty& InProperty)
 	{
-		TOptional<UObjectPropertyType> Property = FTrackInstancePropertyBindings::StaticValue<UObjectPropertyType>(InObject, *PropertyPath.ToString());
-		if (Property)
+		FVariantPropertyTypeIndex Unused;
+		return ComputeVariantIndex(InProperty, Unused);
+	}
+
+	static bool ComputeVariantIndex(const FProperty& Property, FVariantPropertyTypeIndex& OutTypeIndex)
+	{
+		using FuncDef = bool (*)(const FProperty& InProperty);
+		static constexpr FuncDef Funcs[] = { &TPropertyMatch<UObjectPropertyTypes>::SupportsProperty... };
+
+		for (int32 Index = 0; Index < CompileTimeNum; ++Index)
 		{
-			ConvertOperationalProperty(Property.GetValue(), OutValue);
+			if (Funcs[Index](Property))
+			{
+				OutTypeIndex.Index = Index;
+				return true;
+			}
 		}
+		return false;
 	}
 
-	/** Property Value Setters  */
-	static void SetObjectPropertyValue(UObject* InObject, const FCustomPropertyAccessor& BaseCustomAccessor, OperationalParamType InValue)
+	static FIntermediatePropertyValue CoercePropertyValue(const FProperty& InProperty, const FSourcePropertyValue& InPropertyValue)
+	{
+		FVariantPropertyTypeIndex VariantIndex;
+
+		const bool bIsSupported = ComputeVariantIndex(InProperty, VariantIndex);
+		check(bIsSupported);
+
+		using FuncDef = FIntermediatePropertyValue (*)(const FProperty& InProperty, const FSourcePropertyValue& InPropertyValue);
+		static constexpr FuncDef Funcs[] = { &TPropertyTraits<UObjectPropertyTypes, StorageType>::CoercePropertyValue... };
+
+		return CoercePropertyValueChecked(InProperty, VariantIndex, InPropertyValue);
+	}
+
+	static FIntermediatePropertyValue CoercePropertyValueChecked(const FProperty& InProperty, FVariantPropertyTypeIndex VariantIndex, const FSourcePropertyValue& InPropertyValue)
+	{
+		using FuncDef = FIntermediatePropertyValue (*)(const FProperty& InProperty, const FSourcePropertyValue& InPropertyValue);
+		static constexpr FuncDef Funcs[] = { &TPropertyTraits<UObjectPropertyTypes, StorageType>::CoercePropertyValue... };
+
+		return Funcs[VariantIndex.Index](InProperty, InPropertyValue);
+	}
+
+	static void UnpackChannels(const InMemoryType& Value, const FProperty& Property, FUnpackedChannelValues& OutUnpackedValues)
+	{
+		UnpackChannelsFromOperational(Value, Property, OutUnpackedValues);
+	}
+
+	template<typename ...Composites>
+	static TVariantPropertyComponentHandler<TraitsType, Composites...> MakeHandler()
+	{
+		return TVariantPropertyComponentHandler<TraitsType, Composites...>();
+	}
+
+	static bool NeedsMetaData()
+	{
+		return sizeof...(UObjectPropertyTypes) > 1;
+	}
+
+	static void GetObjectPropertyValue(const UObject* InObject, FVariantPropertyTypeIndex VariantTypeIndex, const FCustomPropertyAccessor& BaseCustomAccessor, StorageType& OutValue)
+	{
+		const TCustomPropertyAccessor<TraitsType>& CustomAccessor = static_cast<const TCustomPropertyAccessor<TraitsType>&>(BaseCustomAccessor);
+
+		OutValue = (*CustomAccessor.Functions.Getter)(InObject);
+	}
+
+	static void GetObjectPropertyValue(const UObject* InObject, FVariantPropertyTypeIndex VariantTypeIndex, uint16 PropertyOffset, StorageType& OutValue)
+	{
+		using FuncDef = void (*)(const UObject*, uint16, StorageType&);
+		static constexpr FuncDef Funcs[] = { &TPropertyTraits<UObjectPropertyTypes, StorageType>::GetObjectPropertyValue... };
+
+		if (VariantTypeIndex.Index >= 0 && VariantTypeIndex.Index < CompileTimeNum)
+		{
+			Funcs[VariantTypeIndex.Index](InObject, PropertyOffset, OutValue);
+		}
+	}
+	static void GetObjectPropertyValue(const UObject* InObject, FVariantPropertyTypeIndex VariantTypeIndex, FTrackInstancePropertyBindings* PropertyBindings, StorageType& OutValue)
+	{
+		using FuncDef = void (*)(const UObject*, FTrackInstancePropertyBindings*, StorageType&);
+		static constexpr FuncDef Funcs[] = { &TPropertyTraits<UObjectPropertyTypes, StorageType>::GetObjectPropertyValue... };
+
+		if (VariantTypeIndex.Index >= 0 && VariantTypeIndex.Index < CompileTimeNum)
+		{
+			Funcs[VariantTypeIndex.Index](InObject, PropertyBindings, OutValue);
+		}
+	}
+
+	static void SetObjectPropertyValue(UObject* InObject, FVariantPropertyTypeIndex VariantTypeIndex, const FCustomPropertyAccessor& BaseCustomAccessor, StorageTypeParam InValue)
 	{
 		const TCustomPropertyAccessor<TraitsType>& CustomAccessor = static_cast<const TCustomPropertyAccessor<TraitsType>&>(BaseCustomAccessor);
 		(*CustomAccessor.Functions.Setter)(InObject, InValue);
 	}
-	static void SetObjectPropertyValue(UObject* InObject, uint16 PropertyOffset, OperationalParamType InValue)
+	static void SetObjectPropertyValue(UObject* InObject, FVariantPropertyTypeIndex VariantTypeIndex, uint16 PropertyOffset, StorageTypeParam InValue)
 	{
-		UObjectPropertyType* PropertyAddress = reinterpret_cast<UObjectPropertyType*>( reinterpret_cast<uint8*>(InObject) + PropertyOffset );
-		ConvertOperationalProperty(InValue, *PropertyAddress);
+		using FuncDef = void (*)(UObject*, uint16, StorageTypeParam);
+		static constexpr FuncDef Funcs[] = { &TPropertyTraits<UObjectPropertyTypes, StorageType>::SetObjectPropertyValue... };
+
+		if (VariantTypeIndex.Index >= 0 && VariantTypeIndex.Index < CompileTimeNum)
+		{
+			Funcs[VariantTypeIndex.Index](InObject, PropertyOffset, InValue);
+		}
 	}
-	static void SetObjectPropertyValue(UObject* InObject, FTrackInstancePropertyBindings* PropertyBindings, OperationalParamType InValue)
+	static void SetObjectPropertyValue(UObject* InObject, FVariantPropertyTypeIndex VariantTypeIndex, FTrackInstancePropertyBindings* PropertyBindings, StorageTypeParam InValue)
 	{
-		UObjectPropertyType NewValue{};
-		ConvertOperationalProperty(InValue, NewValue);
+		using FuncDef = void (*)(UObject*, FTrackInstancePropertyBindings*, StorageTypeParam);
+		static constexpr FuncDef Funcs[] = { &TPropertyTraits<UObjectPropertyTypes, StorageType>::SetObjectPropertyValue... };
 
-		PropertyBindings->CallFunction<UObjectPropertyType>(*InObject, NewValue);
+		if (VariantTypeIndex.Index >= 0 && VariantTypeIndex.Index < CompileTimeNum)
+		{
+			Funcs[VariantTypeIndex.Index](InObject, PropertyBindings, InValue);
+		}
 	}
 
 	template<typename ...T>
-	static StorageType CombineComposites(T&&... InComposites)
+	static StorageType CombineComposites(FVariantPropertyTypeIndex VariantTypeIndex, T&&... InComposites)
 	{
 		return StorageType{ Forward<T>(InComposites)... };
 	}
 };
 
 
+template<typename InMemoryType, typename ...UObjectPropertyTypes>
+struct TDynamicVariantPropertyTraits : IPropertyTraits
+{
+	using StaticTraits     = TVariantPropertyTraits<InMemoryType, UObjectPropertyTypes...>;
+	using TraitsType       = TDynamicVariantPropertyTraits<InMemoryType, UObjectPropertyTypes...>;
+	using StorageType      = InMemoryType;
+	using StorageTypeParam = typename TCallTraits<StorageType>::ParamType;
+
+	using MetaDataType   = TPropertyMetaData<FVariantPropertyTypeIndex>;
+	using PublicMetaData = TPropertyMetaData<>;
+
+	static constexpr bool  bIsComposite   = true;
+
+	TArray<TDynamicVariantTraitsBase<StorageType>> DynamicTraits;
+
+	bool NeedsMetaData() const
+	{
+		return sizeof...(UObjectPropertyTypes) + DynamicTraits.Num() > 1;
+	}
+
+	static void UnpackChannels(const InMemoryType& Value, const FProperty& Property, FUnpackedChannelValues& OutUnpackedValues)
+	{
+		StaticTraits::UnpackChannels(Value, Property, OutUnpackedValues);
+	}
+
+	template<typename ...Composites>
+	static TVariantPropertyComponentHandler<TraitsType, Composites...> MakeHandler()
+	{
+		return TVariantPropertyComponentHandler<TraitsType, Composites...>();
+	}
+
+	virtual TSubclassOf<UMovieSceneTrack> GetTrackClass(const FProperty* InProperty) const override
+	{
+		if (const FStructProperty* StructProperty = CastField<const FStructProperty>(InProperty))
+		{
+			const TDynamicVariantTraitsBase<StorageType>* Variant = Algo::FindBy(DynamicTraits, StructProperty->Struct, &TDynamicVariantTraitsBase<StorageType>::Struct);
+			if (Variant)
+			{
+				return Variant->TrackClass;
+			}
+		}
+		return nullptr;
+	}
+	virtual bool InitializeTrackFromProperty(UMovieSceneTrack* InTrack, const FProperty* InProperty) const override
+	{
+		if (const FStructProperty* StructProperty = CastField<const FStructProperty>(InProperty))
+		{
+			const TDynamicVariantTraitsBase<StorageType>* Variant = Algo::FindBy(DynamicTraits, StructProperty->Struct, &TDynamicVariantTraitsBase<StorageType>::Struct);
+			if (Variant && Variant->InitializeNewTrack)
+			{
+				(*Variant->InitializeNewTrack)(InTrack, InProperty);
+				return true;
+			}
+		}
+		return false;
+	}
+
+	bool ComputeVariantIndex(const FProperty& InProperty, FVariantPropertyTypeIndex& OutTypeIndex) const
+	{
+		if (StaticTraits::ComputeVariantIndex(InProperty, OutTypeIndex))
+		{
+			return true;
+		}
+
+		if (const FStructProperty* StructProperty = CastField<const FStructProperty>(&InProperty))
+		{
+			for (int32 DynamicIndex = 0; DynamicIndex < DynamicTraits.Num(); ++DynamicIndex)
+			{
+				if (StructProperty->Struct == DynamicTraits[DynamicIndex].Struct)
+				{
+					OutTypeIndex.Index = StaticTraits::CompileTimeNum + DynamicIndex;
+					return true;
+				}
+			}
+		}
+		return false;
+	}
+	bool SupportsProperty(const FProperty& InProperty) const
+	{
+		FVariantPropertyTypeIndex Unused;
+		return ComputeVariantIndex(InProperty, Unused);
+	}
+	FIntermediatePropertyValue CoercePropertyValue(const FProperty& InProperty, const FSourcePropertyValue& InPropertyValue) const
+	{
+		FVariantPropertyTypeIndex StaticVariantIndex;
+		if (StaticTraits::ComputeVariantIndex(InProperty, StaticVariantIndex))
+		{
+			return StaticTraits::CoercePropertyValueChecked(InProperty, StaticVariantIndex, InPropertyValue);
+		}
+
+		check(DynamicTraits.Num() > 0);
+
+		const FStructProperty* StructProperty = CastFieldChecked<const FStructProperty>(&InProperty);
+
+		for (int32 DynamicIndex = 0; DynamicIndex < DynamicTraits.Num()-1; ++DynamicIndex)
+		{
+			if (StructProperty->Struct == DynamicTraits[DynamicIndex].Struct)
+			{
+				StorageType NewValue;
+				DynamicTraits[DynamicIndex].CastToOperational(InPropertyValue.Get(), NewValue);
+				return FIntermediatePropertyValue::FromValue(NewValue);
+			}
+		}
+
+		check(StructProperty->Struct == DynamicTraits.Last().Struct);
+
+		StorageType NewValue;
+		DynamicTraits.Last().CastToOperational(InPropertyValue.Get(), NewValue);
+		return FIntermediatePropertyValue::FromValue(NewValue);
+	}
+
+	void AddDynamicType(const TDynamicVariantTraitsBase<StorageType>& Variant)
+	{
+		DynamicTraits.Emplace(Variant);
+	}
+
+	void GetObjectPropertyValue(const UObject* InObject, FVariantPropertyTypeIndex VariantTypeIndex, uint16 PropertyOffset, StorageType& OutValue) const
+	{
+		if (VariantTypeIndex.Index >= StaticTraits::CompileTimeNum)
+		{
+			const int32 DynamicIndex = int32(VariantTypeIndex.Index) - StaticTraits::CompileTimeNum;
+			if (DynamicTraits.IsValidIndex(DynamicIndex))
+			{
+				const void* PropertyAddress = reinterpret_cast<const uint8*>(InObject) + PropertyOffset;
+				DynamicTraits[DynamicIndex].CastToOperational(PropertyAddress, OutValue);
+			}
+		}
+		else
+		{
+			StaticTraits::GetObjectPropertyValue(InObject, VariantTypeIndex, PropertyOffset, OutValue);
+		}
+	}
+	void GetObjectPropertyValue(const UObject* InObject, FVariantPropertyTypeIndex VariantTypeIndex, FTrackInstancePropertyBindings* PropertyBindings, StorageType& OutValue) const
+	{
+		if (VariantTypeIndex.Index >= StaticTraits::CompileTimeNum)
+		{
+			const int32 DynamicIndex = int32(VariantTypeIndex.Index) - StaticTraits::CompileTimeNum;
+			if (DynamicTraits.IsValidIndex(DynamicIndex))
+			{
+				DynamicTraits[DynamicIndex].RetrieveSlowProperty(InObject, PropertyBindings, OutValue);
+			}
+		}
+		else
+		{
+			StaticTraits::GetObjectPropertyValue(InObject, VariantTypeIndex, PropertyBindings, OutValue);
+		}
+	}
+
+	static void GetObjectPropertyValue(const UObject* InObject, FVariantPropertyTypeIndex VariantTypeIndex, const FCustomPropertyAccessor& BaseCustomAccessor, StorageType& OutValue)
+	{
+		const TCustomPropertyAccessor<TraitsType>& CustomAccessor = static_cast<const TCustomPropertyAccessor<TraitsType>&>(BaseCustomAccessor);
+
+		OutValue = (*CustomAccessor.Functions.Getter)(InObject);
+	}
+
+	void SetObjectPropertyValue(UObject* InObject, FVariantPropertyTypeIndex VariantTypeIndex, uint16 PropertyOffset, StorageTypeParam InValue) const
+	{
+		if (VariantTypeIndex.Index >= StaticTraits::CompileTimeNum)
+		{
+			const int32 DynamicIndex = int32(VariantTypeIndex.Index) - StaticTraits::CompileTimeNum;
+			if (DynamicTraits.IsValidIndex(DynamicIndex))
+			{
+				void* PropertyAddress = reinterpret_cast<uint8*>(InObject) + PropertyOffset;
+				DynamicTraits[DynamicIndex].CastToFinal(InValue, PropertyAddress);
+			}
+		}
+		else
+		{
+			StaticTraits::SetObjectPropertyValue(InObject, VariantTypeIndex, PropertyOffset, InValue);
+		}
+	}
+	void SetObjectPropertyValue(UObject* InObject, FVariantPropertyTypeIndex VariantTypeIndex, FTrackInstancePropertyBindings* PropertyBindings, StorageTypeParam InValue) const
+	{
+		if (VariantTypeIndex.Index >= StaticTraits::CompileTimeNum)
+		{
+			const int32 DynamicIndex = int32(VariantTypeIndex.Index) - StaticTraits::CompileTimeNum;
+			if (DynamicTraits.IsValidIndex(DynamicIndex))
+			{
+				DynamicTraits[DynamicIndex].ApplySlowProperty(InObject, PropertyBindings, InValue);
+			}
+		}
+		else
+		{
+			StaticTraits::SetObjectPropertyValue(InObject, VariantTypeIndex, PropertyBindings, InValue);
+		}
+	}
+
+	static void SetObjectPropertyValue(UObject* InObject, FVariantPropertyTypeIndex VariantTypeIndex, const FCustomPropertyAccessor& BaseCustomAccessor, StorageTypeParam InValue)
+	{
+		const TCustomPropertyAccessor<TraitsType>& CustomAccessor = static_cast<const TCustomPropertyAccessor<TraitsType>&>(BaseCustomAccessor);
+		(*CustomAccessor.Functions.Setter)(InObject, InValue);
+	}
+
+	template<typename ...T>
+	static StorageType CombineComposites(FVariantPropertyTypeIndex VariantTypeIndex, T&&... InComposites)
+	{
+		return StorageType{ Forward<T>(InComposites)... };
+	}
+};
+
 
 
 /**
@@ -161,10 +610,10 @@ struct TIndirectPropertyTraits
 template<typename RuntimeType, typename ...MetaDataTypes>
 struct TRuntimePropertyTraits
 {
-	using StorageType  = RuntimeType;
-	using MetaDataType = TPropertyMetaData<MetaDataTypes...>;
+	using StorageType      = RuntimeType;
+	using StorageTypeParam = typename TCallTraits<StorageType>::ParamType;
+	using MetaDataType     = TPropertyMetaData<MetaDataTypes...>;
 
-	using OperationalParamType = typename TCallTraits<StorageType>::ParamType;
 
 	/** Property Value Getters  */
 	static void GetObjectPropertyValue(const UObject* InObject, const FCustomPropertyAccessor& BaseCustomAccessor, typename TCallTraits<MetaDataTypes>::ParamType... MetaData, StorageType& OutValue)
@@ -180,17 +629,26 @@ struct TRuntimePropertyTraits
 	{}
 
 	/** Property Value Setters  */
-	static void SetObjectPropertyValue(UObject* InObject, const FCustomPropertyAccessor& BaseCustomAccessor, typename TCallTraits<MetaDataTypes>::ParamType... MetaData, OperationalParamType InValue)
+	static void SetObjectPropertyValue(UObject* InObject, const FCustomPropertyAccessor& BaseCustomAccessor, typename TCallTraits<MetaDataTypes>::ParamType... MetaData, StorageTypeParam InValue)
 	{}
 
-	static void SetObjectPropertyValue(UObject* InObject, uint16 PropertyOffset, typename TCallTraits<MetaDataTypes>::ParamType... MetaData, OperationalParamType InValue)
+	static void SetObjectPropertyValue(UObject* InObject, uint16 PropertyOffset, typename TCallTraits<MetaDataTypes>::ParamType... MetaData, StorageTypeParam InValue)
 	{}
 
-	static void SetObjectPropertyValue(UObject* InObject, FTrackInstancePropertyBindings* PropertyBindings, typename TCallTraits<MetaDataTypes>::ParamType... MetaData, OperationalParamType InValue)
+	static void SetObjectPropertyValue(UObject* InObject, FTrackInstancePropertyBindings* PropertyBindings, typename TCallTraits<MetaDataTypes>::ParamType... MetaData, StorageTypeParam InValue)
 	{}
 };
 
 
+
+template<typename UObjectPropertyType, bool bInIsComposite = true>
+using TDirectPropertyTraits = TPropertyTraits<UObjectPropertyType, UObjectPropertyType, bInIsComposite>;
+
+template<typename UObjectPropertyType, typename InMemoryType, bool bInIsComposite = true>
+struct UE_DEPRECATED(5.4, "Please use TPropertyTraits directly.") TIndirectPropertyTraits : TPropertyTraits<UObjectPropertyType, InMemoryType, bInIsComposite>
+{};
+
+
 } // namespace MovieScene
 } // namespace UE
 
diff --git a/Engine/Source/Runtime/MovieScene/Public/Evaluation/PreAnimatedState/MovieScenePreAnimatedPropertyStorage.h b/Engine/Source/Runtime/MovieScene/Public/Evaluation/PreAnimatedState/MovieScenePreAnimatedPropertyStorage.h
--- a/Engine/Source/Runtime/MovieScene/Public/Evaluation/PreAnimatedState/MovieScenePreAnimatedPropertyStorage.h
+++ b/Engine/Source/Runtime/MovieScene/Public/Evaluation/PreAnimatedState/MovieScenePreAnimatedPropertyStorage.h
@@ -66,7 +66,13 @@ struct TPreAnimatedPropertyTraits<PropertyTraits, TIntegerSequence<int, MetaData
 	using KeyType = TTuple<FObjectKey, FName>;
 	using StorageType = TPreAnimatedPropertyValue<typename PropertyTraits::StorageType, MetaDataTypes...>;
 
-	static void RestorePreAnimatedValue(const KeyType& InKey, StorageType& CachedValue, const FRestoreStateParams& Params)
+	const PropertyTraits* PropertyTraitsInstance;
+
+	TPreAnimatedPropertyTraits(const PropertyTraits* InPropertyTraitsInstance)
+		: PropertyTraitsInstance(InPropertyTraitsInstance)
+	{}
+
+	void RestorePreAnimatedValue(const KeyType& InKey, StorageType& CachedValue, const FRestoreStateParams& Params)
 	{
 		UObject* Object = InKey.Get<0>().ResolveObjectPtr();
 		if (!Object)
@@ -76,15 +82,15 @@ struct TPreAnimatedPropertyTraits<PropertyTraits, TIntegerSequence<int, MetaData
 
 		if (const uint16* FastOffset = CachedValue.Binding.template TryGet<uint16>())
 		{
-			PropertyTraits::SetObjectPropertyValue(Object, CachedValue.MetaData.template Get<MetaDataIndices>()..., *FastOffset, CachedValue.Data);
+			PropertyTraitsInstance->SetObjectPropertyValue(Object, CachedValue.MetaData.template Get<MetaDataIndices>()..., *FastOffset, CachedValue.Data);
 		}
 		else  if (const TSharedPtr<FTrackInstancePropertyBindings>* Bindings = CachedValue.Binding.template TryGet<TSharedPtr<FTrackInstancePropertyBindings>>())
 		{
-			PropertyTraits::SetObjectPropertyValue(Object, CachedValue.MetaData.template Get<MetaDataIndices>()..., Bindings->Get(), CachedValue.Data);
+			PropertyTraitsInstance->SetObjectPropertyValue(Object, CachedValue.MetaData.template Get<MetaDataIndices>()..., Bindings->Get(), CachedValue.Data);
 		}
 		else if (const FCustomPropertyAccessor* CustomAccessor = CachedValue.Binding.template Get<const FCustomPropertyAccessor*>())
 		{
-			PropertyTraits::SetObjectPropertyValue(Object, CachedValue.MetaData.template Get<MetaDataIndices>()..., *CustomAccessor, CachedValue.Data);
+			PropertyTraitsInstance->SetObjectPropertyValue(Object, CachedValue.MetaData.template Get<MetaDataIndices>()..., *CustomAccessor, CachedValue.Data);
 		}
 	}
 };
@@ -100,7 +106,13 @@ struct TPreAnimatedPropertyTraits<PropertyTraits, TPropertyMetaData<>, TIntegerS
 	using KeyType = TTuple<FObjectKey, FName>;
 	using StorageType = TPreAnimatedPropertyValue<typename PropertyTraits::StorageType>;
 
-	static void RestorePreAnimatedValue(const KeyType& InKey, StorageType& CachedValue, const FRestoreStateParams& Params)
+	const PropertyTraits* PropertyTraitsInstance;
+
+	TPreAnimatedPropertyTraits(const PropertyTraits* InPropertyTraitsInstance)
+		: PropertyTraitsInstance(InPropertyTraitsInstance)
+	{}
+
+	void RestorePreAnimatedValue(const KeyType& InKey, StorageType& CachedValue, const FRestoreStateParams& Params)
 	{
 		UObject* Object = InKey.Get<0>().ResolveObjectPtr();
 		if (!Object)
@@ -110,15 +122,15 @@ struct TPreAnimatedPropertyTraits<PropertyTraits, TPropertyMetaData<>, TIntegerS
 
 		if (const uint16* FastOffset = CachedValue.Binding.template TryGet<uint16>())
 		{
-			PropertyTraits::SetObjectPropertyValue(Object, *FastOffset, CachedValue.Data);
+			PropertyTraitsInstance->SetObjectPropertyValue(Object, *FastOffset, CachedValue.Data);
 		}
 		else  if (const TSharedPtr<FTrackInstancePropertyBindings>* Bindings = CachedValue.Binding.template TryGet<TSharedPtr<FTrackInstancePropertyBindings>>())
 		{
-			PropertyTraits::SetObjectPropertyValue(Object, Bindings->Get(), CachedValue.Data);
+			PropertyTraitsInstance->SetObjectPropertyValue(Object, Bindings->Get(), CachedValue.Data);
 		}
 		else if (const FCustomPropertyAccessor* CustomAccessor = CachedValue.Binding.template Get<const FCustomPropertyAccessor*>())
 		{
-			PropertyTraits::SetObjectPropertyValue(Object, *CustomAccessor, CachedValue.Data);
+			PropertyTraitsInstance->SetObjectPropertyValue(Object, *CustomAccessor, CachedValue.Data);
 		}
 	}
 };
@@ -138,7 +150,9 @@ struct TPreAnimatedPropertyStorageImpl<PropertyTraits, TPropertyMetaData<MetaDat
 	static_assert(StorageTraits::SupportsGrouping, "Pre-animated storage for properties should support grouping by object");
 
 	TPreAnimatedPropertyStorageImpl(const FPropertyDefinition& InPropertyDefinition)
-		: MetaDataComponents(InPropertyDefinition.MetaDataTypes)
+		: TPreAnimatedStateStorage<StorageTraits>(static_cast<const PropertyTraits*>(InPropertyDefinition.TraitsInstance))
+		, PropertyTraitsInstance(static_cast<const PropertyTraits*>(InPropertyDefinition.TraitsInstance))
+		, MetaDataComponents(InPropertyDefinition.MetaDataTypes)
 	{
 		check(MetaDataComponents.Num() == sizeof...(MetaDataIndices));
 
@@ -215,21 +229,21 @@ struct TPreAnimatedPropertyStorageImpl<PropertyTraits, TPropertyMetaData<MetaDat
 				if (Fast)
 				{
 					NewValue.Binding.template Set<uint16>(Fast[Index]);
-					PropertyTraits::GetObjectPropertyValue(BoundObjects[Index], MetaData.template Get<MetaDataIndices>()[Index]..., Fast[Index], NewValue.Data);
+					PropertyTraitsInstance->GetObjectPropertyValue(BoundObjects[Index], MetaData.template Get<MetaDataIndices>()[Index]..., Fast[Index], NewValue.Data);
 				}
 				else if (Custom)
 				{
 					const FCustomPropertyAccessor& Accessor = this->CustomAccessors[Custom[Index].Value];
 
 					NewValue.Binding.template Set<const FCustomPropertyAccessor*>(&Accessor);
-					PropertyTraits::GetObjectPropertyValue(BoundObjects[Index], MetaData.template Get<MetaDataIndices>()[Index]..., Accessor, NewValue.Data);
+					PropertyTraitsInstance->GetObjectPropertyValue(BoundObjects[Index], MetaData.template Get<MetaDataIndices>()[Index]..., Accessor, NewValue.Data);
 				}
 				else if (Slow)
 				{
 					const TSharedPtr<FTrackInstancePropertyBindings>& Bindings = Slow[Index];
 
 					NewValue.Binding.template Set<TSharedPtr<FTrackInstancePropertyBindings>>(Bindings);
-					PropertyTraits::GetObjectPropertyValue(BoundObjects[Index], MetaData.template Get<MetaDataIndices>()[Index]..., Bindings.Get(), NewValue.Data);
+					PropertyTraitsInstance->GetObjectPropertyValue(BoundObjects[Index], MetaData.template Get<MetaDataIndices>()[Index]..., Bindings.Get(), NewValue.Data);
 				}
 
 				this->AssignPreAnimatedValue(Entry.ValueHandle.StorageIndex, StorageRequirement, MoveTemp(NewValue));
@@ -244,6 +258,7 @@ struct TPreAnimatedPropertyStorageImpl<PropertyTraits, TPropertyMetaData<MetaDat
 
 protected:
 
+	const PropertyTraits* PropertyTraitsInstance;
 	TArrayView<const FComponentTypeID> MetaDataComponents;
 	FCustomAccessorView CustomAccessors;
 };
diff --git a/Engine/Source/Runtime/MovieScene/Public/MovieSceneCommonHelpers.h b/Engine/Source/Runtime/MovieScene/Public/MovieSceneCommonHelpers.h
--- a/Engine/Source/Runtime/MovieScene/Public/MovieSceneCommonHelpers.h
+++ b/Engine/Source/Runtime/MovieScene/Public/MovieSceneCommonHelpers.h
@@ -26,6 +26,11 @@ struct FMovieSceneSequenceID;
 class UMovieSceneCondition;
 class UMovieSceneTrack;
 
+namespace UE::MovieScene
+{
+	struct FSourcePropertyValue;
+} // namespace UE::MovieScene
+
 namespace UE::MovieScene
 {
 	struct FSharedPlaybackState;
@@ -409,6 +414,16 @@ public:
 		return TOptional<ValueType>();
 	}
 
+	/**
+	 * Static function for accessing a property value on an object without caching its address
+	 *
+	 * @param Object			The object to get the property from
+	 * @param InPropertyPath	The path to the property to retrieve
+	 * @return (Optional) The current value of the property on the object
+	 */
+	MOVIESCENE_API static TOptional<UE::MovieScene::FSourcePropertyValue> StaticValue(const UObject* Object, FStringView InPropertyPath);
+	MOVIESCENE_API static TOptional<TPair<const FProperty*, UE::MovieScene::FSourcePropertyValue>> StaticPropertyAndValue(const UObject* Object, FStringView InPropertyPath);
+
 	/**
 	 * Static function for accessing a property value on an object without caching its address
 	 *
@@ -417,12 +432,11 @@ public:
 	 * @return (Optional) The current value of the property on the object
 	 */
 	template <typename ValueType>
-	static TOptional<ValueType> StaticValue(const UObject* Object, const FString& InPropertyPath)
+	static TOptional<ValueType> StaticValue(const UObject* Object, FStringView InPropertyPath)
 	{
 		checkf(Object, TEXT("No object specified"));
 
-		FTrackInstancePropertyBindings Temp(NAME_None, InPropertyPath);
-		FResolvedProperty ResolvedProperty = ResolveProperty(Temp, *Object);
+		FResolvedProperty ResolvedProperty = ResolveProperty(InPropertyPath, *Object);
 
 		ValueType Value;
 		if (TryGetPropertyValue<ValueType>(ResolvedProperty, Value))
@@ -475,7 +489,7 @@ public:
 		return PropertyName;
 	}
 
-	static MOVIESCENE_API FProperty* FindProperty(const UObject* Object, const FString& InPropertyPath);
+	static MOVIESCENE_API FProperty* FindProperty(const UObject* Object, FStringView InPropertyPath);
 
 private:
 
@@ -560,11 +574,11 @@ private:
 		return false;
 	}
 
-	static MOVIESCENE_API FResolvedProperty ResolvePropertyRecursive(FTrackInstancePropertyBindings& Bindings, void* BasePointer, UStruct* InStruct, TArray<FString>& InPropertyNames, uint32 Index);
-	static MOVIESCENE_API FResolvedProperty ResolveProperty(FTrackInstancePropertyBindings& Bindings, const UObject& Object);
+	static MOVIESCENE_API FResolvedProperty ResolvePropertyRecursive(void* BasePointer, UStruct* InStruct, TArrayView<const FStringView> InPropertyNames, uint32 Index);
+	static MOVIESCENE_API FResolvedProperty ResolveProperty(FStringView PropertyPath, const UObject& Object);
 
-	static MOVIESCENE_API void FindProperty(FTrackInstancePropertyBindings& Bindings, void* BasePointer, UStruct* InStruct, const FString& InPropertyName, FResolvedProperty& OutResolvedProperty);
-	static MOVIESCENE_API FResolvedProperty FindPropertyAndArrayIndex(FTrackInstancePropertyBindings& Bindings, void* BasePointer, UStruct* InStruct, const FString& PropertyName);
+	static MOVIESCENE_API void FindProperty(void* BasePointer, UStruct* InStruct, FStringView InPropertyName, FResolvedProperty& OutResolvedProperty);
+	static MOVIESCENE_API FResolvedProperty FindPropertyAndArrayIndex(void* BasePointer, UStruct* InStruct, FStringView InPropertyName);
 
 	/** Find or add the FResolvedPropertyAndFunction for the specified object */
 	MOVIESCENE_API FResolvedPropertyAndFunction FindOrAdd(const UObject& InObject);
