Sequencer: Added safety checks to ensure that possessables and their associated bindings have matching GUIDs

diff --git a/Engine/Source/Runtime/MovieScene/Private/MovieScene.cpp b/Engine/Source/Runtime/MovieScene/Private/MovieScene.cpp
--- a/Engine/Source/Runtime/MovieScene/Private/MovieScene.cpp
+++ b/Engine/Source/Runtime/MovieScene/Private/MovieScene.cpp
@@ -404,13 +404,30 @@ FGuid UMovieScene::AddPossessable( const FString& Name, UClass* Class )
 
 void UMovieScene::AddPossessable(const FMovieScenePossessable& InNewPossessable, const FMovieSceneBinding& InNewBinding)
 {
+	const FGuid PossessableID = InNewPossessable.GetGuid();
+	const FGuid BindingID     = InNewBinding.GetObjectGuid();
+
+	FGuid FinalID = PossessableID.IsValid() ? PossessableID : BindingID;
+
 	Modify();
 
-	FMovieScenePossessable NewPossessable;
-	NewPossessable = InNewPossessable;
+	FMovieScenePossessable NewPossessable = InNewPossessable;
+	FMovieSceneBinding     NewBinding     = InNewBinding;
+
+	if (!ensureMsgf(PossessableID == BindingID, TEXT("GUIDs for possessable and binding objects do not match")))
+	{
+		if (!FinalID.IsValid())
+		{
+			FinalID = FGuid::NewGuid();
+		}
+
+		NewBinding.SetObjectGuid_Private(FinalID);
+		NewPossessable.SetGuid(FinalID);
+	}
+
+	NewBinding.SetObjectGuid_Private(InNewPossessable.GetGuid());
 	InsertSorted(Possessables, MoveTemp(NewPossessable));
 
-	FMovieSceneBinding NewBinding = InNewBinding;
 	for (UMovieSceneTrack* Track : NewBinding.GetTracks())
 	{
 		if (!Track->IsIn(this))
