Sequencer: Tweaked Binding ID validation logic

diff --git a/Engine/Source/Runtime/MovieScene/Private/MovieScene.cpp b/Engine/Source/Runtime/MovieScene/Private/MovieScene.cpp
--- a/Engine/Source/Runtime/MovieScene/Private/MovieScene.cpp
+++ b/Engine/Source/Runtime/MovieScene/Private/MovieScene.cpp
@@ -414,7 +414,7 @@ void UMovieScene::AddPossessable(const FMovieScenePossessable& InNewPossessable,
 	FMovieScenePossessable NewPossessable = InNewPossessable;
 	FMovieSceneBinding     NewBinding     = InNewBinding;
 
-	if (!ensureMsgf(PossessableID == BindingID, TEXT("GUIDs for possessable and binding objects do not match")))
+	if (!ensureMsgf(FinalID.IsValid() && PossessableID == BindingID, TEXT("GUIDs for possessable and binding objects are Invalid or do not match")))
 	{
 		if (!FinalID.IsValid())
 		{
@@ -425,7 +425,6 @@ void UMovieScene::AddPossessable(const FMovieScenePossessable& InNewPossessable,
 		NewPossessable.SetGuid(FinalID);
 	}
 
-	NewBinding.SetObjectGuid_Private(InNewPossessable.GetGuid());
 	InsertSorted(Possessables, MoveTemp(NewPossessable));
 
 	for (UMovieSceneTrack* Track : NewBinding.GetTracks())
