Sequencer: Deprecate FMovieSceneBinding::Name since it's redudant and not kept in sync with actor labels.

FMovieScenePossessable and FMovieSceneSpawnable have Name that is kept in sync (as much as possible) with actor labels. Consolidate code that tries to retrieve the FMovieSceneBinding::Name by instead calling UMovieScene::GetObjectDisplayName.

MovieSceneSequenceExtensions::FindBindingByName now returns a more reliable binding because it gathers the names from FMovieScenePossessable and FMovieSceneSpawnable which are kept in sync with the actor labels.

diff --git a/Engine/Plugins/Animation/ControlRig/Source/ControlRigEditor/Private/EditMode/SControlRigOutliner.cpp b/Engine/Plugins/Animation/ControlRig/Source/ControlRigEditor/Private/EditMode/SControlRigOutliner.cpp
--- a/Engine/Plugins/Animation/ControlRig/Source/ControlRigEditor/Private/EditMode/SControlRigOutliner.cpp
+++ b/Engine/Plugins/Animation/ControlRig/Source/ControlRigEditor/Private/EditMode/SControlRigOutliner.cpp
@@ -1280,24 +1280,7 @@ void SMultiRigHierarchyTreeView::RefreshTreeView()
 			FGuid Id = Sequencer->FindCachedObjectId(*Object, Sequencer->GetFocusedTemplateID());
 			if (UMovieScene* MovieScene = Sequencer->GetFocusedMovieSceneSequence()->GetMovieScene())
 			{
-				const FString* Name = nullptr;
-				if (FMovieSceneSpawnable* Spawnable = MovieScene->FindSpawnable(Id))
-				{
-					Name = &Spawnable->GetName();
-				}
-				else if (FMovieScenePossessable* Possesable = MovieScene->FindPossessable(Id))
-				{
-					Name = &Possesable->GetName();
-				}
-				else if (FMovieSceneBinding* Binding = MovieScene->FindBinding(Id))
-				{
-					Name = &Binding->GetName();
-				}
-
-				if (Name && !Name->IsEmpty())
-				{
-					InOutName = *Name;
-				}
+				InOutName = MovieScene->GetObjectDisplayName(Id).ToString();
 			}
 		}
 	};
diff --git a/Engine/Plugins/Experimental/SequenceNavigator/Source/SequenceNavigator/Private/Items/NavigationToolSequence.cpp b/Engine/Plugins/Experimental/SequenceNavigator/Source/SequenceNavigator/Private/Items/NavigationToolSequence.cpp
--- a/Engine/Plugins/Experimental/SequenceNavigator/Source/SequenceNavigator/Private/Items/NavigationToolSequence.cpp
+++ b/Engine/Plugins/Experimental/SequenceNavigator/Source/SequenceNavigator/Private/Items/NavigationToolSequence.cpp
@@ -724,13 +724,13 @@ TArray<FMovieSceneBinding> FNavigationToolSequence::GetSortedBindings() const
 
 	TArray<FMovieSceneBinding> Bindings = MovieScene->GetBindings();
 
-	Bindings.Sort([](const FMovieSceneBinding& InA, const FMovieSceneBinding& InB)
+	Bindings.Sort([MovieScene](const FMovieSceneBinding& InA, const FMovieSceneBinding& InB)
 		{
 			const int32 SortingOrderA = InA.GetSortingOrder();
 			const int32 SortingOrderB = InB.GetSortingOrder();
 			if (SortingOrderA == SortingOrderB)
 			{
-				return InA.GetName() < InB.GetName();
+				return MovieScene->GetObjectDisplayName(InA.GetObjectGuid()).ToString() < MovieScene->GetObjectDisplayName(InB.GetObjectGuid()).ToString();
 			}
 			return SortingOrderA < SortingOrderB;
 		});
diff --git a/Engine/Plugins/MovieScene/MovieRenderPipeline/Source/MovieRenderPipelineCore/Private/MoviePipelineUtils.cpp b/Engine/Plugins/MovieScene/MovieRenderPipeline/Source/MovieRenderPipelineCore/Private/MoviePipelineUtils.cpp
--- a/Engine/Plugins/MovieScene/MovieRenderPipeline/Source/MovieRenderPipelineCore/Private/MoviePipelineUtils.cpp
+++ b/Engine/Plugins/MovieScene/MovieRenderPipeline/Source/MovieRenderPipelineCore/Private/MoviePipelineUtils.cpp
@@ -649,7 +649,7 @@ namespace MoviePipeline
 							{
 								if (BindingTrack->HasSection(*Section))
 								{
-									BindingName = Binding.GetName();
+									BindingName = Node->MovieScene.Get()->GetObjectDisplayName(Binding.GetObjectGuid()).ToString();
 									break;
 								}
 							}
@@ -725,22 +725,7 @@ namespace MoviePipeline
 
 				if (OwningSequence && OwningSequence->GetMovieScene())
 				{
-					const FMovieSceneBinding* Binding = OwningSequence->GetMovieScene()->FindBinding(CameraObjectBindingId.GetGuid());
-					if (Binding)
-					{
-						if (const FMovieSceneSpawnable* Spawnable = OwningSequence->GetMovieScene()->FindSpawnable(Binding->GetObjectGuid()))
-						{
-							InnerName = Spawnable->GetName();
-						}
-						else if (const FMovieScenePossessable* Posssessable = OwningSequence->GetMovieScene()->FindPossessable(Binding->GetObjectGuid()))
-						{
-							InnerName = Posssessable->GetName();
-						}
-						else
-						{
-							InnerName = Binding->GetName();
-						}
-					}
+					InnerName = OwningSequence->GetMovieScene()->GetObjectDisplayName(CameraObjectBindingId.GetGuid()).ToString();
 				}
 			}
 		}
diff --git a/Engine/Plugins/MovieScene/SequencerScripting/Source/SequencerScriptingEditor/Private/SequencerTools.cpp b/Engine/Plugins/MovieScene/SequencerScripting/Source/SequencerScriptingEditor/Private/SequencerTools.cpp
--- a/Engine/Plugins/MovieScene/SequencerScripting/Source/SequencerScriptingEditor/Private/SequencerTools.cpp
+++ b/Engine/Plugins/MovieScene/SequencerScripting/Source/SequencerScriptingEditor/Private/SequencerTools.cpp
@@ -835,9 +835,12 @@ bool USequencerToolsFunctionLibrary::ImportFBXToControlRig(UWorld* World, ULevel
 	bool bValid = false;
 
 	const TArray<FMovieSceneBinding>& Bindings =  ((const UMovieScene*)MovieScene)->GetBindings();
+
+
+
 	for (const FMovieSceneBinding& Binding : Bindings)
 	{
-		if (Binding.GetName() == ControlRigTrackName)
+		if (MovieScene->GetObjectDisplayName(Binding.GetObjectGuid()).ToString() == ControlRigTrackName)
 		{
 			
 			ALevelSequenceActor* OutActor;
@@ -898,7 +901,7 @@ bool USequencerToolsFunctionLibrary::ExportFBXFromControlRig(ULevelSequence* Seq
 	const TArray<FMovieSceneBinding>& Bindings = ((const UMovieScene*)MovieScene)->GetBindings();
 	for (const FMovieSceneBinding& Binding : Bindings)
 	{
-		if (Binding.GetName() != ActorWithControlRigTrack)
+		if (MovieScene->GetObjectDisplayName(Binding.GetObjectGuid()).ToString() != ActorWithControlRigTrack)
 		{
 			continue;
 		}
diff --git a/Engine/Source/Editor/Sequencer/Private/Sequencer.cpp b/Engine/Source/Editor/Sequencer/Private/Sequencer.cpp
--- a/Engine/Source/Editor/Sequencer/Private/Sequencer.cpp
+++ b/Engine/Source/Editor/Sequencer/Private/Sequencer.cpp
@@ -9469,9 +9469,10 @@ void FixSortingOrders(FMovieSceneBinding* InBinding, const UMovieScene* MovieSce
 {
 	// Fix the sorting orders if InBinding is going to conflict with any of the existing bindings (ie. same class and same name)
 	TOptional<int32> MaxSortingOrder;
+	FText InBindingName = MovieScene->GetObjectDisplayName(InBinding->GetObjectGuid());
 	for (const FMovieSceneBinding& Binding : MovieScene->GetBindings())
 	{
-		if (&Binding != InBinding && Binding.GetName() == InBinding->GetName())
+		if (&Binding != InBinding && MovieScene->GetObjectDisplayName(Binding.GetObjectGuid()).EqualTo(InBindingName))
 		{
 			if (MaxSortingOrder.IsSet())
 			{
diff --git a/Engine/Source/Editor/UMGEditor/Private/WidgetBlueprintEditor.cpp b/Engine/Source/Editor/UMGEditor/Private/WidgetBlueprintEditor.cpp
--- a/Engine/Source/Editor/UMGEditor/Private/WidgetBlueprintEditor.cpp
+++ b/Engine/Source/Editor/UMGEditor/Private/WidgetBlueprintEditor.cpp
@@ -2701,11 +2701,12 @@ void FWidgetBlueprintEditor::OnMovieSceneBindingsPasted(const TArray<FMovieScene
 	for (const FMovieSceneBinding& BindingPasted : BindingsPasted)
 	{
 		FMovieScenePossessable* Possessable = AnimationSequence->GetMovieScene()->FindPossessable(BindingPasted.GetObjectGuid());
+		FText BindingPastedName = AnimationSequence->GetMovieScene()->GetObjectDisplayName(BindingPasted.GetObjectGuid());
 		if (Possessable && !Possessable->GetParent().IsValid())
 		{
 			for (FObjectAndDisplayName& BindableObject : BindableObjects)
 			{
-				if (BindableObject.DisplayName.ToString() == BindingPasted.GetName())
+				if (BindableObject.DisplayName.EqualTo(BindingPastedName))
 				{
 					AnimationSequence->BindPossessableObject(BindingPasted.GetObjectGuid(), *BindableObject.Object, BindingContext);			
 					break;
@@ -2718,6 +2719,7 @@ void FWidgetBlueprintEditor::OnMovieSceneBindingsPasted(const TArray<FMovieScene
 	for (const FMovieSceneBinding& BindingPasted : BindingsPasted)
 	{
 		FMovieScenePossessable* Possessable = AnimationSequence->GetMovieScene()->FindPossessable(BindingPasted.GetObjectGuid());
+		FText BindingPastedName = AnimationSequence->GetMovieScene()->GetObjectDisplayName(BindingPasted.GetObjectGuid());
 		if (Possessable && Possessable->GetParent().IsValid())
 		{
 			for (FObjectAndDisplayName& BindableObject : BindableObjects)
@@ -2734,7 +2736,7 @@ void FWidgetBlueprintEditor::OnMovieSceneBindingsPasted(const TArray<FMovieScene
 					}
 
 					// Special case for canvas slots, they need to be added again
-					if (BindableObject.Object->GetFName().ToString() == BindingPasted.GetName())
+					if (BindableObject.Object->GetFName() == BindingPastedName.ToString())
 					{
 						// Create handle, to rebind correctly
 						ActiveSequencer->GetHandleToObject(BindableObject.Object);
diff --git a/Engine/Source/Runtime/MovieScene/Private/Evaluation/MovieSceneEvaluationState.cpp b/Engine/Source/Runtime/MovieScene/Private/Evaluation/MovieSceneEvaluationState.cpp
--- a/Engine/Source/Runtime/MovieScene/Private/Evaluation/MovieSceneEvaluationState.cpp
+++ b/Engine/Source/Runtime/MovieScene/Private/Evaluation/MovieSceneEvaluationState.cpp
@@ -543,11 +543,10 @@ void FMovieSceneObjectCache::UpdateBindings(const FGuid& InGuid, TSharedRef<cons
 		const int32 EndIndex = Algo::Unique(Bindings->Objects);
 		if (EndIndex < NumBoundObjects)
 		{
-			FMovieSceneBinding* Binding = Sequence->GetMovieScene()->FindBinding(InGuid);
-			FString BindingName = Binding ? Binding->GetName() : TEXT("<invalid binding>");
+			FText ObjectDisplayName = Sequence->GetMovieScene()->GetObjectDisplayName(InGuid);
 			UE_LOG(LogMovieScene, Warning, TEXT("Found %d duplicate object(s) while resolving binding %s (%s) in %s"),
 					NumBoundObjects - EndIndex, 
-					*BindingName, *LexToString(InGuid), 
+					*ObjectDisplayName.ToString(), *LexToString(InGuid), 
 					*Sequence->GetPathName());
 			Bindings->Objects.SetNum(EndIndex);
 		}
diff --git a/Engine/Source/Runtime/MovieScene/Private/MovieScene.cpp b/Engine/Source/Runtime/MovieScene/Private/MovieScene.cpp
--- a/Engine/Source/Runtime/MovieScene/Private/MovieScene.cpp
+++ b/Engine/Source/Runtime/MovieScene/Private/MovieScene.cpp
@@ -303,7 +303,7 @@ FGuid UMovieScene::AddSpawnable( const FString& Name, UObject& ObjectTemplate )
 	check(Spawnables.IsValidIndex(NewSpawnableIndex));
 
 	// Add a new binding (sorted) so that tracks can be added to it
-	int32 NewBindingIndex = InsertSorted(ObjectBindings, FMovieSceneBinding(Spawnables[NewSpawnableIndex].GetGuid(), Spawnables[NewSpawnableIndex].GetName()));
+	int32 NewBindingIndex = InsertSorted(ObjectBindings, FMovieSceneBinding(Spawnables[NewSpawnableIndex].GetGuid()));
 
 	EventHandlers.Trigger(&UE::MovieScene::ISequenceDataEventHandler::OnBindingAdded, ObjectBindings[NewBindingIndex]);
 
@@ -394,7 +394,7 @@ FGuid UMovieScene::AddPossessable( const FString& Name, UClass* Class )
 	int32 NewPossessableIndex = InsertSorted(Possessables, FMovieScenePossessable(Name, Class));
 	check(Possessables.IsValidIndex(NewPossessableIndex));
 	// Add a new binding so that tracks can be added to it
-	int32 NewBindingIndex = InsertSorted(ObjectBindings, FMovieSceneBinding(Possessables[NewPossessableIndex].GetGuid(), Possessables[NewPossessableIndex].GetName()));
+	int32 NewBindingIndex = InsertSorted(ObjectBindings, FMovieSceneBinding(Possessables[NewPossessableIndex].GetGuid()));
 	check(ObjectBindings.IsValidIndex(NewBindingIndex));
 	EventHandlers.Trigger(&UE::MovieScene::ISequenceDataEventHandler::OnBindingAdded, ObjectBindings[NewBindingIndex]);
 
@@ -546,29 +546,28 @@ FMovieScenePossessable& UMovieScene::GetPossessable( const int32 Index )
 }
 
 
-FText UMovieScene::GetObjectDisplayName(const FGuid& ObjectId)
+FText UMovieScene::GetObjectDisplayName(const FGuid& ObjectId) const
 {
 #if WITH_EDITORONLY_DATA
-	FText* Result = ObjectsToDisplayNames.Find(ObjectId.ToString());
+	const FText* Result = ObjectsToDisplayNames.Find(ObjectId.ToString());
 
 	if (Result && !Result->IsEmpty())
 	{
 		return *Result;
 	}
 
-	FMovieSceneSpawnable* Spawnable = FindSpawnable(ObjectId);
-
-	if (Spawnable != nullptr)
+	int32 SpawnableIndex = IndexOfSpawnable(ObjectId);
+	if (SpawnableIndex != INDEX_NONE)
 	{
-		return FText::FromString(Spawnable->GetName());
+		return FText::FromString(Spawnables[SpawnableIndex].GetName());
 	}
 
-	FMovieScenePossessable* Possessable = FindPossessable(ObjectId);
-
-	if (Possessable != nullptr)
+	int32 PossessableIndex = IndexOfPossessable(ObjectId);
+	if (PossessableIndex != INDEX_NONE)
 	{
-		return FText::FromString(Possessable->GetName());
+		return FText::FromString(Possessables[PossessableIndex].GetName());
 	}
+
 #endif
 	return FText::GetEmpty();
 }
@@ -1661,7 +1660,6 @@ void UMovieScene::ReplaceBinding(const FGuid& OldGuid, const FGuid& NewGuid, con
 	if (FMovieSceneBinding* Binding = FindBinding(OldGuid))
 	{
 		Binding->SetObjectGuid_Private(NewGuid);
-		Binding->SetName(Name);
 
 		// Replace directly changes the guid, so force a sort here
 		ObjectBindings.Sort();
diff --git a/Engine/Source/Runtime/MovieScene/Private/MovieSceneSequence.cpp b/Engine/Source/Runtime/MovieScene/Private/MovieSceneSequence.cpp
--- a/Engine/Source/Runtime/MovieScene/Private/MovieSceneSequence.cpp
+++ b/Engine/Source/Runtime/MovieScene/Private/MovieSceneSequence.cpp
@@ -324,8 +324,9 @@ void UMovieSceneSequence::OptimizeForCook()
 		
 		if (bRemoveObject)
 		{
-			UE_LOG(LogMovieScene, Display, TEXT("Removing muted object: %s from: %s"), *Bindings[ObjectBindingIndex].GetName(), *GetPathName());
 			FGuid GuidToRemove = Bindings[ObjectBindingIndex].GetObjectGuid();
+			FText ObjectDisplayName = MovieScene->GetObjectDisplayName(GuidToRemove);
+			UE_LOG(LogMovieScene, Display, TEXT("Removing muted object: %s from: %s"), *ObjectDisplayName.ToString(), *GetPathName());
 			MovieScene->RemoveSpawnable(GuidToRemove);
 			MovieScene->RemovePossessable(GuidToRemove);
 		}
diff --git a/Engine/Source/Runtime/MovieScene/Public/MovieScene.h b/Engine/Source/Runtime/MovieScene/Public/MovieScene.h
--- a/Engine/Source/Runtime/MovieScene/Public/MovieScene.h
+++ b/Engine/Source/Runtime/MovieScene/Public/MovieScene.h
@@ -793,7 +793,7 @@ public:
 	 * @param ObjectId The object identifier.
 	 * @return The object's display name.
 	 */
-	MOVIESCENE_API FText GetObjectDisplayName(const FGuid& ObjectId);
+	MOVIESCENE_API FText GetObjectDisplayName(const FGuid& ObjectId) const;
 
 	/** Get the playback time range of this movie scene, relative to its 0-time offset. */
 	TRange<FFrameNumber> GetPlaybackRange() const
diff --git a/Engine/Source/Runtime/MovieScene/Public/MovieSceneBinding.h b/Engine/Source/Runtime/MovieScene/Public/MovieSceneBinding.h
--- a/Engine/Source/Runtime/MovieScene/Public/MovieSceneBinding.h
+++ b/Engine/Source/Runtime/MovieScene/Public/MovieSceneBinding.h
@@ -36,12 +36,10 @@ private:
 	 * Creates and initializes a new instance.
 	 *
 	 * @param InObjectGuid
-	 * @param InBindingName
 	 * @param InTracks
 	 */
-	FMovieSceneBinding(const FGuid& InObjectGuid, const FString& InBindingName, const TArray<UMovieSceneTrack*>& InTracks)
+	FMovieSceneBinding(const FGuid& InObjectGuid, const TArray<UMovieSceneTrack*>& InTracks)
 		: ObjectGuid(InObjectGuid)
-		, BindingName(InBindingName)
 		, Tracks(InTracks)
 #if WITH_EDITORONLY_DATA
 		, SortingOrder(-1)
@@ -52,11 +50,9 @@ private:
 	 * Creates and initializes a new instance.
 	 *
 	 * @param InObjectGuid
-	 * @param InBindingName
 	 */
-	FMovieSceneBinding(const FGuid& InObjectGuid, const FString& InBindingName)
+	FMovieSceneBinding(const FGuid& InObjectGuid)
 		: ObjectGuid(InObjectGuid)
-		, BindingName(InBindingName)
 #if WITH_EDITORONLY_DATA
 		, SortingOrder(-1)
 #endif
@@ -88,21 +84,11 @@ public:
 		return ObjectGuid;
 	}
 
-	/**
-	 * Set display name of the binding
-	 */
-	void SetName(const FString& InBindingName)
-	{
-		BindingName = InBindingName;
-	} 
+	UE_DEPRECATED(5.7, "SetName is deprecated. Use SetName on FMovieScenePossessable or FMovieSceneSpawnable")
+	void SetName(const FString& InBindingName) { } 
 
-	/**
-	 * @return The display name of the binding
-	 */
-	const FString& GetName() const
-	{
-		return BindingName;
-	}
+	UE_DEPRECATED(5.7, "GetName is deprecated. Use GetName on FMovieScenePossessable or FMovieSceneSpawnable")
+	const FString& GetName() const { return BindingName_DEPRECATED; }
 
 	/**
 	 * Adds a new track to this binding
@@ -177,7 +163,7 @@ private:
 	
 	/** Display name */
 	UPROPERTY()
-	FString BindingName;
+	FString BindingName_DEPRECATED;
 
 	/** All tracks in this binding */
 	UPROPERTY(Instanced)
