When fixed timescale is set negative, ranges 2 levels down collapse causing subscenes of subscenes not to play correctly.
This changelist only fixes the case of fixed timescale. Similar issues occur using the playrate/timescale curve. These issues are too complex to fix for 5.7 and will be addressed in 5.8.

diff --git a/Engine/Source/Runtime/MovieScene/Private/Compilation/MovieSceneCompiledDataManager.cpp b/Engine/Source/Runtime/MovieScene/Private/Compilation/MovieSceneCompiledDataManager.cpp
--- a/Engine/Source/Runtime/MovieScene/Private/Compilation/MovieSceneCompiledDataManager.cpp
+++ b/Engine/Source/Runtime/MovieScene/Private/Compilation/MovieSceneCompiledDataManager.cpp
@@ -213,6 +213,16 @@ struct FGatherParameters
 				Range.SetUpperBoundValue(Range.GetUpperBoundValue() * LinearTransform);
 			}
 
+			// Normalize inside-out ranges due to negative TimeScale
+			if (Range.GetLowerBound().IsClosed() && Range.GetUpperBound().IsClosed()
+				&& Range.GetLowerBoundValue() > Range.GetUpperBoundValue())
+			{
+				const auto OldLower = Range.GetLowerBound();
+				const auto OldUpper = Range.GetUpperBound();
+				Range.SetLowerBound(OldUpper);
+				Range.SetUpperBound(OldLower);
+			}
+
 			InVisitor(Range);
 			return;
 		}
@@ -254,6 +264,16 @@ struct FGatherParameters
 			}
 		}
 
+		// Normalize inside-out ranges due to negative TimeScale
+		if (Range.GetLowerBound().IsClosed() && Range.GetUpperBound().IsClosed()
+			&& Range.GetLowerBoundValue() > Range.GetUpperBoundValue())
+		{
+			const auto OldLower = Range.GetLowerBound();
+			const auto OldUpper = Range.GetUpperBound();
+			Range.SetLowerBound(OldUpper);
+			Range.SetUpperBound(OldLower);
+		}
+
 		if (Range.GetLowerBound().IsOpen() && Range.GetUpperBound().IsOpen())
 		{
 			// If the range is infinite we just have to add it all since there's no way for us to transform it.
diff --git a/Engine/Source/Runtime/MovieScene/Private/EntitySystem/MovieSceneEntitySystemRunner.cpp b/Engine/Source/Runtime/MovieScene/Private/EntitySystem/MovieSceneEntitySystemRunner.cpp
--- a/Engine/Source/Runtime/MovieScene/Private/EntitySystem/MovieSceneEntitySystemRunner.cpp
+++ b/Engine/Source/Runtime/MovieScene/Private/EntitySystem/MovieSceneEntitySystemRunner.cpp
@@ -819,6 +819,12 @@ UE::MovieScene::ERunnerFlushResult FMovieSceneEntitySystemRunner::GameThread_Upd
 				};
 				DissectedUpdates.Add(Dissection);
 
+				// Only mark the first dissection for update this round
+				if (Dissections.Num() == 1)
+				{
+					MarkForUpdate(Request.Params.InstanceHandle, Request.Params.UpdateFlags);
+				}
+
 				Dissections.Reset();
 			}
 			else
diff --git a/Engine/Source/Runtime/MovieScene/Private/EntitySystem/MovieSceneSequenceUpdaters.cpp b/Engine/Source/Runtime/MovieScene/Private/EntitySystem/MovieSceneSequenceUpdaters.cpp
--- a/Engine/Source/Runtime/MovieScene/Private/EntitySystem/MovieSceneSequenceUpdaters.cpp
+++ b/Engine/Source/Runtime/MovieScene/Private/EntitySystem/MovieSceneSequenceUpdaters.cpp
@@ -963,6 +963,15 @@ void FSequenceUpdater_Hierarchical::Update(TSharedRef<const FSharedPlaybackState
 						}
 					}
 
+					// Time-warp can result in inside-out ranges
+					if (SubCachedRange.GetLowerBound().IsClosed() && SubCachedRange.GetUpperBound().IsClosed() && SubCachedRange.GetLowerBoundValue() > SubCachedRange.GetUpperBoundValue())
+					{
+						TRangeBound<FFrameNumber> OldLower = SubCachedRange.GetLowerBound();
+						TRangeBound<FFrameNumber> OldUpper = SubCachedRange.GetUpperBound();
+						SubCachedRange.SetLowerBound(OldUpper);
+						SubCachedRange.SetUpperBound(OldLower);
+					}
+
 					CachedEntityRange = TRange<FFrameNumber>::Intersection(CachedEntityRange, SubCachedRange);
 				}
 				else if (FMovieSceneEvaluationFieldEntitySet* SubSequenceCachedConditionalEntries = CachedPerTickConditionalEntities.Find(SequenceIDFromRoot))
