Sequencer: refactor track instance bindings and add the concept of "volatile" paths to fix possible memory stomps

- Move all the property resolution code into a local namespace with local functions.
- Make these functions use some temporary structs to keep track of a few more things compared to before.
- The resolved property can now be either a "cached" property that cached the tail of the property path, or a "volatile" property that will always walk the objects down the path.
- We use the "volatile" version when going into an array or an instanced struct.
- Removed some code that had accumulated over time and that isn't useful anymore.

diff --git a/Engine/Source/Runtime/MovieScene/Private/MovieSceneCommonHelpers.cpp b/Engine/Source/Runtime/MovieScene/Private/MovieSceneCommonHelpers.cpp
--- a/Engine/Source/Runtime/MovieScene/Private/MovieSceneCommonHelpers.cpp
+++ b/Engine/Source/Runtime/MovieScene/Private/MovieSceneCommonHelpers.cpp
@@ -24,9 +24,7 @@
 #include "EntitySystem/MovieSceneSequenceInstance.h"
 #include "EntitySystem/MovieSceneEntitySystemLinker.h"
 #include "Conditions/MovieSceneGroupCondition.h"
-#include "StructUtils/InstancedStruct.h"
 #include "String/ParseTokens.h"
-#include "StructUtils/PropertyBag.h"
 
 bool MovieSceneHelpers::IsSectionKeyable(const UMovieSceneSection* Section)
 {
@@ -1219,474 +1217,3 @@ MovieSceneHelpers::FMovieSceneScopedPackageDirtyGuard::~FMovieSceneScopedPackage
 	}
 }
 
-FTrackInstancePropertyBindings::FTrackInstancePropertyBindings(FName InPropertyName, const FString& InPropertyPath)
-	: PropertyPath(InPropertyPath)
-	, PropertyName(InPropertyName)
-{
-	static const FString Set(TEXT("Set"));
-	const FString FunctionString = Set + PropertyName.ToString();
-
-	FunctionName = FName(*FunctionString);
-}
-
-void FTrackInstancePropertyBindings::FindProperty(void* BasePointer, UStruct* InStruct, FStringView InPropertyName, FResolvedProperty& OutResolvedProperty)
-{
-	// Find a property via UObject reflection first.
-	FProperty* FoundProperty = FindFProperty<FProperty>(InStruct, FName(InPropertyName, FNAME_Find));
-
-	if (FoundProperty)
-	{
-		OutResolvedProperty.Property = FoundProperty;
-	}
-}
-
-FTrackInstancePropertyBindings::FResolvedProperty FTrackInstancePropertyBindings::FindPropertyAndArrayIndex(void* BasePointer, UStruct* InStruct, FStringView PropertyName)
-{
-	FResolvedProperty PropertyAndIndex;
-
-	// Calculate the array index if possible.
-	int32 ArrayIndex = -1;
-	if (PropertyName.Len() > 0 && PropertyName[PropertyName.Len() - 1] == ']')
-	{
-		int32 OpenIndex = 0;
-		if (PropertyName.FindLastChar('[', OpenIndex))
-		{
-			// We have a property name of the form "Foo[123]". Resolve the property itself ("Foo") and then
-			// parse the array element index (123).
-			FStringView TruncatedPropertyName(PropertyName.GetData(), OpenIndex);
-			FindProperty(BasePointer, InStruct, TruncatedPropertyName, PropertyAndIndex);
-
-			const int32 NumberLength = PropertyName.Len() - OpenIndex - 2;
-			if (NumberLength > 0 && NumberLength <= 10)
-			{
-				TCHAR NumberBuffer[11];
-				FMemory::Memcpy(NumberBuffer, &PropertyName[OpenIndex + 1], sizeof(TCHAR) * NumberLength);
-				LexFromString(PropertyAndIndex.ArrayIndex, NumberBuffer);
-			}
-
-			return PropertyAndIndex;
-		}
-	}
-
-	// Not index found in this property name, just find the property normally.
-	FindProperty(BasePointer, InStruct, PropertyName, PropertyAndIndex);
-	return PropertyAndIndex;
-}
-
-TOptional<TPair<const FProperty*, UE::MovieScene::FSourcePropertyValue>> FTrackInstancePropertyBindings::StaticPropertyAndValue(const UObject* Object, FStringView InPropertyPath)
-{
-	using namespace UE::MovieScene;
-
-	checkf(Object, TEXT("No object specified"));
-
-	FResolvedProperty Address = ResolveProperty(InPropertyPath, *Object);
-
-	if (const FProperty* Property = Address.Property.Get())
-	{
-		const void* PropertyAddress = Property->ContainerPtrToValuePtr<void>(Address.ContainerAddress);
-
-		// Bool property values are stored in a bit field so using a straight cast of the pointer to get their value does not
-		// work.  Instead use the actual property to get the correct value.
-		if (const FBoolProperty* BoolProperty = CastField<const FBoolProperty>(Property))
-		{
-			return MakeTuple(Property, FSourcePropertyValue::FromValue(BoolProperty->GetPropertyValue(PropertyAddress)));
-		}
-		// Object properties might have various different types of storage, but we always expose them as a raw ptr
-		else if (const FObjectPropertyBase* ObjectProperty = CastField<const FObjectPropertyBase>(Property))
-		{
-			return MakeTuple(Property, FSourcePropertyValue::FromValue(ObjectProperty->GetObjectPropertyValue(PropertyAddress)));
-		}
-
-		return MakeTuple(Property, FSourcePropertyValue::FromAddress(PropertyAddress, *Property));
-	}
-
-	return TOptional<TPair<const FProperty*, UE::MovieScene::FSourcePropertyValue>>();
-}
-
-TOptional<UE::MovieScene::FSourcePropertyValue> FTrackInstancePropertyBindings::StaticValue(const UObject* Object, FStringView InPropertyPath)
-{
-	TOptional<TPair<const FProperty*, UE::MovieScene::FSourcePropertyValue>> Result = StaticPropertyAndValue(Object, InPropertyPath);
-	if (Result.IsSet())
-	{
-		return MoveTemp(Result->Value);
-	}
-	return TOptional<UE::MovieScene::FSourcePropertyValue>();
-}
-
-FTrackInstancePropertyBindings::FResolvedProperty FTrackInstancePropertyBindings::ResolvePropertyRecursive(void* BasePointer, UStruct* InStruct, TArrayView<const FStringView> InPropertyNames, uint32 Index)
-{
-	FResolvedProperty ResolvedProperty = FindPropertyAndArrayIndex(BasePointer, InStruct, InPropertyNames[Index]);
-	
-	FResolvedProperty NewProperty;
-
-	FProperty* ResolvedPropertyProperty = ResolvedProperty.Property.Get();
-
-	if (ResolvedProperty.ArrayIndex != INDEX_NONE)
-	{
-		// We found that this segment of the property path reaches an element inside an array.
-		if (ResolvedPropertyProperty)
-		{
-			if (FArrayProperty* ArrayProp = CastField<FArrayProperty>(ResolvedPropertyProperty))
-			{
-				FScriptArrayHelper ArrayHelper(ArrayProp, ArrayProp->ContainerPtrToValuePtr<void>(BasePointer));
-				if (ArrayHelper.IsValidIndex(ResolvedProperty.ArrayIndex))
-				{
-					FStructProperty* InnerStructProp = CastField<FStructProperty>(ArrayProp->Inner);
-					if (InnerStructProp && InPropertyNames.IsValidIndex(Index + 1))
-					{
-						// Move the BasePointer to the array element and keep resolving the property path on it.
-						void* ArrayElement = ArrayHelper.GetRawPtr(ResolvedProperty.ArrayIndex);
-						return ResolvePropertyRecursive(ArrayElement, InnerStructProp->Struct, InPropertyNames, Index + 1);
-					}
-					else
-					{
-						// The property path ends here (e.g. "Foo.Bar[1]").
-						NewProperty.Property = ArrayProp;
-						NewProperty.ContainerAddress = BasePointer;
-						NewProperty.ArrayIndex = ResolvedProperty.ArrayIndex;
-					}
-				}
-			}
-			else
-			{
-				UE_LOG(LogMovieScene, Error, TEXT("Mismatch in property evaluation. %s is not of type: %s"), *ResolvedPropertyProperty->GetName(), *FArrayProperty::StaticClass()->GetName());
-			}
-		}
-	}
-	else if (FStructProperty* StructProp = CastField<FStructProperty>(ResolvedPropertyProperty))
-	{
-		// This segment of the property path reaches a struct FProperty.
-		NewProperty.Property = StructProp;
-		NewProperty.ContainerAddress = BasePointer;
-
-		if (InPropertyNames.IsValidIndex(Index + 1))
-		{
-			// Instanced structs are technically just a memory buffer with no real sub-properties, but
-			// they do have sub-properties if we ask them about their "logical" struct type. Let's do that,
-			// which makes it possible to animate the properties inside.
-			if (StructProp->Struct == FInstancedStruct::StaticStruct())
-			{
-				FInstancedStruct* InstancedStruct = StructProp->ContainerPtrToValuePtr<FInstancedStruct>(BasePointer);
-				const UScriptStruct* InstancedStructType = InstancedStruct->GetScriptStruct();
-				uint8* InstancedStructMemory = InstancedStruct->GetMutableMemory();
-				return ResolvePropertyRecursive(InstancedStructMemory, const_cast<UScriptStruct*>(InstancedStructType), InPropertyNames, Index + 1);
-			}
-			// As above but for a property bag, unless the property path is specifically targeting its inner instanced struct.
-			else if (StructProp->Struct == FInstancedPropertyBag::StaticStruct() && InPropertyNames[Index + 1] != TEXT("Value"))
-			{
-				FInstancedPropertyBag* PropertyBag = StructProp->ContainerPtrToValuePtr<FInstancedPropertyBag>(BasePointer);
-				const UPropertyBag* PropertyBagType = PropertyBag->GetPropertyBagStruct();
-				uint8* PropertyBagMemory = PropertyBag->GetMutableValue().GetMemory();
-				return ResolvePropertyRecursive(PropertyBagMemory, const_cast<UPropertyBag*>(PropertyBagType), InPropertyNames, Index + 1);
-			}
-			else
-			{
-				void* StructContainer = StructProp->ContainerPtrToValuePtr<void>(BasePointer);
-				return ResolvePropertyRecursive(StructContainer, StructProp->Struct, InPropertyNames, Index + 1);
-			}
-		}
-		else
-		{
-			check(StructProp->GetName() == InPropertyNames[Index]);
-		}
-	}
-	else if (ResolvedPropertyProperty)
-	{
-		NewProperty.Property = ResolvedProperty.Property;
-		NewProperty.ContainerAddress = BasePointer;
-	}
-
-	return NewProperty;
-}
-
-FTrackInstancePropertyBindings::FResolvedProperty FTrackInstancePropertyBindings::ResolveProperty(FStringView PropertyPath, const UObject& InObject)
-{
-	using namespace UE::String;
-
-	TArray<FStringView, TInlineAllocator<4>> PropertyNames;
-
-	// Parse property paths into component parts separated by '.'
-	auto ProcessToken = [&PropertyNames](FStringView Token) { PropertyNames.Emplace(Token); };
-	ParseTokens(PropertyPath, TEXT('.'), TFunctionRef<void(FWideStringView)>(ProcessToken), EParseTokensOptions::IgnoreCase | EParseTokensOptions::SkipEmpty);
-
-	if (IsValid(&InObject) && PropertyNames.Num() > 0)
-	{
-		return ResolvePropertyRecursive((void*)&InObject, InObject.GetClass(), PropertyNames, 0);
-	}
-	else
-	{
-		return FTrackInstancePropertyBindings::FResolvedProperty();
-	}
-}
-
-FProperty* FTrackInstancePropertyBindings::FindProperty(const UObject* Object, FStringView InPropertyPath)
-{
-	FResolvedProperty ResolvedProperty = ResolveProperty(InPropertyPath, *Object);
-	return ResolvedProperty.Property.Get();
-}
-
-FTrackInstancePropertyBindings::FResolvedPropertyAndFunction FTrackInstancePropertyBindings::FindOrAdd(const UObject& InObject)
-{
-	FObjectKey ObjectKey(&InObject);
-
-	const FResolvedPropertyAndFunction* PropAndFunction = RuntimeObjectToFunctionMap.Find(ObjectKey);
-	if (PropAndFunction && (
-				PropAndFunction->SetterFunction.IsValid() ||
-				PropAndFunction->ResolvedProperty.GetValidProperty()))
-	{
-		return *PropAndFunction;
-	}
-
-	CacheBinding(InObject);
-	return RuntimeObjectToFunctionMap.FindRef(ObjectKey);
-}
-
-void FTrackInstancePropertyBindings::CallFunctionForEnum(UObject& InRuntimeObject, int64 PropertyValue)
-{
-	FResolvedPropertyAndFunction PropAndFunction = FindOrAdd(InRuntimeObject);
-
-	FProperty* Property = PropAndFunction.ResolvedProperty.GetValidProperty();
-	if (Property && Property->HasSetter())
-	{
-		Property->CallSetter(&InRuntimeObject, &PropertyValue);
-	}
-	else if (UFunction* SetterFunction = PropAndFunction.SetterFunction.Get())
-	{
-		InvokeSetterFunction(&InRuntimeObject, SetterFunction, PropertyValue);
-	}
-	else if (Property)
-	{
-		if (FEnumProperty* EnumProperty = CastField<FEnumProperty>(Property))
-		{
-			FNumericProperty* UnderlyingProperty = EnumProperty->GetUnderlyingProperty();
-			void* ValueAddr = EnumProperty->ContainerPtrToValuePtr<void>(PropAndFunction.ResolvedProperty.ContainerAddress);
-			UnderlyingProperty->SetIntPropertyValue(ValueAddr, PropertyValue);
-		}
-		else
-		{
-			UE_LOG(LogMovieScene, Error, TEXT("Mismatch in property evaluation. %s is not of type: %s"), *Property->GetName(), *FEnumProperty::StaticClass()->GetName());
-		}
-	}
-
-	if (UFunction* NotifyFunction = PropAndFunction.NotifyFunction.Get())
-	{
-		InRuntimeObject.ProcessEvent(NotifyFunction, nullptr);
-	}
-}
-
-void FTrackInstancePropertyBindings::CacheBinding(const UObject& Object)
-{
-	FResolvedPropertyAndFunction PropAndFunction;
-	{
-		PropAndFunction.ResolvedProperty = ResolveProperty(PropertyPath, Object);
-
-		UFunction* SetterFunction = Object.FindFunction(FunctionName);
-		if (SetterFunction && SetterFunction->NumParms >= 1)
-		{
-			PropAndFunction.SetterFunction = SetterFunction;
-		}
-
-		UFunction* NotifyFunction = NotifyFunctionName != NAME_None ? Object.FindFunction(NotifyFunctionName) : nullptr;
-		if (NotifyFunction && NotifyFunction->NumParms == 0 && NotifyFunction->ReturnValueOffset == MAX_uint16)
-		{
-			PropAndFunction.NotifyFunction = NotifyFunction;
-		}
-	}
-
-	RuntimeObjectToFunctionMap.Add(FObjectKey(&Object), PropAndFunction);
-}
-
-FProperty* FTrackInstancePropertyBindings::GetProperty(const UObject& Object)
-{
-	FResolvedPropertyAndFunction PropAndFunction = FindOrAdd(Object);
-	return PropAndFunction.ResolvedProperty.GetValidProperty();
-}
-
-bool FTrackInstancePropertyBindings::HasValidBinding(const UObject& Object)
-{
-	FResolvedPropertyAndFunction PropAndFunction = FindOrAdd(Object);
-	return PropAndFunction.ResolvedProperty.GetValidProperty() != nullptr;
-}
-
-const UStruct* FTrackInstancePropertyBindings::GetPropertyStruct(const UObject& Object)
-{
-	FResolvedPropertyAndFunction PropAndFunction = FindOrAdd(Object);
-
-	if (FProperty* Property = PropAndFunction.ResolvedProperty.GetValidProperty())
-	{
-		if (FStructProperty* StructProperty = CastField<FStructProperty>(Property))
-		{
-			return StructProperty->Struct;
-		}
-		return nullptr;
-	}
-
-	return nullptr;
-}
-
-int64 FTrackInstancePropertyBindings::GetCurrentValueForEnum(const UObject& Object)
-{
-	FResolvedPropertyAndFunction PropAndFunction = FindOrAdd(Object);
-
-	FProperty* Property = PropAndFunction.ResolvedProperty.GetValidProperty();
-
-	if (Property)
-	{
-		if (FEnumProperty* EnumProperty = CastField<FEnumProperty>(Property))
-		{
-			FNumericProperty* UnderlyingProperty = EnumProperty->GetUnderlyingProperty();
-			void* ValueAddr = EnumProperty->ContainerPtrToValuePtr<void>(PropAndFunction.ResolvedProperty.ContainerAddress);
-			int64 Result = UnderlyingProperty->GetSignedIntPropertyValue(ValueAddr);
-			return Result;
-		}
-		else
-		{
-			UE_LOG(LogMovieScene, Error, TEXT("Mismatch in property evaluation. %s is not of type: %s"), *Property->GetName(), *FEnumProperty::StaticClass()->GetName());
-		}
-	}
-
-	return 0;
-}
-
-template<> void FTrackInstancePropertyBindings::CallFunction<bool>(UObject& InRuntimeObject, TCallTraits<bool>::ParamType PropertyValue)
-{
-	FResolvedPropertyAndFunction PropAndFunction = FindOrAdd(InRuntimeObject);
-
-	FProperty* Property = PropAndFunction.ResolvedProperty.GetValidProperty();
-
-	if (Property && Property->HasSetter())
-	{
-		Property->CallSetter(&InRuntimeObject, &PropertyValue);
-	}
-	else if (UFunction* SetterFunction = PropAndFunction.SetterFunction.Get())
-	{
-		InvokeSetterFunction(&InRuntimeObject, SetterFunction, PropertyValue);
-	}
-	else if (Property)
-	{
-		if (FBoolProperty* BoolProperty = CastField<FBoolProperty>(Property))
-		{
-			uint8* ValuePtr = BoolProperty->ContainerPtrToValuePtr<uint8>(PropAndFunction.ResolvedProperty.ContainerAddress);
-			BoolProperty->SetPropertyValue(ValuePtr, PropertyValue);
-		}
-		else
-		{
-			UE_LOG(LogMovieScene, Error, TEXT("Mismatch in property evaluation. %s is not of type: %s"), *Property->GetName(), *FBoolProperty::StaticClass()->GetName());
-		}
-	}
-
-	if (UFunction* NotifyFunction = PropAndFunction.NotifyFunction.Get())
-	{
-		InRuntimeObject.ProcessEvent(NotifyFunction, nullptr);
-	}
-}
-
-template<> bool FTrackInstancePropertyBindings::TryGetPropertyValue<bool>(const FResolvedProperty& ResolvedProperty, bool& OutValue)
-{
-	if (FProperty* Property = ResolvedProperty.GetValidProperty())
-	{
-		if (FBoolProperty* BoolProperty = CastField<FBoolProperty>(Property))
-		{
-			const uint8* ValuePtr = BoolProperty->ContainerPtrToValuePtr<uint8>(ResolvedProperty.ContainerAddress);
-			OutValue = BoolProperty->GetPropertyValue(ValuePtr);
-			return true;
-		}
-		else
-		{
-			UE_LOG(LogMovieScene, Error, TEXT("Mismatch in property evaluation. %s is not of type: %s"), *Property->GetName(), *FBoolProperty::StaticClass()->GetName());
-		}
-	}
-
-	return false;
-}
-
-template<> void FTrackInstancePropertyBindings::SetCurrentValue<bool>(UObject& Object, TCallTraits<bool>::ParamType InValue)
-{
-	FResolvedPropertyAndFunction PropAndFunction = FindOrAdd(Object);
-
-	if (FProperty* Property = PropAndFunction.ResolvedProperty.GetValidProperty())
-	{
-		if (FBoolProperty* BoolProperty = CastField<FBoolProperty>(Property))
-		{
-			uint8* ValuePtr = BoolProperty->ContainerPtrToValuePtr<uint8>(PropAndFunction.ResolvedProperty.ContainerAddress);
-			BoolProperty->SetPropertyValue(ValuePtr, InValue);
-		}
-	}
-
-	if (UFunction* NotifyFunction = PropAndFunction.NotifyFunction.Get())
-	{
-		Object.ProcessEvent(NotifyFunction, nullptr);
-	}
-}
-
-template<> void FTrackInstancePropertyBindings::CallFunction<UObject*>(UObject& InRuntimeObject, UObject* PropertyValue)
-{
-	FResolvedPropertyAndFunction PropAndFunction = FindOrAdd(InRuntimeObject);
-
-	FProperty* Property = PropAndFunction.ResolvedProperty.GetValidProperty();
-
-	if (Property && Property->HasSetter())
-	{
-		Property->CallSetter(&InRuntimeObject, &PropertyValue);
-	}
-	else if (UFunction* SetterFunction = PropAndFunction.SetterFunction.Get())
-	{
-		InvokeSetterFunction(&InRuntimeObject, SetterFunction, PropertyValue);
-	}
-	else if (Property)
-	{
-		if (FObjectPropertyBase* ObjectProperty = CastField<FObjectPropertyBase>(Property))
-		{
-			uint8* ValuePtr = ObjectProperty->ContainerPtrToValuePtr<uint8>(PropAndFunction.ResolvedProperty.ContainerAddress);
-			ObjectProperty->SetObjectPropertyValue(ValuePtr, PropertyValue);
-		}
-		else
-		{
-			UE_LOG(LogMovieScene, Error, TEXT("Mismatch in property evaluation. %s is not of type: %s"), *Property->GetName(), *FObjectPropertyBase::StaticClass()->GetName());
-		}
-	}
-
-	if (UFunction* NotifyFunction = PropAndFunction.NotifyFunction.Get())
-	{
-		InRuntimeObject.ProcessEvent(NotifyFunction, nullptr);
-	}
-}
-
-template<> bool FTrackInstancePropertyBindings::TryGetPropertyValue<UObject*>(const FResolvedProperty& ResolvedProperty, UObject*& OutValue)
-{
-	if (FProperty* Property = ResolvedProperty.GetValidProperty())
-	{
-		if (FObjectPropertyBase* ObjectProperty = CastField<FObjectPropertyBase>(ResolvedProperty.GetValidProperty()))
-		{
-			const uint8* ValuePtr = ObjectProperty->ContainerPtrToValuePtr<uint8>(ResolvedProperty.ContainerAddress);
-			OutValue = ObjectProperty->GetObjectPropertyValue(ValuePtr);
-			return true;
-		}
-		else
-		{
-			UE_LOG(LogMovieScene, Error, TEXT("Mismatch in property evaluation. %s is not of type: %s"), *Property->GetName(), *FObjectPropertyBase::StaticClass()->GetName());
-		}
-	}
-
-	return false;
-}
-
-template<> void FTrackInstancePropertyBindings::SetCurrentValue<UObject*>(UObject& Object, UObject* InValue)
-{
-	FResolvedPropertyAndFunction PropAndFunction = FindOrAdd(Object);
-
-	if (FProperty* Property = PropAndFunction.ResolvedProperty.GetValidProperty())
-	{
-		if (FObjectPropertyBase* ObjectProperty = CastField<FObjectPropertyBase>(Property))
-		{
-			uint8* ValuePtr = ObjectProperty->ContainerPtrToValuePtr<uint8>(PropAndFunction.ResolvedProperty.ContainerAddress);
-			ObjectProperty->SetObjectPropertyValue(ValuePtr, InValue);
-		}
-	}
-
-	if (UFunction* NotifyFunction = PropAndFunction.NotifyFunction.Get())
-	{
-		Object.ProcessEvent(NotifyFunction, nullptr);
-	}
-}
-
diff --git a/Engine/Source/Runtime/MovieScene/Private/Tests/MovieSceneTestObjects.h b/Engine/Source/Runtime/MovieScene/Private/Tests/MovieSceneTestObjects.h
--- a/Engine/Source/Runtime/MovieScene/Private/Tests/MovieSceneTestObjects.h
+++ b/Engine/Source/Runtime/MovieScene/Private/Tests/MovieSceneTestObjects.h
@@ -2,16 +2,17 @@
 
 #pragma once
 
+#include "Compilation/IMovieSceneTrackTemplateProducer.h"
+#include "Evaluation/MovieSceneEvalTemplate.h"
+#include "GameFramework/Actor.h"
 #include "MovieScene.h"
-#include "MovieSceneTrack.h"
 #include "MovieSceneSection.h"
 #include "MovieSceneSequence.h"
-#include "Evaluation/MovieSceneEvalTemplate.h"
-#include "GameFramework/Actor.h"
-#include "Sections/MovieSceneSubSection.h"
+#include "MovieSceneTrack.h"
 #include "Sections/MovieSceneHookSection.h"
+#include "Sections/MovieSceneSubSection.h"
+#include "StructUtils/InstancedStruct.h"
 #include "Tracks/MovieSceneSubTrack.h"
-#include "Compilation/IMovieSceneTrackTemplateProducer.h"
 
 #include "MovieSceneTestObjects.generated.h"
 
@@ -159,6 +160,14 @@ public:
 	}
 };
 
+UENUM(BlueprintType)
+enum class ETestMovieSceneEnum : uint8
+{
+	One,
+	Two,
+	Three
+};
+
 USTRUCT(BlueprintType)
 struct FTestMovieSceneStruct
 {
@@ -170,8 +179,35 @@ struct FTestMovieSceneStruct
 	UPROPERTY(EditAnywhere, BlueprintReadWrite, Interp, Category="Test")
 	float Second = 0.f;
 
+	UPROPERTY(EditAnywhere, BlueprintReadWrite, Interp, Category="Test")
+	ETestMovieSceneEnum Enum = ETestMovieSceneEnum::One;
+
 	UPROPERTY(EditAnywhere, BlueprintReadWrite, Interp, Category="Test")
 	FVector Vector = FVector::ZeroVector;
+
+	UPROPERTY(EditAnywhere, BlueprintReadWrite, Interp, Category="Test")
+	TArray<int32> MultipleIntegers;
+
+	UPROPERTY(EditAnywhere, BlueprintReadWrite, Interp, Category="Test")
+	TArray<FVector> MultipleVectors;
+};
+
+USTRUCT(BlueprintType)
+struct FTestMovieSceneStruct2
+{
+	GENERATED_BODY()
+
+	UPROPERTY(EditAnywhere, BlueprintReadWrite, Interp, Category="Test")
+	float First = 0.f;
+
+	UPROPERTY(EditAnywhere, BlueprintReadWrite, Interp, Category="Test")
+	float Second = 0.f;
+};
+
+UCLASS(MinimalAPI, BlueprintType)
+class UTestMovieSceneObject : public UObject
+{
+	GENERATED_BODY()
 };
 
 UCLASS(MinimalAPI, BlueprintType)
@@ -182,7 +218,19 @@ class ATestMovieSceneArrayPropertiesActor : public AActor
 public:
 
 	UPROPERTY(EditAnywhere, BlueprintReadWrite, Interp, Category="Test")
-	FVector TestVector;
+	bool TestBool = false;
+
+	UPROPERTY(EditAnywhere, BlueprintReadWrite, Interp, Category="Test")
+	ETestMovieSceneEnum TestEnum = ETestMovieSceneEnum::One;
+
+	UPROPERTY(EditAnywhere, BlueprintReadWrite, Interp, Category="Test")
+	int32 TestInt32 = 0;
+
+	UPROPERTY(EditAnywhere, BlueprintReadWrite, Interp, Category="Test")
+	TObjectPtr<UTestMovieSceneObject> TestObject;
+
+	UPROPERTY(EditAnywhere, BlueprintReadWrite, Interp, Category="Test")
+	FVector TestVector = FVector::ZeroVector;
 
 	UPROPERTY(EditAnywhere, BlueprintReadWrite, Interp, Category="Test")
 	TArray<float> MultipleFloats;
@@ -192,9 +240,31 @@ public:
 
 	UPROPERTY(EditAnywhere, BlueprintReadWrite, Interp, Category="Test")
 	TArray<FTestMovieSceneStruct> MultipleStructs;
+	
+	UPROPERTY(EditAnywhere, BlueprintReadWrite, Interp, Category="Test")
+	FInstancedStruct SingleInstancedStruct;
+
+	UPROPERTY(EditAnywhere, BlueprintReadWrite, Interp, Category="Test")
+	TArray<FInstancedStruct> MultipleInstancedStructs;
+
+	UPROPERTY(EditAnywhere, BlueprintReadWrite, BlueprintSetter="SetTestSetterFloat", Interp, Category="Test")
+	float TestSetterFloat = 0;
+
+public:
+
+	int32 NumSetTestSetterFloatCalls = 0;
+
+public:
 
 	ATestMovieSceneArrayPropertiesActor(const FObjectInitializer& ObjInit)
 		: AActor(ObjInit)
 	{}
+
+	UFUNCTION(BlueprintCallable, Category="Test")
+	void SetTestSetterFloat(float InValue)
+	{
+		TestSetterFloat = InValue;
+		++NumSetTestSetterFloatCalls;
+	}
 };
 
diff --git a/Engine/Source/Runtime/MovieScene/Private/Tests/TrackInstancePropertyBindingsTests.cpp b/Engine/Source/Runtime/MovieScene/Private/Tests/TrackInstancePropertyBindingsTests.cpp
new file mode 100644
--- /dev/null
+++ b/Engine/Source/Runtime/MovieScene/Private/Tests/TrackInstancePropertyBindingsTests.cpp
@@ -0,0 +1,370 @@
+// Copyright Epic Games, Inc. All Rights Reserved.
+
+#include "CoreTypes.h"
+#include "MovieSceneFwd.h"
+#include "Misc/AutomationTest.h"
+#include "MovieSceneTestObjects.h"
+#include "TrackInstancePropertyBindings.h"
+
+#if WITH_DEV_AUTOMATION_TESTS
+
+#define UTEST_CACHED_BINDING(Message, Bindings, Object)\
+	{\
+		FTrackInstancePropertyBindings::FResolvedPropertyAndFunction PropAndFunction = Bindings.FindOrAdd(Object);\
+		UTEST_EQUAL(Message, PropAndFunction.ResolvedProperty.IsType<FTrackInstancePropertyBindings::FCachedProperty>(), true);\
+	}
+
+#define UTEST_VOLATILE_BINDING(Message, Bindings, Object)\
+	{\
+		FTrackInstancePropertyBindings::FResolvedPropertyAndFunction PropAndFunction = Bindings.FindOrAdd(Object);\
+		UTEST_EQUAL(Message, PropAndFunction.ResolvedProperty.IsType<FTrackInstancePropertyBindings::FVolatileProperty>(), true);\
+	}
+
+IMPLEMENT_SIMPLE_AUTOMATION_TEST(FTrackInstancePropertyBindingsTests, 
+		"System.Engine.Sequencer.TrackInstancePropertyBindings", 
+		EAutomationTestFlags::EditorContext | EAutomationTestFlags::EngineFilter)
+bool FTrackInstancePropertyBindingsTests::RunTest(const FString& Parameters)
+{
+	using namespace UE::MovieScene;
+
+	// TestBool property
+	{
+		ATestMovieSceneArrayPropertiesActor* TestActor = NewObject<ATestMovieSceneArrayPropertiesActor>();
+
+		FTrackInstancePropertyBindings Bindings(TEXT("TestBool"), TEXT("TestBool"));
+		UTEST_EQUAL("", TestActor->TestBool, false);
+		Bindings.CallFunction<bool>(*TestActor, true);
+		UTEST_EQUAL("", TestActor->TestBool, true);
+		UTEST_CACHED_BINDING("", Bindings, *TestActor);
+	}
+
+	// TestEnum property
+	{
+		ATestMovieSceneArrayPropertiesActor* TestActor = NewObject<ATestMovieSceneArrayPropertiesActor>();
+
+		FTrackInstancePropertyBindings Bindings(TEXT("TestEnum"), TEXT("TestEnum"));
+		UTEST_EQUAL("", TestActor->TestEnum, ETestMovieSceneEnum::One);
+		Bindings.CallFunctionForEnum(*TestActor, (int64)ETestMovieSceneEnum::Two);
+		UTEST_EQUAL("", TestActor->TestEnum, ETestMovieSceneEnum::Two);
+		UTEST_CACHED_BINDING("", Bindings, *TestActor);
+	}
+
+	// TestInt32 property
+	{
+		ATestMovieSceneArrayPropertiesActor* TestActor = NewObject<ATestMovieSceneArrayPropertiesActor>();
+
+		FTrackInstancePropertyBindings Bindings(TEXT("TestInt32"), TEXT("TestInt32"));
+		UTEST_EQUAL("", TestActor->TestInt32, 0);
+		Bindings.CallFunction<int32>(*TestActor, 42);
+		UTEST_EQUAL("", TestActor->TestInt32, 42);
+		UTEST_CACHED_BINDING("", Bindings, *TestActor);
+	}
+
+	// TestObject property
+	{
+		ATestMovieSceneArrayPropertiesActor* TestActor = NewObject<ATestMovieSceneArrayPropertiesActor>();
+
+		UTestMovieSceneObject* TestObject = NewObject<UTestMovieSceneObject>();
+
+		FTrackInstancePropertyBindings Bindings(TEXT("TestObject"), TEXT("TestObject"));
+		UTEST_EQUAL("", TestActor->TestObject.Get(), (UTestMovieSceneObject*)nullptr);
+		Bindings.CallFunction<UObject*>(*TestActor, (UObject*)TestObject);
+		UTEST_EQUAL("", TestActor->TestObject.Get(), TestObject);
+		UTEST_CACHED_BINDING("", Bindings, *TestActor);
+	}
+
+	// TestVector.X/Y property
+	{
+		ATestMovieSceneArrayPropertiesActor* TestActor = NewObject<ATestMovieSceneArrayPropertiesActor>();
+
+		FTrackInstancePropertyBindings BindingsX(TEXT("X"), TEXT("TestVector.X"));
+		UTEST_EQUAL("", TestActor->TestVector, FVector(0, 0, 0));
+		BindingsX.CallFunction<double>(*TestActor, 1.0);
+		UTEST_EQUAL("", TestActor->TestVector, FVector(1, 0, 0));
+		UTEST_CACHED_BINDING("", BindingsX, *TestActor);
+
+		FTrackInstancePropertyBindings BindingsY(TEXT("Y"), TEXT("TestVector.Y"));
+		UTEST_EQUAL("", TestActor->TestVector, FVector(1, 0, 0));
+		BindingsY.CallFunction<double>(*TestActor, 2.0);
+		UTEST_EQUAL("", TestActor->TestVector, FVector(1, 2, 0));
+		UTEST_CACHED_BINDING("", BindingsY, *TestActor);
+	}
+
+	// MultipleFloats property
+	{
+		ATestMovieSceneArrayPropertiesActor* TestActor = NewObject<ATestMovieSceneArrayPropertiesActor>();
+
+		TestActor->MultipleFloats.Add(0);
+		TestActor->MultipleFloats.Add(1);
+
+		FTrackInstancePropertyBindings Bindings(TEXT("MultipleFloats[1]"), TEXT("MultipleFloats[1]"));
+		UTEST_EQUAL("", TestActor->MultipleFloats[1], 1.f);
+		Bindings.CallFunction<float>(*TestActor, 11.f);
+		UTEST_EQUAL("", TestActor->MultipleFloats[1], 11.f);
+		UTEST_VOLATILE_BINDING("", Bindings, *TestActor);
+	}
+
+	// SingleStruct.Second property
+	{
+		ATestMovieSceneArrayPropertiesActor* TestActor = NewObject<ATestMovieSceneArrayPropertiesActor>();
+
+		FTrackInstancePropertyBindings Bindings(TEXT("Second"), TEXT("SingleStruct.Second"));
+		UTEST_EQUAL("", TestActor->SingleStruct.Second, 0.f);
+		Bindings.CallFunction<float>(*TestActor, 1.f);
+		UTEST_EQUAL("", TestActor->SingleStruct.Second, 1.f);
+		UTEST_CACHED_BINDING("", Bindings, *TestActor);
+	}
+
+	// SingleStruct.MultipleIntegers property
+	{
+		ATestMovieSceneArrayPropertiesActor* TestActor = NewObject<ATestMovieSceneArrayPropertiesActor>();
+
+		TestActor->SingleStruct.MultipleIntegers.Add(2);
+		TestActor->SingleStruct.MultipleIntegers.Add(3);
+
+		FTrackInstancePropertyBindings Bindings(TEXT("MultipleIntegers[1]"), TEXT("SingleStruct.MultipleIntegers[1]"));
+		UTEST_EQUAL("", TestActor->SingleStruct.MultipleIntegers[1], 3);
+		Bindings.CallFunction<int32>(*TestActor, 4);
+		UTEST_EQUAL("", TestActor->SingleStruct.MultipleIntegers[1], 4);
+		UTEST_VOLATILE_BINDING("", Bindings, *TestActor);
+	}
+
+	// SingleStruct.MultipleVectors.X property
+	{
+		ATestMovieSceneArrayPropertiesActor* TestActor = NewObject<ATestMovieSceneArrayPropertiesActor>();
+
+		TestActor->SingleStruct.MultipleVectors.Add(FVector(2, 3, 4));
+		TestActor->SingleStruct.MultipleVectors.Add(FVector(5, 6, 7));
+
+		FTrackInstancePropertyBindings Bindings(TEXT("X"), TEXT("SingleStruct.MultipleVectors[1].X"));
+		UTEST_EQUAL("", TestActor->SingleStruct.MultipleVectors[1].X, 5.0);
+		Bindings.CallFunction<double>(*TestActor, 10.0);
+		UTEST_EQUAL("", TestActor->SingleStruct.MultipleVectors[1].X, 10.0);
+		UTEST_VOLATILE_BINDING("", Bindings, *TestActor);
+	}
+
+	// MultipleStruct.Second property
+	{
+		ATestMovieSceneArrayPropertiesActor* TestActor = NewObject<ATestMovieSceneArrayPropertiesActor>();
+
+		TestActor->MultipleStructs.Reserve(2);
+		TestActor->MultipleStructs.Add(FTestMovieSceneStruct());
+		TestActor->MultipleStructs.Add(FTestMovieSceneStruct{ 3.f, 5.f });
+
+		FTrackInstancePropertyBindings Bindings(TEXT("Second"), TEXT("MultipleStructs[1].Second"));
+		UTEST_EQUAL("", TestActor->MultipleStructs[1].Second, 5.f);
+		Bindings.CallFunction<float>(*TestActor, 10.f);
+		UTEST_EQUAL("", TestActor->MultipleStructs[1].Second, 10.f);
+		UTEST_VOLATILE_BINDING("", Bindings, *TestActor);
+
+		const FTrackInstancePropertyBindings::FResolvedPropertyAndFunction& PropAndFunction = Bindings.FindOrAdd(*TestActor);
+		void* OldContainerAddress = PropAndFunction.GetContainerAddress();
+		UTEST_EQUAL("", OldContainerAddress, (void*)&TestActor->MultipleStructs[1]);
+
+		// Force re-allocating the array and check that we haven't cached the old addresses.
+		TestActor->MultipleStructs.Reserve(4);
+		UTEST_NOT_EQUAL("", OldContainerAddress, (void*)&TestActor->MultipleStructs[1]); 
+		void* NewContainerAddress = PropAndFunction.GetContainerAddress();
+		UTEST_EQUAL("", NewContainerAddress, (void*)&TestActor->MultipleStructs[1]);
+	}
+
+	// MultipleStruct.Enum property
+	{
+		ATestMovieSceneArrayPropertiesActor* TestActor = NewObject<ATestMovieSceneArrayPropertiesActor>();
+
+		TestActor->MultipleStructs.Reserve(2);
+		TestActor->MultipleStructs.Add(FTestMovieSceneStruct());
+		TestActor->MultipleStructs.Add(FTestMovieSceneStruct());
+
+		FTrackInstancePropertyBindings Bindings(TEXT("Enum"), TEXT("MultipleStructs[1].Enum"));
+		UTEST_EQUAL("", TestActor->MultipleStructs[1].Enum, ETestMovieSceneEnum::One);
+		Bindings.CallFunctionForEnum(*TestActor, (int64)ETestMovieSceneEnum::Two);
+		UTEST_EQUAL("", TestActor->MultipleStructs[1].Enum, ETestMovieSceneEnum::Two);
+		UTEST_VOLATILE_BINDING("", Bindings, *TestActor);
+
+		const FTrackInstancePropertyBindings::FResolvedPropertyAndFunction& PropAndFunction = Bindings.FindOrAdd(*TestActor);
+		void* OldContainerAddress = PropAndFunction.GetContainerAddress();
+		UTEST_EQUAL("", OldContainerAddress, (void*)&TestActor->MultipleStructs[1]);
+
+		// Force re-allocating the array and check that we haven't cached the old addresses.
+		TestActor->MultipleStructs.Reserve(4);
+		UTEST_NOT_EQUAL("", OldContainerAddress, (void*)&TestActor->MultipleStructs[1]); 
+		void* NewContainerAddress = PropAndFunction.GetContainerAddress();
+		UTEST_EQUAL("", NewContainerAddress, (void*)&TestActor->MultipleStructs[1]);
+	}
+
+	// MultipleStruct.MultipleIntegers property
+	{
+		ATestMovieSceneArrayPropertiesActor* TestActor = NewObject<ATestMovieSceneArrayPropertiesActor>();
+
+		TestActor->MultipleStructs.Reserve(2);
+		TestActor->MultipleStructs.Add(FTestMovieSceneStruct());
+		TestActor->MultipleStructs.Add(FTestMovieSceneStruct());
+		TestActor->MultipleStructs[1].MultipleIntegers.Add(7);
+		TestActor->MultipleStructs[1].MultipleIntegers.Add(8);
+
+		FTrackInstancePropertyBindings Bindings(TEXT("MultipleIntegers[1]"), TEXT("MultipleStructs[1].MultipleIntegers[1]"));
+		UTEST_EQUAL("", TestActor->MultipleStructs[1].MultipleIntegers[1], 8);
+		Bindings.CallFunction<int32>(*TestActor, 10);
+		UTEST_EQUAL("", TestActor->MultipleStructs[1].MultipleIntegers[1], 10);
+		UTEST_VOLATILE_BINDING("", Bindings, *TestActor);
+	}
+
+	// MultipleStruct.MultipleVectors.X property
+	{
+		ATestMovieSceneArrayPropertiesActor* TestActor = NewObject<ATestMovieSceneArrayPropertiesActor>();
+
+		TestActor->MultipleStructs.Reserve(2);
+		TestActor->MultipleStructs.Add(FTestMovieSceneStruct());
+		TestActor->MultipleStructs.Add(FTestMovieSceneStruct{ 3.f, 5.f });
+		TestActor->MultipleStructs[1].MultipleVectors.Add(FVector(11, 12, 13));
+		TestActor->MultipleStructs[1].MultipleVectors.Add(FVector(14, 15, 16));
+
+		FTrackInstancePropertyBindings Bindings(TEXT("X"), TEXT("MultipleStructs[1].MultipleVectors[1].X"));
+		UTEST_EQUAL("", TestActor->MultipleStructs[1].MultipleVectors[1].X, 14.0);
+		Bindings.CallFunction<double>(*TestActor, 20.0);
+		UTEST_EQUAL("", TestActor->MultipleStructs[1].MultipleVectors[1].X, 20.0);
+		UTEST_VOLATILE_BINDING("", Bindings, *TestActor);
+	}
+
+	// SingleInstancedStruct.Second property
+	{
+		ATestMovieSceneArrayPropertiesActor* TestActor = NewObject<ATestMovieSceneArrayPropertiesActor>();
+
+		FTestMovieSceneStruct& TestStruct = TestActor->SingleInstancedStruct.InitializeAs<FTestMovieSceneStruct>();
+		TestStruct.First = 2.f;
+		TestStruct.Second = 4.f;
+
+		FTrackInstancePropertyBindings Bindings(TEXT("Second"), TEXT("SingleInstancedStruct.Second"));
+		UTEST_EQUAL("", TestStruct.Second, 4.f);
+		Bindings.CallFunction<float>(*TestActor, 8.f);
+		UTEST_EQUAL("", TestStruct.Second, 8.f);
+		UTEST_VOLATILE_BINDING("", Bindings, *TestActor);
+
+		FMovieSceneEasingSettings& Easing = TestActor->SingleInstancedStruct.InitializeAs<FMovieSceneEasingSettings>();
+		Easing.AutoEaseInDuration = 5000;
+		Easing.AutoEaseOutDuration = 5000;
+		Bindings.CallFunction<float>(*TestActor, 8.f);  // The struct is different so it should now fail.
+		// Everything else should be untouched (this checks for memory corruption)
+		UTEST_EQUAL("", Easing.AutoEaseInDuration, 5000);
+		UTEST_EQUAL("", Easing.AutoEaseOutDuration, 5000);
+		UTEST_EQUAL("", Easing.EaseIn, TScriptInterface<IMovieSceneEasingFunction>());
+		UTEST_EQUAL("", Easing.bManualEaseIn, false);
+		UTEST_EQUAL("", Easing.ManualEaseInDuration, 0);
+		UTEST_EQUAL("", Easing.EaseOut, TScriptInterface<IMovieSceneEasingFunction>());
+		UTEST_EQUAL("", Easing.bManualEaseOut, false);
+		UTEST_EQUAL("", Easing.ManualEaseOutDuration, 0);
+	}
+
+	// SingleInstancedStruct.Second property
+	{
+		ATestMovieSceneArrayPropertiesActor* TestActor = NewObject<ATestMovieSceneArrayPropertiesActor>();
+
+		FTestMovieSceneStruct& TestStruct = TestActor->SingleInstancedStruct.InitializeAs<FTestMovieSceneStruct>();
+		TestStruct.First = 2.f;
+		TestStruct.Second = 4.f;
+
+		FTrackInstancePropertyBindings Bindings(TEXT("Second"), TEXT("SingleInstancedStruct.Second"));
+		UTEST_EQUAL("", TestStruct.Second, 4.f);
+		Bindings.CallFunction<float>(*TestActor, 8.f);
+		UTEST_EQUAL("", TestStruct.Second, 8.f);
+		UTEST_VOLATILE_BINDING("", Bindings, *TestActor);
+
+		FTestMovieSceneStruct2& TestStruct2 = TestActor->SingleInstancedStruct.InitializeAs<FTestMovieSceneStruct2>();
+		TestStruct2.Second = 12.f;
+		Bindings.CallFunction<float>(*TestActor, 18.f);  // The struct is different but duck typing should let us key this.
+		UTEST_EQUAL("", TestStruct2.Second, 18.f);
+		UTEST_VOLATILE_BINDING("", Bindings, *TestActor);
+	}
+
+	// SingleInstancedStruct.MultipleIntegers property
+	{
+		ATestMovieSceneArrayPropertiesActor* TestActor = NewObject<ATestMovieSceneArrayPropertiesActor>();
+
+		FTestMovieSceneStruct& TestStruct = TestActor->SingleInstancedStruct.InitializeAs<FTestMovieSceneStruct>();
+		TestStruct.MultipleIntegers.Add(3);
+		TestStruct.MultipleIntegers.Add(4);
+
+		FTrackInstancePropertyBindings Bindings(TEXT("MultipleIntegers[1]"), TEXT("SingleInstancedStruct.MultipleIntegers[1]"));
+		UTEST_EQUAL("", TestStruct.MultipleIntegers[1], 4);
+		Bindings.CallFunction<int32>(*TestActor, 8);
+		UTEST_EQUAL("", TestStruct.MultipleIntegers[1], 8);
+		UTEST_VOLATILE_BINDING("", Bindings, *TestActor);
+	}
+
+	// SingleInstancedStruct.MultipleVectors.X property
+	{
+		ATestMovieSceneArrayPropertiesActor* TestActor = NewObject<ATestMovieSceneArrayPropertiesActor>();
+
+		FTestMovieSceneStruct& TestStruct = TestActor->SingleInstancedStruct.InitializeAs<FTestMovieSceneStruct>();
+		TestStruct.MultipleVectors.Add(FVector(2, 3, 4));
+		TestStruct.MultipleVectors.Add(FVector(5, 6, 7));
+
+		FTrackInstancePropertyBindings Bindings(TEXT("X"), TEXT("SingleInstancedStruct.MultipleVectors[1].X"));
+		UTEST_EQUAL("", TestStruct.MultipleVectors[1].X, 5.0);
+		Bindings.CallFunction<double>(*TestActor, 10.0);
+		UTEST_EQUAL("", TestStruct.MultipleVectors[1].X, 10.0);
+		UTEST_VOLATILE_BINDING("", Bindings, *TestActor);
+	}
+
+	// MultipleInstancedStructs.Second property
+	{
+		ATestMovieSceneArrayPropertiesActor* TestActor = NewObject<ATestMovieSceneArrayPropertiesActor>();
+
+		TestActor->MultipleInstancedStructs.Reserve(2);
+		TestActor->MultipleInstancedStructs.Emplace();
+		TestActor->MultipleInstancedStructs.Emplace();
+
+		FTestMovieSceneStruct& TestStruct = TestActor->MultipleInstancedStructs[1].InitializeAs<FTestMovieSceneStruct>();
+		TestStruct.First = 5.f;
+		TestStruct.Second = 9.f;
+
+		FTrackInstancePropertyBindings Bindings(TEXT("Second"), TEXT("MultipleInstancedStructs[1].Second"));
+		UTEST_EQUAL("", TestStruct.Second, 9.f);
+		Bindings.CallFunction<float>(*TestActor, 12.f);
+		UTEST_EQUAL("", TestStruct.Second, 12.f);
+		UTEST_VOLATILE_BINDING("", Bindings, *TestActor);
+	}
+
+	// MultipleInstancedStructs.MultipleIntegers property
+	{
+		ATestMovieSceneArrayPropertiesActor* TestActor = NewObject<ATestMovieSceneArrayPropertiesActor>();
+
+		TestActor->MultipleInstancedStructs.Reserve(2);
+		TestActor->MultipleInstancedStructs.Emplace();
+		TestActor->MultipleInstancedStructs.Emplace();
+
+		FTestMovieSceneStruct& TestStruct = TestActor->MultipleInstancedStructs[1].InitializeAs<FTestMovieSceneStruct>();
+		TestStruct.MultipleIntegers.Add(8);
+		TestStruct.MultipleIntegers.Add(9);
+
+		FTrackInstancePropertyBindings Bindings(TEXT("MultipleIntegers[1]"), TEXT("MultipleInstancedStructs[1].MultipleIntegers[1]"));
+		UTEST_EQUAL("", TestStruct.MultipleIntegers[1], 9);
+		Bindings.CallFunction<int32>(*TestActor, 12);
+		UTEST_EQUAL("", TestStruct.MultipleIntegers[1], 12);
+		UTEST_VOLATILE_BINDING("", Bindings, *TestActor);
+	}
+
+	// MultipleInstancedStructs.MultipleVectors.Y property
+	{
+		ATestMovieSceneArrayPropertiesActor* TestActor = NewObject<ATestMovieSceneArrayPropertiesActor>();
+
+		TestActor->MultipleInstancedStructs.Reserve(2);
+		TestActor->MultipleInstancedStructs.Emplace();
+		TestActor->MultipleInstancedStructs.Emplace();
+
+		FTestMovieSceneStruct& TestStruct = TestActor->MultipleInstancedStructs[1].InitializeAs<FTestMovieSceneStruct>();
+		TestStruct.MultipleVectors.Add(FVector(5, 4, 3));
+		TestStruct.MultipleVectors.Add(FVector(2, 1, 0));
+
+		FTrackInstancePropertyBindings Bindings(TEXT("Y"), TEXT("MultipleInstancedStructs[1].MultipleVectors[1].Y"));
+		UTEST_EQUAL("", TestStruct.MultipleVectors[1].Y, 1.0);
+		Bindings.CallFunction<double>(*TestActor, 11.0);
+		UTEST_EQUAL("", TestStruct.MultipleVectors[1].Y, 11.0);
+		UTEST_VOLATILE_BINDING("", Bindings, *TestActor);
+	}
+
+	return true;
+}
+
+#endif // WITH_DEV_AUTOMATION_TESTS
diff --git a/Engine/Source/Runtime/MovieScene/Private/TrackInstancePropertyBindings.cpp b/Engine/Source/Runtime/MovieScene/Private/TrackInstancePropertyBindings.cpp
new file mode 100644
--- /dev/null
+++ b/Engine/Source/Runtime/MovieScene/Private/TrackInstancePropertyBindings.cpp
@@ -0,0 +1,797 @@
+// Copyright Epic Games, Inc. All Rights Reserved.
+
+#include "TrackInstancePropertyBindings.h"
+
+#include "EntitySystem/MovieSceneIntermediatePropertyValue.h"
+#include "MovieSceneFwd.h"
+#include "String/ParseTokens.h"
+#include "StructUtils/InstancedStruct.h"
+#include "StructUtils/PropertyBag.h"
+
+namespace UE::MovieScene
+{
+
+struct FPropertyResolutionStep
+{
+	FProperty* Property = nullptr;
+	int32 ArrayIndex = INDEX_NONE;
+	void* ContainerAddress = nullptr;
+};
+
+struct FPropertyResolutionState
+{
+	TArray<FPropertyResolutionStep> PropertySteps;
+	bool bIsVolatile = false;
+	bool bIsValid = true;
+
+	const FPropertyResolutionStep* GetLastStep() const
+	{
+		if (!PropertySteps.IsEmpty())
+		{
+			return &PropertySteps.Last();
+		}
+		return nullptr;
+	}
+};
+
+FPropertyResolutionStep FindProperty(void* BasePointer, UStruct* InStruct, FStringView PropertyName)
+{
+	FPropertyResolutionStep PropertyStep;
+	PropertyStep.ContainerAddress = BasePointer;
+	PropertyStep.Property = FindFProperty<FProperty>(InStruct, FName(PropertyName, FNAME_Find));
+	return PropertyStep;
+}
+
+FPropertyResolutionStep FindPropertyAndArrayIndex(void* BasePointer, UStruct* InStruct, FStringView PropertyName)
+{
+	// Calculate the array index if possible.
+	int32 ArrayIndex = -1;
+	if (PropertyName.Len() > 0 && PropertyName[PropertyName.Len() - 1] == ']')
+	{
+		int32 OpenIndex = 0;
+		if (PropertyName.FindLastChar('[', OpenIndex))
+		{
+			// We have a property name of the form "Foo[123]". Resolve the property itself ("Foo") and then
+			// parse the array element index (123).
+			FStringView TruncatedPropertyName(PropertyName.GetData(), OpenIndex);
+			FPropertyResolutionStep PropertyStep = FindProperty(BasePointer, InStruct, TruncatedPropertyName);
+
+			const int32 NumberLength = PropertyName.Len() - OpenIndex - 2;
+			if (NumberLength > 0 && NumberLength <= 10)
+			{
+				TCHAR NumberBuffer[11];
+				FMemory::Memzero(NumberBuffer);
+				FMemory::Memcpy(NumberBuffer, &PropertyName[OpenIndex + 1], sizeof(TCHAR) * NumberLength);
+				LexFromString(PropertyStep.ArrayIndex, NumberBuffer);
+			}
+
+			return PropertyStep;
+		}
+	}
+
+	// No index found in this property name, just find the property normally.
+	return FindProperty(BasePointer, InStruct, PropertyName);
+}
+
+void ResolvePropertyRecursive(void* BasePointer, UStruct* InStruct, TArrayView<const FStringView> InPropertyNames, uint32 Index, FPropertyResolutionState& OutResolutionState)
+{
+	// If we need to resovle a property on an instanced struct or property bag, we need to first dive into the correct struct.
+	void* ActualBasePointer = BasePointer;
+	UStruct* ActualStruct = InStruct;
+	{
+		static FProperty* const PropertyBagValueProperty = FInstancedPropertyBag::StaticStruct()->FindPropertyByName(TEXT("Value"));
+
+		// Instanced structs are technically just a memory buffer with no real sub-properties, but
+		// they do have sub-properties if we ask them about their "logical" struct type. Let's do that,
+		// which makes it possible to animate the properties inside.
+		if (InStruct == FInstancedStruct::StaticStruct())
+		{
+			// Instanced structs may be reallocated, flag this property path as volatile.
+			OutResolutionState.bIsVolatile = true;
+
+			FInstancedStruct* InstancedStruct = (FInstancedStruct*)BasePointer;
+			const UScriptStruct* InstancedStructType = InstancedStruct->GetScriptStruct();
+			uint8* InstancedStructMemory = InstancedStruct->GetMutableMemory();
+
+			ActualBasePointer = (void*)InstancedStructMemory;
+			ActualStruct = const_cast<UScriptStruct*>(InstancedStructType);
+		}
+		// As above but for a property bag, unless the property path is specifically targeting its inner 
+		// instanced struct, in which case we don't need to do anything, we'll end up in the previous code
+		// block on the next iteration.
+		else if (InStruct == FInstancedPropertyBag::StaticStruct() && InPropertyNames[Index] != PropertyBagValueProperty->GetFName())
+		{
+			// Property bags may be reallocated, flag this property path as volatile.
+			OutResolutionState.bIsVolatile = true;
+
+			FInstancedPropertyBag* PropertyBag = (FInstancedPropertyBag*)BasePointer;
+
+			// Add a step for diving into the instanced struct, so we don't need to handle a special case later.
+			OutResolutionState.PropertySteps.Add(FPropertyResolutionStep{ PropertyBagValueProperty, INDEX_NONE, (void*)PropertyBag });
+
+			const UPropertyBag* PropertyBagType = PropertyBag->GetPropertyBagStruct();
+			uint8* PropertyBagMemory = PropertyBag->GetMutableValue().GetMemory();
+
+			ActualBasePointer = (void*)PropertyBagMemory;
+			ActualStruct = const_cast<UPropertyBag*>(PropertyBagType);
+		}
+	}
+
+	// Find the property on the given struct.
+	const FPropertyResolutionStep PropertyStep = FindPropertyAndArrayIndex(ActualBasePointer, ActualStruct, InPropertyNames[Index]);
+	if (!PropertyStep.Property)
+	{
+		OutResolutionState.bIsValid = false;
+		return;
+	}
+
+	const bool bHasMoreSteps = InPropertyNames.IsValidIndex(Index + 1);
+
+	if (PropertyStep.ArrayIndex != INDEX_NONE)
+	{
+		// We found that this segment of the property path reaches an element inside an array.
+		if (FArrayProperty* ArrayProp = CastField<FArrayProperty>(PropertyStep.Property))
+		{
+			FScriptArrayHelper ArrayHelper(ArrayProp, ArrayProp->ContainerPtrToValuePtr<void>(ActualBasePointer));
+			if (ArrayHelper.IsValidIndex(PropertyStep.ArrayIndex))
+			{
+				// Arrays may be resized, flag this property path as volatile.
+				OutResolutionState.bIsVolatile = true;
+
+				OutResolutionState.PropertySteps.Add(PropertyStep);
+
+				FStructProperty* InnerStructProp = CastField<FStructProperty>(ArrayProp->Inner);
+				// InnerStructProp is null for arrays of basic types like floats, integers, etc.
+				if (InnerStructProp && bHasMoreSteps)
+				{
+					// Move the BasePointer to the array element and keep resolving the property path on it.
+					void* ArrayElement = ArrayHelper.GetRawPtr(PropertyStep.ArrayIndex);
+					ResolvePropertyRecursive(ArrayElement, InnerStructProp->Struct, InPropertyNames, Index + 1, OutResolutionState);
+				}
+				else
+				{
+					// The property path ends here (e.g. "Foo.Bar[1]").
+					return;
+				}
+			}
+		}
+		else
+		{
+			UE_LOG(LogMovieScene, Error, TEXT("Mismatch in property evaluation. %s is not of type: %s"), *PropertyStep.Property->GetName(), *FArrayProperty::StaticClass()->GetName());
+			OutResolutionState.bIsValid = false;
+		}
+	}
+	else if (FStructProperty* StructProp = CastField<FStructProperty>(PropertyStep.Property))
+	{
+		// This segment of the property path reaches a struct FProperty.
+		OutResolutionState.PropertySteps.Add(PropertyStep);
+
+		if (bHasMoreSteps)
+		{
+			// Move the BasePointer to the struct and keep resolving the property path on it.
+			void* StructContainer = StructProp->ContainerPtrToValuePtr<void>(ActualBasePointer);
+			ResolvePropertyRecursive(StructContainer, StructProp->Struct, InPropertyNames, Index + 1, OutResolutionState);
+		}
+		else
+		{
+			// We stop at the struct, probably because we can directly animate it, like a vector/rotator/etc.
+			check(StructProp->GetName() == InPropertyNames[Index]);
+		}
+	}
+	else
+	{
+		// This segment of the property path reaches something else... probably a final float/integer/etc.
+		OutResolutionState.PropertySteps.Add(PropertyStep);
+	}
+}
+
+void ResolveProperty(const UObject& InObject, FStringView PropertyPath, FPropertyResolutionState& OutResolutionState)
+{
+	using namespace UE::String;
+
+	TArray<FStringView, TInlineAllocator<4>> PropertyNames;
+
+	// Parse property paths into component parts separated by '.'
+	auto ProcessToken = [&PropertyNames](FStringView Token) { PropertyNames.Emplace(Token); };
+	ParseTokens(PropertyPath, TEXT('.'), TFunctionRef<void(FWideStringView)>(ProcessToken), EParseTokensOptions::IgnoreCase | EParseTokensOptions::SkipEmpty);
+	
+	if (IsValid(&InObject) && PropertyNames.Num() > 0)
+	{
+		ResolvePropertyRecursive((void*)&InObject, InObject.GetClass(), PropertyNames, 0, OutResolutionState);
+	}
+}
+
+FVolatileProperty BuildVolatileProperty(const UObject* Object, FStringView PropertyPath, const FPropertyResolutionState& ResolutionState)
+{
+	FVolatileProperty VolatileProperty;
+	VolatileProperty.RootContainer = Object;
+	VolatileProperty.PropertyPath = PropertyPath;
+	VolatileProperty.LeafProperty = ResolutionState.PropertySteps.Last().Property;
+
+	void* BaseContainer = (void*)VolatileProperty.RootContainer;
+
+	for (int32 StepIndex = 0; StepIndex < ResolutionState.PropertySteps.Num(); ++StepIndex)
+	{
+		const FPropertyResolutionStep& ResolutionStep(ResolutionState.PropertySteps[StepIndex]);
+
+		// Here are some situations we need to handle. The property name in brackets is the current
+		// resolution step.
+		//
+		// Obj.<Value> : add a fixed offset step
+		//
+		// Obj.<Foo>.Bar : don't do anything, wait until next step to add a fixed offset step
+		//
+		// Obj.<Value[1]> : add a fixed offset step, an array jump, and another fixed offset step
+		//
+		// Obj.<Foo[1]>.Bar : add a fixed offset step, an array jump, another fixed offset step, and 
+		//		reset the base address to the array element.
+		//
+		// Obj.<InstancedStruct> : add a fixed offset step
+		//
+		// Obj.<InstancedStruct>.Foo : add a fixed offset step, an instanced struct property jump,
+		//		and reset the base address to that property's value address.
+		//
+		// Obj.<InstancedStructs[1]> : add a fixed offset step, an array jump, and another fixed 
+		//		offset step.
+		//
+		// Obj.<InstancedStructs[1]>.Foo : add a fixed offset step, an array jump, another fixed
+		//		offset step, an instanced struct property jump, and reset the base address to that
+		//		property's value address.
+		//
+
+		const bool bIsLastStep = (StepIndex == ResolutionState.PropertySteps.Num() - 1);
+
+		bool bIsInstancedStruct = false;
+		if (FStructProperty* StructProp = CastField<FStructProperty>(ResolutionStep.Property))
+		{
+			if (StructProp->Struct == FInstancedStruct::StaticStruct())
+			{
+				bIsInstancedStruct = true;
+			}
+		}
+		else if (FArrayProperty* ArrayProp = CastField<FArrayProperty>(ResolutionStep.Property))
+		{
+			if (FStructProperty* InnerStructProp = CastField<FStructProperty>(ArrayProp->Inner))
+			{
+				if (InnerStructProp->Struct == FInstancedStruct::StaticStruct())
+				{
+					bIsInstancedStruct = true;
+				}
+			}
+		}
+
+		// If this is the last step, force-add an entry so that we can resolve the leaf container inside
+		// the GetLeafContainerAddress function.
+		if (bIsLastStep)
+		{
+			FVolatilePropertyStep LeafContainerStep;
+			LeafContainerStep.SetContainerOffset((uintptr_t)ResolutionStep.ContainerAddress - (uintptr_t)BaseContainer);
+			VolatileProperty.PropertySteps.Add(LeafContainerStep);
+
+			// Set the leaf container step index, and make further offsets based off this leaf container.
+			VolatileProperty.LeafContainerStepIndex = VolatileProperty.PropertySteps.Num() - 1;
+			BaseContainer = ResolutionStep.ContainerAddress;
+		}
+
+		// If this is the last step before the end, or the step before jumping into an instanced struct 
+		// or an array, add a step to jump to the value so far. Otherwise, continue to follow the property
+		// path and we'll compute the offset later when needed.
+		if (bIsLastStep || bIsInstancedStruct || ResolutionStep.ArrayIndex != INDEX_NONE)
+		{
+			void* ValuePtr = ResolutionStep.Property->ContainerPtrToValuePtr<void>(ResolutionStep.ContainerAddress);
+			const uint32 ValueOffset = ((uintptr_t)ValuePtr - (uintptr_t)BaseContainer);
+			if (ValueOffset > 0)
+			{
+				FVolatilePropertyStep ValueJumpStep;
+				ValueJumpStep.SetContainerOffset(ValueOffset);
+				VolatileProperty.PropertySteps.Add(ValueJumpStep);
+			}
+		}
+
+		// If the current step is accessing a value in an array, we need to add some steps for that.
+		if (ResolutionStep.ArrayIndex != INDEX_NONE)
+		{
+			// Add a step to jump into the array.
+			FVolatilePropertyStep ArrayJumpStep;
+			ArrayJumpStep.SetCheckArrayIndex(ResolutionStep.ArrayIndex);
+			VolatileProperty.PropertySteps.Add(ArrayJumpStep);
+
+			// Add a step to jump from the beginning of the array to the actual element.
+			FArrayProperty* ArrayProp = CastFieldChecked<FArrayProperty>(ResolutionStep.Property);
+			FScriptArrayHelper ArrayHelper(ArrayProp, ArrayProp->ContainerPtrToValuePtr<void>(ResolutionStep.ContainerAddress));
+			void* ArrayDataPtr = ArrayHelper.GetRawPtr(0);
+			void* RawElementPtr = ArrayHelper.GetRawPtr(ResolutionStep.ArrayIndex);
+			FVolatilePropertyStep ElementJumpStep;
+			ElementJumpStep.SetContainerOffset((uintptr_t)RawElementPtr - (uintptr_t)ArrayDataPtr);
+			VolatileProperty.PropertySteps.Add(ElementJumpStep);
+
+			// Reset the base container so that offsets are now calculated from the array element.
+			BaseContainer = RawElementPtr;
+		}
+
+		// If the current step leads to an instanced struct and we have more steps after that, we need
+		// to add a step for jumping into the struct's memory buffer and checking that the next step is
+		// valid for that instanced struct's type.
+		if (bIsInstancedStruct && !bIsLastStep)
+		{
+			const FPropertyResolutionStep& NextStep = ResolutionState.PropertySteps[StepIndex + 1];
+
+			FVolatilePropertyStep InstancedStructStep;
+			InstancedStructStep.SetCheckStruct(CastChecked<UScriptStruct>(NextStep.Property->GetOwnerStruct()));
+			VolatileProperty.PropertySteps.Add(InstancedStructStep);
+
+			// Also reset the base container so that offsets are now calculated from the instanced struct.
+			BaseContainer = NextStep.ContainerAddress;
+		}
+	}
+
+	ensure(VolatileProperty.LeafContainerStepIndex >= 0);
+
+	return VolatileProperty;
+}
+
+void* FVolatilePropertyStep::ResolveAddress(void* ContainerAddress, bool& bNeedsRecaching) const
+{
+	switch (Data.GetIndex())
+	{
+		case 0:
+			{
+				// Just offset the pointer.
+				const uint32 ContainerOffset = Data.Get<uint32>();
+				return ((uint8*)ContainerAddress + ContainerOffset);
+			}
+		case 1:
+			{
+				// Jump into the array's memory if we know the index we will jump to next is valid
+				// for this array.
+				const int32 CheckArrayIndex = Data.Get<int32>();
+				FScriptArray* ScriptArray = (FScriptArray*)ContainerAddress;
+				if (ScriptArray->IsValidIndex(CheckArrayIndex))
+				{
+					return ScriptArray->GetData();
+				}
+				return nullptr;
+			}
+		case 2:
+			{
+				// Jump into the instanced struct's memory after checking that it's the type we
+				// expect. If not, return null.
+				FInstancedStruct* InstancedStruct = (FInstancedStruct*)ContainerAddress;
+				const UScriptStruct* InstancedStructType = InstancedStruct->GetScriptStruct();
+				const UScriptStruct* CheckStructType = Data.Get<TWeakObjectPtr<UScriptStruct>>().Get();
+				if (ensure(CheckStructType) && CheckStructType == InstancedStructType)
+				{
+					return InstancedStruct->GetMutableMemory();
+				}
+				bNeedsRecaching = true;
+				return nullptr;
+			}
+		default:
+			{
+				return nullptr;
+			}
+	}
+}
+
+void* FVolatileProperty::ResolvePropertySteps(bool bStopAtLeafStep) const
+{
+	bool bNeedsRecaching = false;
+	void* Address = ResolvePropertyStepsImpl(bStopAtLeafStep, bNeedsRecaching);
+	if (bNeedsRecaching && RootContainer && !PropertyPath.IsEmpty())
+	{
+		FPropertyResolutionState ResolutionState;
+		ResolveProperty(*RootContainer, PropertyPath, ResolutionState);
+		FVolatileProperty* MutableThis = const_cast<FVolatileProperty*>(this);
+		if (ResolutionState.bIsValid)
+		{
+			// If the path resolves with the new object trail, let's accept that. This happens for example
+			// with a new instanced struct that looks the same as the previous one, such as a re-generated 
+			// property bag.
+			FVolatileProperty NewVolatileProperty = BuildVolatileProperty(RootContainer, PropertyPath, ResolutionState);
+			ensure(NewVolatileProperty.RootContainer == RootContainer);
+			ensure(NewVolatileProperty.PropertyPath == PropertyPath);
+			*MutableThis = MoveTemp(NewVolatileProperty);
+
+			// Re-resolve.
+			bNeedsRecaching = false;
+			Address = ResolvePropertyStepsImpl(bStopAtLeafStep, bNeedsRecaching);
+			ensure(!bNeedsRecaching);
+		}
+		else
+		{
+			// The path doesn't resolve anymore. Leave it failing, but disable its ability to re-resolve 
+			// itself, so that we don't try to re-resolve it every time we call into it.
+			MutableThis->PropertyPath.Reset();
+		}
+	}
+	return Address;
+}
+
+void* FVolatileProperty::ResolvePropertyStepsImpl(bool bStopAtLeafStep, bool& bNeedsRecaching) const
+{
+	if (RootContainer)
+	{
+		void* CurAddress = (void*)RootContainer;
+		const int32 NumSteps = (bStopAtLeafStep ? LeafContainerStepIndex + 1 : PropertySteps.Num());
+		for (int32 Index = 0; Index < NumSteps; ++Index)
+		{
+			const FVolatilePropertyStep& PropertyStep = PropertySteps[Index];
+			CurAddress = PropertyStep.ResolveAddress(CurAddress, bNeedsRecaching);
+			if (CurAddress == nullptr || bNeedsRecaching)
+			{
+				break;
+			}
+		}
+		return CurAddress;
+	}
+	return nullptr;
+}
+
+}  // namespace UE::MovieScene
+
+TOptional<TPair<const FProperty*, UE::MovieScene::FSourcePropertyValue>> FTrackInstancePropertyBindings::StaticPropertyAndValue(const UObject* Object, FStringView InPropertyPath)
+{
+	using namespace UE::MovieScene;
+
+	checkf(Object, TEXT("No object specified"));
+
+	FPropertyResolutionState ResolutionState;
+	UE::MovieScene::ResolveProperty(*Object, InPropertyPath, ResolutionState);
+
+	if (const FPropertyResolutionStep* LastStep = ResolutionState.GetLastStep())
+	{
+		const FProperty* Property = LastStep->Property;
+		if (ensure(Property && LastStep->ContainerAddress && LastStep->ArrayIndex == INDEX_NONE))
+		{
+			const void* PropertyAddress = Property->ContainerPtrToValuePtr<void>(LastStep->ContainerAddress);
+
+			// Bool property values are stored in a bit field so using a straight cast of the pointer to get their value does not
+			// work.  Instead use the actual property to get the correct value.
+			if (const FBoolProperty* BoolProperty = CastField<const FBoolProperty>(Property))
+			{
+				return MakeTuple(Property, FSourcePropertyValue::FromValue(BoolProperty->GetPropertyValue(PropertyAddress)));
+			}
+			// Object properties might have various different types of storage, but we always expose them as a raw ptr
+			else if (const FObjectPropertyBase* ObjectProperty = CastField<const FObjectPropertyBase>(Property))
+			{
+				return MakeTuple(Property, FSourcePropertyValue::FromValue(ObjectProperty->GetObjectPropertyValue(PropertyAddress)));
+			}
+
+			return MakeTuple(Property, FSourcePropertyValue::FromAddress(PropertyAddress, *Property));
+		}
+	}
+
+	return TOptional<TPair<const FProperty*, UE::MovieScene::FSourcePropertyValue>>();
+}
+
+TOptional<UE::MovieScene::FSourcePropertyValue> FTrackInstancePropertyBindings::StaticValue(const UObject* Object, FStringView InPropertyPath)
+{
+	TOptional<TPair<const FProperty*, UE::MovieScene::FSourcePropertyValue>> Result = StaticPropertyAndValue(Object, InPropertyPath);
+	if (Result.IsSet())
+	{
+		return MoveTemp(Result->Value);
+	}
+	return TOptional<UE::MovieScene::FSourcePropertyValue>();
+}
+
+FProperty* FTrackInstancePropertyBindings::FResolvedPropertyAndFunction::GetValidProperty() const
+{
+	if (const FCachedProperty* CachedProperty = ResolvedProperty.TryGet<FCachedProperty>())
+	{
+		return CachedProperty->GetValidProperty();
+	}
+	else if (const FVolatileProperty* VolatileProperty = ResolvedProperty.TryGet<FVolatileProperty>())
+	{
+		return VolatileProperty->GetValidProperty();
+	}
+	else
+	{
+		return nullptr;
+	}
+}
+
+void* FTrackInstancePropertyBindings::FResolvedPropertyAndFunction::GetContainerAddress() const
+{
+	if (const FCachedProperty* CachedProperty = ResolvedProperty.TryGet<FCachedProperty>())
+	{
+		return CachedProperty->ContainerAddress;
+	}
+	else if (const FVolatileProperty* VolatileProperty = ResolvedProperty.TryGet<FVolatileProperty>())
+	{
+		return VolatileProperty->GetLeafContainerAddress();
+	}
+	else
+	{
+		return nullptr;
+	}
+}
+
+FTrackInstancePropertyBindings::FTrackInstancePropertyBindings(FName InPropertyName, const FString& InPropertyPath)
+	: PropertyPath(InPropertyPath)
+	, PropertyName(InPropertyName)
+{
+	static const FString Set(TEXT("Set"));
+	const FString FunctionString = Set + PropertyName.ToString();
+
+	FunctionName = FName(*FunctionString);
+}
+
+FProperty* FTrackInstancePropertyBindings::FindProperty(const UObject* Object, FStringView InPropertyPath)
+{
+	using namespace UE::MovieScene;
+	FPropertyResolutionState ResolutionState;
+	UE::MovieScene::ResolveProperty(*Object, InPropertyPath, ResolutionState);
+	if (const FPropertyResolutionStep* LastStep = ResolutionState.GetLastStep())
+	{
+		return LastStep->Property;
+	}
+	return nullptr;
+}
+
+FTrackInstancePropertyBindings::FResolvedPropertyAndFunction FTrackInstancePropertyBindings::FindPropertyAndFunction(const UObject* Object, FStringView InPropertyPath)
+{
+	using namespace UE::MovieScene;
+	FPropertyResolutionState ResolutionState;
+	UE::MovieScene::ResolveProperty(*Object, InPropertyPath, ResolutionState);
+
+	if (ResolutionState.PropertySteps.IsEmpty())
+	{
+		return FResolvedPropertyAndFunction();
+	}
+	
+	FResolvedPropertyAndFunction PropAndFunction;
+	if (!ResolutionState.bIsVolatile)
+	{
+		// No volatility was found while resolving this property path. Just use the tail property and
+		// container address, they should be fixed.
+		const FPropertyResolutionStep& LastStep(ResolutionState.PropertySteps.Last());
+		FCachedProperty CachedProperty{ LastStep.Property, LastStep.ContainerAddress, LastStep.ArrayIndex };
+		PropAndFunction.ResolvedProperty = TVariant<FCachedProperty, FVolatileProperty>(TInPlaceType<FCachedProperty>(), CachedProperty);
+	}
+	else
+	{
+		// We found some volatility while resolving this property path, such as an array or an instanced
+		// struct. Let's compress the steps as much as possible, keeping only the steps where we need to
+		// jump into some other memory buffer.
+		FVolatileProperty VolatileProperty = BuildVolatileProperty(Object, InPropertyPath, ResolutionState);
+		PropAndFunction.ResolvedProperty = TVariant<FCachedProperty, FVolatileProperty>(TInPlaceType<FVolatileProperty>(), VolatileProperty);
+	}
+	return PropAndFunction;
+}
+
+const FTrackInstancePropertyBindings::FResolvedPropertyAndFunction& FTrackInstancePropertyBindings::FindOrAdd(const UObject& InObject)
+{
+	FObjectKey ObjectKey(&InObject);
+
+	const FResolvedPropertyAndFunction* PropAndFunction = RuntimeObjectToFunctionMap.Find(ObjectKey);
+	if (PropAndFunction && (
+				PropAndFunction->SetterFunction.IsValid() ||
+				PropAndFunction->GetValidProperty()))
+	{
+		return *PropAndFunction;
+	}
+
+	CacheBinding(InObject);
+	return RuntimeObjectToFunctionMap.FindChecked(ObjectKey);
+}
+
+void FTrackInstancePropertyBindings::CallFunctionForEnum(UObject& InRuntimeObject, int64 PropertyValue)
+{
+	const FResolvedPropertyAndFunction& PropAndFunction = FindOrAdd(InRuntimeObject);
+
+	FProperty* Property = PropAndFunction.GetValidProperty();
+	if (Property && Property->HasSetter())
+	{
+		Property->CallSetter(&InRuntimeObject, &PropertyValue);
+	}
+	else if (UFunction* SetterFunction = PropAndFunction.SetterFunction.Get())
+	{
+		InvokeSetterFunction(&InRuntimeObject, SetterFunction, PropertyValue);
+	}
+	else if (Property)
+	{
+		if (FEnumProperty* EnumProperty = CastField<FEnumProperty>(Property))
+		{
+			FNumericProperty* UnderlyingProperty = EnumProperty->GetUnderlyingProperty();
+			void* ValueAddr = EnumProperty->ContainerPtrToValuePtr<void>(PropAndFunction.GetContainerAddress());
+			UnderlyingProperty->SetIntPropertyValue(ValueAddr, PropertyValue);
+		}
+		else
+		{
+			UE_LOG(LogMovieScene, Error, TEXT("Mismatch in property evaluation. %s is not of type: %s"), *Property->GetName(), *FEnumProperty::StaticClass()->GetName());
+		}
+	}
+}
+
+void FTrackInstancePropertyBindings::CacheBinding(const UObject& Object)
+{
+	FResolvedPropertyAndFunction PropAndFunction = FindPropertyAndFunction(&Object, PropertyPath);
+	{
+		UFunction* SetterFunction = Object.FindFunction(FunctionName);
+		if (SetterFunction && SetterFunction->NumParms >= 1)
+		{
+			PropAndFunction.SetterFunction = SetterFunction;
+		}
+	}
+
+	RuntimeObjectToFunctionMap.Add(FObjectKey(&Object), PropAndFunction);
+}
+
+FProperty* FTrackInstancePropertyBindings::GetProperty(const UObject& Object)
+{
+	const FResolvedPropertyAndFunction& PropAndFunction = FindOrAdd(Object);
+	return PropAndFunction.GetValidProperty();
+}
+
+bool FTrackInstancePropertyBindings::HasValidBinding(const UObject& Object)
+{
+	const FResolvedPropertyAndFunction& PropAndFunction = FindOrAdd(Object);
+	return PropAndFunction.GetValidProperty() != nullptr;
+}
+
+const UStruct* FTrackInstancePropertyBindings::GetPropertyStruct(const UObject& Object)
+{
+	const FResolvedPropertyAndFunction& PropAndFunction = FindOrAdd(Object);
+
+	if (FProperty* Property = PropAndFunction.GetValidProperty())
+	{
+		if (FStructProperty* StructProperty = CastField<FStructProperty>(Property))
+		{
+			return StructProperty->Struct;
+		}
+		return nullptr;
+	}
+
+	return nullptr;
+}
+
+int64 FTrackInstancePropertyBindings::GetCurrentValueForEnum(const UObject& Object)
+{
+	const FResolvedPropertyAndFunction& PropAndFunction = FindOrAdd(Object);
+
+	FProperty* Property = PropAndFunction.GetValidProperty();
+
+	if (Property)
+	{
+		if (FEnumProperty* EnumProperty = CastField<FEnumProperty>(Property))
+		{
+			FNumericProperty* UnderlyingProperty = EnumProperty->GetUnderlyingProperty();
+			void* ValueAddr = EnumProperty->ContainerPtrToValuePtr<void>(PropAndFunction.GetContainerAddress());
+			int64 Result = UnderlyingProperty->GetSignedIntPropertyValue(ValueAddr);
+			return Result;
+		}
+		else
+		{
+			UE_LOG(LogMovieScene, Error, TEXT("Mismatch in property evaluation. %s is not of type: %s"), *Property->GetName(), *FEnumProperty::StaticClass()->GetName());
+		}
+	}
+
+	return 0;
+}
+
+// Explicit specializations for bools.
+
+template<> void FTrackInstancePropertyBindings::CallFunction<bool>(UObject& InRuntimeObject, TCallTraits<bool>::ParamType PropertyValue)
+{
+	const FResolvedPropertyAndFunction& PropAndFunction = FindOrAdd(InRuntimeObject);
+
+	FProperty* Property = PropAndFunction.GetValidProperty();
+
+	if (Property && Property->HasSetter())
+	{
+		Property->CallSetter(&InRuntimeObject, &PropertyValue);
+	}
+	else if (UFunction* SetterFunction = PropAndFunction.SetterFunction.Get())
+	{
+		InvokeSetterFunction(&InRuntimeObject, SetterFunction, PropertyValue);
+	}
+	else if (Property)
+	{
+		if (FBoolProperty* BoolProperty = CastField<FBoolProperty>(Property))
+		{
+			uint8* ValuePtr = BoolProperty->ContainerPtrToValuePtr<uint8>(PropAndFunction.GetContainerAddress());
+			BoolProperty->SetPropertyValue(ValuePtr, PropertyValue);
+		}
+		else
+		{
+			UE_LOG(LogMovieScene, Error, TEXT("Mismatch in property evaluation. %s is not of type: %s"), *Property->GetName(), *FBoolProperty::StaticClass()->GetName());
+		}
+	}
+}
+
+template<> bool FTrackInstancePropertyBindings::TryGetPropertyValue<bool>(const FResolvedPropertyAndFunction& PropAndFunction, bool& OutValue)
+{
+	if (FProperty* Property = PropAndFunction.GetValidProperty())
+	{
+		if (FBoolProperty* BoolProperty = CastField<FBoolProperty>(Property))
+		{
+			const uint8* ValuePtr = BoolProperty->ContainerPtrToValuePtr<uint8>(PropAndFunction.GetContainerAddress());
+			OutValue = BoolProperty->GetPropertyValue(ValuePtr);
+			return true;
+		}
+		else
+		{
+			UE_LOG(LogMovieScene, Error, TEXT("Mismatch in property evaluation. %s is not of type: %s"), *Property->GetName(), *FBoolProperty::StaticClass()->GetName());
+		}
+	}
+
+	return false;
+}
+
+template<> void FTrackInstancePropertyBindings::SetCurrentValue<bool>(UObject& Object, TCallTraits<bool>::ParamType InValue)
+{
+	const FResolvedPropertyAndFunction& PropAndFunction = FindOrAdd(Object);
+
+	if (FProperty* Property = PropAndFunction.GetValidProperty())
+	{
+		if (FBoolProperty* BoolProperty = CastField<FBoolProperty>(Property))
+		{
+			uint8* ValuePtr = BoolProperty->ContainerPtrToValuePtr<uint8>(PropAndFunction.GetContainerAddress());
+			BoolProperty->SetPropertyValue(ValuePtr, InValue);
+		}
+	}
+}
+
+// Explicit specializations for object pointers.
+
+template<> void FTrackInstancePropertyBindings::CallFunction<UObject*>(UObject& InRuntimeObject, UObject* PropertyValue)
+{
+	const FResolvedPropertyAndFunction& PropAndFunction = FindOrAdd(InRuntimeObject);
+
+	FProperty* Property = PropAndFunction.GetValidProperty();
+
+	if (Property && Property->HasSetter())
+	{
+		Property->CallSetter(&InRuntimeObject, &PropertyValue);
+	}
+	else if (UFunction* SetterFunction = PropAndFunction.SetterFunction.Get())
+	{
+		InvokeSetterFunction(&InRuntimeObject, SetterFunction, PropertyValue);
+	}
+	else if (Property)
+	{
+		if (FObjectPropertyBase* ObjectProperty = CastField<FObjectPropertyBase>(Property))
+		{
+			uint8* ValuePtr = ObjectProperty->ContainerPtrToValuePtr<uint8>(PropAndFunction.GetContainerAddress());
+			ObjectProperty->SetObjectPropertyValue(ValuePtr, PropertyValue);
+		}
+		else
+		{
+			UE_LOG(LogMovieScene, Error, TEXT("Mismatch in property evaluation. %s is not of type: %s"), *Property->GetName(), *FObjectPropertyBase::StaticClass()->GetName());
+		}
+	}
+}
+
+template<> bool FTrackInstancePropertyBindings::TryGetPropertyValue<UObject*>(const FResolvedPropertyAndFunction& PropAndFunction, UObject*& OutValue)
+{
+	if (FProperty* Property = PropAndFunction.GetValidProperty())
+	{
+		if (FObjectPropertyBase* ObjectProperty = CastField<FObjectPropertyBase>(PropAndFunction.GetValidProperty()))
+		{
+			const uint8* ValuePtr = ObjectProperty->ContainerPtrToValuePtr<uint8>(PropAndFunction.GetContainerAddress());
+			OutValue = ObjectProperty->GetObjectPropertyValue(ValuePtr);
+			return true;
+		}
+		else
+		{
+			UE_LOG(LogMovieScene, Error, TEXT("Mismatch in property evaluation. %s is not of type: %s"), *Property->GetName(), *FObjectPropertyBase::StaticClass()->GetName());
+		}
+	}
+
+	return false;
+}
+
+template<> void FTrackInstancePropertyBindings::SetCurrentValue<UObject*>(UObject& Object, UObject* InValue)
+{
+	const FResolvedPropertyAndFunction& PropAndFunction = FindOrAdd(Object);
+
+	if (FProperty* Property = PropAndFunction.GetValidProperty())
+	{
+		if (FObjectPropertyBase* ObjectProperty = CastField<FObjectPropertyBase>(Property))
+		{
+			uint8* ValuePtr = ObjectProperty->ContainerPtrToValuePtr<uint8>(PropAndFunction.GetContainerAddress());
+			ObjectProperty->SetObjectPropertyValue(ValuePtr, InValue);
+		}
+	}
+}
+
diff --git a/Engine/Source/Runtime/MovieScene/Public/MovieSceneCommonHelpers.h b/Engine/Source/Runtime/MovieScene/Public/MovieSceneCommonHelpers.h
--- a/Engine/Source/Runtime/MovieScene/Public/MovieSceneCommonHelpers.h
+++ b/Engine/Source/Runtime/MovieScene/Public/MovieSceneCommonHelpers.h
@@ -10,6 +10,8 @@
 #include "Misc/FrameNumber.h"
 #include "UObject/WeakFieldPtr.h"
 
+#include "TrackInstancePropertyBindings.h"  // For backwards compatibility
+
 class AActor;
 class UCameraComponent;
 class UMovieScene;
@@ -26,11 +28,6 @@ struct FMovieSceneSequenceID;
 class UMovieSceneCondition;
 class UMovieSceneTrack;
 
-namespace UE::MovieScene
-{
-	struct FSourcePropertyValue;
-} // namespace UE::MovieScene
-
 namespace UE::MovieScene
 {
 	struct FSharedPlaybackState;
@@ -300,366 +297,3 @@ public:
 	static MOVIESCENE_API bool EvaluateSequenceCondition(const FGuid& BindingID, const FMovieSceneSequenceID& SequenceID, const UMovieSceneCondition* Condition, UObject* ConditionOwnerObject, TSharedRef<const UE::MovieScene::FSharedPlaybackState> SharedPlaybackState);
 };
 
-/**
- * Manages bindings to keyed properties for a track instance. 
- * Calls UFunctions to set the value on runtime objects
- */
-class FTrackInstancePropertyBindings
-{
-public:
-
-	MOVIESCENE_API FTrackInstancePropertyBindings(FName InPropertyName, const FString& InPropertyPath);
-
-	/**
-	 * Calls the setter function for a specific runtime object or if the setter function does not exist, the property is set directly
-	 *
-	 * @param InRuntimeObject The runtime object whose function to call
-	 * @param PropertyValue The new value to assign to the property
-	 */
-	template <typename ValueType>
-	void CallFunction(UObject& InRuntimeObject, typename TCallTraits<ValueType>::ParamType PropertyValue)
-	{
-		FResolvedPropertyAndFunction PropAndFunction = FindOrAdd(InRuntimeObject);
-
-		FProperty* Property = PropAndFunction.ResolvedProperty.GetValidProperty();
-		if (Property && Property->HasSetter())
-		{
-			Property->CallSetter(&InRuntimeObject, &PropertyValue);
-		}
-		else if (UFunction* SetterFunction = PropAndFunction.SetterFunction.Get())
-		{
-			InvokeSetterFunction(&InRuntimeObject, SetterFunction, PropertyValue);
-		}
-		else if (ValueType* Val = PropAndFunction.GetPropertyAddress<ValueType>())
-		{
-			*Val = MoveTempIfPossible(PropertyValue);
-		}
-
-		if (UFunction* NotifyFunction = PropAndFunction.NotifyFunction.Get())
-		{
-			InRuntimeObject.ProcessEvent(NotifyFunction, nullptr);
-		}
-	}
-
-	/**
-	 * Calls the setter function for a specific runtime object or if the setter function does not exist, the property is set directly
-	 *
-	 * @param InRuntimeObject The runtime object whose function to call
-	 * @param PropertyValue The new value to assign to the property
-	 */
-	MOVIESCENE_API void CallFunctionForEnum( UObject& InRuntimeObject, int64 PropertyValue );
-
-	/**
-	 * Rebuilds the property and function mappings for a single runtime object, and adds them to the cache
-	 *
-	 * @param InRuntimeObject	The object to cache mappings for
-	 */
-	MOVIESCENE_API void CacheBinding(const UObject& InRuntimeObject);
-
-	/**
-	 * Gets the FProperty that is bound to the track instance.
-	 *
-	 * @param Object	The Object that owns the property
-	 * @return			The property on the object if it exists
-	 */
-	MOVIESCENE_API FProperty* GetProperty(const UObject& Object);
-
-	/**
-	 * Returns whether this binding is valid for the given object.
-	 *
-	 * @param Object	The Object that owns the property
-	 * @return          Whether the property binding is valid
-	 */
-	MOVIESCENE_API bool HasValidBinding(const UObject& Object);
-
-	/**
-	 * Gets the structure type of the bound property if it is a property of that type.
-	 *
-	 * @param Object	The Object that owns the property
-	 * @return          The structure type
-	 */
-	MOVIESCENE_API const UStruct* GetPropertyStruct(const UObject& Object);
-
-	/**
-	 * Gets the current value of a property on an object
-	 *
-	 * @param Object	The object to get the property from
-	 * @return ValueType	The current value
-	 */
-	template <typename ValueType>
-	ValueType GetCurrentValue(const UObject& Object)
-	{
-		ValueType Value{};
-
-		FResolvedPropertyAndFunction PropAndFunction = FindOrAdd(Object);
-		TryGetPropertyValue<ValueType>(PropAndFunction.ResolvedProperty, Value);
-
-		return Value;
-	}
-
-	/**
-	 * Optionally gets the current value of a property on an object
-	 *
-	 * @param Object	The object to get the property from
-	 * @return (Optional) The current value of the property on the object
-	 */
-	template <typename ValueType>
-	TOptional<ValueType> GetOptionalValue(const UObject& Object)
-	{
-		ValueType Value{};
-
-		FResolvedPropertyAndFunction PropAndFunction = FindOrAdd(Object);
-		if (TryGetPropertyValue<ValueType>(PropAndFunction.ResolvedProperty, Value))
-		{
-			return Value;
-		}
-
-		return TOptional<ValueType>();
-	}
-
-	/**
-	 * Static function for accessing a property value on an object without caching its address
-	 *
-	 * @param Object			The object to get the property from
-	 * @param InPropertyPath	The path to the property to retrieve
-	 * @return (Optional) The current value of the property on the object
-	 */
-	MOVIESCENE_API static TOptional<UE::MovieScene::FSourcePropertyValue> StaticValue(const UObject* Object, FStringView InPropertyPath);
-	MOVIESCENE_API static TOptional<TPair<const FProperty*, UE::MovieScene::FSourcePropertyValue>> StaticPropertyAndValue(const UObject* Object, FStringView InPropertyPath);
-
-	/**
-	 * Static function for accessing a property value on an object without caching its address
-	 *
-	 * @param Object			The object to get the property from
-	 * @param InPropertyPath	The path to the property to retrieve
-	 * @return (Optional) The current value of the property on the object
-	 */
-	template <typename ValueType>
-	static TOptional<ValueType> StaticValue(const UObject* Object, FStringView InPropertyPath)
-	{
-		checkf(Object, TEXT("No object specified"));
-
-		FResolvedProperty ResolvedProperty = ResolveProperty(InPropertyPath, *Object);
-
-		ValueType Value;
-		if (TryGetPropertyValue<ValueType>(ResolvedProperty, Value))
-		{
-			return Value;
-		}
-
-		return TOptional<ValueType>();
-	}
-
-	/**
-	 * Gets the current value of a property on an object
-	 *
-	 * @param Object	The object to get the property from
-	 * @return ValueType	The current value
-	 */
-	MOVIESCENE_API int64 GetCurrentValueForEnum(const UObject& Object);
-
-	/**
-	 * Sets the current value of a property on an object
-	 *
-	 * @param Object	The object to set the property on
-	 * @param InValue   The value to set
-	 */
-	template <typename ValueType>
-	void SetCurrentValue(UObject& Object, typename TCallTraits<ValueType>::ParamType InValue)
-	{
-		FResolvedPropertyAndFunction PropAndFunction = FindOrAdd(Object);
-
-		if (ValueType* Val = PropAndFunction.GetPropertyAddress<ValueType>())
-		{
-			*Val = InValue;
-
-			if (UFunction* NotifyFunction = PropAndFunction.NotifyFunction.Get())
-			{
-				Object.ProcessEvent(NotifyFunction, nullptr);
-			}
-		}
-	}
-
-	/** @return the property path that this binding was initialized from */
-	const FString& GetPropertyPath() const
-	{
-		return PropertyPath;
-	}
-
-	/** @return the property name that this binding was initialized from */
-	const FName& GetPropertyName() const
-	{
-		return PropertyName;
-	}
-
-	static MOVIESCENE_API FProperty* FindProperty(const UObject* Object, FStringView InPropertyPath);
-
-private:
-
-	/**
-	 * Wrapper for UObject::ProcessEvent that attempts to pass the new property value directly to the function as a parameter,
-	 * but handles cases where multiple parameters or a return value exists. The setter parameter must be the first in the list,
-	 * any other parameters will be default constructed.
-	 */
-	template<typename T>
-	static void InvokeSetterFunction(UObject* InRuntimeObject, UFunction* Setter, T&& InPropertyValue);
-
-	struct FResolvedProperty
-	{
-		TWeakFieldPtr<FProperty> Property;
-
-		void* ContainerAddress;
-
-		int32 ArrayIndex = INDEX_NONE;
-
-		FProperty* GetValidProperty() const
-		{
-			FProperty* PropertyPtr = Property.Get();
-			if (PropertyPtr && ContainerAddress && !PropertyPtr->HasAnyFlags(RF_BeginDestroyed | RF_FinishDestroyed))
-			{
-				return PropertyPtr;
-			}
-			return nullptr;
-		}
-
-		template<typename ValueType>
-		ValueType* GetPropertyAddress() const
-		{
-			if (FProperty* PropertyPtr = GetValidProperty())
-			{
-				if (ArrayIndex == INDEX_NONE)
-				{
-					return PropertyPtr->ContainerPtrToValuePtr<ValueType>(ContainerAddress);
-				}
-				else if (FArrayProperty* ArrayProp = CastFieldChecked<FArrayProperty>(PropertyPtr))
-				{
-					FScriptArrayHelper ArrayHelper(ArrayProp, ArrayProp->ContainerPtrToValuePtr<void>(ContainerAddress));
-					check(ArrayHelper.IsValidIndex(ArrayIndex));
-					return (ValueType*)ArrayHelper.GetRawPtr(ArrayIndex);
-				}
-			}
-			return nullptr;
-		}
-
-		FResolvedProperty()
-			: Property(nullptr)
-			, ContainerAddress(nullptr)
-		{}
-	};
-
-	struct FResolvedPropertyAndFunction
-	{
-		FResolvedProperty ResolvedProperty;
-		TWeakObjectPtr<UFunction> SetterFunction;
-		TWeakObjectPtr<UFunction> NotifyFunction;
-
-		template<typename ValueType>
-		ValueType* GetPropertyAddress() const
-		{
-			return ResolvedProperty.GetPropertyAddress<ValueType>();
-		}
-
-		FResolvedPropertyAndFunction()
-			: ResolvedProperty()
-			, SetterFunction( nullptr )
-			, NotifyFunction( nullptr )
-		{}
-	};
-
-	template <typename ValueType>
-	static bool TryGetPropertyValue(const FResolvedProperty& ResolvedProperty, ValueType& OutValue)
-	{
-		if (const ValueType* Value = ResolvedProperty.GetPropertyAddress<ValueType>())
-		{
-			OutValue = *Value;
-			return true;
-		}
-		return false;
-	}
-
-	static MOVIESCENE_API FResolvedProperty ResolvePropertyRecursive(void* BasePointer, UStruct* InStruct, TArrayView<const FStringView> InPropertyNames, uint32 Index);
-	static MOVIESCENE_API FResolvedProperty ResolveProperty(FStringView PropertyPath, const UObject& Object);
-
-	static MOVIESCENE_API void FindProperty(void* BasePointer, UStruct* InStruct, FStringView InPropertyName, FResolvedProperty& OutResolvedProperty);
-	static MOVIESCENE_API FResolvedProperty FindPropertyAndArrayIndex(void* BasePointer, UStruct* InStruct, FStringView InPropertyName);
-
-	/** Find or add the FResolvedPropertyAndFunction for the specified object */
-	MOVIESCENE_API FResolvedPropertyAndFunction FindOrAdd(const UObject& InObject);
-
-private:
-
-	/** Mapping of objects to bound functions that will be called to update data on the track */
-	TMap< FObjectKey, FResolvedPropertyAndFunction > RuntimeObjectToFunctionMap;
-
-	/** Path to the property we are bound to */
-	FString PropertyPath;
-
-	/** Name of the function to call to set values */
-	FName FunctionName;
-
-	/** Name of a function to call when a value has been set */
-	FName NotifyFunctionName;
-
-	/** Actual name of the property we are bound to */
-	FName PropertyName;
-};
-
-/** Explicit specializations for bools */
-template<> MOVIESCENE_API void FTrackInstancePropertyBindings::CallFunction<bool>(UObject& InRuntimeObject, TCallTraits<bool>::ParamType PropertyValue);
-template<> MOVIESCENE_API bool FTrackInstancePropertyBindings::TryGetPropertyValue<bool>(const FResolvedProperty& ResolvedProperty, bool& OutValue);
-template<> MOVIESCENE_API void FTrackInstancePropertyBindings::SetCurrentValue<bool>(UObject& Object, TCallTraits<bool>::ParamType InValue);
-
-/** Explicit specializations for object pointers */
-template<> MOVIESCENE_API void FTrackInstancePropertyBindings::CallFunction<UObject*>(UObject& InRuntimeObject, UObject* PropertyValue);
-template<> MOVIESCENE_API bool FTrackInstancePropertyBindings::TryGetPropertyValue<UObject*>(const FResolvedProperty& ResolvedProperty, UObject*& OutValue);
-template<> MOVIESCENE_API void FTrackInstancePropertyBindings::SetCurrentValue<UObject*>(UObject& Object, UObject* InValue);
-
-template<typename T>
-void FTrackInstancePropertyBindings::InvokeSetterFunction(UObject* InRuntimeObject, UFunction* Setter, T&& InPropertyValue)
-{
-	// CacheBinding already guarantees that the function has >= 1 parameters
-	const int32 ParmsSize = Setter->ParmsSize;
-
-	// This should all be const really, but ProcessEvent only takes a non-const void*
-	void* InputParameter = const_cast<typename TDecay<T>::Type*>(&InPropertyValue);
-
-	// By default we try and use the existing stack value
-	uint8* Params = reinterpret_cast<uint8*>(InputParameter);
-
-	check(InRuntimeObject && Setter);
-	if (Setter->ReturnValueOffset != MAX_uint16 || Setter->NumParms > 0)
-	{
-		// Function has a return value or multiple parameters, we need to initialize memory for the entire parameter pack
-		// We use alloca here (as in UObject::ProcessEvent) to avoid a heap allocation. Alloca memory survives the current function's stack frame.
-		Params = reinterpret_cast<uint8*>(FMemory_Alloca(ParmsSize));
-
-		bool bFirstProperty = true;
-		for (FProperty* Property = Setter->PropertyLink; Property; Property = Property->PropertyLinkNext)
-		{
-			// Initialize the parameter pack with any param properties that reside in the container
-			if (Property->IsInContainer(ParmsSize))
-			{
-				Property->InitializeValue_InContainer(Params);
-
-				// The first encountered property is assumed to be the input value so initialize this with the user-specified value from InPropertyValue
-				if (Property->HasAnyPropertyFlags(CPF_Parm) && !Property->HasAnyPropertyFlags(CPF_ReturnParm) && bFirstProperty)
-				{
-					const bool bIsValid = ensureMsgf(sizeof(T) == Property->GetElementSize(), TEXT("Property type does not match for Sequencer setter function %s::%s (%" SIZE_T_FMT "bytes != %ibytes"), *InRuntimeObject->GetName(), *Setter->GetName(), sizeof(T), Property->GetElementSize());
-					if (bIsValid)
-					{
-						Property->CopyCompleteValue(Property->ContainerPtrToValuePtr<void>(Params), &InPropertyValue);
-					}
-					else
-					{
-						return;
-					}
-				}
-				bFirstProperty = false;
-			}
-		}
-	}
-
-	// Now we have the parameters set up correctly, call the function
-	InRuntimeObject->ProcessEvent(Setter, Params);
-}
-
diff --git a/Engine/Source/Runtime/MovieScene/Public/TrackInstancePropertyBindings.h b/Engine/Source/Runtime/MovieScene/Public/TrackInstancePropertyBindings.h
new file mode 100644
--- /dev/null
+++ b/Engine/Source/Runtime/MovieScene/Public/TrackInstancePropertyBindings.h
@@ -0,0 +1,461 @@
+// Copyright Epic Games, Inc. All Rights Reserved.
+
+#pragma once
+
+#include "CoreTypes.h"
+
+#include "Containers/Map.h"
+#include "Misc/TVariant.h"
+#include "UObject/Class.h"
+#include "UObject/ObjectKey.h"
+#include "UObject/UnrealType.h"
+#include "UObject/WeakFieldPtr.h"
+
+namespace UE::MovieScene
+{
+
+struct FSourcePropertyValue;
+
+/**
+ * A property path whose only tail propery information is cached, for faster (or at least less slow) access.
+ */
+struct FCachedProperty
+{
+	TWeakFieldPtr<FProperty> Property = nullptr;
+
+	void* ContainerAddress = nullptr;
+
+	int32 ArrayIndex = INDEX_NONE;
+
+	FProperty* GetValidProperty() const
+	{
+		FProperty* PropertyPtr = Property.Get();
+		if (ContainerAddress && PropertyPtr && !PropertyPtr->HasAnyFlags(RF_BeginDestroyed | RF_FinishDestroyed))
+		{
+			return PropertyPtr;
+		}
+		return nullptr;
+	}
+
+	template<typename ValueType>
+	ValueType* GetPropertyAddress() const
+	{
+		if (FProperty* PropertyPtr = GetValidProperty())
+		{
+			if (ArrayIndex == INDEX_NONE)
+			{
+				return PropertyPtr->ContainerPtrToValuePtr<ValueType>(ContainerAddress);
+			}
+			else if (FArrayProperty* ArrayProp = CastFieldChecked<FArrayProperty>(PropertyPtr))
+			{
+				FScriptArrayHelper ArrayHelper(ArrayProp, ArrayProp->ContainerPtrToValuePtr<void>(ContainerAddress));
+				check(ArrayHelper.IsValidIndex(ArrayIndex));
+				return (ValueType*)ArrayHelper.GetRawPtr(ArrayIndex);
+			}
+		}
+		return nullptr;
+	}
+};
+
+/**
+ * A step in the property path, for when that property path has volatile elements in it.
+ */
+struct FVolatilePropertyStep
+{
+private:
+
+	using FData = TVariant<
+		// ContainerOffset : A fixed offset to jump by.
+		// Size = 4 bytes
+		uint32,
+		// CheckArrayIndex : A jump into an array's buffer, checking that the element index we'll 
+		// offset to next is still valid.
+		// Size = 4 bytes
+		int32,
+		// CheckStruct : A jump into an instanced struct's buffer, checking that the struct type
+		// is what we expect.
+		// Size = 8 bytes
+		TWeakObjectPtr<UScriptStruct>
+	>;
+	FData Data;
+
+public:
+
+	void* ResolveAddress(void* ContainerAddress, bool& bNeedsRecaching) const;
+
+	void SetContainerOffset(uint32 InOffset) { Data.Set<uint32>(InOffset); }
+	void SetCheckArrayIndex(int32 InArrayIndex) { Data.Set<int32>(InArrayIndex); }
+	void SetCheckStruct(TWeakObjectPtr<UScriptStruct> InScriptStruct) { Data.Set<TWeakObjectPtr<UScriptStruct>>(InScriptStruct); }
+};
+// The TVariant data type flag is 8 bytes so with the biggest union field being 8 bytes, it should 
+// all fit inside of 16 bytes.
+static_assert(sizeof(FVolatilePropertyStep) <= 16, "Try to fit FVolatilePropertyStep inside 16 bytes");
+
+/**
+ * A property path where no caching is done, and the address of the end value is re-resolved from
+ * the root object every time. This is for property paths that go over re-allocatable memory such as
+ * arrays, instanced structs, and other containers.
+ */
+struct FVolatileProperty
+{
+	const UObject* RootContainer = nullptr;
+	FString PropertyPath;
+
+	TArray<FVolatilePropertyStep> PropertySteps;
+	TWeakFieldPtr<FProperty> LeafProperty;
+	int32 LeafContainerStepIndex = INDEX_NONE;
+
+	FProperty* GetValidProperty() const
+	{
+		FProperty* PropertyPtr = LeafProperty.Get();
+		if (RootContainer && PropertyPtr && !PropertyPtr->HasAnyFlags(RF_BeginDestroyed | RF_FinishDestroyed))
+		{
+			return PropertyPtr;
+		}
+		return nullptr;
+	}
+
+	void* GetLeafContainerAddress() const
+	{
+		return ResolvePropertySteps(true);
+	}
+
+	template<typename ValueType>
+	ValueType* GetPropertyAddress() const
+	{
+		return (ValueType*)ResolvePropertySteps(false);
+	}
+
+private:
+
+	MOVIESCENE_API void* ResolvePropertySteps(bool bStopAtLeafStep) const;
+
+	void* ResolvePropertyStepsImpl(bool bStopAtLeafStep, bool& bNeedsRecaching) const;
+};
+
+}  // namespace UE::MovieScene
+
+/**
+ * Manages bindings to keyed properties for a track instance. 
+ * Calls UFunctions to set the value on runtime objects
+ */
+class FTrackInstancePropertyBindings
+{
+public:
+
+	MOVIESCENE_API FTrackInstancePropertyBindings(FName InPropertyName, const FString& InPropertyPath);
+
+	/**
+	 * Calls the setter function for a specific runtime object or if the setter function does not exist, the property is set directly
+	 *
+	 * @param InRuntimeObject The runtime object whose function to call
+	 * @param PropertyValue The new value to assign to the property
+	 */
+	template <typename ValueType>
+	void CallFunction(UObject& InRuntimeObject, typename TCallTraits<ValueType>::ParamType PropertyValue)
+	{
+		const FResolvedPropertyAndFunction& PropAndFunction = FindOrAdd(InRuntimeObject);
+
+		FProperty* Property = PropAndFunction.GetValidProperty();
+		if (Property && Property->HasSetter())
+		{
+			Property->CallSetter(&InRuntimeObject, &PropertyValue);
+		}
+		else if (UFunction* SetterFunction = PropAndFunction.SetterFunction.Get())
+		{
+			InvokeSetterFunction(&InRuntimeObject, SetterFunction, PropertyValue);
+		}
+		else if (ValueType* Val = PropAndFunction.GetPropertyAddress<ValueType>())
+		{
+			*Val = MoveTempIfPossible(PropertyValue);
+		}
+	}
+
+	/**
+	 * Calls the setter function for a specific runtime object or if the setter function does not exist, the property is set directly
+	 *
+	 * @param InRuntimeObject The runtime object whose function to call
+	 * @param PropertyValue The new value to assign to the property
+	 */
+	MOVIESCENE_API void CallFunctionForEnum( UObject& InRuntimeObject, int64 PropertyValue );
+
+	/**
+	 * Rebuilds the property and function mappings for a single runtime object, and adds them to the cache
+	 *
+	 * @param InRuntimeObject	The object to cache mappings for
+	 */
+	MOVIESCENE_API void CacheBinding(const UObject& InRuntimeObject);
+
+	/**
+	 * Gets the FProperty that is bound to the track instance.
+	 *
+	 * @param Object	The Object that owns the property
+	 * @return			The property on the object if it exists
+	 */
+	MOVIESCENE_API FProperty* GetProperty(const UObject& Object);
+
+	/**
+	 * Returns whether this binding is valid for the given object.
+	 *
+	 * @param Object	The Object that owns the property
+	 * @return          Whether the property binding is valid
+	 */
+	MOVIESCENE_API bool HasValidBinding(const UObject& Object);
+
+	/**
+	 * Gets the structure type of the bound property if it is a property of that type.
+	 *
+	 * @param Object	The Object that owns the property
+	 * @return          The structure type
+	 */
+	MOVIESCENE_API const UStruct* GetPropertyStruct(const UObject& Object);
+
+	/**
+	 * Gets the current value of a property on an object
+	 *
+	 * @param Object	The object to get the property from
+	 * @return ValueType	The current value
+	 */
+	template <typename ValueType>
+	ValueType GetCurrentValue(const UObject& Object)
+	{
+		ValueType Value{};
+
+		const FResolvedPropertyAndFunction& PropAndFunction = FindOrAdd(Object);
+		TryGetPropertyValue<ValueType>(PropAndFunction, Value);
+
+		return Value;
+	}
+
+	/**
+	 * Optionally gets the current value of a property on an object
+	 *
+	 * @param Object	The object to get the property from
+	 * @return (Optional) The current value of the property on the object
+	 */
+	template <typename ValueType>
+	TOptional<ValueType> GetOptionalValue(const UObject& Object)
+	{
+		ValueType Value{};
+
+		const FResolvedPropertyAndFunction& PropAndFunction = FindOrAdd(Object);
+		if (TryGetPropertyValue<ValueType>(PropAndFunction, Value))
+		{
+			return Value;
+		}
+
+		return TOptional<ValueType>();
+	}
+
+	/**
+	 * Gets the current value of a property on an object
+	 *
+	 * @param Object	The object to get the property from
+	 * @return ValueType	The current value
+	 */
+	MOVIESCENE_API int64 GetCurrentValueForEnum(const UObject& Object);
+
+	/**
+	 * Sets the current value of a property on an object
+	 *
+	 * @param Object	The object to set the property on
+	 * @param InValue   The value to set
+	 */
+	template <typename ValueType>
+	void SetCurrentValue(UObject& Object, typename TCallTraits<ValueType>::ParamType InValue)
+	{
+		const FResolvedPropertyAndFunction& PropAndFunction = FindOrAdd(Object);
+
+		if (ValueType* Val = PropAndFunction.GetPropertyAddress<ValueType>())
+		{
+			*Val = InValue;
+		}
+	}
+
+	/** @return the property path that this binding was initialized from */
+	const FString& GetPropertyPath() const
+	{
+		return PropertyPath;
+	}
+
+	/** @return the property name that this binding was initialized from */
+	const FName& GetPropertyName() const
+	{
+		return PropertyName;
+	}
+
+public:
+
+	/**
+	 * Static function for accessing a property value on an object without caching its address
+	 *
+	 * @param Object			The object to get the property from
+	 * @param InPropertyPath	The path to the property to retrieve
+	 * @return (Optional) The current value of the property on the object
+	 */
+	MOVIESCENE_API static TOptional<UE::MovieScene::FSourcePropertyValue> StaticValue(const UObject* Object, FStringView InPropertyPath);
+	MOVIESCENE_API static TOptional<TPair<const FProperty*, UE::MovieScene::FSourcePropertyValue>> StaticPropertyAndValue(const UObject* Object, FStringView InPropertyPath);
+
+	/**
+	 * Static function for accessing a property value on an object without caching its address
+	 *
+	 * @param Object			The object to get the property from
+	 * @param InPropertyPath	The path to the property to retrieve
+	 * @return (Optional) The current value of the property on the object
+	 */
+	template <typename ValueType>
+	static TOptional<ValueType> StaticValue(const UObject* Object, FStringView InPropertyPath)
+	{
+		checkf(Object, TEXT("No object specified"));
+
+		FResolvedPropertyAndFunction PropAndFunction = FindPropertyAndFunction(Object, InPropertyPath);
+
+		ValueType Value;
+		if (TryGetPropertyValue<ValueType>(PropAndFunction, Value))
+		{
+			return Value;
+		}
+
+		return TOptional<ValueType>();
+	}
+
+	/** Finds the property at the end of the given property path. */
+	static MOVIESCENE_API FProperty* FindProperty(const UObject* Object, FStringView InPropertyPath);
+
+private:
+
+	/**
+	 * Wrapper for UObject::ProcessEvent that attempts to pass the new property value directly to the function as a parameter,
+	 * but handles cases where multiple parameters or a return value exists. The setter parameter must be the first in the list,
+	 * any other parameters will be default constructed.
+	 */
+	template<typename T>
+	static void InvokeSetterFunction(UObject* InRuntimeObject, UFunction* Setter, T&& InPropertyValue);
+
+	using FCachedProperty = UE::MovieScene::FCachedProperty;
+	using FVolatileProperty = UE::MovieScene::FVolatileProperty;
+
+	struct FResolvedPropertyAndFunction
+	{
+		TVariant<FCachedProperty, FVolatileProperty> ResolvedProperty;
+		TWeakObjectPtr<UFunction> SetterFunction;
+		
+		MOVIESCENE_API FProperty* GetValidProperty() const;
+
+		MOVIESCENE_API void* GetContainerAddress() const;
+
+		template<typename ValueType>
+		ValueType* GetPropertyAddress() const
+		{
+			if (const FCachedProperty* CachedProperty = ResolvedProperty.TryGet<FCachedProperty>())
+			{
+				return CachedProperty->GetPropertyAddress<ValueType>();
+			}
+			else if (const FVolatileProperty* VolatileProperty = ResolvedProperty.TryGet<FVolatileProperty>())
+			{
+				return VolatileProperty->GetPropertyAddress<ValueType>();
+			}
+			else
+			{
+				return nullptr;
+			}
+		}
+
+		FResolvedPropertyAndFunction()
+			: ResolvedProperty()
+			, SetterFunction( nullptr )
+		{}
+	};
+
+	MOVIESCENE_API static FResolvedPropertyAndFunction FindPropertyAndFunction(const UObject* Object, FStringView InPropertyPath);
+
+	template <typename ValueType>
+	static bool TryGetPropertyValue(const FResolvedPropertyAndFunction& PropAndFunction, ValueType& OutValue)
+	{
+		if (const ValueType* Value = PropAndFunction.GetPropertyAddress<ValueType>())
+		{
+			OutValue = *Value;
+			return true;
+		}
+		return false;
+	}
+
+	/** Find or add the FResolvedPropertyAndFunction for the specified object */
+	MOVIESCENE_API const FResolvedPropertyAndFunction& FindOrAdd(const UObject& InObject);
+
+private:
+
+	/** Mapping of objects to bound functions that will be called to update data on the track */
+	TMap< FObjectKey, FResolvedPropertyAndFunction > RuntimeObjectToFunctionMap;
+
+	/** Path to the property we are bound to */
+	FString PropertyPath;
+
+	/** Name of the function to call to set values */
+	FName FunctionName;
+
+	/** Actual name of the property we are bound to */
+	FName PropertyName;
+
+	friend class FTrackInstancePropertyBindingsTests;
+};
+
+/** Explicit specializations for bools */
+template<> MOVIESCENE_API void FTrackInstancePropertyBindings::CallFunction<bool>(UObject& InRuntimeObject, TCallTraits<bool>::ParamType PropertyValue);
+template<> MOVIESCENE_API bool FTrackInstancePropertyBindings::TryGetPropertyValue<bool>(const FResolvedPropertyAndFunction& PropAndFunction, bool& OutValue);
+template<> MOVIESCENE_API void FTrackInstancePropertyBindings::SetCurrentValue<bool>(UObject& Object, TCallTraits<bool>::ParamType InValue);
+
+/** Explicit specializations for object pointers */
+template<> MOVIESCENE_API void FTrackInstancePropertyBindings::CallFunction<UObject*>(UObject& InRuntimeObject, UObject* PropertyValue);
+template<> MOVIESCENE_API bool FTrackInstancePropertyBindings::TryGetPropertyValue<UObject*>(const FResolvedPropertyAndFunction& PropAndFunction, UObject*& OutValue);
+template<> MOVIESCENE_API void FTrackInstancePropertyBindings::SetCurrentValue<UObject*>(UObject& Object, UObject* InValue);
+
+template<typename T>
+void FTrackInstancePropertyBindings::InvokeSetterFunction(UObject* InRuntimeObject, UFunction* Setter, T&& InPropertyValue)
+{
+	// CacheBinding already guarantees that the function has >= 1 parameters
+	const int32 ParmsSize = Setter->ParmsSize;
+
+	// This should all be const really, but ProcessEvent only takes a non-const void*
+	void* InputParameter = const_cast<typename TDecay<T>::Type*>(&InPropertyValue);
+
+	// By default we try and use the existing stack value
+	uint8* Params = reinterpret_cast<uint8*>(InputParameter);
+
+	check(InRuntimeObject && Setter);
+	if (Setter->ReturnValueOffset != MAX_uint16 || Setter->NumParms > 0)
+	{
+		// Function has a return value or multiple parameters, we need to initialize memory for the entire parameter pack
+		// We use alloca here (as in UObject::ProcessEvent) to avoid a heap allocation. Alloca memory survives the current function's stack frame.
+		Params = reinterpret_cast<uint8*>(FMemory_Alloca(ParmsSize));
+
+		bool bFirstProperty = true;
+		for (FProperty* Property = Setter->PropertyLink; Property; Property = Property->PropertyLinkNext)
+		{
+			// Initialize the parameter pack with any param properties that reside in the container
+			if (Property->IsInContainer(ParmsSize))
+			{
+				Property->InitializeValue_InContainer(Params);
+
+				// The first encountered property is assumed to be the input value so initialize this with the user-specified value from InPropertyValue
+				if (Property->HasAnyPropertyFlags(CPF_Parm) && !Property->HasAnyPropertyFlags(CPF_ReturnParm) && bFirstProperty)
+				{
+					const bool bIsValid = ensureMsgf(sizeof(T) == Property->GetElementSize(), TEXT("Property type does not match for Sequencer setter function %s::%s (%" SIZE_T_FMT "bytes != %ibytes"), *InRuntimeObject->GetName(), *Setter->GetName(), sizeof(T), Property->GetElementSize());
+					if (bIsValid)
+					{
+						Property->CopyCompleteValue(Property->ContainerPtrToValuePtr<void>(Params), &InPropertyValue);
+					}
+					else
+					{
+						return;
+					}
+				}
+				bFirstProperty = false;
+			}
+		}
+	}
+
+	// Now we have the parameters set up correctly, call the function
+	InRuntimeObject->ProcessEvent(Setter, Params);
+}
+
