Sequencer: Initialize bool
#rnx

diff --git a/Engine/Source/Runtime/MovieScene/Public/Decorations/MovieSceneMuteSoloDecoration.h b/Engine/Source/Runtime/MovieScene/Public/Decorations/MovieSceneMuteSoloDecoration.h
--- a/Engine/Source/Runtime/MovieScene/Public/Decorations/MovieSceneMuteSoloDecoration.h
+++ b/Engine/Source/Runtime/MovieScene/Public/Decorations/MovieSceneMuteSoloDecoration.h
@@ -10,8 +10,8 @@
 
 struct FMovieSceneMuteSoloData
 {
-	bool bMuted;
-	bool bSoloed;
+	bool bMuted = false;
+	bool bSoloed = false;
 };
 
 UCLASS(MinimalAPI, Transient)
