VCam: Fix GetXByAssetData suite of functions, such as ULevelSequenceShotMetaDataLibrary::GetIsFlaggedByAssetData, not parsing bool tags correctly.

diff --git a/Engine/Source/Runtime/MovieScene/Private/MetaData/MovieSceneShotMetaData.cpp b/Engine/Source/Runtime/MovieScene/Private/MetaData/MovieSceneShotMetaData.cpp
--- a/Engine/Source/Runtime/MovieScene/Private/MetaData/MovieSceneShotMetaData.cpp
+++ b/Engine/Source/Runtime/MovieScene/Private/MetaData/MovieSceneShotMetaData.cpp
@@ -21,10 +21,45 @@ static bool ParseAssetTag(const FAssetData& InAssetData, FName InTag, TTagValue&
 	// ExtendAssetRegistryTags sets "Unset" if the optional value was unset. Users could have also messed with the registry.
 	if (StringValue.IsNumeric()) 
 	{
-		// Ok to demote int to bool
 		Value = static_cast<TTagValue>(FCString::Atoi(*StringValue));
 		return true;
 	}
+	
+	return false;
+}
+
+static const TCHAR* ConvertBoolToYesNoString(bool bInValue)
+{
+	return bInValue ? TEXT("Yes") : TEXT("No");
+}
+
+/** We need special handling for "bool" tags. We save them as "Yes" / "No" strings. */
+template<>
+bool ParseAssetTag(const FAssetData& InAssetData, FName InTag, bool& Value)
+{
+	FString StringValue;
+	InAssetData.GetTagValue(InTag, StringValue);
+
+	// Starting 5.7, bool tags are "Yes" and "No"
+	if (StringValue == ConvertBoolToYesNoString(true))
+	{
+		Value = true;
+		return true;
+	}
+	if (StringValue == ConvertBoolToYesNoString(false))
+	{
+		Value = false;
+		return true;
+	}
+
+	// Handle pre 5.7, where the tags were stored as strings "1" -> true, "0" -> false
+	if (StringValue.IsNumeric()) 
+	{
+		// Ok to demote int to bool
+		Value = static_cast<bool>(FCString::Atoi(*StringValue));
+		return true;
+	}
+
 	return false;
 }
 }
@@ -62,26 +97,23 @@ bool UMovieSceneShotMetaData::GetFavoriteRatingByAssetData(const FAssetData& InA
 
 void UMovieSceneShotMetaData::ExtendAssetRegistryTags(FAssetRegistryTagsContext Context) const
 {
-	auto BoolToYesNoString = [](bool bInValue) -> FString
-		{
-			return bInValue ? TEXT("Yes") : TEXT("No");
-		};
-
+	using namespace UE::MovieScene::ShotMetaDataDetail;
+	
 	Context.AddTag(FAssetRegistryTag(
-		AssetRegistryTag_bIsNoGood, bIsNoGood.IsSet() ? BoolToYesNoString(*bIsNoGood) : TEXT(""),
-		FAssetRegistryTag::ETagType::TT_Numerical, FAssetRegistryTag::TD_None
+		AssetRegistryTag_bIsNoGood, bIsNoGood.IsSet() ? ConvertBoolToYesNoString(*bIsNoGood) : TEXT(""),
+		FAssetRegistryTag::ETagType::TT_Alphabetical, FAssetRegistryTag::TD_None
 	));
 	Context.AddTag(FAssetRegistryTag(
-		AssetRegistryTag_bIsFlagged, bIsFlagged.IsSet() ? BoolToYesNoString(*bIsFlagged) : TEXT(""),
-		FAssetRegistryTag::ETagType::TT_Numerical, FAssetRegistryTag::TD_None
+		AssetRegistryTag_bIsFlagged, bIsFlagged.IsSet() ? ConvertBoolToYesNoString(*bIsFlagged) : TEXT(""),
+		FAssetRegistryTag::ETagType::TT_Alphabetical, FAssetRegistryTag::TD_None
 	));
 	Context.AddTag(FAssetRegistryTag(
-		AssetRegistryTag_bIsRecorded, bIsRecorded.IsSet() ? BoolToYesNoString(*bIsRecorded) : TEXT(""),
-		FAssetRegistryTag::ETagType::TT_Numerical, FAssetRegistryTag::TD_None
+		AssetRegistryTag_bIsRecorded, bIsRecorded.IsSet() ? ConvertBoolToYesNoString(*bIsRecorded) : TEXT(""),
+		FAssetRegistryTag::ETagType::TT_Alphabetical, FAssetRegistryTag::TD_None
 	));
 	Context.AddTag(FAssetRegistryTag(
-		AssetRegistryTag_bIsSubSequence, bIsSubSequence.IsSet() ? BoolToYesNoString(*bIsSubSequence) : TEXT(""),
-		FAssetRegistryTag::ETagType::TT_Numerical, FAssetRegistryTag::TD_None
+		AssetRegistryTag_bIsSubSequence, bIsSubSequence.IsSet() ? ConvertBoolToYesNoString(*bIsSubSequence) : TEXT(""),
+		FAssetRegistryTag::ETagType::TT_Alphabetical, FAssetRegistryTag::TD_None
 	));
 	Context.AddTag(FAssetRegistryTag(
 		AssetRegistryTag_FavoriteRating, FavoriteRating.IsSet() ? FString::FromInt(*FavoriteRating) : TEXT(""),
