Sequencer: Made FMovieSceneSubSequenceData::WeakCondition a soft object ptr rather than a weak ptr to resolve cooker link errors

#DisabledValidations nonUnityAndStaticAnalysis

diff --git a/Engine/Source/Runtime/MovieScene/Public/Evaluation/MovieSceneSequenceHierarchy.h b/Engine/Source/Runtime/MovieScene/Public/Evaluation/MovieSceneSequenceHierarchy.h
--- a/Engine/Source/Runtime/MovieScene/Public/Evaluation/MovieSceneSequenceHierarchy.h
+++ b/Engine/Source/Runtime/MovieScene/Public/Evaluation/MovieSceneSequenceHierarchy.h
@@ -178,7 +178,7 @@ struct FMovieSceneSubSequenceData
 	
 	/* Optional condition pulled off of the UMovieSceneSubSection or UMovieSceneSubTrack */
 	UPROPERTY()
-	TWeakObjectPtr<const UMovieSceneCondition> WeakCondition;
+	TSoftObjectPtr<const UMovieSceneCondition> WeakCondition;
 
 
 private:
