Sequencer: Follow up to 45319395 where multiple objects could be returned with FMovieSceneBindingResolveResult. A binding that did not resolve with any valid actor was adding a nullptr to OutObjects.

diff --git a/Engine/Source/Runtime/MovieScene/Private/MovieSceneBindingReferences.cpp b/Engine/Source/Runtime/MovieScene/Private/MovieSceneBindingReferences.cpp
--- a/Engine/Source/Runtime/MovieScene/Private/MovieSceneBindingReferences.cpp
+++ b/Engine/Source/Runtime/MovieScene/Private/MovieSceneBindingReferences.cpp
@@ -193,7 +193,10 @@ void FMovieSceneBindingReferences::ResolveBindingInternal(const FMovieSceneBindi
 	else
 	{
 		// Otherwise, attempt to resolve via the locator
-		OutObjects.Add(ResolveBindingFromLocator(InternalIndex, LocatorResolveParams));
+		if (UObject* ResolvedObject = ResolveBindingFromLocator(InternalIndex, LocatorResolveParams))
+		{
+			OutObjects.Add(ResolvedObject);
+		}
 	}
 }
 
