Sequencer: Use GetNameSafe

#rnx

diff --git a/Engine/Source/Runtime/MovieScene/Private/MovieSceneSequencePlayer.cpp b/Engine/Source/Runtime/MovieScene/Private/MovieSceneSequencePlayer.cpp
--- a/Engine/Source/Runtime/MovieScene/Private/MovieSceneSequencePlayer.cpp
+++ b/Engine/Source/Runtime/MovieScene/Private/MovieSceneSequencePlayer.cpp
@@ -1539,7 +1539,7 @@ FString UMovieSceneSequencePlayer::GetSequenceName(bool bAddClientInfo) const
 {
 	if (Sequence)
 	{
-		FString SequenceName = Sequence->GetName();
+		FString SequenceName = GetNameSafe(Sequence);
 		if (bAddClientInfo)
 		{
 			AActor* Actor = GetTypedOuter<AActor>();
