Sequencer: Fix GetBoundObjects so that it now takes a sequence ID instead of a sequence.

Taking a single sequence is ambiguous because the sequence could appear in the hierarchy more than once.

diff --git a/Engine/Plugins/MovieScene/LevelSequenceEditor/Source/LevelSequenceEditor/Private/LevelSequenceEditorSubsystem.cpp b/Engine/Plugins/MovieScene/LevelSequenceEditor/Source/LevelSequenceEditor/Private/LevelSequenceEditorSubsystem.cpp
--- a/Engine/Plugins/MovieScene/LevelSequenceEditor/Source/LevelSequenceEditor/Private/LevelSequenceEditorSubsystem.cpp
+++ b/Engine/Plugins/MovieScene/LevelSequenceEditor/Source/LevelSequenceEditor/Private/LevelSequenceEditorSubsystem.cpp
@@ -271,7 +271,7 @@ void FMovieSceneBindingPropertyInfoListCustomization::CustomizeDetails(IDetailLa
 						}
 					}
 					bool bHasBoundObject = false;
-					TArray<UObject*> BoundObjects = MovieSceneHelpers::GetBoundObjects(Sequence, BindingGuid, Sequencer->GetSharedPlaybackState(), 0);
+					TArray<UObject*> BoundObjects = MovieSceneHelpers::GetBoundObjects(Sequencer->GetFocusedTemplateID(), BindingGuid, Sequencer->GetSharedPlaybackState(), 0);
 					for (UObject* BoundObject : BoundObjects)
 					{
 						if (BoundObject)
diff --git a/Engine/Source/Editor/Sequencer/Private/SequencerUtilities.cpp b/Engine/Source/Editor/Sequencer/Private/SequencerUtilities.cpp
--- a/Engine/Source/Editor/Sequencer/Private/SequencerUtilities.cpp
+++ b/Engine/Source/Editor/Sequencer/Private/SequencerUtilities.cpp
@@ -1345,7 +1345,7 @@ FMovieScenePossessable* FSequencerUtilities::ConvertToPossessable(TSharedRef<ISe
 	}
 
 	// For now, just get the first bound object to convert
-	TArray<UObject*> BoundObjects = MovieSceneHelpers::GetBoundObjects(Sequence, BindingGuid, Sequencer->GetSharedPlaybackState(), BindingIndex);
+	TArray<UObject*> BoundObjects = MovieSceneHelpers::GetBoundObjects(Sequencer->GetFocusedTemplateID(), BindingGuid, Sequencer->GetSharedPlaybackState(), BindingIndex);
 	UObject* ObjectToConvert = BoundObjects.Num() > 0 ? BoundObjects[0] : nullptr;
 
 	// If we have an old-style spawnable, use the template as the object to convert instead.
@@ -1556,7 +1556,7 @@ FMovieScenePossessable* FSequencerUtilities::ConvertToCustomBinding(TSharedRef<I
 		return CreatedPossessable;
 	}
 
-	TArray<UObject*> BoundObjects = MovieSceneHelpers::GetBoundObjects(Sequence, BindingGuid, Sequencer->GetSharedPlaybackState(), BindingIndex);
+	TArray<UObject*> BoundObjects = MovieSceneHelpers::GetBoundObjects(Sequencer->GetFocusedTemplateID(), BindingGuid, Sequencer->GetSharedPlaybackState(), BindingIndex);
 	UObject* ObjectToConvert = BoundObjects.Num() > 0 ? BoundObjects[0] : nullptr;
 
 	bool bConvertFromSpawnable = false;
@@ -2605,7 +2605,7 @@ void FSequencerUtilities::CopyBindings(TSharedRef<ISequencer> Sequencer, const T
 					}
 					else
 					{
-						TArray<UObject*> BoundObjects = MovieSceneHelpers::GetBoundObjects(ObjectBinding.Sequence, ObjectBinding.BindingID, Sequencer->GetSharedPlaybackState(), BindingIndex);
+						TArray<UObject*> BoundObjects = MovieSceneHelpers::GetBoundObjects(Sequencer->GetFocusedTemplateID(), ObjectBinding.BindingID, Sequencer->GetSharedPlaybackState(), BindingIndex);
 						if (BoundObjects.Num() > 0)
 						{
 							CopyableBinding->BoundObjectNames.Add(BoundObjects[0]->GetPathName(World));
diff --git a/Engine/Source/Runtime/MovieScene/Private/MovieSceneCommonHelpers.cpp b/Engine/Source/Runtime/MovieScene/Private/MovieSceneCommonHelpers.cpp
--- a/Engine/Source/Runtime/MovieScene/Private/MovieSceneCommonHelpers.cpp
+++ b/Engine/Source/Runtime/MovieScene/Private/MovieSceneCommonHelpers.cpp
@@ -838,66 +838,68 @@ FGuid MovieSceneHelpers::TryCreateCustomSpawnableBinding(UMovieSceneSequence* Se
 
 UObject* MovieSceneHelpers::GetSingleBoundObject(UMovieSceneSequence* Sequence, const FGuid& ObjectId, TSharedRef<const UE::MovieScene::FSharedPlaybackState> SharedPlaybackState, int32 BindingIndex)
 {
-	for (UObject* BoundObject : GetBoundObjects(Sequence, ObjectId, SharedPlaybackState, BindingIndex))
+	if (FMovieSceneEvaluationState* EvaluationState = SharedPlaybackState->FindCapability<FMovieSceneEvaluationState>())
 	{
-		if (BoundObject)
+		FMovieSceneSequenceIDRef SequenceID = EvaluationState->FindSequenceId(Sequence);
+		for (UObject* BoundObject : GetBoundObjects(SequenceID, ObjectId, SharedPlaybackState, BindingIndex))
 		{
-			return BoundObject;
+			if (BoundObject)
+			{
+				return BoundObject;
+			}
 		}
 	}
 	return nullptr;
 }
 
-TArray<UObject*> MovieSceneHelpers::GetBoundObjects(UMovieSceneSequence* Sequence, const FGuid& ObjectId, TSharedRef<const UE::MovieScene::FSharedPlaybackState> SharedPlaybackState, int32 BindingIndex)
+TArray<UObject*> MovieSceneHelpers::GetBoundObjects(const FMovieSceneSequenceID& SequenceID, const FGuid& ObjectId, TSharedRef<const UE::MovieScene::FSharedPlaybackState> SharedPlaybackState, int32 BindingIndex)
 {
 	TArray<UObject*> BoundObjects;
-	if (Sequence)
+	if (FMovieSceneEvaluationState* EvaluationState = SharedPlaybackState->FindCapability<FMovieSceneEvaluationState>())
 	{
-		if (UMovieScene* MovieScene = Sequence->GetMovieScene())
+		UMovieSceneSequence* Sequence = EvaluationState->FindSequence(SequenceID);
+		UMovieScene* MovieScene = Sequence ? Sequence->GetMovieScene() : nullptr;
+
+		if (MovieScene)
 		{
-			if (FMovieSceneEvaluationState* EvaluationState = SharedPlaybackState->FindCapability<FMovieSceneEvaluationState>())
+			if (MovieScene->FindSpawnable(ObjectId))
 			{
-				FMovieSceneSequenceIDRef SequenceID = EvaluationState->FindSequenceId(Sequence);
-
-				if (MovieScene->FindSpawnable(ObjectId))
+				for (TWeakObjectPtr<> BoundObject : EvaluationState->FindBoundObjects(FMovieSceneEvaluationOperand(SequenceID, ObjectId), SharedPlaybackState))
 				{
-					for (TWeakObjectPtr<> BoundObject : EvaluationState->FindBoundObjects(FMovieSceneEvaluationOperand(SequenceID, ObjectId), SharedPlaybackState))
+					if (BoundObject.Get())
 					{
-						if (BoundObject.Get())
-						{
-							BoundObjects.Add(BoundObject.Get());
-						}
+						BoundObjects.Add(BoundObject.Get());
 					}
 				}
-				else if (const FMovieScenePossessable* Possessable = MovieScene->FindPossessable(ObjectId))
+			}
+			else if (const FMovieScenePossessable* Possessable = MovieScene->FindPossessable(ObjectId))
+			{
+				const FMovieSceneBindingReferences* Refs = Sequence->GetBindingReferences();
+				if (Refs)
 				{
-					const FMovieSceneBindingReferences* Refs = Sequence->GetBindingReferences();
-					if (Refs)
-					{
-						UObject* ResolutionContext = MovieSceneHelpers::GetResolutionContext(Sequence, ObjectId, SequenceID, SharedPlaybackState);
+					UObject* ResolutionContext = MovieSceneHelpers::GetResolutionContext(Sequence, ObjectId, SequenceID, SharedPlaybackState);
 
-						if (Possessable->GetParent().IsValid() && Sequence->AreParentContextsSignificant())
+					if (Possessable->GetParent().IsValid() && Sequence->AreParentContextsSignificant())
+					{
+						TArrayView<TWeakObjectPtr<>> ParentBoundObjects = EvaluationState->FindBoundObjects(FMovieSceneEvaluationOperand(SequenceID, Possessable->GetParent()), SharedPlaybackState);
+						for (TWeakObjectPtr<> Parent : ParentBoundObjects)
 						{
-							TArrayView<TWeakObjectPtr<>> ParentBoundObjects = EvaluationState->FindBoundObjects(FMovieSceneEvaluationOperand(SequenceID, Possessable->GetParent()), SharedPlaybackState);
-							for (TWeakObjectPtr<> Parent : ParentBoundObjects)
+							ResolutionContext = Parent.Get();
+							if (!ResolutionContext)
 							{
-								ResolutionContext = Parent.Get();
-								if (!ResolutionContext)
-								{
-									continue;
-								}
+								continue;
 							}
 						}
+					}
 
-						TArray<UObject*, TInlineAllocator<1>> SingleBoundObjects;
-						UE::UniversalObjectLocator::FResolveParams LocatorResolveParams(ResolutionContext);
-						FMovieSceneBindingResolveParams BindingResolveParams{ Sequence, ObjectId, SequenceID, ResolutionContext };
-						Refs->ResolveSingleBinding(BindingResolveParams, BindingIndex, LocatorResolveParams, SharedPlaybackState, SingleBoundObjects);
+					TArray<UObject*, TInlineAllocator<1>> SingleBoundObjects;
+					UE::UniversalObjectLocator::FResolveParams LocatorResolveParams(ResolutionContext);
+					FMovieSceneBindingResolveParams BindingResolveParams{ Sequence, ObjectId, SequenceID, ResolutionContext };
+					Refs->ResolveSingleBinding(BindingResolveParams, BindingIndex, LocatorResolveParams, SharedPlaybackState, SingleBoundObjects);
 
-						for (UObject* BoundObject : SingleBoundObjects)
-						{
-							BoundObjects.Add(BoundObject);
-						}
+					for (UObject* BoundObject : SingleBoundObjects)
+					{
+						BoundObjects.Add(BoundObject);
 					}
 				}
 			}
diff --git a/Engine/Source/Runtime/MovieScene/Public/MovieSceneCommonHelpers.h b/Engine/Source/Runtime/MovieScene/Public/MovieSceneCommonHelpers.h
--- a/Engine/Source/Runtime/MovieScene/Public/MovieSceneCommonHelpers.h
+++ b/Engine/Source/Runtime/MovieScene/Public/MovieSceneCommonHelpers.h
@@ -246,7 +246,7 @@ public:
 	/*
 	* Returns the objects currently bound to the given objectid and binding index (optional).
 	*/
-	static MOVIESCENE_API TArray<UObject*> GetBoundObjects(UMovieSceneSequence* Sequence, const FGuid& ObjectId, TSharedRef<const UE::MovieScene::FSharedPlaybackState> SharedPlaybackState, int32 BindingIndex = 0);
+	static MOVIESCENE_API TArray<UObject*> GetBoundObjects(const FMovieSceneSequenceID& SequenceID, const FGuid& ObjectId, TSharedRef<const UE::MovieScene::FSharedPlaybackState> SharedPlaybackState, int32 BindingIndex = 0);
 
 	UE_DEPRECATED(5.7, "This method has been deprecated. Please use GetBoundObjects which returns an array of bound objects.")
 	static MOVIESCENE_API UObject* GetSingleBoundObject(UMovieSceneSequence* Sequence, const FGuid& ObjectId, TSharedRef<const UE::MovieScene::FSharedPlaybackState> SharedPlaybackState, int32 BindingIndex = 0);
