ShotManagement: Added in SubSequence compatibility functions to UMovieSceneSequence to allow checking of compatability of a subsequence when adding it to a parent sequence.
CineAssembly's report that they are compatible as a subsequence only if Data Only is false to prevent Data Only CineAssemblys from being added as SubSequences in other Sequences.

diff --git a/Engine/Plugins/VirtualProduction/CinematicAssemblyTools/Source/CineAssemblyTools/Private/CineAssembly.cpp b/Engine/Plugins/VirtualProduction/CinematicAssemblyTools/Source/CineAssemblyTools/Private/CineAssembly.cpp
--- a/Engine/Plugins/VirtualProduction/CinematicAssemblyTools/Source/CineAssemblyTools/Private/CineAssembly.cpp
+++ b/Engine/Plugins/VirtualProduction/CinematicAssemblyTools/Source/CineAssemblyTools/Private/CineAssembly.cpp
@@ -116,6 +116,11 @@ void UCineAssembly::PostLoad()
 	}
 }
 
+bool UCineAssembly::IsCompatibleAsSubSequence(const UMovieSceneSequence& ParentSequence) const
+{
+	return !bIsDataOnly;
+}
+
 void UCineAssembly::Initialize()
 {
 	Super::Initialize();
diff --git a/Engine/Plugins/VirtualProduction/CinematicAssemblyTools/Source/CineAssemblyTools/Public/CineAssembly.h b/Engine/Plugins/VirtualProduction/CinematicAssemblyTools/Source/CineAssemblyTools/Public/CineAssembly.h
--- a/Engine/Plugins/VirtualProduction/CinematicAssemblyTools/Source/CineAssemblyTools/Public/CineAssembly.h
+++ b/Engine/Plugins/VirtualProduction/CinematicAssemblyTools/Source/CineAssemblyTools/Public/CineAssembly.h
@@ -25,6 +25,10 @@ public:
 
 	friend class SCineAssemblyConfigWindow;
 
+	//~ Begin UMovieSceneSequence Interface
+	CINEASSEMBLYTOOLS_API virtual bool IsCompatibleAsSubSequence(const UMovieSceneSequence& ParentSequence) const override;
+	//~ End UMovieSceneSequence Interface
+
 	//~ Begin ULevelSquence Interface
 	CINEASSEMBLYTOOLS_API virtual void Initialize() override;
 	//~ End ULevelSquence Interface
diff --git a/Engine/Source/Editor/MovieSceneTools/Private/TrackEditors/SubTrackEditor.cpp b/Engine/Source/Editor/MovieSceneTools/Private/TrackEditors/SubTrackEditor.cpp
--- a/Engine/Source/Editor/MovieSceneTools/Private/TrackEditors/SubTrackEditor.cpp
+++ b/Engine/Source/Editor/MovieSceneTools/Private/TrackEditors/SubTrackEditor.cpp
@@ -550,7 +550,7 @@ bool FSubTrackEditor::HandleAssetAdded(UObject* Asset, const FGuid& TargetObject
 		return true;
 	}
 
-	FNotificationInfo Info(FText::Format( LOCTEXT("InvalidSequence", "Invalid level sequence {0}. There could be a circular dependency."), Sequence->GetDisplayName()));	
+	FNotificationInfo Info(FText::Format( LOCTEXT("InvalidSequence", "Invalid level sequence {0}. There could be a circular dependency or the sequences are not compatible."), Sequence->GetDisplayName()));	
 	Info.bUseLargeFont = false;
 	FSlateNotificationManager::Get().AddNotification(Info);
 
@@ -1694,7 +1694,7 @@ FKeyPropertyResult FSubTrackEditor::AddKeyInternal(FFrameNumber KeyTime, UMovieS
 		return KeyPropertyResult;
 	}
 
-	FNotificationInfo Info(FText::Format( LOCTEXT("InvalidSequence", "Invalid level sequence {0}. There could be a circular dependency."), InMovieSceneSequence->GetDisplayName()));
+	FNotificationInfo Info(FText::Format( LOCTEXT("InvalidSequence", "Invalid level sequence {0}. There could be a circular dependency or the sequences are not compatible."), InMovieSceneSequence->GetDisplayName()));
 	Info.bUseLargeFont = false;
 	FSlateNotificationManager::Get().AddNotification(Info);
 
diff --git a/Engine/Source/Editor/MovieSceneTools/Private/TrackEditors/SubTrackEditorBase.cpp b/Engine/Source/Editor/MovieSceneTools/Private/TrackEditors/SubTrackEditorBase.cpp
--- a/Engine/Source/Editor/MovieSceneTools/Private/TrackEditors/SubTrackEditorBase.cpp
+++ b/Engine/Source/Editor/MovieSceneTools/Private/TrackEditors/SubTrackEditorBase.cpp
@@ -424,12 +424,17 @@ void FSubSectionEditorUtil::DilateSection(const TRange<FFrameNumber>& NewRange,
 
 bool FSubTrackEditorUtil::CanAddSubSequence(const UMovieSceneSequence* CurrentSequence, const UMovieSceneSequence& SubSequence)
 {
-	// Prevent adding ourselves and ensure we have a valid movie scene.
+	// Prevent adding ourselves, check the SubSequence is a compatible subsequence and ensure we have a valid movie scene.
 	if ((CurrentSequence == nullptr) || (CurrentSequence == &SubSequence) || (CurrentSequence->GetMovieScene() == nullptr))
 	{
 		return false;
 	}
 
+	if (!CurrentSequence->IsSubSequenceCompatible(SubSequence))
+	{
+		return false;
+	}
+
 	// ensure that the other sequence has a valid movie scene
 	UMovieScene* SequenceMovieScene = SubSequence.GetMovieScene();
 
diff --git a/Engine/Source/Runtime/MovieScene/Private/MovieSceneSequence.cpp b/Engine/Source/Runtime/MovieScene/Private/MovieSceneSequence.cpp
--- a/Engine/Source/Runtime/MovieScene/Private/MovieSceneSequence.cpp
+++ b/Engine/Source/Runtime/MovieScene/Private/MovieSceneSequence.cpp
@@ -438,3 +438,12 @@ bool UMovieSceneSequence::IsFilterSupported(const FString& InFilterName) const
 }
 
 #endif
+
+bool UMovieSceneSequence::IsSubSequenceCompatible(const UMovieSceneSequence& SubSequence) const
+{
+	if (IsCompatibleSubSequence(SubSequence))
+	{
+		return SubSequence.IsCompatibleAsSubSequence(*this);
+	}
+	return false;
+}
diff --git a/Engine/Source/Runtime/MovieScene/Private/Sections/MovieSceneSubSection.cpp b/Engine/Source/Runtime/MovieScene/Private/Sections/MovieSceneSubSection.cpp
--- a/Engine/Source/Runtime/MovieScene/Private/Sections/MovieSceneSubSection.cpp
+++ b/Engine/Source/Runtime/MovieScene/Private/Sections/MovieSceneSubSection.cpp
@@ -516,6 +516,16 @@ void UMovieSceneSubSection::ImportEntityImpl(UMovieSceneEntitySystemLinker* Enti
 
 void UMovieSceneSubSection::SetSequence(UMovieSceneSequence* Sequence)
 {
+	if (Sequence != nullptr)
+	{
+		// Confirm compatibility of the sequence as a subsequence with the parent sequence
+		if (UMovieSceneSequence* CurrentSequence = GetTypedOuter<UMovieSceneSequence>();
+			CurrentSequence != nullptr && !CurrentSequence->IsSubSequenceCompatible(*Sequence))
+		{
+			return;
+		}
+	}
+
 	if (!TryModify())
 	{
 		return;
@@ -568,38 +578,47 @@ void UMovieSceneSubSection::PostEditChangeProperty(FPropertyChangedEvent& Proper
 		{
 			if (UMovieSceneSequence* CurrentSequence = TrackOuter->GetTypedOuter<UMovieSceneSequence>())
 			{
-				TArray<UMovieSceneSubTrack*> SubTracks;
-
-				for (UMovieSceneTrack* Track : SubSequenceMovieScene->GetTracks())
+				// Confirm that the newly set subsequence is compatible, otherwise we need to revert to previous
+				if (CurrentSequence->IsSubSequenceCompatible(*SubSequence))
 				{
-					if (UMovieSceneSubTrack* SubTrack = Cast<UMovieSceneSubTrack>(Track))
-					{
-						SubTracks.Add(SubTrack);
-					}
-				}
+					TArray<UMovieSceneSubTrack*> SubTracks;
 
-				for (const FMovieSceneBinding& Binding : SubSequenceMovieScene->GetBindings())
-				{
-					for (UMovieSceneTrack* Track : SubSequenceMovieScene->FindTracks(UMovieSceneSubTrack::StaticClass(), Binding.GetObjectGuid()))
+					for (UMovieSceneTrack* Track : SubSequenceMovieScene->GetTracks())
 					{
 						if (UMovieSceneSubTrack* SubTrack = Cast<UMovieSceneSubTrack>(Track))
 						{
 							SubTracks.Add(SubTrack);
 						}
 					}
-				}
 
-				for (UMovieSceneSubTrack* SubTrack : SubTracks)
-				{
-					if ( SubTrack->ContainsSequence(*CurrentSequence, true))
+					for (const FMovieSceneBinding& Binding : SubSequenceMovieScene->GetBindings())
 					{
-						UE_LOG(LogMovieScene, Error, TEXT("Invalid level sequence %s. It is already contained by: %s."), *SubSequence->GetDisplayName().ToString(), *CurrentSequence->GetDisplayName().ToString());
+						for (UMovieSceneTrack* Track : SubSequenceMovieScene->FindTracks(UMovieSceneSubTrack::StaticClass(), Binding.GetObjectGuid()))
+						{
+							if (UMovieSceneSubTrack* SubTrack = Cast<UMovieSceneSubTrack>(Track))
+							{
+								SubTracks.Add(SubTrack);
+							}
+						}
+					}
 
-						// Restore to the previous sub sequence because there was a circular dependency
-						SubSequence = PreviousSubSequence;
-						break;
+					for (UMovieSceneSubTrack* SubTrack : SubTracks)
+					{
+						if (SubTrack->ContainsSequence(*CurrentSequence, true))
+						{
+							UE_LOG(LogMovieScene, Error, TEXT("Invalid level sequence %s. It is already contained by: %s."), *SubSequence->GetDisplayName().ToString(), *CurrentSequence->GetDisplayName().ToString());
+
+							// Restore to the previous sub sequence because there was a circular dependency
+							SubSequence = PreviousSubSequence;
+							break;
+						}
 					}
 				}
+				else
+				{
+					UE_LOG(LogMovieScene, Error, TEXT("Level sequence %s is not a compatible subsequence for %s."), *SubSequence->GetDisplayName().ToString(), *CurrentSequence->GetDisplayName().ToString());
+					SubSequence = PreviousSubSequence;
+				}
 			}
 		}
 
diff --git a/Engine/Source/Runtime/MovieScene/Public/MovieSceneSequence.h b/Engine/Source/Runtime/MovieScene/Public/MovieSceneSequence.h
--- a/Engine/Source/Runtime/MovieScene/Public/MovieSceneSequence.h
+++ b/Engine/Source/Runtime/MovieScene/Public/MovieSceneSequence.h
@@ -450,4 +450,19 @@ public:
 	virtual bool IsFilterSupportedImpl(const FString& InFilterName) const { return true; }
 
 #endif
+
+	/** 
+	 * Given a sequence, checks if it is compatible as a subsequence to this sequence. 
+	 */
+	MOVIESCENE_API bool IsSubSequenceCompatible(const UMovieSceneSequence& SubSequence) const;
+
+	/**
+	 * Sequences can determine if the given sequence would be compatible as a subsequence.
+	 */
+	virtual bool IsCompatibleSubSequence(const UMovieSceneSequence& SubSequence) const { return true; }
+
+	/**
+	 * Sequences can determine if they would be compatible as a subsequence to the given parent sequence.
+	 */
+	virtual bool IsCompatibleAsSubSequence(const UMovieSceneSequence& ParentSequence) const { return true; }
 };
