Sequencer: Mark the return Object as deprecated rather than deprecating it completely so that licensee code continues to operate as usual and they can upgrade.

diff --git a/Engine/Plugins/Importers/USDImporter/Source/USDStage/Private/USDDynamicBindingResolverLibrary.cpp b/Engine/Plugins/Importers/USDImporter/Source/USDStage/Private/USDDynamicBindingResolverLibrary.cpp
--- a/Engine/Plugins/Importers/USDImporter/Source/USDStage/Private/USDDynamicBindingResolverLibrary.cpp
+++ b/Engine/Plugins/Importers/USDImporter/Source/USDStage/Private/USDDynamicBindingResolverLibrary.cpp
@@ -58,10 +58,12 @@ FMovieSceneDynamicBindingResolveResult UUsdDynamicBindingResolverLibrary::Resolv
 						if (bIsActorBinding)
 						{
 							Result.Objects.Add(FoundComponent->GetOwner());
+							Result.Object = FoundComponent->GetOwner(); // Object deprecated 5.7
 						}
 						else
 						{
 							Result.Objects.Add(FoundComponent);
+							Result.Object = FoundComponent; // Object deprecated 5.7
 						}
 						break;
 					}
diff --git a/Engine/Source/Runtime/MovieScene/Private/Bindings/MovieSceneSpawnableBinding.cpp b/Engine/Source/Runtime/MovieScene/Private/Bindings/MovieSceneSpawnableBinding.cpp
--- a/Engine/Source/Runtime/MovieScene/Private/Bindings/MovieSceneSpawnableBinding.cpp
+++ b/Engine/Source/Runtime/MovieScene/Private/Bindings/MovieSceneSpawnableBinding.cpp
@@ -199,6 +199,7 @@ FMovieSceneBindingResolveResult UMovieSceneSpawnableBindingBase::ResolveBinding(
 	if (SpawnedObject)
 	{
 		Result.Objects.Add(SpawnedObject);
+		Result.Object = SpawnedObject; // Object deprecated 5.7
 	}
 	return Result;
 }
diff --git a/Engine/Source/Runtime/MovieScene/Private/MovieSceneDynamicBinding.cpp b/Engine/Source/Runtime/MovieScene/Private/MovieSceneDynamicBinding.cpp
--- a/Engine/Source/Runtime/MovieScene/Private/MovieSceneDynamicBinding.cpp
+++ b/Engine/Source/Runtime/MovieScene/Private/MovieSceneDynamicBinding.cpp
@@ -32,6 +32,7 @@ FMovieSceneDynamicBindingResolveResult UBuiltInDynamicBindingResolverLibrary::Re
 				APawn* Pawn = PlayerController->GetPawn();
 				FMovieSceneDynamicBindingResolveResult Result;
 				Result.Objects.Add(Pawn);
+				Result.Object = Pawn; // Object deprecated 5.7
 				Result.bIsPossessedObject = true;
 				return Result;
 			}
diff --git a/Engine/Source/Runtime/MovieScene/Public/Bindings/MovieSceneCustomBinding.h b/Engine/Source/Runtime/MovieScene/Public/Bindings/MovieSceneCustomBinding.h
--- a/Engine/Source/Runtime/MovieScene/Public/Bindings/MovieSceneCustomBinding.h
+++ b/Engine/Source/Runtime/MovieScene/Public/Bindings/MovieSceneCustomBinding.h
@@ -32,10 +32,8 @@ struct FMovieSceneBindingResolveResult
 	UPROPERTY(EditAnywhere, BlueprintReadWrite, Category = "General")
 	TArray<TObjectPtr<UObject>> Objects;
 
-#if WITH_EDITORONLY_DATA
-	UPROPERTY(meta = (DeprecatedProperty, DeprecationMessage = "Use Objects instead."))
-	TObjectPtr<UObject> Object_DEPRECATED = nullptr;
-#endif // WITH_EDITORONLY_DATA
+	UPROPERTY(EditAnywhere, BlueprintReadWrite, Category = "General", meta = (DeprecatedProperty, DeprecationMessage = "This property is deprecated, please use the array of Objects instead."))
+	TObjectPtr<UObject> Object = nullptr;
 };
 
 /*
diff --git a/Engine/Source/Runtime/MovieScene/Public/MovieSceneDynamicBinding.h b/Engine/Source/Runtime/MovieScene/Public/MovieSceneDynamicBinding.h
--- a/Engine/Source/Runtime/MovieScene/Public/MovieSceneDynamicBinding.h
+++ b/Engine/Source/Runtime/MovieScene/Public/MovieSceneDynamicBinding.h
@@ -95,10 +95,8 @@ struct FMovieSceneDynamicBindingResolveResult
 	UPROPERTY(EditAnywhere, BlueprintReadWrite, Category = "General")
 	TArray<TObjectPtr<UObject>> Objects;
 
-#if WITH_EDITORONLY_DATA
-	UPROPERTY(meta = (DeprecatedProperty, DeprecationMessage = "This property is deprecated, please use the array of Objects instead."))
-	TObjectPtr<UObject> Object_DEPRECATED = nullptr;
-#endif // WITH_EDITORONLY_DATA
+	UPROPERTY(EditAnywhere, BlueprintReadWrite, Category = "General", meta = (DeprecatedProperty, DeprecationMessage = "This property is deprecated, please use the array of Objects instead."))
+	TObjectPtr<UObject> Object = nullptr;
 
 	/**
 	 * Whether the resolved object is external to the sequence
diff --git a/Engine/Source/Runtime/MovieSceneTracks/Private/Bindings/MovieSceneReplaceableDirectorBlueprintBinding.cpp b/Engine/Source/Runtime/MovieSceneTracks/Private/Bindings/MovieSceneReplaceableDirectorBlueprintBinding.cpp
--- a/Engine/Source/Runtime/MovieSceneTracks/Private/Bindings/MovieSceneReplaceableDirectorBlueprintBinding.cpp
+++ b/Engine/Source/Runtime/MovieSceneTracks/Private/Bindings/MovieSceneReplaceableDirectorBlueprintBinding.cpp
@@ -91,6 +91,7 @@ FMovieSceneBindingResolveResult UMovieSceneReplaceableDirectorBlueprintBinding::
 	FMovieSceneBindingResolveResult ResolveResult;
 	FMovieSceneDynamicBindingResolveResult DynamicResolveResult = FMovieSceneDynamicBindingInvoker::ResolveDynamicBinding(SharedPlaybackState, ResolveParams.Sequence, ResolveParams.SequenceID, ResolveParams.ObjectBindingID, DynamicBinding);
 	ResolveResult.Objects = DynamicResolveResult.Objects;
+	ResolveResult.Object = DynamicResolveResult.Object; // Object deprecated 5.7
 	return ResolveResult;
 }
 
