Sequencer: Set Transient flag in constructor. MuteSolo and TrackRow Decorations are transient and not meant to be saved

#rnx

diff --git a/Engine/Source/Runtime/MovieScene/Private/Decorations/MovieSceneMuteSoloDecoration.cpp b/Engine/Source/Runtime/MovieScene/Private/Decorations/MovieSceneMuteSoloDecoration.cpp
new file mode 100644
--- /dev/null
+++ b/Engine/Source/Runtime/MovieScene/Private/Decorations/MovieSceneMuteSoloDecoration.cpp
@@ -0,0 +1,11 @@
+// Copyright Epic Games, Inc. All Rights Reserved.
+
+#include "Decorations/MovieSceneMuteSoloDecoration.h"
+
+#include UE_INLINE_GENERATED_CPP_BY_NAME(MovieSceneMuteSoloDecoration)
+
+UMovieSceneMuteSoloDecoration::UMovieSceneMuteSoloDecoration()
+{
+	// Mute Solo Decoration is not meant to be saved
+	SetFlags(RF_Transient);
+}
diff --git a/Engine/Source/Runtime/MovieScene/Private/Decorations/MovieSceneTrackRowDecoration.cpp b/Engine/Source/Runtime/MovieScene/Private/Decorations/MovieSceneTrackRowDecoration.cpp
--- a/Engine/Source/Runtime/MovieScene/Private/Decorations/MovieSceneTrackRowDecoration.cpp
+++ b/Engine/Source/Runtime/MovieScene/Private/Decorations/MovieSceneTrackRowDecoration.cpp
@@ -4,6 +4,12 @@
 
 #include UE_INLINE_GENERATED_CPP_BY_NAME(MovieSceneTrackRowDecoration)
 
+UMovieSceneTrackRowDecoration::UMovieSceneTrackRowDecoration()
+{
+	// Track Row Decoration is not meant to be saved
+	SetFlags(RF_Transient); 
+}
+
 void UMovieSceneTrackRowDecoration::SetMuted(int32 InRowIndex, bool bMuted)
 {
 	if (!MuteSoloData.Contains(InRowIndex))
diff --git a/Engine/Source/Runtime/MovieScene/Public/Decorations/MovieSceneMuteSoloDecoration.h b/Engine/Source/Runtime/MovieScene/Public/Decorations/MovieSceneMuteSoloDecoration.h
--- a/Engine/Source/Runtime/MovieScene/Public/Decorations/MovieSceneMuteSoloDecoration.h
+++ b/Engine/Source/Runtime/MovieScene/Public/Decorations/MovieSceneMuteSoloDecoration.h
@@ -21,6 +21,8 @@ class UMovieSceneMuteSoloDecoration
 {
 public:
 
+	UMovieSceneMuteSoloDecoration();
+
 	GENERATED_BODY()
 
 	void SetMuted(bool bMuted) { MuteSoloData.bMuted = bMuted; }
diff --git a/Engine/Source/Runtime/MovieScene/Public/Decorations/MovieSceneTrackRowDecoration.h b/Engine/Source/Runtime/MovieScene/Public/Decorations/MovieSceneTrackRowDecoration.h
--- a/Engine/Source/Runtime/MovieScene/Public/Decorations/MovieSceneTrackRowDecoration.h
+++ b/Engine/Source/Runtime/MovieScene/Public/Decorations/MovieSceneTrackRowDecoration.h
@@ -18,6 +18,8 @@ public:
 
 	GENERATED_BODY()
 
+	UMovieSceneTrackRowDecoration();
+
 	UE_API void SetMuted(int32 InRowIndex, bool bMuted);
 	UE_API bool IsMuted(int32 InRowIndex) const;
 
