[Core] Fix "HasJumped" flag accidentally discarded by MovieSceneEntitySystemRunner

#rnx

diff --git a/Engine/Source/Runtime/MovieScene/Private/EntitySystem/MovieSceneEntitySystemRunner.cpp b/Engine/Source/Runtime/MovieScene/Private/EntitySystem/MovieSceneEntitySystemRunner.cpp
--- a/Engine/Source/Runtime/MovieScene/Private/EntitySystem/MovieSceneEntitySystemRunner.cpp
+++ b/Engine/Source/Runtime/MovieScene/Private/EntitySystem/MovieSceneEntitySystemRunner.cpp
@@ -788,6 +788,8 @@ UE::MovieScene::ERunnerFlushResult FMovieSceneEntitySystemRunner::GameThread_Upd
 					Algo::Reverse(Dissections);
 				}
 
+				bool bHasJumped = Request.Context.HasJumped();
+
 				for (int32 Index = 0; Index < Dissections.Num() - 1; ++Index)
 				{
 					// Never finish or destroy sequence instances until the _last_ dissected update
@@ -795,13 +797,16 @@ UE::MovieScene::ERunnerFlushResult FMovieSceneEntitySystemRunner::GameThread_Upd
 
 					FDissectedUpdate Dissection{
 						FSimpleDelegate(), // Never trigger the OnFlushed delegate until the final dissection
-						FMovieSceneContext(FMovieSceneEvaluationRange(Dissections[Index], Request.Context.GetFrameRate(), Request.Context.GetDirection()), Request.Context.GetStatus()),
+						FMovieSceneContext(FMovieSceneEvaluationRange(Dissections[Index], Request.Context.GetFrameRate(), Request.Context.GetDirection()), Request.Context.GetStatus()).SetHasJumped(bHasJumped),
 						Request.Params.InstanceHandle,
 						Index,
 						Flags
 					};
 					DissectedUpdates.Add(Dissection);
 
+					// Only the first dissection needs to know if we've jumped
+					bHasJumped = false;
+
 					// Only mark the first dissection for update this round
 					if (Index == 0)
 					{
@@ -812,7 +817,7 @@ UE::MovieScene::ERunnerFlushResult FMovieSceneEntitySystemRunner::GameThread_Upd
 				// Add the last one with MAX_int32 so it gets evaluated with all the others in this flush
 				FDissectedUpdate Dissection{
 					MoveTemp(Request.OnFlushed),
-					FMovieSceneContext(FMovieSceneEvaluationRange(Dissections.Last(), Request.Context.GetFrameRate(), Request.Context.GetDirection()), Request.Context.GetStatus()),
+					FMovieSceneContext(FMovieSceneEvaluationRange(Dissections.Last(), Request.Context.GetFrameRate(), Request.Context.GetDirection()), Request.Context.GetStatus()).SetHasJumped(bHasJumped),
 					Request.Params.InstanceHandle,
 					MAX_int32,
 					Request.Params.UpdateFlags
