Sequencer: fix regresion of pre/post-roll behavior in the skeletal animation track

diff --git a/Engine/Source/Runtime/MovieScene/Private/Sections/MovieSceneSectionTimingParameters.cpp b/Engine/Source/Runtime/MovieScene/Private/Sections/MovieSceneSectionTimingParameters.cpp
--- a/Engine/Source/Runtime/MovieScene/Private/Sections/MovieSceneSectionTimingParameters.cpp
+++ b/Engine/Source/Runtime/MovieScene/Private/Sections/MovieSceneSectionTimingParameters.cpp
@@ -24,7 +24,7 @@ FMovieSceneSequenceTransform FMovieSceneSectionTimingParametersSeconds::MakeTran
 
 	// ----------------------------------------------------------------------------
 	// First things first, subtract the section start bound
-	AddPositionInOuterAsOffset(Result, -OuterRange.GetLowerBoundValue());
+	AddPositionInOuterAsOffset(Result, OuterRange, -OuterRange.GetLowerBoundValue());
 
 	// ----------------------------------------------------------------------------
 	// Time warp
@@ -59,9 +59,18 @@ FMovieSceneSequenceTransform FMovieSceneSectionTimingParametersSeconds::MakeTran
 	return Result;
 }
 
-void FMovieSceneSectionTimingParametersSeconds::AddPositionInOuterAsOffset(FMovieSceneSequenceTransform& Transform, const FFrameTime& Offset) const
+void FMovieSceneSectionTimingParametersSeconds::AddPositionInOuterAsOffset(FMovieSceneSequenceTransform& Transform, const TRange<FFrameNumber>& OuterRange, const FFrameTime& Offset) const
 {
-	Transform.Add(FMovieSceneTimeTransform(Offset));
+	if (bClampToOuterRange)
+	{
+		FMovieSceneTimeWarpVariant Clamp;
+		Clamp.Set(FMovieSceneTimeWarpClamp{ OuterRange.Size<FFrameNumber>() });
+		Transform.Add(Offset, MoveTemp(Clamp));
+	}
+	else
+	{
+		Transform.Add(FMovieSceneTimeTransform(Offset));
+	}
 }
 
 void FMovieSceneSectionTimingParametersSeconds::AddPlayRate(FMovieSceneSequenceTransform& Transform) const
@@ -92,7 +101,7 @@ void FMovieSceneSectionTimingParametersSeconds::AddLoopingOrClampingAndReverse(F
 		Loop.Set(FMovieSceneTimeWarpLoopFloat{ static_cast<float>(Duration) });
 		Result.Add(FFrameTime::FromDecimal(Offset), MoveTemp(Loop));
 	}
-	else if (bClamp)
+	else if (bClampToInnerRange)
 	{
 		// Clamp
 		FMovieSceneTimeWarpVariant Clamp;
diff --git a/Engine/Source/Runtime/MovieScene/Public/Sections/MovieSceneSectionTimingParameters.h b/Engine/Source/Runtime/MovieScene/Public/Sections/MovieSceneSectionTimingParameters.h
--- a/Engine/Source/Runtime/MovieScene/Public/Sections/MovieSceneSectionTimingParameters.h
+++ b/Engine/Source/Runtime/MovieScene/Public/Sections/MovieSceneSectionTimingParameters.h
@@ -62,7 +62,7 @@ struct FMovieSceneSectionTimingParametersSeconds
 	 * When true, apply clamping to the inner range. Mutually exclusive with bLoop.
 	 */
 	UPROPERTY(BlueprintReadWrite, EditAnywhere, Category="Timing")
-	uint8 bClamp : 1 = 0;
+	uint8 bClampToInnerRange : 1 = 0;
 
 	/**
 	 * When true, reverses the play direction. Applied after all other transformations
@@ -70,12 +70,18 @@ struct FMovieSceneSectionTimingParametersSeconds
 	UPROPERTY(BlueprintReadWrite, EditAnywhere, Category="Timing")
 	uint8 bReverse : 1 = 0;
 
+	/**
+	 * When true, apply clamping to the outer range before anything else.
+	 */
+	UPROPERTY(BlueprintReadWrite, EditAnywhere, Category="Timing")
+	uint8 bClampToOuterRange : 1 = 0;
+
 	/**
 	 * Make a transform out of these parameters
 	 */
 	MOVIESCENE_API FMovieSceneSequenceTransform MakeTransform(const FFrameRate& OuterFrameRate, const TRange<FFrameNumber>& OuterRange, double SourceDuration, double InnerPlayRate) const;
 
-	MOVIESCENE_API void AddPositionInOuterAsOffset(FMovieSceneSequenceTransform& Transform, const FFrameTime& Offset) const;
+	MOVIESCENE_API void AddPositionInOuterAsOffset(FMovieSceneSequenceTransform& Transform, const TRange<FFrameNumber>& OuterRange, const FFrameTime& Offset) const;
 
 	MOVIESCENE_API void AddPlayRate(FMovieSceneSequenceTransform& Transform) const;
 
diff --git a/Engine/Source/Runtime/MovieSceneTracks/Private/Sections/MovieSceneSkeletalAnimationSection.cpp b/Engine/Source/Runtime/MovieSceneTracks/Private/Sections/MovieSceneSkeletalAnimationSection.cpp
--- a/Engine/Source/Runtime/MovieSceneTracks/Private/Sections/MovieSceneSkeletalAnimationSection.cpp
+++ b/Engine/Source/Runtime/MovieSceneTracks/Private/Sections/MovieSceneSkeletalAnimationSection.cpp
@@ -78,7 +78,8 @@ FMovieSceneSequenceTransform FMovieSceneSkeletalAnimationParams::MakeTransform(c
 		FirstLoopStartFrameOffset / OuterFrameRate,
 		bLooping,
 		!bLooping,
-		bReverse
+		bReverse,
+		true  // bClampToOuterRange
 	};
 
 	return TimingParams.MakeTransform(OuterFrameRate, OuterRange, SequenceLength, OverrideSequence->RateScale);
