Sequencer: fix regresion of pre/post-roll behavior in the skeletal animation track

diff --git a/Engine/Config/BaseEngine.ini b/Engine/Config/BaseEngine.ini
--- a/Engine/Config/BaseEngine.ini
+++ b/Engine/Config/BaseEngine.ini
@@ -1377,6 +1377,9 @@ ClassName=/Script/OnlineSubsystemUtils.OnlineEngineInterfaceImpl
 +ClassRedirects=(OldName="/Script/MovieSceneTextTrack.TextChannelEvaluatorSystem",NewName="/Script/MovieSceneTracks.TextChannelEvaluatorSystem")
 +ClassRedirects=(OldName="/Script/MovieSceneTextTrackEditor.MovieSceneTextKeyStruct",NewName="/Script/MovieSceneTools.MovieSceneTextKeyStruct")
 
++PropertyRedirects=(OldName="/Script/MovieScene.MovieSceneSectionTimingParametersSeconds.bClamp",NewName="MovieSceneSectionTimingParametersSeconds.bClampToInnerRange")
++PropertyRedirects=(OldName="/Script/MovieScene.FMovieSceneSectionTimingParametersFrames.bClamp",NewName="FMovieSceneSectionTimingParametersFrames.bClampToInnerRange")
+
 +ClassRedirects=(OldName="/Script/Engine.WorldPartitionDestructibleHLODMeshComponent",NewName="/Script/Engine.WorldPartitionDestructibleHLODComponent")
 
 ; End of releases redirects
diff --git a/Engine/Source/Runtime/MovieScene/Private/Sections/MovieSceneSectionTimingParameters.cpp b/Engine/Source/Runtime/MovieScene/Private/Sections/MovieSceneSectionTimingParameters.cpp
--- a/Engine/Source/Runtime/MovieScene/Private/Sections/MovieSceneSectionTimingParameters.cpp
+++ b/Engine/Source/Runtime/MovieScene/Private/Sections/MovieSceneSectionTimingParameters.cpp
@@ -24,7 +24,7 @@ FMovieSceneSequenceTransform FMovieSceneSectionTimingParametersSeconds::MakeTran
 
 	// ----------------------------------------------------------------------------
 	// First things first, subtract the section start bound
-	AddPositionInOuterAsOffset(Result, OuterRange, -OuterRange.GetLowerBoundValue());
+	AddPositionInOuterAsOffset(Result, OuterRange);
 
 	// ----------------------------------------------------------------------------
 	// Time warp
@@ -59,18 +59,22 @@ FMovieSceneSequenceTransform FMovieSceneSectionTimingParametersSeconds::MakeTran
 	return Result;
 }
 
-void FMovieSceneSectionTimingParametersSeconds::AddPositionInOuterAsOffset(FMovieSceneSequenceTransform& Transform, const TRange<FFrameNumber>& OuterRange, const FFrameTime& Offset) const
+void FMovieSceneSectionTimingParametersSeconds::AddOffset(FMovieSceneSequenceTransform& Transform, const FFrameTime& Offset) const
 {
+	Transform.Add(FMovieSceneTimeTransform(Offset));
+}
+
+void FMovieSceneSectionTimingParametersSeconds::AddPositionInOuterAsOffset(FMovieSceneSequenceTransform& Transform, const TRange<FFrameNumber>& OuterRange) const
+{
+	const FFrameNumber Offset = -OuterRange.GetLowerBoundValue();
 	if (bClampToOuterRange)
 	{
 		FMovieSceneTimeWarpVariant Clamp;
 		Clamp.Set(FMovieSceneTimeWarpClamp{ OuterRange.Size<FFrameNumber>() });
 		Transform.Add(Offset, MoveTemp(Clamp));
 	}
-	else
-	{
-		Transform.Add(FMovieSceneTimeTransform(Offset));
-	}
+
+	Transform.Add(FMovieSceneTimeTransform(Offset));
 }
 
 void FMovieSceneSectionTimingParametersSeconds::AddPlayRate(FMovieSceneSequenceTransform& Transform) const
@@ -126,7 +130,7 @@ FMovieSceneSequenceTransform FMovieSceneSectionTimingParametersFrames::MakeTrans
 
 	// ----------------------------------------------------------------------------
 	// First things first, subtract the section start bound
-	AddPositionInOuterAsOffset(Result, -OuterRange.GetLowerBoundValue());
+	AddPositionInOuterAsOffset(Result, OuterRange);
 
 	// ----------------------------------------------------------------------------
 	// Time warp
@@ -161,8 +165,21 @@ FMovieSceneSequenceTransform FMovieSceneSectionTimingParametersFrames::MakeTrans
 	return Result;
 }
 
-void FMovieSceneSectionTimingParametersFrames::AddPositionInOuterAsOffset(FMovieSceneSequenceTransform& Transform, const FFrameTime& Offset) const
+void FMovieSceneSectionTimingParametersFrames::AddOffset(FMovieSceneSequenceTransform& Transform, const FFrameTime& Offset) const
+{
+	Transform.Add(FMovieSceneTimeTransform(Offset));
+}
+
+void FMovieSceneSectionTimingParametersFrames::AddPositionInOuterAsOffset(FMovieSceneSequenceTransform& Transform, const TRange<FFrameNumber>& OuterRange) const
 {
+	const FFrameNumber Offset = -OuterRange.GetLowerBoundValue();
+	if (bClampToOuterRange)
+	{
+		FMovieSceneTimeWarpVariant Clamp;
+		Clamp.Set(FMovieSceneTimeWarpClamp{ OuterRange.Size<FFrameNumber>() });
+		Transform.Add(Offset, MoveTemp(Clamp));
+	}
+
 	Transform.Add(FMovieSceneTimeTransform(Offset));
 }
 
@@ -192,7 +209,7 @@ void FMovieSceneSectionTimingParametersFrames::AddLoopingOrClampingAndReverse(FM
 		Loop.Set(FMovieSceneTimeWarpLoop{ Duration });
 		Result.Add(Offset, MoveTemp(Loop));
 	}
-	else if (bClamp)
+	else if (bClampToInnerRange)
 	{
 		// Clamp
 		FMovieSceneTimeWarpVariant Clamp;
diff --git a/Engine/Source/Runtime/MovieScene/Public/Sections/MovieSceneSectionTimingParameters.h b/Engine/Source/Runtime/MovieScene/Public/Sections/MovieSceneSectionTimingParameters.h
--- a/Engine/Source/Runtime/MovieScene/Public/Sections/MovieSceneSectionTimingParameters.h
+++ b/Engine/Source/Runtime/MovieScene/Public/Sections/MovieSceneSectionTimingParameters.h
@@ -53,7 +53,7 @@ struct FMovieSceneSectionTimingParametersSeconds
 	float FirstLoopStartOffset = 0.f;
 
 	/**
-	 * When true, apply looping to the inner range. Mutually exclusive with bClamp.
+	 * When true, apply looping to the inner range. Mutually exclusive with bClampToInnerRange.
 	 */
 	UPROPERTY(BlueprintReadWrite, EditAnywhere, Category="Timing")
 	uint8 bLoop : 1 = 0;
@@ -81,7 +81,12 @@ struct FMovieSceneSectionTimingParametersSeconds
 	 */
 	MOVIESCENE_API FMovieSceneSequenceTransform MakeTransform(const FFrameRate& OuterFrameRate, const TRange<FFrameNumber>& OuterRange, double SourceDuration, double InnerPlayRate) const;
 
-	MOVIESCENE_API void AddPositionInOuterAsOffset(FMovieSceneSequenceTransform& Transform, const TRange<FFrameNumber>& OuterRange, const FFrameTime& Offset) const;
+	/**
+	 * Helpers
+	 */
+	MOVIESCENE_API void AddOffset(FMovieSceneSequenceTransform& Transform, const FFrameTime& Offset) const;
+
+	MOVIESCENE_API void AddPositionInOuterAsOffset(FMovieSceneSequenceTransform& Transform, const TRange<FFrameNumber>& OuterRange) const;
 
 	MOVIESCENE_API void AddPlayRate(FMovieSceneSequenceTransform& Transform) const;
 
@@ -90,6 +95,12 @@ struct FMovieSceneSectionTimingParametersSeconds
 	MOVIESCENE_API void AddInnerStartOffset(FMovieSceneSequenceTransform& Transform, const FFrameTime& Offset) const;
 
 	MOVIESCENE_API void AddLoopingOrClampingAndReverse(FMovieSceneSequenceTransform& Result, const double Duration, const double Offset) const;
+
+	/**
+	 * Deprecated
+	 */
+	UE_DEPRECATED(5.7, "Please use AddOffset, or the version of AddPositionInOuterAsOffset that takes an outer range")
+	void AddPositionInOuterAsOffset(FMovieSceneSequenceTransform& Transform, const FFrameTime& Offset) const { AddOffset(Transform, Offset); }
 };
 
 
@@ -135,7 +146,7 @@ struct FMovieSceneSectionTimingParametersFrames
 	FFrameNumber FirstLoopStartOffset = 0;
 
 	/**
-	 * When true, apply looping to the inner range. Mutually exclusive with bClamp.
+	 * When true, apply looping to the inner range. Mutually exclusive with bClampToInnerRange.
 	 */
 	UPROPERTY(BlueprintReadWrite, EditAnywhere, Category="Timing")
 	uint8 bLoop : 1 = 0;
@@ -144,7 +155,7 @@ struct FMovieSceneSectionTimingParametersFrames
 	 * When true, apply clamping to the inner range. Mutually exclusive with bLoop.
 	 */
 	UPROPERTY(BlueprintReadWrite, EditAnywhere, Category="Timing")
-	uint8 bClamp : 1 = 0;
+	uint8 bClampToInnerRange : 1 = 0;
 
 	/**
 	 * When true, reverses the play direction. Applied after all other transformations
@@ -152,6 +163,12 @@ struct FMovieSceneSectionTimingParametersFrames
 	UPROPERTY(BlueprintReadWrite, EditAnywhere, Category="Timing")
 	uint8 bReverse : 1 = 0;
 
+	/**
+	 * When true, apply clamping to the outer range before anything else.
+	 */
+	UPROPERTY(BlueprintReadWrite, EditAnywhere, Category="Timing")
+	uint8 bClampToOuterRange : 1 = 0;
+
 	/**
 	 * Make a transform out of these parameters
 	 */
@@ -160,7 +177,9 @@ struct FMovieSceneSectionTimingParametersFrames
 	/**
 	 * Helpers
 	 */
-	MOVIESCENE_API void AddPositionInOuterAsOffset(FMovieSceneSequenceTransform& Transform, const FFrameTime& Offset) const;
+	MOVIESCENE_API void AddOffset(FMovieSceneSequenceTransform& Transform, const FFrameTime& Offset) const;
+
+	MOVIESCENE_API void AddPositionInOuterAsOffset(FMovieSceneSequenceTransform& Transform, const TRange<FFrameNumber>& OuterRange) const;
 
 	MOVIESCENE_API void AddPlayRate(FMovieSceneSequenceTransform& Transform) const;
 
@@ -170,6 +189,11 @@ struct FMovieSceneSectionTimingParametersFrames
 
 	MOVIESCENE_API void AddLoopingOrClampingAndReverse(FMovieSceneSequenceTransform& Result, const FFrameNumber& Duration, const FFrameNumber& Offset) const;
 
+	/**
+	 * Deprecated
+	 */
+	UE_DEPRECATED(5.7, "Please use AddOffset, or the version of AddPositionInOuterAsOffset that takes an outer range")
+	void AddPositionInOuterAsOffset(FMovieSceneSequenceTransform& Transform, const FFrameTime& Offset) const { AddOffset(Transform, Offset); }
 };
 
 namespace UE::MovieScene
