Sequencer: fix possible memory corruption when trying to animate invalid property paths or subtypes of property bags

- If the property path was invalid past the first step, we could incorrectly write unrelated data at the last address where the property binding resolution stopped. We now check if the resolution succeeded, or just don't do anything at all.

- Subclasses of FInstancedPropertyBag were not handled correctly in the case of animating a property inside the property bag.

diff --git a/Engine/Source/Runtime/MovieScene/Private/TrackInstancePropertyBindings.cpp b/Engine/Source/Runtime/MovieScene/Private/TrackInstancePropertyBindings.cpp
--- a/Engine/Source/Runtime/MovieScene/Private/TrackInstancePropertyBindings.cpp
+++ b/Engine/Source/Runtime/MovieScene/Private/TrackInstancePropertyBindings.cpp
@@ -99,7 +99,7 @@ void ResolvePropertyRecursive(void* BasePointer, UStruct* InStruct, TArrayView<c
 		// As above but for a property bag, unless the property path is specifically targeting its inner 
 		// instanced struct, in which case we don't need to do anything, we'll end up in the previous code
 		// block on the next iteration.
-		else if (InStruct == FInstancedPropertyBag::StaticStruct() && InPropertyNames[Index] != PropertyBagValueProperty->GetFName())
+		else if (InStruct->IsChildOf<FInstancedPropertyBag>() && InPropertyNames[Index] != PropertyBagValueProperty->GetFName())
 		{
 			// Property bags may be reallocated, flag this property path as volatile.
 			OutResolutionState.bIsVolatile = true;
@@ -438,26 +438,29 @@ TOptional<TPair<const FProperty*, UE::MovieScene::FSourcePropertyValue>> FTrackI
 	FPropertyResolutionState ResolutionState;
 	UE::MovieScene::ResolveProperty(*Object, InPropertyPath, ResolutionState);
 
-	if (const FPropertyResolutionStep* LastStep = ResolutionState.GetLastStep())
+	if (ResolutionState.bIsValid)
 	{
-		const FProperty* Property = LastStep->Property;
-		if (ensure(Property && LastStep->ContainerAddress && LastStep->ArrayIndex == INDEX_NONE))
+		if (const FPropertyResolutionStep* LastStep = ResolutionState.GetLastStep())
 		{
-			const void* PropertyAddress = Property->ContainerPtrToValuePtr<void>(LastStep->ContainerAddress);
-
-			// Bool property values are stored in a bit field so using a straight cast of the pointer to get their value does not
-			// work.  Instead use the actual property to get the correct value.
-			if (const FBoolProperty* BoolProperty = CastField<const FBoolProperty>(Property))
-			{
-				return MakeTuple(Property, FSourcePropertyValue::FromValue(BoolProperty->GetPropertyValue(PropertyAddress)));
-			}
-			// Object properties might have various different types of storage, but we always expose them as a raw ptr
-			else if (const FObjectPropertyBase* ObjectProperty = CastField<const FObjectPropertyBase>(Property))
+			const FProperty* Property = LastStep->Property;
+			if (ensure(Property && LastStep->ContainerAddress && LastStep->ArrayIndex == INDEX_NONE))
 			{
-				return MakeTuple(Property, FSourcePropertyValue::FromValue(ObjectProperty->GetObjectPropertyValue(PropertyAddress)));
-			}
+				const void* PropertyAddress = Property->ContainerPtrToValuePtr<void>(LastStep->ContainerAddress);
+
+				// Bool property values are stored in a bit field so using a straight cast of the pointer to get their value does not
+				// work.  Instead use the actual property to get the correct value.
+				if (const FBoolProperty* BoolProperty = CastField<const FBoolProperty>(Property))
+				{
+					return MakeTuple(Property, FSourcePropertyValue::FromValue(BoolProperty->GetPropertyValue(PropertyAddress)));
+				}
+				// Object properties might have various different types of storage, but we always expose them as a raw ptr
+				else if (const FObjectPropertyBase* ObjectProperty = CastField<const FObjectPropertyBase>(Property))
+				{
+					return MakeTuple(Property, FSourcePropertyValue::FromValue(ObjectProperty->GetObjectPropertyValue(PropertyAddress)));
+				}
 
-			return MakeTuple(Property, FSourcePropertyValue::FromAddress(PropertyAddress, *Property));
+				return MakeTuple(Property, FSourcePropertyValue::FromAddress(PropertyAddress, *Property));
+			}
 		}
 	}
 
@@ -521,9 +524,12 @@ FProperty* FTrackInstancePropertyBindings::FindProperty(const UObject* Object, F
 	using namespace UE::MovieScene;
 	FPropertyResolutionState ResolutionState;
 	UE::MovieScene::ResolveProperty(*Object, InPropertyPath, ResolutionState);
-	if (const FPropertyResolutionStep* LastStep = ResolutionState.GetLastStep())
+	if (ResolutionState.bIsValid)
 	{
-		return LastStep->Property;
+		if (const FPropertyResolutionStep* LastStep = ResolutionState.GetLastStep())
+		{
+			return LastStep->Property;
+		}
 	}
 	return nullptr;
 }
@@ -534,7 +540,7 @@ FTrackInstancePropertyBindings::FResolvedPropertyAndFunction FTrackInstancePrope
 	FPropertyResolutionState ResolutionState;
 	UE::MovieScene::ResolveProperty(*Object, InPropertyPath, ResolutionState);
 
-	if (ResolutionState.PropertySteps.IsEmpty())
+	if (!ResolutionState.bIsValid || ResolutionState.PropertySteps.IsEmpty())
 	{
 		return FResolvedPropertyAndFunction();
 	}
