Sequence Validator: Better display names for output

Warnings now produce something like: Cube.Transform.MovieScene3DTransformSection_0.Location.X or Audio.Dialog_01

#rnx

diff --git a/Engine/Plugins/MovieScene/SequenceValidator/Source/SequenceValidator/Private/Validation/Rules/SequenceValidationRule_SectionAlignments.cpp b/Engine/Plugins/MovieScene/SequenceValidator/Source/SequenceValidator/Private/Validation/Rules/SequenceValidationRule_SectionAlignments.cpp
--- a/Engine/Plugins/MovieScene/SequenceValidator/Source/SequenceValidator/Private/Validation/Rules/SequenceValidationRule_SectionAlignments.cpp
+++ b/Engine/Plugins/MovieScene/SequenceValidator/Source/SequenceValidator/Private/Validation/Rules/SequenceValidationRule_SectionAlignments.cpp
@@ -5,6 +5,7 @@
 #include "Algo/BinarySearch.h"
 #include "Internationalization/Internationalization.h"
 #include "MovieScene.h"
+#include "MovieSceneCommonHelpers.h"
 #include "MovieSceneTimeHelpers.h"
 #include "MovieSceneSection.h"
 #include "MovieSceneSequence.h"
@@ -151,7 +152,7 @@ public:
 
 private:
 
-	void MaybeWarnAbout(FSequenceValidationResults& OutResults, const FMovieSceneInverseSequenceTransform& SequenceToRootTransform, FFrameTime InTime, UObject* InTarget = nullptr)
+	void MaybeWarnAbout(FSequenceValidationResults& OutResults, const FMovieSceneInverseSequenceTransform& SequenceToRootTransform, FFrameTime InTime, UMovieSceneSection* Section)
 	{
 		// Convert the local time to a root sequence time, so that we can compare it to all the important times that
 		// are also in root time space.
@@ -189,12 +190,12 @@ private:
 
 		if (bDoWarn)
 		{
-			TSharedRef<FSequenceValidationResult> NewResult = MakeShared<FSequenceValidationResult>(EMessageSeverity::Warning, InTarget, RuleInfo);
+			TSharedRef<FSequenceValidationResult> NewResult = MakeShared<FSequenceValidationResult>(EMessageSeverity::Warning, Section, RuleInfo);
 			NewResult->SetSubSectionTrail(CurrentSubSectionTrail);
 			NewResult->SetLocalTime(InTime);
 			NewResult->SetUserMessage(FText::Format(
-						LOCTEXT("AlignmentWarningFmt", "Section bound is within two frames of {0}."),
-						TimeLabel));
+				LOCTEXT("AlignmentWarningFmt", "Section {0} : bound is within two frames of {1}."),
+						MovieSceneHelpers::GetDisplayPathName(Section), TimeLabel));
 			OutResults.AddResult(NewResult);
 		}
 	}
diff --git a/Engine/Plugins/MovieScene/SequenceValidator/Source/SequenceValidator/Private/Validation/Rules/SequenceValidationRule_UnassignedBindingsAndAssets.cpp b/Engine/Plugins/MovieScene/SequenceValidator/Source/SequenceValidator/Private/Validation/Rules/SequenceValidationRule_UnassignedBindingsAndAssets.cpp
--- a/Engine/Plugins/MovieScene/SequenceValidator/Source/SequenceValidator/Private/Validation/Rules/SequenceValidationRule_UnassignedBindingsAndAssets.cpp
+++ b/Engine/Plugins/MovieScene/SequenceValidator/Source/SequenceValidator/Private/Validation/Rules/SequenceValidationRule_UnassignedBindingsAndAssets.cpp
@@ -5,6 +5,7 @@
 #include "Algo/BinarySearch.h"
 #include "Internationalization/Internationalization.h"
 #include "MovieScene.h"
+#include "MovieSceneCommonHelpers.h"
 #include "MovieSceneSection.h"
 #include "MovieSceneSequence.h"
 #include "MovieSceneSequenceVisitor.h"
@@ -92,13 +93,13 @@ public:
 
 private:
 
-	void WarnAbout(FSequenceValidationResults& OutResults, UObject* InTarget = nullptr)
+	void WarnAbout(FSequenceValidationResults& OutResults, UMovieSceneSection* Section)
 	{
-		TSharedRef<FSequenceValidationResult> NewResult = MakeShared<FSequenceValidationResult>(EMessageSeverity::Warning, InTarget, RuleInfo);
+		TSharedRef<FSequenceValidationResult> NewResult = MakeShared<FSequenceValidationResult>(EMessageSeverity::Warning, Section, RuleInfo);
 		NewResult->SetSubSectionTrail(CurrentSubSectionTrail);
 		NewResult->SetUserMessage(FText::Format(
-			LOCTEXT("UnassignedBindingsAndAssetsWarningFmt", "Section {0} has an unassigned binding/asset."),
-			FText::FromString(InTarget->GetName())));
+			LOCTEXT("UnassignedBindingsAndAssetsWarningFmt", "Section {0} : contains an unassigned binding/asset."),
+			MovieSceneHelpers::GetDisplayPathName(Section)));
 		OutResults.AddResult(NewResult);
 	}
 
diff --git a/Engine/Plugins/MovieScene/SequenceValidator/Source/SequenceValidator/Private/Validation/Rules/SequenceValidationRule_WholeSectionRanges.cpp b/Engine/Plugins/MovieScene/SequenceValidator/Source/SequenceValidator/Private/Validation/Rules/SequenceValidationRule_WholeSectionRanges.cpp
--- a/Engine/Plugins/MovieScene/SequenceValidator/Source/SequenceValidator/Private/Validation/Rules/SequenceValidationRule_WholeSectionRanges.cpp
+++ b/Engine/Plugins/MovieScene/SequenceValidator/Source/SequenceValidator/Private/Validation/Rules/SequenceValidationRule_WholeSectionRanges.cpp
@@ -4,6 +4,7 @@
 
 #include "Internationalization/Internationalization.h"
 #include "MovieScene.h"
+#include "MovieSceneCommonHelpers.h"
 #include "MovieSceneSection.h"
 #include "MovieSceneSequence.h"
 #include "MovieSceneSequenceVisitor.h"
@@ -51,7 +52,8 @@ public:
 			const FFrameTime StartTime = ConvertFrameTime(SectionRange.GetLowerBoundValue(), TickResolution, DisplayRate);
 			if (StartTime.GetSubFrame() != 0.f)
 			{
-				AddWarning(Section, StartTime);
+				const bool bIsLowerBound = true;
+				AddWarning(Section, StartTime, bIsLowerBound);
 			}
 		}
 		if (SectionRange.HasUpperBound() && ShouldSectionHaveWholeFrameUpperBound(Section))
@@ -59,7 +61,8 @@ public:
 			const FFrameTime EndTime = ConvertFrameTime(SectionRange.GetUpperBoundValue(), TickResolution, DisplayRate);
 			if (EndTime.GetSubFrame() != 0.f)
 			{
-				AddWarning(Section, EndTime);
+				const bool bIsLowerBound = true;
+				AddWarning(Section, EndTime, bIsLowerBound);
 			}
 		}
 	}
@@ -79,12 +82,19 @@ private:
 		return true;
 	}
 
-	void AddWarning(UMovieSceneSection* Section, FFrameTime LocalTime)
+	void AddWarning(UMovieSceneSection* Section, FFrameTime LocalTime, const bool bIsLowerBound)
 	{
 		TSharedRef<FSequenceValidationResult> NewResult = MakeShared<FSequenceValidationResult>(EMessageSeverity::Warning, Section, RuleInfo);
 		NewResult->SetSubSectionTrail(CurrentSubSectionTrail);
 		NewResult->SetLocalTime(LocalTime);
-		NewResult->SetUserMessage(LOCTEXT("WholeFrameWarningFmt", "Section bound is not on a whole frame."));
+		if (bIsLowerBound)
+		{
+			NewResult->SetUserMessage(FText::Format(LOCTEXT("WholeFrameLowerBoundWarningFmt", "Section {0} : lower bound is not on a whole frame."), MovieSceneHelpers::GetDisplayPathName(Section)));
+		}
+		else
+		{
+			NewResult->SetUserMessage(FText::Format(LOCTEXT("WholeFrameUpperBoundWarningFmt", "Section {0} : upper bound is not on a whole frame."), MovieSceneHelpers::GetDisplayPathName(Section)));
+		}
 		Results.AddResult(NewResult);
 	}
 
diff --git a/Engine/Source/Runtime/MovieScene/Private/MovieSceneCommonHelpers.cpp b/Engine/Source/Runtime/MovieScene/Private/MovieSceneCommonHelpers.cpp
--- a/Engine/Source/Runtime/MovieScene/Private/MovieSceneCommonHelpers.cpp
+++ b/Engine/Source/Runtime/MovieScene/Private/MovieSceneCommonHelpers.cpp
@@ -1,5 +1,7 @@
 // Copyright Epic Games, Inc. All Rights Reserved.
 #include "MovieSceneCommonHelpers.h"
+#include "Channels/MovieSceneChannelEditorData.h"
+#include "Channels/MovieSceneChannel.h"
 #include "EntitySystem/MovieSceneIntermediatePropertyValue.h"
 #include "Components/SceneComponent.h"
 #include "GameFramework/Actor.h"
@@ -26,6 +28,67 @@
 #include "Conditions/MovieSceneGroupCondition.h"
 #include "String/ParseTokens.h"
 
+FText MovieSceneHelpers::GetDisplayPathName(const UMovieSceneTrack* Track)
+{
+	FText DisplayPathName;
+
+	FGuid ObjectGuid = Track->FindObjectBindingGuid();
+	if (ObjectGuid.IsValid())
+	{
+		if (UMovieScene* MovieScene = Cast<UMovieScene>(Track->GetOuter()))
+		{
+			DisplayPathName = MovieScene->GetObjectDisplayName(ObjectGuid);
+		}
+	}
+
+#if WITH_EDITOR
+	DisplayPathName = DisplayPathName.IsEmpty() ? Track->GetDisplayName() : FText::Join(FText::FromString("."), DisplayPathName, Track->GetDisplayName());
+#endif
+
+	return DisplayPathName;
+}
+
+FText MovieSceneHelpers::GetDisplayPathName(const UMovieSceneSection* Section)
+{
+	FText DisplayPathName;
+
+	UMovieSceneTrack* Track = Cast<UMovieSceneTrack>(Section->GetOuter());
+	if (Track)
+	{
+		DisplayPathName = GetDisplayPathName(Track);
+	}
+
+	FString SectionName;
+	if (const UObject* SourceObject = Section->GetSourceObject())
+	{
+		SectionName = SourceObject->GetName();
+	}
+	else
+	{
+		Section->GetName(SectionName);
+	}
+	
+	DisplayPathName = DisplayPathName.IsEmpty() ? FText::FromString(SectionName) : FText::Join(FText::FromString("."), DisplayPathName, FText::FromString(SectionName));
+	
+	return DisplayPathName;
+}
+
+FText MovieSceneHelpers::GetDisplayPathName(const UMovieSceneSection* Section, const FMovieSceneChannel* Channel, const FMovieSceneChannelMetaData& ChannelMetaData)
+{
+	FText DisplayPathName;
+
+	DisplayPathName = GetDisplayPathName(Section);
+
+#if WITH_EDITOR
+	if (Channel)
+	{
+		DisplayPathName = DisplayPathName.IsEmpty() ? FText::Join(FText::FromString("."), ChannelMetaData.Group, ChannelMetaData.DisplayText) : FText::Join(FText::FromString("."), DisplayPathName, ChannelMetaData.Group, ChannelMetaData.DisplayText);
+	}
+#endif
+
+	return DisplayPathName;
+}
+
 bool MovieSceneHelpers::IsSectionKeyable(const UMovieSceneSection* Section)
 {
 	if (!Section)
diff --git a/Engine/Source/Runtime/MovieScene/Public/MovieSceneCommonHelpers.h b/Engine/Source/Runtime/MovieScene/Public/MovieSceneCommonHelpers.h
--- a/Engine/Source/Runtime/MovieScene/Public/MovieSceneCommonHelpers.h
+++ b/Engine/Source/Runtime/MovieScene/Public/MovieSceneCommonHelpers.h
@@ -24,6 +24,8 @@ template <class TClass> class TSubclassOf;
 class UMovieSceneCustomBinding;
 struct FRichCurve;
 enum class EMovieSceneKeyInterpolation : uint8;
+struct FMovieSceneChannel;
+struct FMovieSceneChannelMetaData;
 struct FMovieSceneSequenceID;
 class UMovieSceneCondition;
 class UMovieSceneTrack;
@@ -54,6 +56,11 @@ public:
 		bool bPackageWasDirty;
 	};
 
+	/** Get displayable names */
+	static MOVIESCENE_API FText GetDisplayPathName(const UMovieSceneTrack* Track);
+	static MOVIESCENE_API FText GetDisplayPathName(const UMovieSceneSection* Section);
+	static MOVIESCENE_API FText GetDisplayPathName(const UMovieSceneSection* Section, const FMovieSceneChannel* Channel, const FMovieSceneChannelMetaData& ChannelMetaData);
+
 	/** 
 	 * @return Whether the section is keyable (active, on a track that is not muted, etc 
 	 */
