Sequencer: fix reference-leak in the Track Instance system

The Track Instance system was able to tag and handle garbage when related to referenced Sequence data going away (e.g. one of its source Sections being GC'ed). However it wasn't handling the case of another system tagging something as garbage (e.g. one of its bound objects being GC'ed). The Track Instances for these entries were left behind and never cleaned up. This fix adds such Track Instaces to the list of invalided ones so that they are removed on the next update.

diff --git a/Engine/Source/Runtime/MovieScene/Private/EntitySystem/TrackInstance/MovieSceneTrackInstanceSystem.cpp b/Engine/Source/Runtime/MovieScene/Private/EntitySystem/TrackInstance/MovieSceneTrackInstanceSystem.cpp
--- a/Engine/Source/Runtime/MovieScene/Private/EntitySystem/TrackInstance/MovieSceneTrackInstanceSystem.cpp
+++ b/Engine/Source/Runtime/MovieScene/Private/EntitySystem/TrackInstance/MovieSceneTrackInstanceSystem.cpp
@@ -179,6 +179,28 @@ void UMovieSceneTrackInstanceInstantiator::OnTagGarbage()
 	}
 }
 
+void UMovieSceneTrackInstanceInstantiator::OnCleanTaggedGarbage()
+{
+	using namespace UE::MovieScene;
+
+	// Garbage has been tagged with NeedsUnlink, so visit those and remove any track instances that are
+	// associated with them. Note that this is different garbage compared to what we tagged in OnTagGarbage.
+	// Here, we remove things that were tagged as garbage by *other* systems, such as GC'ed bound objects.
+	auto DoCleanTaggedGarbage = [this](FMovieSceneEntityID EntityID, FTrackInstanceInputComponent InputComponent)
+	{
+		this->InvalidatedOutputs.PadToNum(InputComponent.OutputIndex + 1, false);
+		this->InvalidatedOutputs[InputComponent.OutputIndex] = true;
+	};
+
+	FBuiltInComponentTypes* BuiltInComponents = FBuiltInComponentTypes::Get();
+
+	FEntityTaskBuilder()
+	.ReadEntityIDs()
+	.Read(FBuiltInComponentTypes::Get()->TrackInstanceInput)
+	.FilterAny({ BuiltInComponents->Tags.NeedsUnlink })
+	.Iterate_PerEntity(&Linker->EntityManager, DoCleanTaggedGarbage);
+}
+
 void UMovieSceneTrackInstanceInstantiator::OnRun(FSystemTaskPrerequisites& InPrerequisites, FSystemSubsequentTasks& Subsequents)
 {
 	using namespace UE::MovieScene;
@@ -308,4 +330,4 @@ void UMovieSceneTrackInstanceSystem::EvaluateAllInstances()
 			Entry.TrackInstance->Animate();
 		}
 	}
-}
\ No newline at end of file
+}
diff --git a/Engine/Source/Runtime/MovieScene/Public/EntitySystem/TrackInstance/MovieSceneTrackInstanceSystem.h b/Engine/Source/Runtime/MovieScene/Public/EntitySystem/TrackInstance/MovieSceneTrackInstanceSystem.h
--- a/Engine/Source/Runtime/MovieScene/Public/EntitySystem/TrackInstance/MovieSceneTrackInstanceSystem.h
+++ b/Engine/Source/Runtime/MovieScene/Public/EntitySystem/TrackInstance/MovieSceneTrackInstanceSystem.h
@@ -63,6 +63,7 @@ private:
 	MOVIESCENE_API virtual void OnRun(FSystemTaskPrerequisites& InPrerequisites, FSystemSubsequentTasks& Subsequents) override final;
 
 	MOVIESCENE_API virtual void OnTagGarbage() override;
+	MOVIESCENE_API virtual void OnCleanTaggedGarbage() override;
 
 	MOVIESCENE_API virtual void Serialize(FArchive& Ar) override;
 
