UMG: fix first and last frame evaluation of animations

- Don't tick the widget on the first frame, wait until the next frame. The flag for that was lost during the animation refactor in 5.6
- Fix the evaluation time for when stopping which should indeed be the first frame of the playrange, not frame 0. This was incorrectly ported during the animation refactor.

#FailedValidations Preflight

diff --git a/Engine/Source/Runtime/MovieScene/Private/Evaluation/MovieScenePlaybackManager.cpp b/Engine/Source/Runtime/MovieScene/Private/Evaluation/MovieScenePlaybackManager.cpp
--- a/Engine/Source/Runtime/MovieScene/Private/Evaluation/MovieScenePlaybackManager.cpp
+++ b/Engine/Source/Runtime/MovieScene/Private/Evaluation/MovieScenePlaybackManager.cpp
@@ -370,6 +370,15 @@ void FMovieScenePlaybackManager::SetCurrentTime(const FFrameTime& InFrameTime)
 	PlaybackPosition.Reset(FMath::Clamp(CurrentTick, EffectiveStartTick, LastValidTick));
 }
 
+void FMovieScenePlaybackManager::SetCurrentTimeOffset(const FFrameTime& InFrameTimeOffset)
+{
+	const FFrameRate TickResolution = PlaybackPosition.GetOutputRate();
+	const FFrameNumber CurrentTickOffset = ConvertFrameTime(InFrameTimeOffset, DisplayRate, TickResolution).RoundToFrame();
+	const FFrameNumber EffectiveStartTick = SequenceStartTick + StartOffsetTicks;
+	const FFrameNumber LastValidTick = GetLastValidTick();
+	PlaybackPosition.Reset(FMath::Clamp(EffectiveStartTick + CurrentTickOffset, EffectiveStartTick, LastValidTick));
+}
+
 TRange<FFrameTime> FMovieScenePlaybackManager::GetEffectivePlaybackRange() const
 {
 	const FFrameNumber StartTick = SequenceStartTick + StartOffsetTicks;
diff --git a/Engine/Source/Runtime/MovieScene/Public/Evaluation/MovieScenePlaybackManager.h b/Engine/Source/Runtime/MovieScene/Public/Evaluation/MovieScenePlaybackManager.h
--- a/Engine/Source/Runtime/MovieScene/Public/Evaluation/MovieScenePlaybackManager.h
+++ b/Engine/Source/Runtime/MovieScene/Public/Evaluation/MovieScenePlaybackManager.h
@@ -96,6 +96,8 @@ public:
 
 	/** Sets the current playback position, in display rate. */
 	MOVIESCENE_API void SetCurrentTime(const FFrameTime& InFrameTime);
+	/** Sets the current playback position, in display rate, as an offset from the effective start time. */
+	MOVIESCENE_API void SetCurrentTimeOffset(const FFrameTime& InFrameTimeOffset);
 
 	/** Get the effective playback range, in display rate, taking into account start/end offsets. */
 	MOVIESCENE_API TRange<FFrameTime> GetEffectivePlaybackRange() const;
diff --git a/Engine/Source/Runtime/UMG/Private/Animation/WidgetAnimationState.cpp b/Engine/Source/Runtime/UMG/Private/Animation/WidgetAnimationState.cpp
--- a/Engine/Source/Runtime/UMG/Private/Animation/WidgetAnimationState.cpp
+++ b/Engine/Source/Runtime/UMG/Private/Animation/WidgetAnimationState.cpp
@@ -456,6 +456,8 @@ void FWidgetAnimationState::Play(const FWidgetAnimationStatePlayParams& PlayPara
 		SharedPlaybackState->GetPreAnimatedState().EnableGlobalPreAnimatedStateCapture();
 	}
 
+	bIsBeginningPlay = true;
+
 	const FFrameRate DisplayRate = PlaybackManager.GetDisplayRate();
 
 	PlaybackManager.SetStartOffset(FFrameTime(0));
@@ -578,8 +580,7 @@ void FWidgetAnimationState::Stop()
 
 	PlaybackManager.SetPlaybackStatus(EMovieScenePlayerStatus::Stopped);
 
-	// TODO: we should use PlaybackManager.GetEffectiveStartTime(), but the old behavior always used frame zero.
-	PlaybackManager.SetCurrentTime(0);
+	PlaybackManager.SetCurrentTimeOffset(0);
 
 	bool bFlushPrivateLinker = false;
 
