Sequencer: Fix division by zero caused by adjoining scaling drivers

diff --git a/Engine/Source/Runtime/MovieScene/Private/Decorations/MovieSceneScalingAnchors.cpp b/Engine/Source/Runtime/MovieScene/Private/Decorations/MovieSceneScalingAnchors.cpp
--- a/Engine/Source/Runtime/MovieScene/Private/Decorations/MovieSceneScalingAnchors.cpp
+++ b/Engine/Source/Runtime/MovieScene/Private/Decorations/MovieSceneScalingAnchors.cpp
@@ -279,7 +279,7 @@ void UMovieSceneScalingAnchors::UpdateCurve(UMovieScenePlayRateCurve* Curve) con
 		CumulativeScaleOffset += (Entry.CurrentAnchor.Position - Entry.InitialAnchor.Position);
 
 		// If we have a previous anchor position, add play rate for the empty spae up until this anchor
-		if (LastAnchorInitialPosition.IsSet())
+		if (LastAnchorInitialPosition.IsSet() && !FMath::IsNearlyZero(double(Entry.ScaledAnchor.Position.Value - LastAnchorPosition->Value)))
 		{
 			const double PlayRateToThisAnchor = double(Entry.ScaledAnchor.Position.Value - LastAnchorPosition->Value) / double(Entry.InitialAnchor.Position.Value - LastAnchorInitialPosition->Value);
 			Curve->PlayRate.AddConstantKey(KeyTime, PlayRateToThisAnchor);
