Sequencer: fix missing initialization of new field bHasLooped

diff --git a/Engine/Source/Runtime/MovieScene/Public/Evaluation/MovieScenePlayback.h b/Engine/Source/Runtime/MovieScene/Public/Evaluation/MovieScenePlayback.h
--- a/Engine/Source/Runtime/MovieScene/Public/Evaluation/MovieScenePlayback.h
+++ b/Engine/Source/Runtime/MovieScene/Public/Evaluation/MovieScenePlayback.h
@@ -225,6 +225,7 @@ struct FMovieSceneContext : FMovieSceneEvaluationRange
 		, PrePostRollStartEndTime(TNumericLimits<int32>::Lowest())
 		, HierarchicalBias(0)
 		, bHasJumped(false)
+		, bHasLooped(false)
 		, bSilent(false)
 		, bSectionPreRoll(false)
 		, bSectionPostRoll(false)
