Take Recorder: Prevent crash when timecode rolls over when recording a camera track. This converts a check into an ensure.

diff --git a/Engine/Source/Runtime/MovieScene/Public/MovieSceneSection.h b/Engine/Source/Runtime/MovieScene/Public/MovieSceneSection.h
--- a/Engine/Source/Runtime/MovieScene/Public/MovieSceneSection.h
+++ b/Engine/Source/Runtime/MovieScene/Public/MovieSceneSection.h
@@ -330,9 +330,12 @@ public:
 			bCanSetRange = TryModify();
 		}
 
-		if (bCanSetRange)
+		const bool bInvariants = ensure(
+			NewRange.GetLowerBound().IsOpen() || NewRange.GetUpperBound().IsOpen()
+			|| NewRange.GetLowerBoundValue() <= NewRange.GetUpperBoundValue()
+			);
+		if (bCanSetRange && bInvariants)
 		{
-			check(NewRange.GetLowerBound().IsOpen() || NewRange.GetUpperBound().IsOpen() || NewRange.GetLowerBoundValue() <= NewRange.GetUpperBoundValue());
 			SectionRange.Value = NewRange;
 		}
 	}
