Sequencer: Strip off 'b' when finding setter function for bool properties

Hotfix compatible version which doesn't include removal of FunctionName from TrackInstancePropertyBindings

diff --git a/Engine/Source/Runtime/MovieScene/Private/EntitySystem/MovieScenePropertyRegistry.cpp b/Engine/Source/Runtime/MovieScene/Private/EntitySystem/MovieScenePropertyRegistry.cpp
--- a/Engine/Source/Runtime/MovieScene/Private/EntitySystem/MovieScenePropertyRegistry.cpp
+++ b/Engine/Source/Runtime/MovieScene/Private/EntitySystem/MovieScenePropertyRegistry.cpp
@@ -129,16 +129,27 @@ TOptional<uint16> ComputeFastPropertyPtrOffset(UClass* ObjectClass, const FMovie
 		return TOptional<uint16>();
 	}
 
+	FString PropertyName = PropertyBinding.PropertyName.ToString();
+	
+	// If this is a bool property, strip off the 'b' so that the "Set" functions to be 
+	// found are, for example, "SetHidden" instead of "SetbHidden"
+	FBoolProperty* BoolProperty = CastField<FBoolProperty>(PropertyAndAddress.Property);
+	if (BoolProperty)
+	{
+		PropertyName.RemoveFromStart("b", ESearchCase::CaseSensitive);
+	}
+
 	// @todo: Constructing FNames from strings is _very_ costly and we really shouldn't be doing this at runtime.
 	//        This is a little better now we use a string builder and an FNAME_Find, but it's still not ideal
 	TStringBuilder<128> SetterName;
 	SetterName.Append(TEXT("Set"));
-	PropertyBinding.PropertyName.AppendString(SetterName);
+	SetterName = SetterName + PropertyName;
 
 	FName SetterFunctionName(SetterName.ToString(), FNAME_Find);
+
 	if (SetterFunctionName.IsNone() || ObjectClass->FindFunctionByName(SetterFunctionName) == nullptr)
 	{
-		if (FBoolProperty* BoolProperty = CastField<FBoolProperty>(PropertyAndAddress.Property))
+		if (BoolProperty)
 		{
 			// bitfield booleans potentially have an additional byte offset.
 			// In practice this is always 0 because the property internal offset itself is incremented, but this is here for completeness
diff --git a/Engine/Source/Runtime/MovieScene/Private/TrackInstancePropertyBindings.cpp b/Engine/Source/Runtime/MovieScene/Private/TrackInstancePropertyBindings.cpp
--- a/Engine/Source/Runtime/MovieScene/Private/TrackInstancePropertyBindings.cpp
+++ b/Engine/Source/Runtime/MovieScene/Private/TrackInstancePropertyBindings.cpp
@@ -513,10 +513,6 @@ FTrackInstancePropertyBindings::FTrackInstancePropertyBindings(FName InPropertyN
 	: PropertyPath(InPropertyPath)
 	, PropertyName(InPropertyName)
 {
-	static const FString Set(TEXT("Set"));
-	const FString FunctionString = Set + PropertyName.ToString();
-
-	FunctionName = FName(*FunctionString);
 }
 
 FProperty* FTrackInstancePropertyBindings::FindProperty(const UObject* Object, FStringView InPropertyPath)
@@ -613,7 +609,20 @@ void FTrackInstancePropertyBindings::CacheBinding(const UObject& Object)
 {
 	FResolvedPropertyAndFunction PropAndFunction = FindPropertyAndFunction(&Object, PropertyPath);
 	{
-		UFunction* SetterFunction = Object.FindFunction(FunctionName);
+		FString PropertyVarName = PropertyName.ToString();
+
+		// If this is a bool property, strip off the 'b' so that the "Set" functions to be 
+		// found are, for example, "SetHidden" instead of "SetbHidden"
+		FProperty* Property = PropAndFunction.GetValidProperty();
+		if (FBoolProperty* BoolProperty = CastField<FBoolProperty>(Property))
+		{
+			PropertyVarName.RemoveFromStart("b", ESearchCase::CaseSensitive);
+		}
+
+		static const FString Set("Set");
+		const FString FunctionString = Set + PropertyVarName;
+
+		UFunction* SetterFunction = Object.FindFunction(FName(*FunctionString));
 		if (SetterFunction && SetterFunction->NumParms >= 1)
 		{
 			PropAndFunction.SetterFunction = SetterFunction;
