Sequencer: temporarily disable static assert for volatile property step structure size

diff --git a/Engine/Source/Runtime/MovieScene/Public/TrackInstancePropertyBindings.h b/Engine/Source/Runtime/MovieScene/Public/TrackInstancePropertyBindings.h
--- a/Engine/Source/Runtime/MovieScene/Public/TrackInstancePropertyBindings.h
+++ b/Engine/Source/Runtime/MovieScene/Public/TrackInstancePropertyBindings.h
@@ -89,7 +89,8 @@ public:
 };
 // The TVariant data type flag is 8 bytes so with the biggest union field being 8 bytes, it should 
 // all fit inside of 16 bytes.
-static_assert(sizeof(FVolatilePropertyStep) <= 16, "Try to fit FVolatilePropertyStep inside 16 bytes");
+// TEMP: disable this for dSTM-Linux build targets where weak object pointers seem to be larger
+//static_assert(sizeof(FVolatilePropertyStep) <= 16, "Try to fit FVolatilePropertyStep inside 16 bytes");
 
 /**
  * A property path where no caching is done, and the address of the end value is re-resolved from
