Sequencer: Add null check to FMovieSceneEntitySystemGraph::Shutdown() to fix a crash caused by reference elimination nulling out system ptrs

diff --git a/Engine/Source/Runtime/MovieScene/Private/EntitySystem/MovieSceneEntitySystemGraphs.cpp b/Engine/Source/Runtime/MovieScene/Private/EntitySystem/MovieSceneEntitySystemGraphs.cpp
--- a/Engine/Source/Runtime/MovieScene/Private/EntitySystem/MovieSceneEntitySystemGraphs.cpp
+++ b/Engine/Source/Runtime/MovieScene/Private/EntitySystem/MovieSceneEntitySystemGraphs.cpp
@@ -602,9 +602,14 @@ void FMovieSceneEntitySystemGraph::Shutdown()
 
 	for (int32 Index = 0; Index < Nodes.Array.GetMaxIndex(); ++Index)
 	{
-		if (Nodes.Array.IsAllocated(Index))
+		// When reference elimination is enabled, it's possible for the pointers in this array
+		//    to become null during a GC pass which is also likely to be when Shutdown is called,
+		//    so we have to check for null here, even though that is not necessary if reference
+		//    elimination is disabled
+		UMovieSceneEntitySystem* System = Nodes.Array.IsAllocated(Index) ? Nodes.Array[Index].System : nullptr;
+		if (System)
 		{
-			Nodes.Array[Index].System->Abandon();
+			System->Abandon();
 		}
 	}
 
