Sequencer: let procedural channels declare that they can't be keyed

This prevents us from considering, say, noise channels, when trying to create keys. Otherwise we end up extending these channels' sections to the keying time, before realizing the add-key functions do nothing.

diff --git a/Engine/Source/Runtime/MovieScene/Private/Channels/MovieSceneChannelEditorData.cpp b/Engine/Source/Runtime/MovieScene/Private/Channels/MovieSceneChannelEditorData.cpp
--- a/Engine/Source/Runtime/MovieScene/Private/Channels/MovieSceneChannelEditorData.cpp
+++ b/Engine/Source/Runtime/MovieScene/Private/Channels/MovieSceneChannelEditorData.cpp
@@ -27,6 +27,7 @@ FMovieSceneChannelMetaData::FMovieSceneChannelMetaData()
 	, bCanCollapseToTrack(true)
 	, bRelativeToSection(false)
 	, SortOrder(0)
+	, bCanBeKeyed(true)
 	, bSortEmptyGroupsLast(true)
 	, bInvertValue(false)
 	, bCanPreserveRatio(false)
@@ -40,6 +41,7 @@ FMovieSceneChannelMetaData::FMovieSceneChannelMetaData(FName InName, FText InDis
 	, bCanCollapseToTrack(true)
 	, bRelativeToSection(false)
 	, SortOrder(0)
+	, bCanBeKeyed(true)
 	, bSortEmptyGroupsLast(true)
 	, bInvertValue(false)
 	, bCanPreserveRatio(false)
@@ -70,4 +72,4 @@ FFrameNumber FMovieSceneChannelMetaData::GetOffsetTime(const UMovieSceneSection*
 	return OffsetTime + KeyOffset.Get();
 }
 
-#endif	// WITH_EDITOR
\ No newline at end of file
+#endif	// WITH_EDITOR
diff --git a/Engine/Source/Runtime/MovieScene/Private/MovieSceneCommonHelpers.cpp b/Engine/Source/Runtime/MovieScene/Private/MovieSceneCommonHelpers.cpp
--- a/Engine/Source/Runtime/MovieScene/Private/MovieSceneCommonHelpers.cpp
+++ b/Engine/Source/Runtime/MovieScene/Private/MovieSceneCommonHelpers.cpp
@@ -2,6 +2,7 @@
 #include "MovieSceneCommonHelpers.h"
 #include "Channels/MovieSceneChannelEditorData.h"
 #include "Channels/MovieSceneChannel.h"
+#include "Channels/MovieSceneChannelProxy.h"
 #include "EntitySystem/MovieSceneIntermediatePropertyValue.h"
 #include "Components/SceneComponent.h"
 #include "GameFramework/Actor.h"
@@ -91,7 +92,7 @@ FText MovieSceneHelpers::GetDisplayPathName(const UMovieSceneSection* Section, c
 
 bool MovieSceneHelpers::IsSectionKeyable(const UMovieSceneSection* Section)
 {
-	if (!Section)
+	if (!Section || !Section->IsActive())
 	{
 		return false;
 	}
@@ -102,7 +103,33 @@ bool MovieSceneHelpers::IsSectionKeyable(const UMovieSceneSection* Section)
 		return false;
 	}
 
-	return !Track->IsRowEvalDisabled(Section->GetRowIndex()) && !Track->IsEvalDisabled() && Section->IsActive();
+	if (Track->IsEvalDisabled() || Track->IsRowEvalDisabled(Section->GetRowIndex()))
+	{
+		return false;
+	}
+
+#if WITH_EDITOR
+	bool bHasAnyKeyableChannels = false;
+	const FMovieSceneChannelProxy& ChannelProxy = Section->GetChannelProxy();
+	for (const FMovieSceneChannelEntry& ChannelEntry : ChannelProxy.GetAllEntries())
+	{
+		for (const FMovieSceneChannelMetaData& MetaData : ChannelEntry.GetMetaData())
+		{
+			if (MetaData.bEnabled && MetaData.bCanBeKeyed)
+			{
+				bHasAnyKeyableChannels = true;
+				goto DoneCheckingChannels;
+			}
+		}
+	}
+DoneCheckingChannels:
+	if (!bHasAnyKeyableChannels)
+	{
+		return false;
+	}
+#endif  // WITH_EDITOR
+
+	return true;
 }
 
 UMovieSceneSection* MovieSceneHelpers::FindSectionAtTime( TArrayView<UMovieSceneSection* const> Sections, FFrameNumber Time, int32 RowIndex )
diff --git a/Engine/Source/Runtime/MovieScene/Public/Channels/MovieSceneChannelEditorData.h b/Engine/Source/Runtime/MovieScene/Public/Channels/MovieSceneChannelEditorData.h
--- a/Engine/Source/Runtime/MovieScene/Public/Channels/MovieSceneChannelEditorData.h
+++ b/Engine/Source/Runtime/MovieScene/Public/Channels/MovieSceneChannelEditorData.h
@@ -74,6 +74,8 @@ struct FMovieSceneChannelMetaData
 	uint8 bRelativeToSection : 1;
 	/** A sort order for this channel. Channels are sorted by this order, then by name. Groups are sorted by the channel with the lowest sort order. */
 	uint32 SortOrder;
+	/** Whether this channel can be keyed */
+	bool bCanBeKeyed;
 	/** By default if a channel has no FText::Group specified, we put it last, by setting this to false we use SortIndex instead */
 	bool bSortEmptyGroupsLast;
 	/** Invert the value for display purposes */
diff --git a/Engine/Source/Runtime/MovieSceneTracks/Private/Channels/MovieSceneDoublePerlinNoiseChannelContainer.cpp b/Engine/Source/Runtime/MovieSceneTracks/Private/Channels/MovieSceneDoublePerlinNoiseChannelContainer.cpp
--- a/Engine/Source/Runtime/MovieSceneTracks/Private/Channels/MovieSceneDoublePerlinNoiseChannelContainer.cpp
+++ b/Engine/Source/Runtime/MovieSceneTracks/Private/Channels/MovieSceneDoublePerlinNoiseChannelContainer.cpp
@@ -37,7 +37,9 @@ void UMovieSceneDoublePerlinNoiseChannelContainer::ImportEntityImpl(const UE::Mo
 #if WITH_EDITOR
 FMovieSceneChannelHandle UMovieSceneDoublePerlinNoiseChannelContainer::AddChannelProxy(FName ChannelName, FMovieSceneChannelProxyData& ProxyData, const FMovieSceneChannelMetaData& MetaData)
 {
-	return ProxyData.AddWithDefaultEditorData(PerlinNoiseChannel, MetaData);
+	FMovieSceneChannelMetaData NoiseMetaData(MetaData);
+	NoiseMetaData.bCanBeKeyed = false;
+	return ProxyData.AddWithDefaultEditorData(PerlinNoiseChannel, NoiseMetaData);
 }
 #else
 void UMovieSceneDoublePerlinNoiseChannelContainer::AddChannelProxy(FName ChannelName, FMovieSceneChannelProxyData& ProxyData)
diff --git a/Engine/Source/Runtime/MovieSceneTracks/Private/Channels/MovieSceneFloatPerlinNoiseChannelContainer.cpp b/Engine/Source/Runtime/MovieSceneTracks/Private/Channels/MovieSceneFloatPerlinNoiseChannelContainer.cpp
--- a/Engine/Source/Runtime/MovieSceneTracks/Private/Channels/MovieSceneFloatPerlinNoiseChannelContainer.cpp
+++ b/Engine/Source/Runtime/MovieSceneTracks/Private/Channels/MovieSceneFloatPerlinNoiseChannelContainer.cpp
@@ -38,7 +38,9 @@ void UMovieSceneFloatPerlinNoiseChannelContainer::ImportEntityImpl(const UE::Mov
 #if WITH_EDITOR
 FMovieSceneChannelHandle UMovieSceneFloatPerlinNoiseChannelContainer::AddChannelProxy(FName ChannelName, FMovieSceneChannelProxyData& ProxyData, const FMovieSceneChannelMetaData& MetaData)
 {
-	return ProxyData.AddWithDefaultEditorData(PerlinNoiseChannel, MetaData);
+	FMovieSceneChannelMetaData NoiseMetaData(MetaData);
+	NoiseMetaData.bCanBeKeyed = false;
+	return ProxyData.AddWithDefaultEditorData(PerlinNoiseChannel, NoiseMetaData);
 }
 #else
 void UMovieSceneFloatPerlinNoiseChannelContainer::AddChannelProxy(FName ChannelName, FMovieSceneChannelProxyData& ProxyData)
