Sequencer: Allow game pause to pause the sequencer clock

This is now the default behavior and can be disabled with Sequencer.AllowPausedClock 0

#Synced-CL 48032220

diff --git a/Engine/Source/Runtime/MovieScene/Private/MovieSceneTimeController.cpp b/Engine/Source/Runtime/MovieScene/Private/MovieSceneTimeController.cpp
--- a/Engine/Source/Runtime/MovieScene/Private/MovieSceneTimeController.cpp
+++ b/Engine/Source/Runtime/MovieScene/Private/MovieSceneTimeController.cpp
@@ -19,6 +19,12 @@ static TAutoConsoleVariable<float> CVarSequencerSecondsPerFrame(
 	TEXT("Seconds per frame to wait when in play every frame mode."),
 	ECVF_Default);
 
+static TAutoConsoleVariable<bool> CVarSequencerAllowPausedClock(
+	TEXT("Sequencer.AllowPausedClock"),
+	true,
+	TEXT("Allow game pause to pause the sequencer clock."),
+	ECVF_Default);
+
 void FMovieSceneTimeController::Tick(float DeltaSeconds, float InPlayRate)
 {
 	OnTick(DeltaSeconds, InPlayRate);
@@ -109,7 +115,21 @@ FFrameTime FMovieSceneTimeController_ExternalClock::OnRequestCurrentTime(const F
 
 	ClockLastUpdateTime = CurrentTime;
 
-	double CurrentSequenceTimeSeconds = UndilatedDeltaTime + AccumulatedDilation;
+	double PausedTime = 0;
+	if (CVarSequencerAllowPausedClock->GetBool())
+	{
+		const TIndirectArray<FWorldContext>& WorldContexts = GEngine->GetWorldContexts();
+		for (const FWorldContext& Context : WorldContexts)
+		{	
+			if (((Context.WorldType == EWorldType::PIE) || (Context.WorldType == EWorldType::Game)) && (Context.World() != NULL))
+			{
+				PausedTime = Context.World()->RealTimeSeconds - Context.World()->AudioTimeSeconds;
+				break;
+			}
+		}
+	}
+
+	double CurrentSequenceTimeSeconds = UndilatedDeltaTime + AccumulatedDilation - PausedTime;
 
 	FFrameTime StartTime = StartTimeIfPlaying->ConvertTo(InCurrentTime.Rate);
 	FFrameTime NewTime   = StartTime + CurrentSequenceTimeSeconds * InCurrentTime.Rate;
