Sequencer: Fix SetRange call of constructor being skipped when added to a read-only sequence.

Due to the change 44192297(Sequencer x Curve Editor: Avoid Modify call when data is read-only / does not change... by dominik.peacock) when an existing section is loaded into a read-only movie scene it gives different defaults than when it is loaded into a normal one

If you look at UMovieSceneBoolSection::UMovieSceneBoolSection it's trying to set TRange::All as default, but it fails due to movie scene being read-only. There are more section types that use this method in constructor, so this might be more wide-spread

#Synced-CL 47637010

diff --git a/Engine/Source/Runtime/MovieScene/Public/MovieSceneSection.h b/Engine/Source/Runtime/MovieScene/Public/MovieSceneSection.h
--- a/Engine/Source/Runtime/MovieScene/Public/MovieSceneSection.h
+++ b/Engine/Source/Runtime/MovieScene/Public/MovieSceneSection.h
@@ -317,17 +317,12 @@ public:
 	 */
 	virtual void SetRange(const TRange<FFrameNumber>& NewRange)
 	{
-		const bool bHasChanged = NewRange != SectionRange.Value; 
-		if (IsReadOnly() || !bHasChanged)
-		{
-			return;
-		}
-		
 		// Skip TryModify for objects that still need initialization (i.e. we're in the object's constructor), because modifying objects in their constructor can lead to non-deterministic cook issues.
 		bool bCanSetRange = true;
 		if (!HasAnyFlags(RF_NeedInitialization))
 		{
-			bCanSetRange = TryModify();
+			const bool bHasChanged = NewRange != SectionRange.Value; 
+			bCanSetRange = !IsReadOnly() && bHasChanged && TryModify();
 		}
 
 		const bool bInvariants = ensure(
