Removed unhelpful VeryVerbose logging

#Synced-CL 47550262

diff --git a/Engine/Source/Runtime/MovieScene/Public/EntitySystem/MovieSceneEntitySystemTask.h b/Engine/Source/Runtime/MovieScene/Public/EntitySystem/MovieSceneEntitySystemTask.h
--- a/Engine/Source/Runtime/MovieScene/Public/EntitySystem/MovieSceneEntitySystemTask.h
+++ b/Engine/Source/Runtime/MovieScene/Public/EntitySystem/MovieSceneEntitySystemTask.h
@@ -1586,8 +1586,6 @@ struct TEntityTaskBase
 
 	void Run(TaskImpl& TaskImplInstance)
 	{
-		UE_LOG(LogMovieSceneECS, VeryVerbose, TEXT("Running entity task the following components: %s"), *FilteredTask.GetComponents().ToString(EntityManager));
-
 		PreTask(&TaskImplInstance);
 
 		EComponentHeaderLockMode LockMode = EntityManager->GetThreadingModel() == EEntityThreadingModel::NoThreading
@@ -1741,8 +1739,6 @@ struct TEntityAllocationTaskBase
 
 	void Run(TaskImpl& TaskImplInstance)
 	{
-		UE_LOG(LogMovieSceneECS, VeryVerbose, TEXT("Running entity task the following components: %s"), *ComponentFilter.GetComponents().ToString(EntityManager));
-
 		PreTask(&TaskImplInstance);
 
 		EComponentHeaderLockMode LockMode = EntityManager->GetThreadingModel() == EEntityThreadingModel::NoThreading
@@ -1857,8 +1853,6 @@ struct TUnstructuredTask
 			checkf(CurrentThread == CommonParams.DesiredThread, TEXT("MovieScene evaluation task is not being run on its desired thread"));
 		}
 
-		UE_LOG(LogMovieSceneECS, VeryVerbose, TEXT("Running unstructured task"));
-
 		TaskImplInstance.Run();
 	}
 
