[Sequencer] Make the Jump method return the sequence's existing status instead of always returning Stopped. Prevents audio components from stopping for one frame following a de-sync correction where we jump backwards instead of play forwards.

#tests PIE
#rnx
#Synced-CL 48398309

diff --git a/Engine/Source/Runtime/MovieScene/Private/MovieSceneSequencePlayer.cpp b/Engine/Source/Runtime/MovieScene/Private/MovieSceneSequencePlayer.cpp
--- a/Engine/Source/Runtime/MovieScene/Private/MovieSceneSequencePlayer.cpp
+++ b/Engine/Source/Runtime/MovieScene/Private/MovieSceneSequencePlayer.cpp
@@ -1146,12 +1146,12 @@ void UMovieSceneSequencePlayer::UpdateTimeCursorPosition(FFrameTime NewPosition,
 	}
 }
 
-EMovieScenePlayerStatus::Type UpdateMethodToStatus(EUpdatePositionMethod Method)
+EMovieScenePlayerStatus::Type UpdateMethodToStatus(EUpdatePositionMethod Method, EMovieScenePlayerStatus::Type SequenceStatus)
 {
 	switch(Method)
 	{
 	case EUpdatePositionMethod::Scrub: return EMovieScenePlayerStatus::Scrubbing;
-	case EUpdatePositionMethod::Jump:  return EMovieScenePlayerStatus::Stopped;
+	case EUpdatePositionMethod::Jump:  return SequenceStatus;						// If jumping, the new status after the jump should be whatever the current status of the sequence is
 	case EUpdatePositionMethod::Play:  return EMovieScenePlayerStatus::Playing;
 	default:                           return EMovieScenePlayerStatus::Stopped;
 	}
@@ -1221,7 +1221,8 @@ void UMovieSceneSequencePlayer::StartTimeControllerAndBroadcastPlayState()
 
 void UMovieSceneSequencePlayer::UpdateTimeCursorPosition_Internal(FFrameTime NewPosition, EUpdatePositionMethod Method, bool bHasJumpedOverride)
 {
-	EMovieScenePlayerStatus::Type StatusOverride = UpdateMethodToStatus(Method);
+	bHasJumpedOverride = bHasJumpedOverride || (Method == EUpdatePositionMethod::Jump);
+	EMovieScenePlayerStatus::Type StatusOverride = UpdateMethodToStatus(Method, Status);
 
 	const int32 Duration = DurationFrames;
 	if (Duration == 0 && DurationSubFrames == 0.f)
@@ -1289,7 +1290,7 @@ void UMovieSceneSequencePlayer::UpdateTimeCursorPosition_Internal(FFrameTime New
 				(!bReversePlayback && CurrentPosition < EndTimeWithReversed))
 			{
 				FMovieSceneEvaluationRange Range = PlayPosition.PlayTo(EndTimeWithReversed);
-				UpdateMovieSceneInstance(Range, StatusOverride);
+				UpdateMovieSceneInstance(Range, StatusOverride, bHasJumpedOverride);
 			}
 
 			const FFrameTime Overplay = FFrameTime::FromDecimal(FMath::Fmod(PositionRelativeToStart, DurationWithSubFrames));
@@ -1341,7 +1342,7 @@ void UMovieSceneSequencePlayer::UpdateTimeCursorPosition_Internal(FFrameTime New
 			NewPosition = FMath::Clamp(NewPosition, FFrameTime(StartTime), GetLastValidTime());
 
 			FMovieSceneEvaluationRange Range = UpdatePlayPosition(PlayPosition, NewPosition, Method, bReversePlayback);
-			UpdateMovieSceneInstance(Range, StatusOverride);
+			UpdateMovieSceneInstance(Range, StatusOverride, bHasJumpedOverride);
 
 			// We have authority to finish playback if:
 			// 1. There's no playback replication (standalone sequence)
