Fix Loop Count being cleared when playback ends

Previously PlayLooping made a permanent change to the playback settings, but it's more fitting that its changes should be temporary. With this change PlayLooping will cache the playback settings before changing them, then the temporary changes will be reverted back to the cached values when playback ends, instead of defaulting to setting the loop count back to 0. This was originally introduced as a result of UE-67637.

#Synced-CL 49032165

diff --git a/Engine/Source/Runtime/MovieScene/Private/MovieSceneSequencePlayer.cpp b/Engine/Source/Runtime/MovieScene/Private/MovieSceneSequencePlayer.cpp
--- a/Engine/Source/Runtime/MovieScene/Private/MovieSceneSequencePlayer.cpp
+++ b/Engine/Source/Runtime/MovieScene/Private/MovieSceneSequencePlayer.cpp
@@ -297,6 +297,12 @@ void UMovieSceneSequencePlayer::PlayLooping(int32 NumLoops)
 		return;
 	}
 
+	// PlayLooping should only override num loops temporarily, so cache the current settings here for restoring later when playback stops.
+	if (!CachedPlaybackSettings.IsSet())
+	{
+		CachedPlaybackSettings = PlaybackSettings;
+	}
+
 	PlaybackSettings.LoopCount.Value = NumLoops;
 	PlayInternal();
 }
@@ -472,8 +478,12 @@ void UMovieSceneSequencePlayer::StopInternal(FFrameTime TimeToResetTo)
 		PauseOnFrame.Reset();
 		LastTickGameTimeSeconds.Reset();
 
-		// Reset loop count on stop so that it doesn't persist to the next call to play
-		PlaybackSettings.LoopCount.Value = 0;
+		// Restore playback settings to cached state to clear any temporary overrides.
+		if (CachedPlaybackSettings.IsSet())
+		{
+			PlaybackSettings = CachedPlaybackSettings.GetValue();
+			CachedPlaybackSettings.Reset();
+		}
 
 		if (PlaybackSettings.FinishCompletionStateOverride == EMovieSceneCompletionModeOverride::ForceRestoreState)
 		{
diff --git a/Engine/Source/Runtime/MovieScene/Public/MovieSceneSequencePlayer.h b/Engine/Source/Runtime/MovieScene/Public/MovieSceneSequencePlayer.h
--- a/Engine/Source/Runtime/MovieScene/Public/MovieSceneSequencePlayer.h
+++ b/Engine/Source/Runtime/MovieScene/Public/MovieSceneSequencePlayer.h
@@ -724,6 +724,9 @@ protected:
 	UPROPERTY(replicated)
 	FMovieSceneSequencePlaybackSettings PlaybackSettings;
 
+	/** Cached playback settings for restoring when playback stops. This is used for temporary playback overrides, e.g. PlayLooping **/
+	TOptional<FMovieSceneSequencePlaybackSettings> CachedPlaybackSettings;
+
 	/** The root template instance we're evaluating */
 	UPROPERTY(transient)
 	FMovieSceneRootEvaluationTemplateInstance RootTemplateInstance;
