PR #14239: [Sequencer] Don't crash when a level sequence has a null track

#Synced-CL 49563507

diff --git a/Engine/Source/Runtime/MovieScene/Private/MovieScene.cpp b/Engine/Source/Runtime/MovieScene/Private/MovieScene.cpp
--- a/Engine/Source/Runtime/MovieScene/Private/MovieScene.cpp
+++ b/Engine/Source/Runtime/MovieScene/Private/MovieScene.cpp
@@ -1126,6 +1126,11 @@ UMovieSceneTrack* UMovieScene::FindTrack(TSubclassOf<UMovieSceneTrack> TrackClas
 	{
 		for (const auto& Track : Binding->GetTracks())
 		{
+			if (!IsValid(Track))
+			{
+				continue;
+			}
+
 			if (TrackClass.GetDefaultObject() == nullptr || Track->GetClass()->IsChildOf(TrackClass))
 			{
 				if (TrackName == NAME_None || Track->GetTrackName() == TrackName)
@@ -1147,6 +1152,11 @@ TArray<UMovieSceneTrack*> UMovieScene::FindTracks(TSubclassOf<UMovieSceneTrack>
 	{
 		for (const auto& Track : Binding->GetTracks())
 		{
+			if (!IsValid(Track))
+			{
+				continue;
+			}
+
 			if (TrackClass.GetDefaultObject() == nullptr || Track->GetClass()->IsChildOf(TrackClass))
 			{
 				if (TrackName == NAME_None || Track->GetTrackName() == TrackName)
@@ -1269,6 +1279,11 @@ UMovieSceneTrack* UMovieScene::FindTrack( TSubclassOf<UMovieSceneTrack> TrackCla
 
 	for (const TObjectPtr<UMovieSceneTrack>& Track : Tracks)
 	{
+		if (!IsValid(Track))
+		{
+			continue;
+		}
+
 		if (Track->GetClass()->IsChildOf(TrackClass))
 		{
 			FoundTrack = Track;
