Fix Split Sections. Previously when splitting, or trimming sections and deleting keys, the key was assumed to be cubic, and tangents were calculated after keys were deleted. This caused the wrong interpolation to be used for keyframes other than cubic, and caused the curves to change as their tangents were calculated when no keys existed before or after the newly added keys.

This change brings the behavior of adding keys in DeleteKeysFrom in line with the behavior of adding keys normally. We get the interpolation mode of the current time to keep keyframe interpolation the same, and rely on the calculated tangents from when keyframes existed on the channels before deleting, so the tangents of the newly added key will match the curve.

FMovieSceneDoubleChannel::DeleteKeysFrom no longer calls AutoSetTangents. This is only used by TrimSection, so it shouldn't impact the expected behavior of other functionality.

#Synced-CL 49559285

diff --git a/Engine/Source/Runtime/MovieScene/Private/Channels/MovieSceneCurveChannelImpl.cpp b/Engine/Source/Runtime/MovieScene/Private/Channels/MovieSceneCurveChannelImpl.cpp
--- a/Engine/Source/Runtime/MovieScene/Private/Channels/MovieSceneCurveChannelImpl.cpp
+++ b/Engine/Source/Runtime/MovieScene/Private/Channels/MovieSceneCurveChannelImpl.cpp
@@ -845,7 +845,8 @@ void TMovieSceneCurveChannelImpl<ChannelType>::DeleteKeysFrom(ChannelType* InCha
 			CurveValueType Value = 0;
 			if (Evaluate(InChannel, InTime, Value))
 			{
-				AddCubicKey(InChannel, InTime, Value);
+				EMovieSceneKeyInterpolation Interpolation = GetInterpolationMode(InChannel, InTime, EMovieSceneKeyInterpolation::SmartAuto);
+				AddKeyToChannel(InChannel, InTime, Value, Interpolation);
 			}
 		}
 	}
diff --git a/Engine/Source/Runtime/MovieScene/Private/Channels/MovieSceneDoubleChannel.cpp b/Engine/Source/Runtime/MovieScene/Private/Channels/MovieSceneDoubleChannel.cpp
--- a/Engine/Source/Runtime/MovieScene/Private/Channels/MovieSceneDoubleChannel.cpp
+++ b/Engine/Source/Runtime/MovieScene/Private/Channels/MovieSceneDoubleChannel.cpp
@@ -232,7 +232,6 @@ void FMovieSceneDoubleChannel::DeleteKeys(TArrayView<const FKeyHandle> InHandles
 void FMovieSceneDoubleChannel::DeleteKeysFrom(FFrameNumber InTime, bool bDeleteKeysBefore)
 {
 	FMovieSceneDoubleChannelImpl::DeleteKeysFrom(this, InTime, bDeleteKeysBefore);
-	AutoSetTangents();
 }
 
 void FMovieSceneDoubleChannel::RemapTimes(const UE::MovieScene::IRetimingInterface& Retimer)
