Anim In Engine: Fix weighted interpolation when tangents overlap to avoid the curve popping out of the tangents. We do this by scaling the tangents that overlap in time so that they no longer do so, need to update X values for each proportionally and then adjust the Y values to match up.   This is cached so there is a one time hit per evaluation of an interval, but only happens if it overlaps. We use ControlRig.AvoidOverlappingWeightedTangents cvar to control this which defaults true which can change existing animations but any animaiton with an overlap will have a bad pop so this should be very very rare.

#Synced-CL 49551760

diff --git a/Engine/Source/Runtime/MovieScene/Private/Channels/MovieSceneInterpolation.cpp b/Engine/Source/Runtime/MovieScene/Private/Channels/MovieSceneInterpolation.cpp
--- a/Engine/Source/Runtime/MovieScene/Private/Channels/MovieSceneInterpolation.cpp
+++ b/Engine/Source/Runtime/MovieScene/Private/Channels/MovieSceneInterpolation.cpp
@@ -3,10 +3,20 @@
 #include "Channels/MovieSceneInterpolation.h"
 #include "Curves/CurveEvaluation.h"
 #include "Algo/Partition.h"
+#include "HAL/IConsoleManager.h"
+
 
 namespace UE::MovieScene::Interpolation
 {
 
+static bool GAvoidOverlappingWeightedTangents = true;
+static FAutoConsoleVariableRef CAvoidOverlappingWeightedTangents(
+	TEXT("ControlRig.AvoidOverlappingWeightedTangents"),
+	GAvoidOverlappingWeightedTangents,
+	TEXT("(Default: true) If true weighted tangents won't overlap in time, legacy behavior was to have them overlap which can cause discontinuities"),
+	ECVF_Default
+);
+
 bool FInterpolationExtents::IsValid() const
 {
 	return MinValue != std::numeric_limits<double>::max() && MaxValue != std::numeric_limits<double>::lowest();
@@ -856,12 +866,11 @@ FWeightedCubicInterpolation::FWeightedCubicInterpolation(
 	StartKeyValue = StartValue;
 	EndKeyValue = EndValue;
 
-	double CosAngle, SinAngle, Angle;
-
+	double StartCosAngle, StartSinAngle, StartAngle;
 	// ---------------------------------------------------------------------------------
 	// Initialize the start key parameters
-	Angle = FMath::Atan(StartTangent * ToSeconds);
-	FMath::SinCos(&SinAngle, &CosAngle, Angle);
+	StartAngle = FMath::Atan(StartTangent * ToSeconds);
+	FMath::SinCos(&StartSinAngle, &StartCosAngle, StartAngle);
 
 	if (bStartIsWeighted)
 	{
@@ -874,14 +883,14 @@ FWeightedCubicInterpolation::FWeightedCubicInterpolation(
 		StartWeight = FMath::Sqrt(DXInSeconds*DXInSeconds + DY*DY) * OneThird;
 	}
 
-	const double StartKeyTanX = CosAngle * StartWeight + Time1;
-	StartKeyTanY              = SinAngle * StartWeight + StartValue;
-	NormalizedStartTanDX = (StartKeyTanX - Time1) / DXInSeconds;
+	double StartKeyTanX = StartCosAngle * StartWeight + Time1;
+	StartKeyTanY              = StartSinAngle * StartWeight + StartValue;
 
+	double EndCosAngle, EndSinAngle, EndAngle;
 	// ---------------------------------------------------------------------------------
 	// Initialize the end key parameters
-	Angle = FMath::Atan(EndTangent * ToSeconds);
-	FMath::SinCos(&SinAngle, &CosAngle, Angle);
+	EndAngle = FMath::Atan(EndTangent * ToSeconds);
+	FMath::SinCos(&EndSinAngle, &EndCosAngle, EndAngle);
 
 	if (bEndIsWeighted)
 	{
@@ -894,9 +903,36 @@ FWeightedCubicInterpolation::FWeightedCubicInterpolation(
 		EndWeight = FMath::Sqrt(DXInSeconds*DXInSeconds + DY*DY) * OneThird;
 	}
 
-	const double EndKeyTanX = -CosAngle * EndWeight + Time2;
-	EndKeyTanY              = -SinAngle * EndWeight + EndValue;
+	double EndKeyTanX = -EndCosAngle * EndWeight + Time2;
+	EndKeyTanY              = -EndSinAngle * EndWeight + EndValue;
+
+	if (GAvoidOverlappingWeightedTangents) //if prevent overlapping tangents 
+	{
+		if (EndKeyTanX < StartKeyTanX)
+		{
+			//shrink tangents based upon proportional distance 
+			const double Overlap = StartKeyTanX - EndKeyTanX;
+			double FirstXDistance = StartKeyTanX - Time1;
+			double SecondXDistance = Time2 - EndKeyTanX;
+			const double Total = FirstXDistance + SecondXDistance;
+			const double FirstProportion = FirstXDistance / Total;
+			const double SecondProption = SecondXDistance / Total;
+			const double MoveFirstBack = Overlap * FirstProportion;
+			const double MoveSecondForward = Overlap * SecondProption;
+
+			StartKeyTanX -= MoveFirstBack;
+			FirstXDistance = StartKeyTanX - Time1;
+			const double FirstYDistance = (FirstXDistance)*FMath::Tan(StartAngle);
+			StartKeyTanY = FirstYDistance + StartValue;
+
+			EndKeyTanX += MoveSecondForward;
+			SecondXDistance = Time2 - EndKeyTanX;
+			const double SecondYDistance = (SecondXDistance) * -FMath::Tan(EndAngle);
+			EndKeyTanY = SecondYDistance + EndValue;
 
+		}
+	}
+	NormalizedStartTanDX = (StartKeyTanX - Time1) / DXInSeconds;
 	NormalizedEndTanDX = (EndKeyTanX - Time1) / DXInSeconds;
 }
 
@@ -992,3 +1028,6 @@ int32 FWeightedCubicInterpolation::Solve(double InValue, TInterpSolutions<FFrame
 }
 
 } // namespace UE::MovieScene
+
+
+
