Sequencer: Speculative null check

#rnx

#Synced-CL 49563507

diff --git a/Engine/Source/Runtime/MovieScene/Private/Decorations/MovieSceneDecorationContainer.cpp b/Engine/Source/Runtime/MovieScene/Private/Decorations/MovieSceneDecorationContainer.cpp
--- a/Engine/Source/Runtime/MovieScene/Private/Decorations/MovieSceneDecorationContainer.cpp
+++ b/Engine/Source/Runtime/MovieScene/Private/Decorations/MovieSceneDecorationContainer.cpp
@@ -16,7 +16,7 @@ UObject* FMovieSceneDecorationContainer::FindDecoration(const TSubclassOf<UObjec
 		// Do an exact match - we intentionally do not support derived decorations
 		for (UObject* Decoration : Decorations)
 		{
-			if (Decoration->IsA(Class))
+			if (Decoration && Decoration->IsA(Class))
 			{
 				return Decoration;
 			}
@@ -61,7 +61,7 @@ void FMovieSceneDecorationContainer::RemoveDecoration(const TSubclassOf<UObject>
 	{
 		for (int32 Index = Decorations.Num()-1; Index >= 0; --Index)
 		{
-			if (Decorations[Index]->IsA(Class))
+			if (Decorations[Index] && Decorations[Index]->IsA(Class))
 			{
 				Event(Decorations[Index]);
 				Decorations.RemoveAtSwap(Index);
