Sequencer: If there are no channels, consider the section to be keyable.

This fixes a regression from 47440003 which checks for channels that can be keyed

#Synced-CL 49956477

#Synced-CL 49956477

diff --git a/Engine/Source/Runtime/MovieScene/Private/MovieSceneCommonHelpers.cpp b/Engine/Source/Runtime/MovieScene/Private/MovieSceneCommonHelpers.cpp
--- a/Engine/Source/Runtime/MovieScene/Private/MovieSceneCommonHelpers.cpp
+++ b/Engine/Source/Runtime/MovieScene/Private/MovieSceneCommonHelpers.cpp
@@ -109,10 +109,12 @@ bool MovieSceneHelpers::IsSectionKeyable(const UMovieSceneSection* Section)
 	}
 
 #if WITH_EDITOR
+	bool bHasAnyChannels = false;
 	bool bHasAnyKeyableChannels = false;
 	const FMovieSceneChannelProxy& ChannelProxy = Section->GetChannelProxy();
 	for (const FMovieSceneChannelEntry& ChannelEntry : ChannelProxy.GetAllEntries())
 	{
+		bHasAnyChannels = true;
 		for (const FMovieSceneChannelMetaData& MetaData : ChannelEntry.GetMetaData())
 		{
 			if (MetaData.bEnabled && MetaData.bCanBeKeyed)
@@ -123,6 +125,12 @@ bool MovieSceneHelpers::IsSectionKeyable(const UMovieSceneSection* Section)
 		}
 	}
 DoneCheckingChannels:
+
+	if (!bHasAnyChannels)
+	{
+		return true;
+	}
+
 	if (!bHasAnyKeyableChannels)
 	{
 		return false;
diff --git a/Engine/Source/Runtime/MovieScene/Public/MovieSceneCommonHelpers.h b/Engine/Source/Runtime/MovieScene/Public/MovieSceneCommonHelpers.h
--- a/Engine/Source/Runtime/MovieScene/Public/MovieSceneCommonHelpers.h
+++ b/Engine/Source/Runtime/MovieScene/Public/MovieSceneCommonHelpers.h
@@ -62,7 +62,9 @@ public:
 	static MOVIESCENE_API FText GetDisplayPathName(const UMovieSceneSection* Section, const FMovieSceneChannel* Channel, const FMovieSceneChannelMetaData& ChannelMetaData);
 
 	/** 
-	 * @return Whether the section is keyable (active, on a track that is not muted, etc 
+	 * @return Whether the section is considered keyable. The section must be active and on a track that is not muted. 
+	 * If there are existing channels, the channels must be active and marked as keyable. A section with no channels is 
+	 * considered keyable.
 	 */
 	static MOVIESCENE_API bool IsSectionKeyable(const UMovieSceneSection*);
 
