[MovieScene/ALT]
- Fix potential assert when loading movie scenes

#Synced-CL 49920339

diff --git a/Engine/Source/Runtime/MovieScene/Private/MovieSceneModule.cpp b/Engine/Source/Runtime/MovieScene/Private/MovieSceneModule.cpp
--- a/Engine/Source/Runtime/MovieScene/Private/MovieSceneModule.cpp
+++ b/Engine/Source/Runtime/MovieScene/Private/MovieSceneModule.cpp
@@ -34,7 +34,7 @@ struct FCachedLegacyConversionFrameRate
 
 	void OnChanged()
 	{
-		TryParseString(FrameRate, *CVarLegacyConversionFrameRate.GetValueOnGameThread());
+		TryParseString(FrameRate, *CVarLegacyConversionFrameRate.GetValueOnAnyThread());
 	}
 
 	FFrameRate FrameRate;
