Sequencer:  Local timewarp support for blender by AndrewR

#Synced-CL 50014700

diff --git a/Engine/Source/Runtime/MovieScene/Private/EntitySystem/BuiltInComponentTypes.cpp b/Engine/Source/Runtime/MovieScene/Private/EntitySystem/BuiltInComponentTypes.cpp
--- a/Engine/Source/Runtime/MovieScene/Private/EntitySystem/BuiltInComponentTypes.cpp
+++ b/Engine/Source/Runtime/MovieScene/Private/EntitySystem/BuiltInComponentTypes.cpp
@@ -371,6 +371,7 @@ FBuiltInComponentTypes::FBuiltInComponentTypes()
 	ComponentRegistry->NewComponentType(&StringResult,          TEXT("String Result"));
 	ComponentRegistry->NewComponentType(&TextResult,            TEXT("Text Result"));
 	ComponentRegistry->NewComponentType(&ObjectResult,          TEXT("Object Result"));
+	ComponentRegistry->NewComponentType(&TimeWarp,              TEXT("Time Warp"));
 
 	ComponentRegistry->NewComponentType(&BaseByte,			    TEXT("Base Byte"));
 	ComponentRegistry->NewComponentType(&BaseInteger,			TEXT("Base Integer"));
diff --git a/Engine/Source/Runtime/MovieScene/Private/EntitySystem/MovieSceneEvalTimeSystem.cpp b/Engine/Source/Runtime/MovieScene/Private/EntitySystem/MovieSceneEvalTimeSystem.cpp
--- a/Engine/Source/Runtime/MovieScene/Private/EntitySystem/MovieSceneEvalTimeSystem.cpp
+++ b/Engine/Source/Runtime/MovieScene/Private/EntitySystem/MovieSceneEvalTimeSystem.cpp
@@ -6,6 +6,7 @@
 #include "EntitySystem/MovieSceneEntitySystemLinker.h"
 #include "EntitySystem/BuiltInComponentTypes.h"
 #include "EntitySystem/MovieSceneInstanceRegistry.h"
+#include "Variants/MovieSceneTimeWarpGetter.h"
 #include "Stats/Stats.h"
 
 #include UE_INLINE_GENERATED_CPP_BY_NAME(MovieSceneEvalTimeSystem)
@@ -34,7 +35,8 @@ struct FGatherTimes
 
 			FEvaluatedTime EvaluatedTime;
 			EvaluatedTime.FrameTime = Context.GetTime();
-			EvaluatedTime.Seconds = Context.GetFrameRate().AsSeconds(EvaluatedTime.FrameTime);
+			EvaluatedTime.FrameRate = Context.GetFrameRate();
+			EvaluatedTime.Seconds = EvaluatedTime.FrameRate.AsSeconds(EvaluatedTime.FrameTime);
 
 			(*EvaluatedTimes)[It.GetIndex()] = EvaluatedTime;
 		}
@@ -69,6 +71,39 @@ struct FAssignEvalSecondsTask
 	}
 };
 
+
+struct FAssignTimeWarpedEvalTimesTask
+{
+	const TArray<FEvaluatedTime, TInlineAllocator<16>>* EvaluatedTimes;
+
+	FAssignTimeWarpedEvalTimesTask(const TArray<FEvaluatedTime, TInlineAllocator<16>>* InEvaluatedTimes)
+		: EvaluatedTimes(InEvaluatedTimes)
+	{}
+
+	void ForEachEntity(FInstanceHandle InstanceHandle, UMovieSceneTimeWarpGetter* TimeWarp, FFrameTime& EvalTime) const
+	{
+		FFrameTime UnwarpedEvalTime = (*EvaluatedTimes)[InstanceHandle.InstanceID].FrameTime;
+		EvalTime = TimeWarp->RemapTime(UnwarpedEvalTime);
+	}
+};
+
+struct FAssignTimeWarpedEvalSecondsTask
+{
+	const TArray<FEvaluatedTime, TInlineAllocator<16>>* EvaluatedTimes;
+
+	FAssignTimeWarpedEvalSecondsTask(const TArray<FEvaluatedTime, TInlineAllocator<16>>* InEvaluatedTimes)
+		: EvaluatedTimes(InEvaluatedTimes)
+	{}
+
+	void ForEachEntity(FInstanceHandle InstanceHandle, UMovieSceneTimeWarpGetter* TimeWarp, double& EvalSeconds) const
+	{
+		const FEvaluatedTime EvaluatedTime = (*EvaluatedTimes)[InstanceHandle.InstanceID];
+		const FFrameTime     EvalTime      = TimeWarp->RemapTime(EvaluatedTime.FrameTime);
+
+		EvalSeconds = EvaluatedTime.FrameRate.AsSeconds(EvalTime);
+	}
+};
+
 } // namespace UE::MovieScene
 
 UMovieSceneEvalTimeSystem::UMovieSceneEvalTimeSystem(const FObjectInitializer& ObjInit)
@@ -110,19 +145,37 @@ void UMovieSceneEvalTimeSystem::OnSchedulePersistentTasks(UE::MovieScene::IEntit
 	FTaskID WriteTask1 = FEntityTaskBuilder()
 	.Read(BuiltInComponents->InstanceHandle)
 	.Write(BuiltInComponents->EvalTime)
-	.FilterNone({ BuiltInComponents->Tags.FixedTime })
+	.FilterNone({ BuiltInComponents->Tags.FixedTime, BuiltInComponents->TimeWarp })
 	.SetStat(GET_STATID(MovieSceneEval_AssignEvalTimes))
 	.Fork_PerEntity<FAssignEvalTimesTask>(&Linker->EntityManager, TaskScheduler, &EvaluatedTimes);
 
 	FTaskID WriteTask2 = FEntityTaskBuilder()
 	.Read(BuiltInComponents->InstanceHandle)
 	.Write(BuiltInComponents->EvalSeconds)
-	.FilterNone({ BuiltInComponents->Tags.FixedTime })
+	.FilterNone({ BuiltInComponents->Tags.FixedTime, BuiltInComponents->TimeWarp })
 	.SetStat(GET_STATID(MovieSceneEval_AssignEvalTimes))
 	.Fork_PerEntity<FAssignEvalSecondsTask>(&Linker->EntityManager, TaskScheduler, &EvaluatedTimes);
 
+	FTaskID WriteTask3 = FEntityTaskBuilder()
+	.Read(BuiltInComponents->InstanceHandle)
+	.Read(BuiltInComponents->TimeWarp)
+	.Write(BuiltInComponents->EvalTime)
+	.FilterNone({ BuiltInComponents->Tags.FixedTime })
+	.SetStat(GET_STATID(MovieSceneEval_AssignEvalTimes))
+	.Fork_PerEntity<FAssignTimeWarpedEvalTimesTask>(&Linker->EntityManager, TaskScheduler, &EvaluatedTimes);
+
+	FTaskID WriteTask4 = FEntityTaskBuilder()
+	.Read(BuiltInComponents->InstanceHandle)
+	.Read(BuiltInComponents->TimeWarp)
+	.Write(BuiltInComponents->EvalSeconds)
+	.FilterNone({ BuiltInComponents->Tags.FixedTime })
+	.SetStat(GET_STATID(MovieSceneEval_AssignEvalTimes))
+	.Fork_PerEntity<FAssignTimeWarpedEvalSecondsTask>(&Linker->EntityManager, TaskScheduler, &EvaluatedTimes);
+
 	TaskScheduler->AddPrerequisite(GatherTask, WriteTask1);
 	TaskScheduler->AddPrerequisite(GatherTask, WriteTask2);
+	TaskScheduler->AddPrerequisite(GatherTask, WriteTask3);
+	TaskScheduler->AddPrerequisite(GatherTask, WriteTask4);
 }
 
 void UMovieSceneEvalTimeSystem::OnRun(FSystemTaskPrerequisites& InPrerequisites, FSystemSubsequentTasks& Subsequents)
@@ -155,4 +208,5 @@ void UMovieSceneEvalTimeSystem::OnRun(FSystemTaskPrerequisites& InPrerequisites,
 	.FilterNone({ FBuiltInComponentTypes::Get()->Tags.FixedTime })
 	.SetStat(GET_STATID(MovieSceneEval_AssignEvalTimes))
 	.Dispatch_PerEntity<FAssignEvalSecondsTask>(&Linker->EntityManager, EvalPrereqs, &Subsequents, &EvaluatedTimes);
-}
\ No newline at end of file
+}
+
diff --git a/Engine/Source/Runtime/MovieScene/Public/EntitySystem/BuiltInComponentTypes.h b/Engine/Source/Runtime/MovieScene/Public/EntitySystem/BuiltInComponentTypes.h
--- a/Engine/Source/Runtime/MovieScene/Public/EntitySystem/BuiltInComponentTypes.h
+++ b/Engine/Source/Runtime/MovieScene/Public/EntitySystem/BuiltInComponentTypes.h
@@ -712,6 +712,9 @@ public:
 	// The result of an evaluated FMovieSceneObjectPathChannel
 	TComponentTypeID<FObjectComponent> ObjectResult;
 
+	// Additional timewarp to apply to Eval Times on the same entity
+	TComponentTypeID<TObjectPtr<UMovieSceneTimeWarpGetter>> TimeWarp;
+
 public:
 
 	// An FEasingComponentData for computing easing curves
diff --git a/Engine/Source/Runtime/MovieScene/Public/EntitySystem/MovieSceneEvalTimeSystem.h b/Engine/Source/Runtime/MovieScene/Public/EntitySystem/MovieSceneEvalTimeSystem.h
--- a/Engine/Source/Runtime/MovieScene/Public/EntitySystem/MovieSceneEvalTimeSystem.h
+++ b/Engine/Source/Runtime/MovieScene/Public/EntitySystem/MovieSceneEvalTimeSystem.h
@@ -21,6 +21,7 @@ namespace UE::MovieScene
 	{
 		FFrameTime FrameTime;
 		double Seconds;
+		FFrameRate FrameRate;
 	};
 }
 
