Sequencer: Fix crash when compiling sequences with null tracks

Runs 'RemoveNullTracks) before calling CompileTrack() to prevent crashes when a
sequence contains null track references. We need to track down what is causing tracks to be null, but this should fix the crash as it currently exists.

Fixes crash: Case 35512

[RN] fix animation/sequencer

#Synced-CL 49920339

diff --git a/Engine/Source/Runtime/MovieScene/Private/Compilation/MovieSceneCompiledDataManager.cpp b/Engine/Source/Runtime/MovieScene/Private/Compilation/MovieSceneCompiledDataManager.cpp
--- a/Engine/Source/Runtime/MovieScene/Private/Compilation/MovieSceneCompiledDataManager.cpp
+++ b/Engine/Source/Runtime/MovieScene/Private/Compilation/MovieSceneCompiledDataManager.cpp
@@ -956,6 +956,9 @@ void UMovieSceneCompiledDataManager::Compile(FMovieSceneCompiledDataID DataID, U
 	{
 		if (ensure(MovieScene))
 		{
+			// Ensure we have no null tracks. This should really not be the case, but we have crash reporter cases of this happening
+			ensure(!MovieScene->RemoveNullTracks());
+
 			for (const FMovieSceneMarkedFrame& Mark : MovieScene->GetMarkedFrames())
 			{
 				if (Mark.bIsDeterminismFence)
@@ -970,14 +973,14 @@ void UMovieSceneCompiledDataManager::Compile(FMovieSceneCompiledDataID DataID, U
 			}
 
 			for (UMovieSceneTrack* Track : MovieScene->GetTracks())
-			{
+			{	
 				CompileTrack(&Entry, nullptr, Track, Params, &GatheredSignatures, &GatheredData);
 			}
 
 			for (const FMovieSceneBinding& ObjectBinding : ((const UMovieScene*)MovieScene)->GetBindings())
 			{
 				for (UMovieSceneTrack* Track : ObjectBinding.GetTracks())
-				{
+				{	
 					CompileTrack(&Entry, &ObjectBinding, Track, Params, &GatheredSignatures, &GatheredData);
 				}
 			}
diff --git a/Engine/Source/Runtime/MovieScene/Private/MovieScene.cpp b/Engine/Source/Runtime/MovieScene/Private/MovieScene.cpp
--- a/Engine/Source/Runtime/MovieScene/Private/MovieScene.cpp
+++ b/Engine/Source/Runtime/MovieScene/Private/MovieScene.cpp
@@ -1531,15 +1531,15 @@ void UMovieScene::UpgradeTimeRanges()
 #endif
 }
 
-#if WITH_EDITOR
-
-void UMovieScene::RemoveNullTracks()
+bool UMovieScene::RemoveNullTracks()
 {
 	// Remove any null tracks
+	bool bNullTracksFound = false;
 	for( int32 TrackIndex = 0; TrackIndex < Tracks.Num(); )
 	{
 		if (Tracks[TrackIndex] == nullptr)
 		{
+			bNullTracksFound = true;
 			Tracks.RemoveAt(TrackIndex);
 		}
 		else
@@ -1550,8 +1550,8 @@ void UMovieScene::RemoveNullTracks()
 
 	for (FMovieSceneBinding& Binding : ObjectBindings)
 	{
-		Binding.RemoveNullTracks();
-		Binding.RemoveNulls();
+		bNullTracksFound |= Binding.RemoveNullTracks();
+		bNullTracksFound |= Binding.RemoveNulls();
 	}
 
 #if WITH_EDITORONLY_DATA
@@ -1559,6 +1559,7 @@ void UMovieScene::RemoveNullTracks()
 	{
 		if (RootFolders[RootFolderIndex] == nullptr)
 		{
+			bNullTracksFound |= true;
 			RootFolders.RemoveAt(RootFolderIndex);
 		}
 		else 
@@ -1579,9 +1580,9 @@ void UMovieScene::RemoveNullTracks()
 	// Clean any section groups which might refer to sections which were not serialized
 	CleanSectionGroups();
 #endif
-}
 
-#endif // WITH_EDITOR
+	return bNullTracksFound;
+}
 
 /* UObject interface
  *****************************************************************************/
diff --git a/Engine/Source/Runtime/MovieScene/Private/MovieSceneBinding.cpp b/Engine/Source/Runtime/MovieScene/Private/MovieSceneBinding.cpp
--- a/Engine/Source/Runtime/MovieScene/Private/MovieSceneBinding.cpp
+++ b/Engine/Source/Runtime/MovieScene/Private/MovieSceneBinding.cpp
@@ -36,16 +36,19 @@ bool FMovieSceneBinding::RemoveTrack(UMovieSceneTrack& Track, UMovieScene* Owner
 	return false;
 }
 
-void FMovieSceneBinding::RemoveNullTracks()
+bool FMovieSceneBinding::RemoveNullTracks()
 {
+	bool bNullTracksFound = false;
 	for (int32 TrackIndex = Tracks.Num()-1; TrackIndex >= 0; --TrackIndex)
 	{
 		if (Tracks[TrackIndex] == nullptr)
 		{
+			bNullTracksFound = true;
 			Tracks.RemoveAt(TrackIndex);
 			// Don't trigger events for null tracks
 		}
 	}
+	return bNullTracksFound;
 }
 
 TArray<UMovieSceneTrack*> FMovieSceneBinding::StealTracks(UMovieScene* Owner)
diff --git a/Engine/Source/Runtime/MovieScene/Public/Decorations/MovieSceneDecorationContainer.h b/Engine/Source/Runtime/MovieScene/Public/Decorations/MovieSceneDecorationContainer.h
--- a/Engine/Source/Runtime/MovieScene/Public/Decorations/MovieSceneDecorationContainer.h
+++ b/Engine/Source/Runtime/MovieScene/Public/Decorations/MovieSceneDecorationContainer.h
@@ -92,9 +92,9 @@ public:
 	/**
 	 * Remove any null decoration ptrs
 	 */
-	void RemoveNulls()
+	bool RemoveNulls()
 	{
-		Decorations.Remove(nullptr);
+		return Decorations.Remove(nullptr) > 0;
 	}
 
 protected:
diff --git a/Engine/Source/Runtime/MovieScene/Public/MovieScene.h b/Engine/Source/Runtime/MovieScene/Public/MovieScene.h
--- a/Engine/Source/Runtime/MovieScene/Public/MovieScene.h
+++ b/Engine/Source/Runtime/MovieScene/Public/MovieScene.h
@@ -1226,6 +1226,9 @@ public:
 	/* Called by the compiler to empty the list of generated conditions*/
 	void ResetGeneratedConditions() { GeneratedConditions.Reset(); }
 
+	/* Remove null tracks on this movie scene. Returns true if any discovered and removed. */
+	MOVIESCENE_API bool RemoveNullTracks();
+
 protected:
 
 	/**
@@ -1275,7 +1278,6 @@ private:
 
 #if WITH_EDITOR
 	MOVIESCENE_API void OptimizeForCook();
-	MOVIESCENE_API void RemoveNullTracks();
 #endif
 
 private:
diff --git a/Engine/Source/Runtime/MovieScene/Public/MovieSceneBinding.h b/Engine/Source/Runtime/MovieScene/Public/MovieSceneBinding.h
--- a/Engine/Source/Runtime/MovieScene/Public/MovieSceneBinding.h
+++ b/Engine/Source/Runtime/MovieScene/Public/MovieSceneBinding.h
@@ -110,7 +110,7 @@ public:
 	/**
 	 * Removes all null tracks from this binding
 	 */
-	MOVIESCENE_API void RemoveNullTracks();
+	MOVIESCENE_API bool RemoveNullTracks();
 
 	/**
 	 * @return All tracks in this binding
