Sequencer: if a property binding path leads to an invalid index of an array, mark that property binding as invalid

#Synced-CL 50723848

diff --git a/Engine/Source/Runtime/MovieScene/Private/TrackInstancePropertyBindings.cpp b/Engine/Source/Runtime/MovieScene/Private/TrackInstancePropertyBindings.cpp
--- a/Engine/Source/Runtime/MovieScene/Private/TrackInstancePropertyBindings.cpp
+++ b/Engine/Source/Runtime/MovieScene/Private/TrackInstancePropertyBindings.cpp
@@ -154,6 +154,11 @@ void ResolvePropertyRecursive(void* BasePointer, UStruct* InStruct, TArrayView<c
 					return;
 				}
 			}
+			else
+			{
+				// The property path goes into an invalid index of the array.
+				OutResolutionState.bIsValid = false;
+			}
 		}
 		else
 		{
