Sequencer: Follow up to 50778806, add Modify

#rnx

#Synced-CL 50752143

diff --git a/Engine/Source/Runtime/MovieScene/Private/MovieSceneCommonHelpers.cpp b/Engine/Source/Runtime/MovieScene/Private/MovieSceneCommonHelpers.cpp
--- a/Engine/Source/Runtime/MovieScene/Private/MovieSceneCommonHelpers.cpp
+++ b/Engine/Source/Runtime/MovieScene/Private/MovieSceneCommonHelpers.cpp
@@ -553,6 +553,8 @@ void MovieSceneHelpers::GetDescendantMovieScenes(UMovieSceneSequence* InSequence
 
 void MovieSceneHelpers::DisableTrack(UMovieSceneTrack* Track)
 {
+	Track->Modify();
+
 	// Disable the entire track or each of the rows
 	if (Track->GetMaxRowIndex() == 0)
 	{
