Revert parts of change 50753723 to fix more pre-animated state issues

The problem with the previous change was that it was a half-step towards the idea of tracking *all* contributors to a value, while saving data for only those who matter for pre-animated state. Right now, we also track only contributors the pre-animated state, not all contributors. Continuing in that direction proved to be a far-ranging and risky refactor. While it was "technically correct" (the best kind of correct, as we all know), it wasn't so pragmatic. So instead, this change now reverts a few parts of the previous change, but *keeps* the tracking of all contributors *specifically* for the Component Mobility system, which was the main issue at hand. We can always revisit the "technically correct" solution if more edge cases are discovered, but since bug reports of that effect have been few and far between, this may be good enough.

#Synced-CL 51024716

diff --git a/Engine/Source/Runtime/MovieScene/Public/Evaluation/PreAnimatedState/MovieScenePreAnimatedObjectStorage.h b/Engine/Source/Runtime/MovieScene/Public/Evaluation/PreAnimatedState/MovieScenePreAnimatedObjectStorage.h
--- a/Engine/Source/Runtime/MovieScene/Public/Evaluation/PreAnimatedState/MovieScenePreAnimatedObjectStorage.h
+++ b/Engine/Source/Runtime/MovieScene/Public/Evaluation/PreAnimatedState/MovieScenePreAnimatedObjectStorage.h
@@ -118,6 +118,11 @@ public:
 		const int32 Num = Params.Num;
 		const bool  bWantsRestore = Params.bWantsRestoreState;
 
+		if (!this->ParentExtension->IsCapturingGlobalState() && !bWantsRestore)
+		{
+			return;
+		}
+
 		FPreAnimatedEntityCaptureSource* EntityMetaData = this->ParentExtension->GetOrCreateEntityMetaData();
 		for (int32 Index = 0; Index < Num; ++Index)
 		{
@@ -143,6 +148,11 @@ public:
 
 	void BeginTrackingEntity(FMovieSceneEntityID EntityID, bool bWantsRestoreState, FRootInstanceHandle RootInstanceHandle, UObject* BoundObject) override
 	{
+		if (!this->ParentExtension->IsCapturingGlobalState() && !bWantsRestoreState)
+		{
+			return;
+		}
+
 		FPreAnimatedEntityCaptureSource* EntityMetaData = this->ParentExtension->GetOrCreateEntityMetaData();
 
 		FPreAnimatedStateEntry Entry = MakeEntry(BoundObject);
diff --git a/Engine/Source/Runtime/MovieScene/Public/Evaluation/PreAnimatedState/MovieScenePreAnimatedPropertyStorage.h b/Engine/Source/Runtime/MovieScene/Public/Evaluation/PreAnimatedState/MovieScenePreAnimatedPropertyStorage.h
--- a/Engine/Source/Runtime/MovieScene/Public/Evaluation/PreAnimatedState/MovieScenePreAnimatedPropertyStorage.h
+++ b/Engine/Source/Runtime/MovieScene/Public/Evaluation/PreAnimatedState/MovieScenePreAnimatedPropertyStorage.h
@@ -172,6 +172,11 @@ struct TPreAnimatedPropertyStorageImpl<PropertyTraits, TPropertyMetaData<MetaDat
 		const int32 Num = Params.Num;
 		const bool  bWantsRestore = Params.bWantsRestoreState;
 
+		if (!this->ParentExtension->IsCapturingGlobalState() && !bWantsRestore)
+		{
+			return;
+		}
+
 		FPreAnimatedEntityCaptureSource* EntityMetaData = this->ParentExtension->GetOrCreateEntityMetaData();
 		for (int32 Index = 0; Index < Num; ++Index)
 		{
diff --git a/Engine/Source/Runtime/MovieSceneTracks/Private/Systems/MovieSceneComponentMobilitySystem.cpp b/Engine/Source/Runtime/MovieSceneTracks/Private/Systems/MovieSceneComponentMobilitySystem.cpp
--- a/Engine/Source/Runtime/MovieSceneTracks/Private/Systems/MovieSceneComponentMobilitySystem.cpp
+++ b/Engine/Source/Runtime/MovieSceneTracks/Private/Systems/MovieSceneComponentMobilitySystem.cpp
@@ -213,6 +213,8 @@ void UMovieSceneComponentMobilitySystem::SavePreAnimatedState(const FPreAnimatio
 			{
 				FPreAnimatedStateEntry Entry = ComponentMobilityStorage->MakeEntry(CurrentSceneComponent);
 				
+				// Always track entities, even KeepState ones, because we don't want to restore the original mobility if there
+				// are still contributors animating the transform.
 				EntityMetaData->BeginTrackingEntity(Entry, EntityID, InstanceHandle, bWantsRestore);
 				ComponentMobilityStorage->CachePreAnimatedValue(CacheValueParams, Entry, CurrentSceneComponent);
 			}
