PR #14376: Fixed a bug where conditions were not actually checked in the FBoundObjectTask for subsequences

#Synced-CL 51192822

diff --git a/Engine/Source/Runtime/MovieScene/Private/EntitySystem/MovieSceneEntityFactory.cpp b/Engine/Source/Runtime/MovieScene/Private/EntitySystem/MovieSceneEntityFactory.cpp
--- a/Engine/Source/Runtime/MovieScene/Private/EntitySystem/MovieSceneEntityFactory.cpp
+++ b/Engine/Source/Runtime/MovieScene/Private/EntitySystem/MovieSceneEntityFactory.cpp
@@ -211,6 +211,7 @@ void FBoundObjectTask::ForEachAllocation(FEntityAllocationProxy AllocationProxy,
 
 		const FSequenceInstance&     SequenceInstance = InstanceRegistry->GetInstance(Instances[Index]);
 		TArrayView<TWeakObjectPtr<>> BoundObjects     = SequenceInstance.GetSharedPlaybackState()->FindBoundObjects(ObjectBindings[Index], SequenceInstance.GetSequenceID());
+		const FSequenceInstance&     RootInstance     = InstanceRegistry->GetInstance(SequenceInstance.GetRootInstanceHandle());
 		
 		bool bCheckedCondition = false;
 		for (TWeakObjectPtr<> WeakObject : BoundObjects)
@@ -235,8 +236,9 @@ void FBoundObjectTask::ForEachAllocation(FEntityAllocationProxy AllocationProxy,
 					// If this entity has a condition that could depend on a bound object, then it hasn't yet been tested, and we must test it here. \
 					// Note that it will only be tested once here, and then the entity ledger will take care of testing it again if it needs to
 					// and it is a per-tick condition.
+					// Conditions can only be evaluated by the root instance, since only it has a SequenceUpdater
 					bCheckedCondition = true;
-					if (!SequenceInstance.EvaluateCondition(ObjectBindings[Index], SequenceInstance.GetSequenceID(), Conditions[Index], Conditions[Index]->GetTypedOuter<UMovieSceneSignedObject>()))
+					if (!RootInstance.EvaluateCondition(ObjectBindings[Index], SequenceInstance.GetSequenceID(), Conditions[Index], Conditions[Index]->GetTypedOuter<UMovieSceneSignedObject>()))
 					{
 						// Condition has failed, don't add this entity to the batch
 						break;
