Sequencer: Fix crash after creating a condition for a subsequence in a non-root sequence.

The condition evaluation path in FSequenceUpdater_Hierarchical::Update was
always passing RootOverrideSequenceID when evaluating conditions on sub-sequence
entries. This caused GetOrCreateDirectorInstance to look up the wrong director
BP for nested sub-sequences, leading to crashes or "no director instance" errors.

Now resolves the parent sequence's root-relative SequenceID from the hierarchy
node so the correct director instance is used for condition evaluation.

#Synced-CL 51609808

diff --git a/Engine/Source/Runtime/MovieScene/Private/EntitySystem/MovieSceneSequenceUpdaters.cpp b/Engine/Source/Runtime/MovieScene/Private/EntitySystem/MovieSceneSequenceUpdaters.cpp
--- a/Engine/Source/Runtime/MovieScene/Private/EntitySystem/MovieSceneSequenceUpdaters.cpp
+++ b/Engine/Source/Runtime/MovieScene/Private/EntitySystem/MovieSceneSequenceUpdaters.cpp
@@ -844,23 +844,35 @@ void FSequenceUpdater_Hierarchical::Update(TSharedRef<const FSharedPlaybackState
 			const UMovieSceneCondition* Condition = SubData->WeakCondition.Get();
 			if (Condition)
 			{
+				// The condition's director BP belongs to the parent sequence (the one
+				// containing the sub-section), so resolve the parent's root-relative ID
+				// for use in director instance lookup during condition evaluation.
+				FMovieSceneSequenceID ConditionOwnerSequenceID = RootOverrideSequenceID;
+				if (const FMovieSceneSequenceHierarchyNode* SubNode = RootOverrideHierarchy->FindNode(Entry.SequenceID))
+				{
+					if (SubNode->ParentID != MovieSceneSequenceID::Root)
+					{
+						ConditionOwnerSequenceID = RootOverridePath.ResolveChildSequenceID(SubNode->ParentID);
+					}
+				}
+
 				// If we're able to cache the condition result, then it should be cached above when its entity got processed- retrieve that value.
 				// Otherwise, test it again.
 				if (Condition->CanCacheResult(SharedPlaybackState))
 				{
-					if (bool* ConditionResult = CachedConditionResults.Find(Condition->ComputeCacheKey(FGuid(), RootOverrideSequenceID, SharedPlaybackState, FindObject<UMovieSceneSubSection>(CompiledDataManager->GetEntryRef(RootCompiledDataID).GetSequence(), *SubData->SectionPath.ToString()))))
+					if (bool* ConditionResult = CachedConditionResults.Find(Condition->ComputeCacheKey(FGuid(), ConditionOwnerSequenceID, SharedPlaybackState, FindObject<UMovieSceneSubSection>(CompiledDataManager->GetEntryRef(RootCompiledDataID).GetSequence(), *SubData->SectionPath.ToString()))))
 					{
 						if (*ConditionResult == false)
 						{
 							bSubSequenceConditionFailed = true;
 						}
 					}
-					else if (!Condition->EvaluateCondition(FGuid(), RootOverrideSequenceID, SharedPlaybackState))
+					else if (!Condition->EvaluateCondition(FGuid(), ConditionOwnerSequenceID, SharedPlaybackState))
 					{
 						bSubSequenceConditionFailed = true;
 					}
 				}
-				else if (!Condition->EvaluateCondition(FGuid(), RootOverrideSequenceID, SharedPlaybackState))
+				else if (!Condition->EvaluateCondition(FGuid(), ConditionOwnerSequenceID, SharedPlaybackState))
 				{
 					bSubSequenceConditionFailed = true;
 				}
