Generalize FSkeletalMeshEditorParams to USkinnedMeshComponent

This struct is used during cinematics to ensure that SKMs still tick pose and bones. Making it work on USkinnedMeshComponent allows this struct to also work for chaos cloth components.

#Synced-CL 51899829
#Synced-CL 51899829

diff --git a/Engine/Source/Runtime/MovieScene/Public/SkeletalMeshRestoreState.h b/Engine/Source/Runtime/MovieScene/Public/SkeletalMeshRestoreState.h
--- a/Engine/Source/Runtime/MovieScene/Public/SkeletalMeshRestoreState.h
+++ b/Engine/Source/Runtime/MovieScene/Public/SkeletalMeshRestoreState.h
@@ -4,24 +4,28 @@
 #include "CoreMinimal.h"
 #include "UObject/Object.h"
 #include "UObject/Interface.h"
-#include "Components/SkeletalMeshComponent.h"
+#include "Components/SkinnedMeshComponent.h"
+#include "Components/SkeletalMeshComponent.h"  // kept for WITH_EDITOR cast below
 
 
 struct FSkeletalMeshEditorParams
 {
-	void SaveState(USkeletalMeshComponent* InSkelMeshComp)
+	void SaveState(USkinnedMeshComponent* InSkinnedMeshComp)
 	{
-		if (InSkelMeshComp)
+		if (InSkinnedMeshComp)
 		{
-			ChildSkelMesh = InSkelMeshComp;
-			VisibilityBasedAnimTickOption = InSkelMeshComp->VisibilityBasedAnimTickOption;
-			InSkelMeshComp->VisibilityBasedAnimTickOption = EVisibilityBasedAnimTickOption::AlwaysTickPoseAndRefreshBones;
-			PredictedLODLevel = InSkelMeshComp->GetPredictedLODLevel();
+			ChildSkelMesh = InSkinnedMeshComp;
+			VisibilityBasedAnimTickOption = InSkinnedMeshComp->VisibilityBasedAnimTickOption;
+			InSkinnedMeshComp->VisibilityBasedAnimTickOption = EVisibilityBasedAnimTickOption::AlwaysTickPoseAndRefreshBones;
+			PredictedLODLevel = InSkinnedMeshComp->GetPredictedLODLevel();
 #if WITH_EDITOR
-			bUpdateAnimationInEditor = InSkelMeshComp->GetUpdateAnimationInEditor();
-			bUpdateClothInEditor = InSkelMeshComp->GetUpdateClothInEditor();
-			InSkelMeshComp->SetUpdateAnimationInEditor(true);
-			InSkelMeshComp->SetUpdateClothInEditor(true);
+			if (USkeletalMeshComponent* SkelComp = Cast<USkeletalMeshComponent>(InSkinnedMeshComp))
+			{
+				bUpdateAnimationInEditor = SkelComp->GetUpdateAnimationInEditor();
+				bUpdateClothInEditor     = SkelComp->GetUpdateClothInEditor();
+				SkelComp->SetUpdateAnimationInEditor(true);
+				SkelComp->SetUpdateClothInEditor(true);
+			}
 #endif
 		}
 	}
@@ -33,9 +37,11 @@ struct FSkeletalMeshEditorParams
 			ChildSkelMesh->VisibilityBasedAnimTickOption = VisibilityBasedAnimTickOption;
 			ChildSkelMesh->SetPredictedLODLevel(PredictedLODLevel);
 #if WITH_EDITOR
-
-			ChildSkelMesh->SetUpdateAnimationInEditor(bUpdateAnimationInEditor);
-			ChildSkelMesh->SetUpdateClothInEditor(bUpdateClothInEditor);
+			if (USkeletalMeshComponent* SkelComp = Cast<USkeletalMeshComponent>(ChildSkelMesh.Get()))
+			{
+				SkelComp->SetUpdateAnimationInEditor(bUpdateAnimationInEditor);
+				SkelComp->SetUpdateClothInEditor(bUpdateClothInEditor);
+			}
 #endif
 		}
 	}
@@ -46,7 +52,7 @@ struct FSkeletalMeshEditorParams
 			ChildSkelMesh->SetPredictedLODLevel(PredictedLODLevel);
 		}
 	}
-	TWeakObjectPtr<USkeletalMeshComponent> ChildSkelMesh;
+	TWeakObjectPtr<USkinnedMeshComponent> ChildSkelMesh;
 	EVisibilityBasedAnimTickOption VisibilityBasedAnimTickOption;
 	int32 PredictedLODLevel;
 #if WITH_EDITOR
@@ -57,7 +63,7 @@ struct FSkeletalMeshEditorParams
 
 struct FSkeletalMeshRestoreState
 {
-	void SaveState(USkeletalMeshComponent* InComponent)
+	void SaveState(USkinnedMeshComponent* InComponent)
 	{
 		SkeletalMeshCompEditorParams.SetNum(0);
 		FSkeletalMeshEditorParams Parent;
@@ -67,11 +73,11 @@ struct FSkeletalMeshRestoreState
 		InComponent->GetChildrenComponents(true, ChildComponents);
 		for (USceneComponent* ChildComponent : ChildComponents)
 		{
-			USkeletalMeshComponent* SkelMeshComp = Cast<USkeletalMeshComponent>(ChildComponent);
-			if (SkelMeshComp)
+			USkinnedMeshComponent* SkinnedMeshComp = Cast<USkinnedMeshComponent>(ChildComponent);
+			if (SkinnedMeshComp)
 			{
 				FSkeletalMeshEditorParams Params;
-				Params.SaveState(SkelMeshComp);
+				Params.SaveState(SkinnedMeshComp);
 				SkeletalMeshCompEditorParams.Add(Params);
 			}
 		}
