Fixed visibility system to define implicit prerequisites only for the CDO. Additionally  FSystemDependencyGraph::MakeRelationship only increments serial number if an edge was truly added to safeguard against situations where MakeRelationship gets called for the global dependency graph outside CDOs

#Synced-CL 51764840

diff --git a/Engine/Source/Runtime/MovieScene/Private/EntitySystem/MovieSceneEntitySystem.cpp b/Engine/Source/Runtime/MovieScene/Private/EntitySystem/MovieSceneEntitySystem.cpp
--- a/Engine/Source/Runtime/MovieScene/Private/EntitySystem/MovieSceneEntitySystem.cpp
+++ b/Engine/Source/Runtime/MovieScene/Private/EntitySystem/MovieSceneEntitySystem.cpp
@@ -67,9 +67,12 @@ struct FSystemDependencyGraph
 
 	void MakeRelationship(uint16 UpstreamGraphID, uint16 DownstreamGraphID)
 	{
-		ImplicitSubsequents.MakeEdge(UpstreamGraphID, DownstreamGraphID);
-		ImplicitPrerequisites.MakeEdge(DownstreamGraphID, UpstreamGraphID);
-		++SerialNumber;
+		const bool bSubsequentEdgeAdded = ImplicitSubsequents.MakeEdge(UpstreamGraphID, DownstreamGraphID);
+		const bool bPrerequisiteEdgeAdded = ImplicitPrerequisites.MakeEdge(DownstreamGraphID, UpstreamGraphID);
+		if (bSubsequentEdgeAdded || bPrerequisiteEdgeAdded)
+		{
+			++SerialNumber;
+		}
 	}
 
 	UClass* ClassFromGraphID(uint16 GraphID) const
diff --git a/Engine/Source/Runtime/MovieScene/Private/EntitySystem/MovieSceneEntitySystemDirectedGraph.cpp b/Engine/Source/Runtime/MovieScene/Private/EntitySystem/MovieSceneEntitySystemDirectedGraph.cpp
--- a/Engine/Source/Runtime/MovieScene/Private/EntitySystem/MovieSceneEntitySystemDirectedGraph.cpp
+++ b/Engine/Source/Runtime/MovieScene/Private/EntitySystem/MovieSceneEntitySystemDirectedGraph.cpp
@@ -246,7 +246,7 @@ bool FDirectedGraph::IsCyclicImpl(uint16 NodeID, TBitArray<>& Visiting) const
 	return false;
 }
 
-void FDirectedGraph::MakeEdge(uint16 FromNode, uint16 ToNode)
+bool FDirectedGraph::MakeEdge(uint16 FromNode, uint16 ToNode)
 {
 	FDirectionalEdge NewEdge(FromNode, ToNode);
 
@@ -254,7 +254,9 @@ void FDirectedGraph::MakeEdge(uint16 FromNode, uint16 ToNode)
 	if (!SortedEdges.IsValidIndex(InsertIndex) || SortedEdges[InsertIndex] != NewEdge)
 	{
 		SortedEdges.Insert(NewEdge, InsertIndex);
+		return true;
 	}
+	return false;
 }
 
 void FDirectedGraph::DestroyEdge(uint16 FromNode, uint16 ToNode)
diff --git a/Engine/Source/Runtime/MovieScene/Public/EntitySystem/MovieSceneEntitySystemDirectedGraph.h b/Engine/Source/Runtime/MovieScene/Public/EntitySystem/MovieSceneEntitySystemDirectedGraph.h
--- a/Engine/Source/Runtime/MovieScene/Public/EntitySystem/MovieSceneEntitySystemDirectedGraph.h
+++ b/Engine/Source/Runtime/MovieScene/Public/EntitySystem/MovieSceneEntitySystemDirectedGraph.h
@@ -148,7 +148,7 @@ public:
 
 	MOVIESCENE_API bool IsCyclic() const;
 
-	MOVIESCENE_API void MakeEdge(uint16 FromNode, uint16 ToNode);
+	MOVIESCENE_API bool MakeEdge(uint16 FromNode, uint16 ToNode);
 
 	MOVIESCENE_API void DestroyEdge(uint16 FromNode, uint16 ToNode);
 
