Sequencer: Guard against a stale template object of a different class occupying InName in the same outer.

#Synced-CL 51953805

diff --git a/Engine/Source/Runtime/MovieScene/Private/MovieSceneCommonHelpers.cpp b/Engine/Source/Runtime/MovieScene/Private/MovieSceneCommonHelpers.cpp
--- a/Engine/Source/Runtime/MovieScene/Private/MovieSceneCommonHelpers.cpp
+++ b/Engine/Source/Runtime/MovieScene/Private/MovieSceneCommonHelpers.cpp
@@ -810,6 +810,17 @@ FString MovieSceneHelpers::MakeUniqueSpawnableName(UMovieScene* MovieScene, cons
 
 UObject* MovieSceneHelpers::MakeSpawnableTemplateFromInstance(UObject& InSourceObject, UMovieScene* InMovieScene, FName InName)
 {
+	// Guard against a stale template object of a different class occupying InName in the same outer.
+	if (UObject* ExistingObject = StaticFindObjectFast(UObject::StaticClass(), InMovieScene, InName))
+	{
+		if (ExistingObject->GetClass() != InSourceObject.GetClass())
+		{
+			const FName ExpiredName = MakeUniqueObjectName(InMovieScene, ExistingObject->GetClass(), "ExpiredSpawnable");
+			ExistingObject->Rename(*ExpiredName.ToString(), nullptr, REN_DoNotDirty | REN_DontCreateRedirectors | REN_NonTransactional);
+			ExistingObject->MarkAsGarbage();
+		}
+	}
+
 	UObject* NewInstance = NewObject<UObject>(InMovieScene, InSourceObject.GetClass(), InName);
 
 	UEngine::FCopyPropertiesForUnrelatedObjectsParams CopyParams;
