PR #14582: [Sequencer] Fix calculation when creating MovieSceneInverseSequenceTransform

#Synced-CL 52445895

diff --git a/Engine/Source/Runtime/MovieScene/Private/Evaluation/MovieSceneSequenceTransform.cpp b/Engine/Source/Runtime/MovieScene/Private/Evaluation/MovieSceneSequenceTransform.cpp
--- a/Engine/Source/Runtime/MovieScene/Private/Evaluation/MovieSceneSequenceTransform.cpp
+++ b/Engine/Source/Runtime/MovieScene/Private/Evaluation/MovieSceneSequenceTransform.cpp
@@ -492,7 +492,7 @@ bool FMovieSceneInverseNestedSequenceTransform::TransformTimeWithinRange(FFrameT
 		return TimeScale.AsLoop().InverseRemapTimeWithinRange(InTime + Offset, RangeStart, RangeEnd, OffsetVisitor);
 
 	case EMovieSceneTimeWarpType::Clamp:
-		if (InTime < -Offset && InTime > TimeScale.AsClamp().Max - Offset)
+		if (InTime < -Offset || InTime > TimeScale.AsClamp().Max - Offset)
 		{
 			return true;
 		}
@@ -513,7 +513,7 @@ bool FMovieSceneInverseNestedSequenceTransform::TransformTimeWithinRange(FFrameT
 		return TimeScale.AsLoopFloat().InverseRemapTimeWithinRange(InTime + Offset, RangeStart, RangeEnd, OffsetVisitor);
 
 	case EMovieSceneTimeWarpType::ClampFloat:
-		if (InTime < -Offset && InTime > FFrameTime::FromDecimal(TimeScale.AsClampFloat().Max) - Offset)
+		if (InTime < -Offset || InTime > FFrameTime::FromDecimal(TimeScale.AsClampFloat().Max) - Offset)
 		{
 			return true;
 		}
@@ -1070,7 +1070,7 @@ FMovieSceneInverseSequenceTransform FMovieSceneSequenceTransform::Inverse() cons
 		const FMovieSceneNestedSequenceTransform& Nested = NestedTransforms[Index];
 		if (Nested.IsLinear())
 		{
-			Result.LinearTransform = Result.LinearTransform * Nested.AsLinear().Inverse();
+			Result.LinearTransform = Nested.AsLinear().Inverse() * Result.LinearTransform;
 			continue;
 		}
 
