Sequencer: Fix race condition in TEntityTaskComponents::Schedule_PerAllocation

Schedule_PerAllocation on TEntityTaskComponents (the unfiltered path)
was setting bSerialTasks = false, making it behave identically to
Fork_PerAllocation and allowing multiple allocations to run in parallel
on different threads. This violated the Schedule_ contract that
guarantees serial execution of ForEachAllocation calls.

The filtered variant (TFilteredEntityTask::Schedule_PerAllocation) and
both Schedule_PerEntity overloads correctly set bSerialTasks = true.

#Synced-CL 52466902

diff --git a/Engine/Source/Runtime/MovieScene/Public/EntitySystem/MovieSceneEntitySystemTask.h b/Engine/Source/Runtime/MovieScene/Public/EntitySystem/MovieSceneEntitySystemTask.h
--- a/Engine/Source/Runtime/MovieScene/Public/EntitySystem/MovieSceneEntitySystemTask.h
+++ b/Engine/Source/Runtime/MovieScene/Public/EntitySystem/MovieSceneEntitySystemTask.h
@@ -551,7 +551,7 @@ struct TEntityTaskComponents : TEntityTaskComponentsImpl<TMakeIntegerSequence<in
 	FTaskID Schedule_PerAllocation(FEntityManager* EntityManager, IEntitySystemScheduler* InScheduler, TaskConstructionArgs&&... InArgs) const
 	{
 		FTaskParams FinalParams = this->CommonParams.TaskParams;
-		FinalParams.bSerialTasks = false;
+		FinalParams.bSerialTasks = true;
 
 #if !(UE_BUILD_SHIPPING || UE_BUILD_TEST)
 		FinalParams.DebugName = GetGeneratedTypeName<TaskImpl>();
