Sequencer: Move Linker->EntityManager.UpdateThreadingModel from ExecutePhase into each of the callsites so that threading model used is after the call to ScheduleTasks

#Synced-CL 52445895

diff --git a/Engine/Source/Runtime/MovieScene/Private/EntitySystem/MovieSceneEntitySystemGraphs.cpp b/Engine/Source/Runtime/MovieScene/Private/EntitySystem/MovieSceneEntitySystemGraphs.cpp
--- a/Engine/Source/Runtime/MovieScene/Private/EntitySystem/MovieSceneEntitySystemGraphs.cpp
+++ b/Engine/Source/Runtime/MovieScene/Private/EntitySystem/MovieSceneEntitySystemGraphs.cpp
@@ -420,8 +420,6 @@ void FMovieSceneEntitySystemGraph::ExecutePhase(UE::MovieScene::ESystemPhase Pha
 
 	const bool bCustomSchedulingEnabled = FEntitySystemScheduler::IsCustomSchedulingEnabled();
 
-	Linker->EntityManager.UpdateThreadingModel();
-
 	const EEntityThreadingModel ThreadingModel = Linker->EntityManager.GetThreadingModel();
 
 	// In a transaction we can only support the no-threading mode.
diff --git a/Engine/Source/Runtime/MovieScene/Private/EntitySystem/MovieSceneEntitySystemRunner.cpp b/Engine/Source/Runtime/MovieScene/Private/EntitySystem/MovieSceneEntitySystemRunner.cpp
--- a/Engine/Source/Runtime/MovieScene/Private/EntitySystem/MovieSceneEntitySystemRunner.cpp
+++ b/Engine/Source/Runtime/MovieScene/Private/EntitySystem/MovieSceneEntitySystemRunner.cpp
@@ -1088,6 +1088,7 @@ UE::MovieScene::ERunnerFlushResult FMovieSceneEntitySystemRunner::GameThread_Spa
 			// The spawn phase can queue events to trigger from the event tracks.
 			bCanQueueEventTriggers = true;
 			{
+				Linker->EntityManager.UpdateThreadingModel();
 				Linker->SystemGraph.ExecutePhase(ESystemPhase::Spawn, Linker, AllTasks);
 			}
 			bCanQueueEventTriggers = false;
@@ -1138,6 +1139,8 @@ UE::MovieScene::ERunnerFlushResult FMovieSceneEntitySystemRunner::GameThread_Ins
 
 	CurrentPhase = ESystemPhase::Instantiation;
 
+	Linker->EntityManager.UpdateThreadingModel();
+
 	FGraphEventArray AllTasks;
 	Linker->SystemGraph.ExecutePhase(ESystemPhase::Instantiation, Linker, AllTasks);
 
@@ -1205,6 +1208,8 @@ UE::MovieScene::ERunnerFlushResult FMovieSceneEntitySystemRunner::GameThread_Eva
 
 	checkf(!Linker->EntityManager.ContainsComponent(FBuiltInComponentTypes::Get()->Tags.NeedsUnlink), TEXT("Stale entities remain in the entity manager during evaluation - these should have been destroyed during the instantiation phase. Did it run?"));
 
+	Linker->EntityManager.UpdateThreadingModel();
+
 	if (FEntitySystemScheduler::IsCustomSchedulingEnabled())
 	{
 		Linker->SystemGraph.ScheduleTasks(&Linker->EntityManager);
@@ -1264,6 +1269,8 @@ void FMovieSceneEntitySystemRunner::GameThread_EvaluationFinalizationPhase(UMovi
 			}
 		}
 
+		Linker->EntityManager.UpdateThreadingModel();
+
 		FGraphEventArray Tasks;
 		Linker->SystemGraph.ExecutePhase(ESystemPhase::Finalization, Linker, Tasks);
 		checkf(Tasks.Num() == 0, TEXT("Cannot dispatch new tasks during finalization"));
@@ -1450,6 +1457,8 @@ bool FMovieSceneEntitySystemRunner::FlushSingleEvaluationPhase()
 
 	Linker->EntityManager.LockDown();
 
+	Linker->EntityManager.UpdateThreadingModel();
+
 	if (FEntitySystemScheduler::IsCustomSchedulingEnabled())
 	{
 		Linker->SystemGraph.ScheduleTasks(&Linker->EntityManager);
