PR #14610: Rep layout mismatch in level sequence player due to with editoronly data property

FMovieSceneSequencePlaybackSettings contains bRestoreState_DEPRECATED wrapped in #if WITH_EDITORONLY_DATA. Since this struct is replicated via DOREPLIFETIME(UMovieSceneSequencePlayer, PlaybackSettings), the RepLayout handle ordering differs between Editor builds (WITH_EDITORONLY_DATA=1) and non-Editor builds (WITH_EDITORONLY_DATA=0).

This causes an immediate connection drop when a PIE client (Windows Editor) connects to a Dedicated Server (Linux packaged build).

#Synced-CL 52516174

diff --git a/Engine/Plugins/MovieScene/ActorSequence/Source/ActorSequence/Private/ActorSequenceComponent.cpp b/Engine/Plugins/MovieScene/ActorSequence/Source/ActorSequence/Private/ActorSequenceComponent.cpp
--- a/Engine/Plugins/MovieScene/ActorSequence/Source/ActorSequence/Private/ActorSequenceComponent.cpp
+++ b/Engine/Plugins/MovieScene/ActorSequence/Source/ActorSequence/Private/ActorSequenceComponent.cpp
@@ -100,13 +100,5 @@ void UActorSequenceComponent::PostLoad()
 {
 	Super::PostLoad();
 
-	// If we previously were using bRestoreState on our PlaybackSettings, upgrade to the enum version.
-#if WITH_EDITORONLY_DATA
-	if (PlaybackSettings.bRestoreState_DEPRECATED)
-	{
-		PlaybackSettings.FinishCompletionStateOverride = EMovieSceneCompletionModeOverride::ForceRestoreState;
-		PlaybackSettings.bRestoreState_DEPRECATED = false;
-	}
-#endif
 }
 
diff --git a/Engine/Plugins/MovieScene/TemplateSequence/Source/TemplateSequence/Private/TemplateSequenceActor.cpp b/Engine/Plugins/MovieScene/TemplateSequence/Source/TemplateSequence/Private/TemplateSequenceActor.cpp
--- a/Engine/Plugins/MovieScene/TemplateSequence/Source/TemplateSequence/Private/TemplateSequenceActor.cpp
+++ b/Engine/Plugins/MovieScene/TemplateSequence/Source/TemplateSequence/Private/TemplateSequenceActor.cpp
@@ -33,14 +33,6 @@ void ATemplateSequenceActor::PostLoad()
 {
 	Super::PostLoad();
 
-	// If we previously were using bRestoreState on our PlaybackSettings, upgrade to the enum version.
-#if WITH_EDITORONLY_DATA
-	if (PlaybackSettings.bRestoreState_DEPRECATED)
-	{
-		PlaybackSettings.FinishCompletionStateOverride = EMovieSceneCompletionModeOverride::ForceRestoreState;
-		PlaybackSettings.bRestoreState_DEPRECATED = false;
-	}
-#endif
 }
 
 void ATemplateSequenceActor::GetLifetimeReplicatedProps(TArray<FLifetimeProperty>& OutLifetimeProps) const
diff --git a/Engine/Source/Runtime/LevelSequence/Private/LevelSequenceActor.cpp b/Engine/Source/Runtime/LevelSequence/Private/LevelSequenceActor.cpp
--- a/Engine/Source/Runtime/LevelSequence/Private/LevelSequenceActor.cpp
+++ b/Engine/Source/Runtime/LevelSequence/Private/LevelSequenceActor.cpp
@@ -292,15 +292,6 @@ void ALevelSequenceActor::PostLoad()
 		bAutoPlay_DEPRECATED = false;
 	}
 
-	// If we previously were using bRestoreState on our PlaybackSettings, upgrade to the enum version.
-#if WITH_EDITORONLY_DATA
-	if (PlaybackSettings.bRestoreState_DEPRECATED)
-	{
-		PlaybackSettings.FinishCompletionStateOverride = EMovieSceneCompletionModeOverride::ForceRestoreState;
-		PlaybackSettings.bRestoreState_DEPRECATED = false;
-	}
-#endif
-
 	GetSequencePlayer()->SetPlaybackSettings(PlaybackSettings);
 
 #if WITH_EDITORONLY_DATA
diff --git a/Engine/Source/Runtime/MovieScene/Public/MovieSceneSequencePlaybackSettings.h b/Engine/Source/Runtime/MovieScene/Public/MovieSceneSequencePlaybackSettings.h
--- a/Engine/Source/Runtime/MovieScene/Public/MovieSceneSequencePlaybackSettings.h
+++ b/Engine/Source/Runtime/MovieScene/Public/MovieSceneSequencePlaybackSettings.h
@@ -53,9 +53,6 @@ struct FMovieSceneSequencePlaybackSettings
 		, PlayRate(1.f)
 		, StartTime(0.f)
 		, bRandomStartTime(false)
-#if WITH_EDITORONLY_DATA
-		, bRestoreState_DEPRECATED(false)
-#endif
 		, bDisableMovementInput(false)
 		, bDisableLookAtInput(false)
 		, bHidePlayer(false)
@@ -92,14 +89,6 @@ struct FMovieSceneSequencePlaybackSettings
 	UPROPERTY(EditAnywhere, BlueprintReadWrite, Category="Playback")
 	uint32 bRandomStartTime : 1;
 
-#if WITH_EDITORONLY_DATA
-	/** Flag used to specify whether actor states should be restored on stop. 
-	* This has been deprecated in favor of FinishCompletionStateOverride.
-	*/
-	UPROPERTY(meta = (DeprecatedProperty, DeprecationMessage = "Use Settings.FinishCompletionStateOverride instead"))
-	uint32 bRestoreState_DEPRECATED : 1;
-#endif
-
 	/** Disable Input from player during play */
 	UPROPERTY(EditAnywhere, BlueprintReadWrite, Category="Cinematic")
 	uint32 bDisableMovementInput : 1;
