Sequencer: Remove unused

#rnx

#Synced-CL 52460642

diff --git a/Engine/Source/Runtime/MovieScene/Public/MovieSceneCommonHelpers.h b/Engine/Source/Runtime/MovieScene/Public/MovieSceneCommonHelpers.h
--- a/Engine/Source/Runtime/MovieScene/Public/MovieSceneCommonHelpers.h
+++ b/Engine/Source/Runtime/MovieScene/Public/MovieSceneCommonHelpers.h
@@ -160,14 +160,6 @@ public:
 	 */	
 	static MOVIESCENE_API UCameraComponent* CameraComponentFromRuntimeObject(UObject* RuntimeObject);
 
-	/**
-	 * Set the runtime object movable
-	 *
-	 * @param Object The object to set the mobility for
-	 * @param Mobility The mobility of the runtime object
-	 */
-	static MOVIESCENE_API void SetRuntimeObjectMobility(UObject* Object, EComponentMobility::Type ComponentMobility = EComponentMobility::Movable);
-
 	/*
 	 * Get the duration for the given sound
 
