Take Recorder - Guarding against MovieSceneModule::UnregisterCompatibleDecoration crash on runtime game shutdown due to trashed structs.

On runtime shutdown, due to shutdown order, the MovieSceneModule can crash due to the Compatible Decoration structs already being trashed. This is apparent when using TakeRecorder in runtime. This fix just adds a guard when not with editor, to not unregister anything is engine exit has been requested and is in progress.

#rnx
#Synced-CL 52611587

diff --git a/Engine/Source/Runtime/MovieScene/Private/MovieSceneModule.cpp b/Engine/Source/Runtime/MovieScene/Private/MovieSceneModule.cpp
--- a/Engine/Source/Runtime/MovieScene/Private/MovieSceneModule.cpp
+++ b/Engine/Source/Runtime/MovieScene/Private/MovieSceneModule.cpp
@@ -159,6 +159,12 @@ public:
 
 	virtual void UnregisterCompatibleDecoration(TSubclassOf<UMovieSceneDecorationContainerObject> ContainerClass, TSubclassOf<UObject> DecorationClass) override
 	{
+		// If the engine is exiting, do not do anything as data may have already been trashed before this call.
+		if (IsEngineExitRequested())
+		{
+			return;
+		}
+
 		if (!ContainerClass || !DecorationClass)
 		{
 			return;
