Sequencer: Fix conditional evaluation in subsequences when a parent sequence is conditionally disabled

#Synced-CL 52681601

diff --git a/Engine/Source/Runtime/MovieScene/Private/EntitySystem/MovieSceneSequenceUpdaters.cpp b/Engine/Source/Runtime/MovieScene/Private/EntitySystem/MovieSceneSequenceUpdaters.cpp
--- a/Engine/Source/Runtime/MovieScene/Private/EntitySystem/MovieSceneSequenceUpdaters.cpp
+++ b/Engine/Source/Runtime/MovieScene/Private/EntitySystem/MovieSceneSequenceUpdaters.cpp
@@ -830,6 +830,11 @@ void FSequenceUpdater_Hierarchical::Update(TSharedRef<const FSharedPlaybackState
 			CachedEntityRange = TRange<FFrameNumber>::Intersection(CachedEntityRange, SubSequenceIt.Range());
 		}
 
+		// Track sequences whose own condition has failed so that their descendants are also skipped.
+		// The compiled tree is populated depth-first (parent AddRange -> recurse -> child AddRange),
+		// so parent entries always appear before their children in the iteration.
+		TSet<FMovieSceneSequenceID> ConditionFailedSequences;
+
 		for (const FMovieSceneSubSequenceTreeEntry& Entry : RootOverrideHierarchy->GetTree().GetAllData(SubSequenceIt.Node()))
 		{
 			// When a root override path is specified, we always remap the 'local' sequence IDs to their equivalents from the root sequence.
@@ -841,8 +846,22 @@ void FSequenceUpdater_Hierarchical::Update(TSharedRef<const FSharedPlaybackState
 			checkf(SubData, TEXT("Sub data does not exist for a SequenceID that exists in the hierarchical tree - this indicates a corrupt compilation product."));
 
 			bool bSubSequenceConditionFailed = false;
+
+			// If a parent sequence's condition failed, propagate that failure to this subsequence.
+			if (const FMovieSceneSequenceHierarchyNode* SubNode = RootOverrideHierarchy->FindNode(Entry.SequenceID))
+			{
+				if (SubNode->ParentID != MovieSceneSequenceID::Root)
+				{
+					const FMovieSceneSequenceID ParentIDFromRoot = RootOverridePath.ResolveChildSequenceID(SubNode->ParentID);
+					if (ConditionFailedSequences.Contains(ParentIDFromRoot))
+					{
+						bSubSequenceConditionFailed = true;
+					}
+				}
+			}
+
 			const UMovieSceneCondition* Condition = SubData->WeakCondition.Get();
-			if (Condition)
+			if (!bSubSequenceConditionFailed && Condition)
 			{
 				// The condition's director BP belongs to the parent sequence (the one
 				// containing the sub-section), so resolve the parent's root-relative ID
@@ -881,6 +900,9 @@ void FSequenceUpdater_Hierarchical::Update(TSharedRef<const FSharedPlaybackState
 			UMovieSceneSequence* SubSequence = SubData->GetSequence();
 			if (SubSequence == nullptr || bSubSequenceConditionFailed)
 			{
+				// Record this sequence so its descendants are also skipped
+				ConditionFailedSequences.Add(SequenceIDFromRoot);
+
 				FInstanceHandle SubSequenceHandle = SequenceInstances.FindRef(SequenceIDFromRoot).Handle;
 				if (SubSequenceHandle.IsValid())
 				{
